Horde: Updating schema documentation

#rnx

diff --git a/Engine/Source/Programs/Horde/Docs/Config/Schema/Globals.md b/Engine/Source/Programs/Horde/Docs/Config/Schema/Globals.md
--- a/Engine/Source/Programs/Horde/Docs/Config/Schema/Globals.md
+++ b/Engine/Source/Programs/Horde/Docs/Config/Schema/Globals.md
@@ -4,128 +4,128 @@
 
 Global configuration
 
-Name | Type | Description
----- | ---- | -----------
-`version` | `integer` | Version number for the server. Values are indicated by the .
-`include` | [`ConfigInclude`](#configinclude)`[]` | Other paths to include
-`macros` | [`ConfigMacro`](#configmacro)`[]` | Macros within the global scope
-`dashboard` | [`DashboardConfig`](#dashboardconfig) | Settings for the dashboard
-`projects` | [`ProjectConfig`](Projects.md)`[]` | List of projects
-`pools` | [`PoolConfig`](#poolconfig)`[]` | List of pools
-`downtime` | [`ScheduledDowntime`](#scheduleddowntime)`[]` | List of scheduled downtime
-`perforceClusters` | [`PerforceCluster`](#perforcecluster)`[]` | List of Perforce clusters
-`software` | [`AgentSoftwareConfig`](#agentsoftwareconfig)`[]` | List of costs of a particular agent type
-`rates` | [`AgentRateConfig`](#agentrateconfig)`[]` | List of costs of a particular agent type
-`networks` | [`NetworkConfig`](#networkconfig)`[]` | List of networks
-`compute` | [`ComputeClusterConfig`](#computeclusterconfig)`[]` | List of compute profiles
-`secrets` | [`SecretConfig`](#secretconfig)`[]` | List of secrets
-`devices` | [`DeviceConfig`](#deviceconfig) | Device configuration
-`tools` | [`ToolConfig`](#toolconfig)`[]` | List of tools hosted by the server
-`maxConformCount` | `integer` | Maximum number of conforms to run at once
-`agentShutdownIfDisabledGracePeriod` | `string` | Time to wait before shutting down an agent that has been disabled Used if no value is set on the actual pool.
-`storage` | [`StorageConfig`](#storageconfig) | Storage configuration
-`artifactTypes` | [`ArtifactTypeConfig`](#artifacttypeconfig)`[]` | Configuration for different artifact types
-`telemetryStores` | [`TelemetryStoreConfig`](#telemetrystoreconfig)`[]` | Metrics to aggregate on the Horde server
-`parameters` | `object` | General parameters for other tools. Can be queried through the api/v1/parameters endpoint.
-`acl` | [`AclConfig`](#aclconfig) | Access control list
+Name | Description
+---- | -----------
+`version` | `integer`<br>Version number for the server. Values are indicated by the .
+`include` | [`ConfigInclude`](#configinclude)`[]`<br>Other paths to include
+`macros` | [`ConfigMacro`](#configmacro)`[]`<br>Macros within the global scope
+`dashboard` | [`DashboardConfig`](#dashboardconfig)<br>Settings for the dashboard
+`projects` | [`ProjectConfig`](Projects.md)`[]`<br>List of projects
+`pools` | [`PoolConfig`](#poolconfig)`[]`<br>List of pools
+`downtime` | [`ScheduledDowntime`](#scheduleddowntime)`[]`<br>List of scheduled downtime
+`perforceClusters` | [`PerforceCluster`](#perforcecluster)`[]`<br>List of Perforce clusters
+`software` | [`AgentSoftwareConfig`](#agentsoftwareconfig)`[]`<br>List of costs of a particular agent type
+`rates` | [`AgentRateConfig`](#agentrateconfig)`[]`<br>List of costs of a particular agent type
+`networks` | [`NetworkConfig`](#networkconfig)`[]`<br>List of networks
+`compute` | [`ComputeClusterConfig`](#computeclusterconfig)`[]`<br>List of compute profiles
+`secrets` | [`SecretConfig`](#secretconfig)`[]`<br>List of secrets
+`devices` | [`DeviceConfig`](#deviceconfig)<br>Device configuration
+`tools` | [`ToolConfig`](#toolconfig)`[]`<br>List of tools hosted by the server
+`maxConformCount` | `integer`<br>Maximum number of conforms to run at once
+`agentShutdownIfDisabledGracePeriod` | `string`<br>Time to wait before shutting down an agent that has been disabled Used if no value is set on the actual pool.
+`storage` | [`StorageConfig`](#storageconfig)<br>Storage configuration
+`artifactTypes` | [`ArtifactTypeConfig`](#artifacttypeconfig)`[]`<br>Configuration for different artifact types
+`telemetryStores` | [`TelemetryStoreConfig`](#telemetrystoreconfig)`[]`<br>Metrics to aggregate on the Horde server
+`parameters` | `object`<br>General parameters for other tools. Can be queried through the api/v1/parameters endpoint.
+`acl` | [`AclConfig`](#aclconfig)<br>Access control list
 
 ## ConfigInclude
 
 Directive to merge config data from another source
 
-Name | Type | Description
----- | ---- | -----------
-`path` | `string` | Path to the config data to be included. May be relative to the including file's location.
+Name | Description
+---- | -----------
+`path` | `string`<br>Path to the config data to be included. May be relative to the including file's location.
 
 ## ConfigMacro
 
 Declares a config macro
 
-Name | Type | Description
----- | ---- | -----------
-`name` | `string` | Name of the macro property
-`value` | `string` | Value for the macro property
+Name | Description
+---- | -----------
+`name` | `string`<br>Name of the macro property
+`value` | `string`<br>Value for the macro property
 
 ## DashboardConfig
 
 Configuration for dashboard features
 
-Name | Type | Description
----- | ---- | -----------
-`showLandingPage` | `boolean` | Navigate to the landing page by default
-`showCI` | `boolean` | Enable CI functionality
-`showAgents` | `boolean` | Whether to show functionality related to agents, pools, and utilization on the dashboard.
-`showAgentRegistration` | `boolean` | Whether to show the agent registration page. When using registration tokens from elsewhere this is not needed.
-`showPerforceServers` | `boolean` | Show the Perforce server option on the server menu
-`showDeviceManager` | `boolean` | Show the device manager on the server menu
-`showTests` | `boolean` | Show automated tests on the server menu
-`agentCategories` | [`DashboardAgentCategoryConfig`](#dashboardagentcategoryconfig)`[]` | Configuration for different agent pages
-`poolCategories` | [`DashboardPoolCategoryConfig`](#dashboardpoolcategoryconfig)`[]` | Configuration for different pool pages
-`analytics` | [`TelemetryViewConfig`](#telemetryviewconfig)`[]` | Configuration for telemetry views
-`include` | [`ConfigInclude`](#configinclude)`[]` | Includes for other configuration files
-`macros` | [`ConfigMacro`](#configmacro)`[]` | Macros within this configuration
+Name | Description
+---- | -----------
+`showLandingPage` | `boolean`<br>Navigate to the landing page by default
+`showCI` | `boolean`<br>Enable CI functionality
+`showAgents` | `boolean`<br>Whether to show functionality related to agents, pools, and utilization on the dashboard.
+`showAgentRegistration` | `boolean`<br>Whether to show the agent registration page. When using registration tokens from elsewhere this is not needed.
+`showPerforceServers` | `boolean`<br>Show the Perforce server option on the server menu
+`showDeviceManager` | `boolean`<br>Show the device manager on the server menu
+`showTests` | `boolean`<br>Show automated tests on the server menu
+`agentCategories` | [`DashboardAgentCategoryConfig`](#dashboardagentcategoryconfig)`[]`<br>Configuration for different agent pages
+`poolCategories` | [`DashboardPoolCategoryConfig`](#dashboardpoolcategoryconfig)`[]`<br>Configuration for different pool pages
+`analytics` | [`TelemetryViewConfig`](#telemetryviewconfig)`[]`<br>Configuration for telemetry views
+`include` | [`ConfigInclude`](#configinclude)`[]`<br>Includes for other configuration files
+`macros` | [`ConfigMacro`](#configmacro)`[]`<br>Macros within this configuration
 
 ## DashboardAgentCategoryConfig
 
 Configuration for a category of agents
 
-Name | Type | Description
----- | ---- | -----------
-`name` | `string` | Name of the category
-`condition` | `string` | Condition string to be evaluated for this page
+Name | Description
+---- | -----------
+`name` | `string`<br>Name of the category
+`condition` | `string`<br>Condition string to be evaluated for this page
 
 ## DashboardPoolCategoryConfig
 
 Configuration for a category of pools
 
-Name | Type | Description
----- | ---- | -----------
-`name` | `string` | Name of the category
-`condition` | `string` | Condition string to be evaluated for this page
+Name | Description
+---- | -----------
+`name` | `string`<br>Name of the category
+`condition` | `string`<br>Condition string to be evaluated for this page
 
 ## TelemetryViewConfig
 
 A telemetry view of related metrics, divided into categofies
 
-Name | Type | Description
----- | ---- | -----------
-`id` | `string` | Identifier for the view
-`name` | `string` | The name of the view
-`telemetryStoreId` | `string` | The telemetry store this view uses
-`variables` | [`TelemetryVariableConfig`](#telemetryvariableconfig)`[]` | The variables used to filter the view data
-`categories` | [`TelemetryCategoryConfig`](#telemetrycategoryconfig)`[]` | The categories contained within the view
+Name | Description
+---- | -----------
+`id` | `string`<br>Identifier for the view
+`name` | `string`<br>The name of the view
+`telemetryStoreId` | `string`<br>The telemetry store this view uses
+`variables` | [`TelemetryVariableConfig`](#telemetryvariableconfig)`[]`<br>The variables used to filter the view data
+`categories` | [`TelemetryCategoryConfig`](#telemetrycategoryconfig)`[]`<br>The categories contained within the view
 
 ## TelemetryVariableConfig
 
 A telemetry view variable used for filtering the charting data
 
-Name | Type | Description
----- | ---- | -----------
-`name` | `string` | The name of the variable for display purposes
-`group` | `string` | The associated data group attached to the variable
-`defaults` | `string[]` | The default values to select
+Name | Description
+---- | -----------
+`name` | `string`<br>The name of the variable for display purposes
+`group` | `string`<br>The associated data group attached to the variable
+`defaults` | `string[]`<br>The default values to select
 
 ## TelemetryCategoryConfig
 
 A chart categody, will be displayed on the dashbord under an associated pivot
 
-Name | Type | Description
----- | ---- | -----------
-`name` | `string` | The name of the category
-`charts` | [`TelemetryChartConfig`](#telemetrychartconfig)`[]` | The charts contained within the category
+Name | Description
+---- | -----------
+`name` | `string`<br>The name of the category
+`charts` | [`TelemetryChartConfig`](#telemetrychartconfig)`[]`<br>The charts contained within the category
 
 ## TelemetryChartConfig
 
 Telemetry chart configuraton
 
-Name | Type | Description
----- | ---- | -----------
-`name` | `string` | The name of the chart, will be displayed on the dashboard
-`display` | [`TelemetryMetricUnitType`](#telemetrymetricunittype-enum) | The unit to display
-`graph` | [`TelemetryMetricGraphType`](#telemetrymetricgraphtype-enum) | The graph type
-`metrics` | [`TelemetryChartMetricConfig`](#telemetrychartmetricconfig)`[]` | List of configured metrics
-`min` | `integer` | The min unit value for clamping chart
-`max` | `integer` | The max unit value for clamping chart
+Name | Description
+---- | -----------
+`name` | `string`<br>The name of the chart, will be displayed on the dashboard
+`display` | [`TelemetryMetricUnitType`](#telemetrymetricunittype-enum)<br>The unit to display
+`graph` | [`TelemetryMetricGraphType`](#telemetrymetricgraphtype-enum)<br>The graph type
+`metrics` | [`TelemetryChartMetricConfig`](#telemetrychartmetricconfig)`[]`<br>List of configured metrics
+`min` | `integer`<br>The min unit value for clamping chart
+`max` | `integer`<br>The max unit value for clamping chart
 
 ## TelemetryMetricUnitType (Enum)
 
@@ -150,37 +150,37 @@ Name | Description
 
 Metric attached to a telemetry chart
 
-Name | Type | Description
----- | ---- | -----------
-`id` | `string` | Associated metric id
-`threshold` | `integer` | The threshold for KPI values
-`alias` | `string` | The metric alias for display purposes
+Name | Description
+---- | -----------
+`id` | `string`<br>Associated metric id
+`threshold` | `integer`<br>The threshold for KPI values
+`alias` | `string`<br>The metric alias for display purposes
 
 ## PoolConfig
 
 Mutable configuration for a pool
 
-Name | Type | Description
----- | ---- | -----------
-`id` | `string` | Unique id for this pool
-`base` | `string` | Base pool config to copy settings from
-`name` | `string` | Name of the pool
-`condition` | `string` | Condition for agents to automatically be included in this pool
-`properties` | `string` `->` `string` | Arbitrary properties related to this pool
-`color` | [`PoolColor`](#poolcolor-enum) | Color to use for this pool on the dashboard
-`enableAutoscaling` | `boolean` | Whether to enable autoscaling for this pool
-`minAgents` | `integer` | The minimum number of agents to keep in the pool
-`numReserveAgents` | `integer` | The minimum number of idle agents to hold in reserve
-`conformInterval` | `string` | Interval between conforms. If zero, the pool will not conform on a schedule.
-`scaleOutCooldown` | `string` | Cooldown time between scale-out events
-`scaleInCooldown` | `string` | Cooldown time between scale-in events
-`shutdownIfDisabledGracePeriod` | `string` | Time to wait before shutting down an agent that has been disabled
-`sizeStrategy` | [`PoolSizeStrategy`](#poolsizestrategy-enum) | 
-`sizeStrategies` | [`PoolSizeStrategyInfo`](#poolsizestrategyinfo)`[]` | List of pool sizing strategies for this pool. The first strategy with a matching condition will be picked.
-`fleetManagers` | [`FleetManagerInfo`](#fleetmanagerinfo)`[]` | List of fleet managers for this pool. The first strategy with a matching condition will be picked. If empty or no conditions match, a default fleet manager will be used.
-`leaseUtilizationSettings` | [`LeaseUtilizationSettings`](#leaseutilizationsettings) | Settings for lease utilization pool sizing strategy (if used)
-`jobQueueSettings` | [`JobQueueSettings`](#jobqueuesettings) | Settings for job queue pool sizing strategy (if used)
-`computeQueueAwsMetricSettings` | [`ComputeQueueAwsMetricSettings`](#computequeueawsmetricsettings) | Settings for job queue pool sizing strategy (if used)
+Name | Description
+---- | -----------
+`id` | `string`<br>Unique id for this pool
+`base` | `string`<br>Base pool config to copy settings from
+`name` | `string`<br>Name of the pool
+`condition` | `string`<br>Condition for agents to automatically be included in this pool
+`properties` | `string` `->` `string`<br>Arbitrary properties related to this pool
+`color` | [`PoolColor`](#poolcolor-enum)<br>Color to use for this pool on the dashboard
+`enableAutoscaling` | `boolean`<br>Whether to enable autoscaling for this pool
+`minAgents` | `integer`<br>The minimum number of agents to keep in the pool
+`numReserveAgents` | `integer`<br>The minimum number of idle agents to hold in reserve
+`conformInterval` | `string`<br>Interval between conforms. If zero, the pool will not conform on a schedule.
+`scaleOutCooldown` | `string`<br>Cooldown time between scale-out events
+`scaleInCooldown` | `string`<br>Cooldown time between scale-in events
+`shutdownIfDisabledGracePeriod` | `string`<br>Time to wait before shutting down an agent that has been disabled
+`sizeStrategy` | [`PoolSizeStrategy`](#poolsizestrategy-enum)<br>
+`sizeStrategies` | [`PoolSizeStrategyInfo`](#poolsizestrategyinfo)`[]`<br>List of pool sizing strategies for this pool. The first strategy with a matching condition will be picked.
+`fleetManagers` | [`FleetManagerInfo`](#fleetmanagerinfo)`[]`<br>List of fleet managers for this pool. The first strategy with a matching condition will be picked. If empty or no conditions match, a default fleet manager will be used.
+`leaseUtilizationSettings` | [`LeaseUtilizationSettings`](#leaseutilizationsettings)<br>Settings for lease utilization pool sizing strategy (if used)
+`jobQueueSettings` | [`JobQueueSettings`](#jobqueuesettings)<br>Settings for job queue pool sizing strategy (if used)
+`computeQueueAwsMetricSettings` | [`ComputeQueueAwsMetricSettings`](#computequeueawsmetricsettings)<br>Settings for job queue pool sizing strategy (if used)
 
 ## PoolColor (Enum)
 
@@ -210,22 +210,22 @@ Name | Description
 
 Metadata for configuring and picking a pool sizing strategy
 
-Name | Type | Description
----- | ---- | -----------
-`type` | [`PoolSizeStrategy`](#poolsizestrategy-enum) | Strategy implementation to use
-`condition` | `string` | Condition if this strategy should be enabled (right now, using date/time as a distinguishing factor)
-`config` | `string` | Configuration for the strategy, serialized as JSON
-`extraAgentCount` | `integer` | Integer to add after pool size has been calculated. Can also be negative.
+Name | Description
+---- | -----------
+`type` | [`PoolSizeStrategy`](#poolsizestrategy-enum)<br>Strategy implementation to use
+`condition` | `string`<br>Condition if this strategy should be enabled (right now, using date/time as a distinguishing factor)
+`config` | `string`<br>Configuration for the strategy, serialized as JSON
+`extraAgentCount` | `integer`<br>Integer to add after pool size has been calculated. Can also be negative.
 
 ## FleetManagerInfo
 
 Metadata for configuring and picking a fleet manager
 
-Name | Type | Description
----- | ---- | -----------
-`type` | [`FleetManagerType`](#fleetmanagertype-enum) | Fleet manager type implementation to use
-`condition` | `string` | Condition if this strategy should be enabled (right now, using date/time as a distinguishing factor)
-`config` | `string` | Configuration for the strategy, serialized as JSON
+Name | Description
+---- | -----------
+`type` | [`FleetManagerType`](#fleetmanagertype-enum)<br>Fleet manager type implementation to use
+`condition` | `string`<br>Condition if this strategy should be enabled (right now, using date/time as a distinguishing factor)
+`config` | `string`<br>Configuration for the strategy, serialized as JSON
 
 ## FleetManagerType (Enum)
 
@@ -244,43 +244,43 @@ Name | Description
 
 Lease utilization sizing settings for a pool
 
-Name | Type | Description
----- | ---- | -----------
-`sampleTimeSec` | `integer` | Time period for each sample
-`numSamples` | `integer` | Number of samples to collect for calculating lease utilization
-`numSamplesForResult` | `integer` | Min number of samples for a valid result
-`minAgents` | `integer` | The minimum number of agents to keep in the pool
-`numReserveAgents` | `integer` | The minimum number of idle agents to hold in reserve
+Name | Description
+---- | -----------
+`sampleTimeSec` | `integer`<br>Time period for each sample
+`numSamples` | `integer`<br>Number of samples to collect for calculating lease utilization
+`numSamplesForResult` | `integer`<br>Min number of samples for a valid result
+`minAgents` | `integer`<br>The minimum number of agents to keep in the pool
+`numReserveAgents` | `integer`<br>The minimum number of idle agents to hold in reserve
 
 ## JobQueueSettings
 
 Job queue sizing settings for a pool
 
-Name | Type | Description
----- | ---- | -----------
-`scaleOutFactor` | `number` | Factor translating queue size to additional agents to grow the pool with The result is always rounded up to nearest integer. Example: if there are 20 jobs in queue, a factor 0.25 will result in 5 new agents being added (20 * 0.25)
-`scaleInFactor` | `number` | Factor by which to shrink the pool size with when queue is empty The result is always rounded up to nearest integer. Example: when the queue size is zero, a default value of 0.9 will shrink the pool by 10% (current agent count * 0.9)
-`samplePeriodMin` | `integer` | How far back in time to look for job batches (that potentially are in the queue)
-`readyTimeThresholdSec` | `integer` | Time spent in ready state before considered truly waiting for an agent<br>A job batch can be in ready state before getting picked up and executed. This threshold will help ensure only batches that have been waiting longer than this value will be considered.
+Name | Description
+---- | -----------
+`scaleOutFactor` | `number`<br>Factor translating queue size to additional agents to grow the pool with The result is always rounded up to nearest integer. Example: if there are 20 jobs in queue, a factor 0.25 will result in 5 new agents being added (20 * 0.25)
+`scaleInFactor` | `number`<br>Factor by which to shrink the pool size with when queue is empty The result is always rounded up to nearest integer. Example: when the queue size is zero, a default value of 0.9 will shrink the pool by 10% (current agent count * 0.9)
+`samplePeriodMin` | `integer`<br>How far back in time to look for job batches (that potentially are in the queue)
+`readyTimeThresholdSec` | `integer`<br>Time spent in ready state before considered truly waiting for an agent<br>A job batch can be in ready state before getting picked up and executed. This threshold will help ensure only batches that have been waiting longer than this value will be considered.
 
 ## ComputeQueueAwsMetricSettings
 
 Settings for
 
-Name | Type | Description
----- | ---- | -----------
-`computeClusterId` | `string` | Compute cluster ID to observe
-`namespace` | `string` | AWS CloudWatch namespace to write metrics in
+Name | Description
+---- | -----------
+`computeClusterId` | `string`<br>Compute cluster ID to observe
+`namespace` | `string`<br>AWS CloudWatch namespace to write metrics in
 
 ## ScheduledDowntime
 
 Settings for the maintenance window
 
-Name | Type | Description
----- | ---- | -----------
-`startTime` | `string` | Start time
-`finishTime` | `string` | Finish time
-`frequency` | [`ScheduledDowntimeFrequency`](#scheduleddowntimefrequency-enum) | Frequency that the window repeats
+Name | Description
+---- | -----------
+`startTime` | `string`<br>Start time
+`finishTime` | `string`<br>Finish time
+`frequency` | [`ScheduledDowntimeFrequency`](#scheduleddowntimefrequency-enum)<br>Frequency that the window repeats
 
 ## ScheduledDowntimeFrequency (Enum)
 
@@ -296,186 +296,186 @@ Name | Description
 
 Information about a cluster of Perforce servers.
 
-Name | Type | Description
----- | ---- | -----------
-`name` | `string` | Name of the cluster
-`serviceAccount` | `string` | Username for Horde to log in to this server. Will use the default user if not set.
-`canImpersonate` | `boolean` | Whether the service account can impersonate other users
-`supportsPartitionedWorkspaces` | `boolean` | Whether to use partitioned workspaces on this server
-`servers` | [`PerforceServer`](#perforceserver)`[]` | List of servers
-`credentials` | [`PerforceCredentials`](#perforcecredentials)`[]` | List of server credentials
-`autoSdk` | [`AutoSdkWorkspace`](#autosdkworkspace)`[]` | List of autosdk streams
+Name | Description
+---- | -----------
+`name` | `string`<br>Name of the cluster
+`serviceAccount` | `string`<br>Username for Horde to log in to this server. Will use the default user if not set.
+`canImpersonate` | `boolean`<br>Whether the service account can impersonate other users
+`supportsPartitionedWorkspaces` | `boolean`<br>Whether to use partitioned workspaces on this server
+`servers` | [`PerforceServer`](#perforceserver)`[]`<br>List of servers
+`credentials` | [`PerforceCredentials`](#perforcecredentials)`[]`<br>List of server credentials
+`autoSdk` | [`AutoSdkWorkspace`](#autosdkworkspace)`[]`<br>List of autosdk streams
 
 ## PerforceServer
 
 Information about an individual Perforce server
 
-Name | Type | Description
----- | ---- | -----------
-`serverAndPort` | `string` | The server and port. The server may be a DNS entry with multiple records, in which case it will be actively load balanced.
-`healthCheck` | `boolean` | Whether to query the healthcheck address under each server
-`resolveDns` | `boolean` | Whether to resolve the DNS entries and load balance between different hosts
-`maxConformCount` | `integer` | Maximum number of simultaneous conforms on this server
-`condition` | `string` | Optional condition for a machine to be eligable to use this server
-`properties` | `string[]` | List of properties for an agent to be eligable to use this server
+Name | Description
+---- | -----------
+`serverAndPort` | `string`<br>The server and port. The server may be a DNS entry with multiple records, in which case it will be actively load balanced.
+`healthCheck` | `boolean`<br>Whether to query the healthcheck address under each server
+`resolveDns` | `boolean`<br>Whether to resolve the DNS entries and load balance between different hosts
+`maxConformCount` | `integer`<br>Maximum number of simultaneous conforms on this server
+`condition` | `string`<br>Optional condition for a machine to be eligable to use this server
+`properties` | `string[]`<br>List of properties for an agent to be eligable to use this server
 
 ## PerforceCredentials
 
 Credentials for a Perforce user
 
-Name | Type | Description
----- | ---- | -----------
-`userName` | `string` | The username
-`password` | `string` | Password for the user
-`ticket` | `string` | Login ticket for the user (will be used instead of password if set)
+Name | Description
+---- | -----------
+`userName` | `string`<br>The username
+`password` | `string`<br>Password for the user
+`ticket` | `string`<br>Login ticket for the user (will be used instead of password if set)
 
 ## AutoSdkWorkspace
 
 Path to a platform and stream to use for syncing AutoSDK
 
-Name | Type | Description
----- | ---- | -----------
-`name` | `string` | Name of this workspace
-`properties` | `string[]` | The agent properties to check (eg. "OSFamily=Windows")
-`userName` | `string` | Username for logging in to the server
-`stream` | `string` | Stream to use
+Name | Description
+---- | -----------
+`name` | `string`<br>Name of this workspace
+`properties` | `string[]`<br>The agent properties to check (eg. "OSFamily=Windows")
+`userName` | `string`<br>Username for logging in to the server
+`stream` | `string`<br>Stream to use
 
 ## AgentSoftwareConfig
 
 Selects different agent software versions by evaluating a condition
 
-Name | Type | Description
----- | ---- | -----------
-`toolId` | `string` | Tool identifier
-`condition` | `string` | Condition for using this channel
+Name | Description
+---- | -----------
+`toolId` | `string`<br>Tool identifier
+`condition` | `string`<br>Condition for using this channel
 
 ## AgentRateConfig
 
 Describes the monetary cost of agents matching a particular criteria
 
-Name | Type | Description
----- | ---- | -----------
-`condition` | `string` | Condition string
-`rate` | `number` | Rate for this agent
+Name | Description
+---- | -----------
+`condition` | `string`<br>Condition string
+`rate` | `number`<br>Rate for this agent
 
 ## NetworkConfig
 
 Describes a network The ID describes any logical grouping, such as region, availability zone, rack or office location.
 
-Name | Type | Description
----- | ---- | -----------
-`id` | `string` | ID for this network
-`cidrBlock` | `string` | CIDR block
-`description` | `string` | Human-readable description
-`computeId` | `string` | Compute ID for this network (used when allocating compute resources)
+Name | Description
+---- | -----------
+`id` | `string`<br>ID for this network
+`cidrBlock` | `string`<br>CIDR block
+`description` | `string`<br>Human-readable description
+`computeId` | `string`<br>Compute ID for this network (used when allocating compute resources)
 
 ## ComputeClusterConfig
 
 Profile for executing compute requests
 
-Name | Type | Description
----- | ---- | -----------
-`id` | `string` | Name of the partition
-`namespaceId` | `string` | Name of the namespace to use
-`requestBucketId` | `string` | Name of the input bucket
-`responseBucketId` | `string` | Name of the output bucket
-`condition` | `string` | Filter for agents to include
-`acl` | [`AclConfig`](#aclconfig) | Access control list
+Name | Description
+---- | -----------
+`id` | `string`<br>Name of the partition
+`namespaceId` | `string`<br>Name of the namespace to use
+`requestBucketId` | `string`<br>Name of the input bucket
+`responseBucketId` | `string`<br>Name of the output bucket
+`condition` | `string`<br>Filter for agents to include
+`acl` | [`AclConfig`](#aclconfig)<br>Access control list
 
 ## AclConfig
 
 Parameters to update an ACL
 
-Name | Type | Description
----- | ---- | -----------
-`entries` | [`AclEntryConfig`](#aclentryconfig)`[]` | Entries to replace the existing ACL
-`profiles` | [`AclProfileConfig`](#aclprofileconfig)`[]` | Defines profiles which allow grouping sets of actions into named collections
-`inherit` | `boolean` | Whether to inherit permissions from the parent ACL
-`exceptions` | `string[]` | List of exceptions to the inherited setting
+Name | Description
+---- | -----------
+`entries` | [`AclEntryConfig`](#aclentryconfig)`[]`<br>Entries to replace the existing ACL
+`profiles` | [`AclProfileConfig`](#aclprofileconfig)`[]`<br>Defines profiles which allow grouping sets of actions into named collections
+`inherit` | `boolean`<br>Whether to inherit permissions from the parent ACL
+`exceptions` | `string[]`<br>List of exceptions to the inherited setting
 
 ## AclEntryConfig
 
 Individual entry in an ACL
 
-Name | Type | Description
----- | ---- | -----------
-`claim` | [`AclClaimConfig`](#aclclaimconfig) | Name of the user or group
-`actions` | `string[]` | Array of actions to allow
-`profiles` | `string[]` | List of profiles to grant
+Name | Description
+---- | -----------
+`claim` | [`AclClaimConfig`](#aclclaimconfig)<br>Name of the user or group
+`actions` | `string[]`<br>Array of actions to allow
+`profiles` | `string[]`<br>List of profiles to grant
 
 ## AclClaimConfig
 
 New claim to create
 
-Name | Type | Description
----- | ---- | -----------
-`type` | `string` | The claim type
-`value` | `string` | The claim value
+Name | Description
+---- | -----------
+`type` | `string`<br>The claim type
+`value` | `string`<br>The claim value
 
 ## AclProfileConfig
 
 Configuration for an ACL profile. This defines a preset group of actions which can be given to a user via an ACL entry.
 
-Name | Type | Description
----- | ---- | -----------
-`id` | `string` | Identifier for this profile
-`actions` | `string[]` | Actions to include
-`excludeActions` | `string[]` | Actions to exclude from the inherited actions
-`extends` | `string[]` | Other profiles to extend from
+Name | Description
+---- | -----------
+`id` | `string`<br>Identifier for this profile
+`actions` | `string[]`<br>Actions to include
+`excludeActions` | `string[]`<br>Actions to exclude from the inherited actions
+`extends` | `string[]`<br>Other profiles to extend from
 
 ## SecretConfig
 
 Configuration for a secret value
 
-Name | Type | Description
----- | ---- | -----------
-`id` | `string` | Identifier for this secret
-`data` | `string` `->` `string` | Key/value pairs associated with this secret
-`sources` | [`ExternalSecretConfig`](#externalsecretconfig)`[]` | Providers to source key/value pairs from
-`acl` | [`AclConfig`](#aclconfig) | Defines access to this particular secret
+Name | Description
+---- | -----------
+`id` | `string`<br>Identifier for this secret
+`data` | `string` `->` `string`<br>Key/value pairs associated with this secret
+`sources` | [`ExternalSecretConfig`](#externalsecretconfig)`[]`<br>Providers to source key/value pairs from
+`acl` | [`AclConfig`](#aclconfig)<br>Defines access to this particular secret
 
 ## ExternalSecretConfig
 
 Configuration for an external secret provider
 
-Name | Type | Description
----- | ---- | -----------
-`provider` | `string` | Name of the provider to use
-`key` | `string` | Optional key indicating the parameter to set in the resulting data array
-`path` | `string` | Optional value indicating what to fetch from the provider
-`arguments` | `string` `->` `string` | Additional provider-specific arguments
+Name | Description
+---- | -----------
+`provider` | `string`<br>Name of the provider to use
+`key` | `string`<br>Optional key indicating the parameter to set in the resulting data array
+`path` | `string`<br>Optional value indicating what to fetch from the provider
+`arguments` | `string` `->` `string`<br>Additional provider-specific arguments
 
 ## DeviceConfig
 
 Configuration for devices
 
-Name | Type | Description
----- | ---- | -----------
-`platforms` | [`DevicePlatformConfig`](#deviceplatformconfig)`[]` | List of device platforms
-`pools` | [`DevicePoolConfig`](#devicepoolconfig)`[]` | List of device pools
+Name | Description
+---- | -----------
+`platforms` | [`DevicePlatformConfig`](#deviceplatformconfig)`[]`<br>List of device platforms
+`pools` | [`DevicePoolConfig`](#devicepoolconfig)`[]`<br>List of device pools
 
 ## DevicePlatformConfig
 
 Configuration for a device platform
 
-Name | Type | Description
----- | ---- | -----------
-`id` | `string` | The id for this platform
-`name` | `string` | Name of the platform
-`models` | `string[]` | A list of platform models
-`legacyNames` | `string[]` | Legacy names which older versions of Gauntlet may be using
-`legacyPerfSpecHighModel` | `string` | Model name for the high perf spec, which may be requested by Gauntlet
+Name | Description
+---- | -----------
+`id` | `string`<br>The id for this platform
+`name` | `string`<br>Name of the platform
+`models` | `string[]`<br>A list of platform models
+`legacyNames` | `string[]`<br>Legacy names which older versions of Gauntlet may be using
+`legacyPerfSpecHighModel` | `string`<br>Model name for the high perf spec, which may be requested by Gauntlet
 
 ## DevicePoolConfig
 
 Configuration for a device pool
 
-Name | Type | Description
----- | ---- | -----------
-`id` | `string` | The id for this platform
-`name` | `string` | The name of the pool
-`poolType` | [`DevicePoolType`](#devicepooltype-enum) | The type of the pool
-`projectIds` | `string[]` | List of project ids associated with pool
+Name | Description
+---- | -----------
+`id` | `string`<br>The id for this platform
+`name` | `string`<br>The name of the pool
+`poolType` | [`DevicePoolType`](#devicepooltype-enum)<br>The type of the pool
+`projectIds` | `string[]`<br>List of project ids associated with pool
 
 ## DevicePoolType (Enum)
 
@@ -490,47 +490,47 @@ Name | Description
 
 Options for configuring a tool
 
-Name | Type | Description
----- | ---- | -----------
-`id` | `string` | Unique identifier for the tool
-`name` | `string` | Name of the tool
-`description` | `string` | Description for the tool
-`parentId` | `string` | Tool id to nest this tool under
-`public` | `boolean` | Whether this tool should be exposed for download on a public endpoint without authentication
-`showInUgs` | `boolean` | Whether to show this tool for download in the UGS tools menu
-`showInDashboard` | `boolean` | Whether to show this tool for download in the dashboard
-`namespaceId` | `string` | Default namespace for new deployments of this tool
-`acl` | [`AclConfig`](#aclconfig) | Permissions for the tool
+Name | Description
+---- | -----------
+`id` | `string`<br>Unique identifier for the tool
+`name` | `string`<br>Name of the tool
+`description` | `string`<br>Description for the tool
+`parentId` | `string`<br>Tool id to nest this tool under
+`public` | `boolean`<br>Whether this tool should be exposed for download on a public endpoint without authentication
+`showInUgs` | `boolean`<br>Whether to show this tool for download in the UGS tools menu
+`showInDashboard` | `boolean`<br>Whether to show this tool for download in the dashboard
+`namespaceId` | `string`<br>Default namespace for new deployments of this tool
+`acl` | [`AclConfig`](#aclconfig)<br>Permissions for the tool
 
 ## StorageConfig
 
 Configuration for storage
 
-Name | Type | Description
----- | ---- | -----------
-`backends` | [`BackendConfig`](#backendconfig)`[]` | List of storage backends
-`namespaces` | [`NamespaceConfig`](#namespaceconfig)`[]` | List of namespaces for storage
+Name | Description
+---- | -----------
+`backends` | [`BackendConfig`](#backendconfig)`[]`<br>List of storage backends
+`namespaces` | [`NamespaceConfig`](#namespaceconfig)`[]`<br>List of namespaces for storage
 
 ## BackendConfig
 
 Common settings object for different providers
 
-Name | Type | Description
----- | ---- | -----------
-`id` | `string` | The storage backend ID
-`base` | `string` | Base backend to copy default settings from
-`type` | [`StorageBackendType`](#storagebackendtype-enum) | The type of storage backend to use
-`baseDir` | `string` | Base directory for filesystem storage
-`awsBucketName` | `string` | Name of the bucket to use
-`awsBucketPath` | `string` | Base path within the bucket
-`awsCredentials` | [`AwsCredentialsType`](#awscredentialstype-enum) | Type of credentials to use
-`awsRole` | `string` | ARN of a role to assume
-`awsProfile` | `string` | The AWS profile to read credentials form
-`awsRegion` | `string` | Region to connect to
-`azureConnectionString` | `string` | Connection string for Azure
-`azureContainerName` | `string` | Name of the container
-`relayServer` | `string` | 
-`relayToken` | `string` | 
+Name | Description
+---- | -----------
+`id` | `string`<br>The storage backend ID
+`base` | `string`<br>Base backend to copy default settings from
+`type` | [`StorageBackendType`](#storagebackendtype-enum)<br>The type of storage backend to use
+`baseDir` | `string`<br>Base directory for filesystem storage
+`awsBucketName` | `string`<br>Name of the bucket to use
+`awsBucketPath` | `string`<br>Base path within the bucket
+`awsCredentials` | [`AwsCredentialsType`](#awscredentialstype-enum)<br>Type of credentials to use
+`awsRole` | `string`<br>ARN of a role to assume
+`awsProfile` | `string`<br>The AWS profile to read credentials form
+`awsRegion` | `string`<br>Region to connect to
+`azureConnectionString` | `string`<br>Connection string for Azure
+`azureContainerName` | `string`<br>Name of the container
+`relayServer` | `string`<br>
+`relayToken` | `string`<br>
 
 ## StorageBackendType (Enum)
 
@@ -558,52 +558,52 @@ Name | Description
 
 Configuration of a particular namespace
 
-Name | Type | Description
----- | ---- | -----------
-`id` | `string` | Identifier for this namespace
-`backend` | `string` | Backend to use for this namespace
-`prefix` | `string` | Prefix for items within this namespace
-`gcFrequencyHrs` | `number` | How frequently to run garbage collection, in hours.
-`gcDelayHrs` | `number` | How long to keep newly uploaded orphanned objects before allowing them to be deleted, in hours.
-`enableAliases` | `boolean` | Support querying exports by their aliases
-`acl` | [`AclConfig`](#aclconfig) | Access list for this namespace
+Name | Description
+---- | -----------
+`id` | `string`<br>Identifier for this namespace
+`backend` | `string`<br>Backend to use for this namespace
+`prefix` | `string`<br>Prefix for items within this namespace
+`gcFrequencyHrs` | `number`<br>How frequently to run garbage collection, in hours.
+`gcDelayHrs` | `number`<br>How long to keep newly uploaded orphanned objects before allowing them to be deleted, in hours.
+`enableAliases` | `boolean`<br>Support querying exports by their aliases
+`acl` | [`AclConfig`](#aclconfig)<br>Access list for this namespace
 
 ## ArtifactTypeConfig
 
 Configuration for an artifact
 
-Name | Type | Description
----- | ---- | -----------
-`type` | `string` | Name of the artifact type
-`name` | `string` | Legacy 'Name' property
-`keepCount` | `integer` | Number of artifacts to retain
-`keepDays` | `integer` | Number of days to retain artifacts of this type
+Name | Description
+---- | -----------
+`type` | `string`<br>Name of the artifact type
+`name` | `string`<br>Legacy 'Name' property
+`keepCount` | `integer`<br>Number of artifacts to retain
+`keepDays` | `integer`<br>Number of days to retain artifacts of this type
 
 ## TelemetryStoreConfig
 
 Config for metrics
 
-Name | Type | Description
----- | ---- | -----------
-`id` | `string` | Identifier for this store
-`acl` | [`AclConfig`](#aclconfig) | Permissions for this store
-`metrics` | [`MetricConfig`](#metricconfig)`[]` | Metrics to aggregate on the Horde server
-`include` | [`ConfigInclude`](#configinclude)`[]` | Includes for other configuration files
-`macros` | [`ConfigMacro`](#configmacro)`[]` | Macros within this configuration
+Name | Description
+---- | -----------
+`id` | `string`<br>Identifier for this store
+`acl` | [`AclConfig`](#aclconfig)<br>Permissions for this store
+`metrics` | [`MetricConfig`](#metricconfig)`[]`<br>Metrics to aggregate on the Horde server
+`include` | [`ConfigInclude`](#configinclude)`[]`<br>Includes for other configuration files
+`macros` | [`ConfigMacro`](#configmacro)`[]`<br>Macros within this configuration
 
 ## MetricConfig
 
 Configures a metric to aggregate on the server
 
-Name | Type | Description
----- | ---- | -----------
-`id` | `string` | Identifier for this metric
-`filter` | `string` | Filter expression to evaluate to determine which events to include. This query is evaluated against an array.
-`property` | `string` | Property to aggregate
-`groupBy` | `string` | Property to group by. Specified as a comma-separated list of JSON path expressions.
-`function` | [`AggregationFunction`](#aggregationfunction-enum) | How to aggregate samples for this metric
-`percentile` | `integer` | For the percentile function, specifies the percentile to measure
-`interval` | `string` | Interval for each metric. Supports times such as "2d", "1h", "1h30m", "20s".
+Name | Description
+---- | -----------
+`id` | `string`<br>Identifier for this metric
+`filter` | `string`<br>Filter expression to evaluate to determine which events to include. This query is evaluated against an array.
+`property` | `string`<br>Property to aggregate
+`groupBy` | `string`<br>Property to group by. Specified as a comma-separated list of JSON path expressions.
+`function` | [`AggregationFunction`](#aggregationfunction-enum)<br>How to aggregate samples for this metric
+`percentile` | `integer`<br>For the percentile function, specifies the percentile to measure
+`interval` | `string`<br>Interval for each metric. Supports times such as "2d", "1h", "1h30m", "20s".
 
 ## AggregationFunction (Enum)
 
diff --git a/Engine/Source/Programs/Horde/Docs/Config/Schema/Projects.md b/Engine/Source/Programs/Horde/Docs/Config/Schema/Projects.md
--- a/Engine/Source/Programs/Horde/Docs/Config/Schema/Projects.md
+++ b/Engine/Source/Programs/Horde/Docs/Config/Schema/Projects.md
@@ -4,63 +4,63 @@
 
 Stores configuration for a project
 
-Name | Type | Description
----- | ---- | -----------
-`id` | `string` | The project id
-`name` | `string` | Name for the new project
-`path` | `string` | Direct include path for the project config. For backwards compatibility with old config files when including from a GlobalConfig object.
-`include` | [`ConfigInclude`](#configinclude)`[]` | Includes for other configuration files
-`macros` | [`ConfigMacro`](#configmacro)`[]` | Macros within the global scope
-`order` | `integer` | Order of this project on the dashboard
-`logo` | `string` | Path to the project logo
-`pools` | [`PoolConfig`](#poolconfig)`[]` | List of pools for this project
-`categories` | [`ProjectCategoryConfig`](#projectcategoryconfig)`[]` | Categories to include in this project
-`jobOptions` | [`JobOptions`](#joboptions) | Default settings for executing jobs
-`streams` | [`StreamConfig`](Streams.md)`[]` | List of streams
-`acl` | [`AclConfig`](#aclconfig) | Acl entries
+Name | Description
+---- | -----------
+`id` | `string`<br>The project id
+`name` | `string`<br>Name for the new project
+`path` | `string`<br>Direct include path for the project config. For backwards compatibility with old config files when including from a GlobalConfig object.
+`include` | [`ConfigInclude`](#configinclude)`[]`<br>Includes for other configuration files
+`macros` | [`ConfigMacro`](#configmacro)`[]`<br>Macros within the global scope
+`order` | `integer`<br>Order of this project on the dashboard
+`logo` | `string`<br>Path to the project logo
+`pools` | [`PoolConfig`](#poolconfig)`[]`<br>List of pools for this project
+`categories` | [`ProjectCategoryConfig`](#projectcategoryconfig)`[]`<br>Categories to include in this project
+`jobOptions` | [`JobOptions`](#joboptions)<br>Default settings for executing jobs
+`streams` | [`StreamConfig`](Streams.md)`[]`<br>List of streams
+`acl` | [`AclConfig`](#aclconfig)<br>Acl entries
 
 ## ConfigInclude
 
 Directive to merge config data from another source
 
-Name | Type | Description
----- | ---- | -----------
-`path` | `string` | Path to the config data to be included. May be relative to the including file's location.
+Name | Description
+---- | -----------
+`path` | `string`<br>Path to the config data to be included. May be relative to the including file's location.
 
 ## ConfigMacro
 
 Declares a config macro
 
-Name | Type | Description
----- | ---- | -----------
-`name` | `string` | Name of the macro property
-`value` | `string` | Value for the macro property
+Name | Description
+---- | -----------
+`name` | `string`<br>Name of the macro property
+`value` | `string`<br>Value for the macro property
 
 ## PoolConfig
 
 Mutable configuration for a pool
 
-Name | Type | Description
----- | ---- | -----------
-`id` | `string` | Unique id for this pool
-`base` | `string` | Base pool config to copy settings from
-`name` | `string` | Name of the pool
-`condition` | `string` | Condition for agents to automatically be included in this pool
-`properties` | `string` `->` `string` | Arbitrary properties related to this pool
-`color` | [`PoolColor`](#poolcolor-enum) | Color to use for this pool on the dashboard
-`enableAutoscaling` | `boolean` | Whether to enable autoscaling for this pool
-`minAgents` | `integer` | The minimum number of agents to keep in the pool
-`numReserveAgents` | `integer` | The minimum number of idle agents to hold in reserve
-`conformInterval` | `string` | Interval between conforms. If zero, the pool will not conform on a schedule.
-`scaleOutCooldown` | `string` | Cooldown time between scale-out events
-`scaleInCooldown` | `string` | Cooldown time between scale-in events
-`shutdownIfDisabledGracePeriod` | `string` | Time to wait before shutting down an agent that has been disabled
-`sizeStrategy` | [`PoolSizeStrategy`](#poolsizestrategy-enum) | 
-`sizeStrategies` | [`PoolSizeStrategyInfo`](#poolsizestrategyinfo)`[]` | List of pool sizing strategies for this pool. The first strategy with a matching condition will be picked.
-`fleetManagers` | [`FleetManagerInfo`](#fleetmanagerinfo)`[]` | List of fleet managers for this pool. The first strategy with a matching condition will be picked. If empty or no conditions match, a default fleet manager will be used.
-`leaseUtilizationSettings` | [`LeaseUtilizationSettings`](#leaseutilizationsettings) | Settings for lease utilization pool sizing strategy (if used)
-`jobQueueSettings` | [`JobQueueSettings`](#jobqueuesettings) | Settings for job queue pool sizing strategy (if used)
-`computeQueueAwsMetricSettings` | [`ComputeQueueAwsMetricSettings`](#computequeueawsmetricsettings) | Settings for job queue pool sizing strategy (if used)
+Name | Description
+---- | -----------
+`id` | `string`<br>Unique id for this pool
+`base` | `string`<br>Base pool config to copy settings from
+`name` | `string`<br>Name of the pool
+`condition` | `string`<br>Condition for agents to automatically be included in this pool
+`properties` | `string` `->` `string`<br>Arbitrary properties related to this pool
+`color` | [`PoolColor`](#poolcolor-enum)<br>Color to use for this pool on the dashboard
+`enableAutoscaling` | `boolean`<br>Whether to enable autoscaling for this pool
+`minAgents` | `integer`<br>The minimum number of agents to keep in the pool
+`numReserveAgents` | `integer`<br>The minimum number of idle agents to hold in reserve
+`conformInterval` | `string`<br>Interval between conforms. If zero, the pool will not conform on a schedule.
+`scaleOutCooldown` | `string`<br>Cooldown time between scale-out events
+`scaleInCooldown` | `string`<br>Cooldown time between scale-in events
+`shutdownIfDisabledGracePeriod` | `string`<br>Time to wait before shutting down an agent that has been disabled
+`sizeStrategy` | [`PoolSizeStrategy`](#poolsizestrategy-enum)<br>
+`sizeStrategies` | [`PoolSizeStrategyInfo`](#poolsizestrategyinfo)`[]`<br>List of pool sizing strategies for this pool. The first strategy with a matching condition will be picked.
+`fleetManagers` | [`FleetManagerInfo`](#fleetmanagerinfo)`[]`<br>List of fleet managers for this pool. The first strategy with a matching condition will be picked. If empty or no conditions match, a default fleet manager will be used.
+`leaseUtilizationSettings` | [`LeaseUtilizationSettings`](#leaseutilizationsettings)<br>Settings for lease utilization pool sizing strategy (if used)
+`jobQueueSettings` | [`JobQueueSettings`](#jobqueuesettings)<br>Settings for job queue pool sizing strategy (if used)
+`computeQueueAwsMetricSettings` | [`ComputeQueueAwsMetricSettings`](#computequeueawsmetricsettings)<br>Settings for job queue pool sizing strategy (if used)
 
 ## PoolColor (Enum)
 
@@ -90,22 +90,22 @@ Name | Description
 
 Metadata for configuring and picking a pool sizing strategy
 
-Name | Type | Description
----- | ---- | -----------
-`type` | [`PoolSizeStrategy`](#poolsizestrategy-enum) | Strategy implementation to use
-`condition` | `string` | Condition if this strategy should be enabled (right now, using date/time as a distinguishing factor)
-`config` | `string` | Configuration for the strategy, serialized as JSON
-`extraAgentCount` | `integer` | Integer to add after pool size has been calculated. Can also be negative.
+Name | Description
+---- | -----------
+`type` | [`PoolSizeStrategy`](#poolsizestrategy-enum)<br>Strategy implementation to use
+`condition` | `string`<br>Condition if this strategy should be enabled (right now, using date/time as a distinguishing factor)
+`config` | `string`<br>Configuration for the strategy, serialized as JSON
+`extraAgentCount` | `integer`<br>Integer to add after pool size has been calculated. Can also be negative.
 
 ## FleetManagerInfo
 
 Metadata for configuring and picking a fleet manager
 
-Name | Type | Description
----- | ---- | -----------
-`type` | [`FleetManagerType`](#fleetmanagertype-enum) | Fleet manager type implementation to use
-`condition` | `string` | Condition if this strategy should be enabled (right now, using date/time as a distinguishing factor)
-`config` | `string` | Configuration for the strategy, serialized as JSON
+Name | Description
+---- | -----------
+`type` | [`FleetManagerType`](#fleetmanagertype-enum)<br>Fleet manager type implementation to use
+`condition` | `string`<br>Condition if this strategy should be enabled (right now, using date/time as a distinguishing factor)
+`config` | `string`<br>Configuration for the strategy, serialized as JSON
 
 ## FleetManagerType (Enum)
 
@@ -124,107 +124,107 @@ Name | Description
 
 Lease utilization sizing settings for a pool
 
-Name | Type | Description
----- | ---- | -----------
-`sampleTimeSec` | `integer` | Time period for each sample
-`numSamples` | `integer` | Number of samples to collect for calculating lease utilization
-`numSamplesForResult` | `integer` | Min number of samples for a valid result
-`minAgents` | `integer` | The minimum number of agents to keep in the pool
-`numReserveAgents` | `integer` | The minimum number of idle agents to hold in reserve
+Name | Description
+---- | -----------
+`sampleTimeSec` | `integer`<br>Time period for each sample
+`numSamples` | `integer`<br>Number of samples to collect for calculating lease utilization
+`numSamplesForResult` | `integer`<br>Min number of samples for a valid result
+`minAgents` | `integer`<br>The minimum number of agents to keep in the pool
+`numReserveAgents` | `integer`<br>The minimum number of idle agents to hold in reserve
 
 ## JobQueueSettings
 
 Job queue sizing settings for a pool
 
-Name | Type | Description
----- | ---- | -----------
-`scaleOutFactor` | `number` | Factor translating queue size to additional agents to grow the pool with The result is always rounded up to nearest integer. Example: if there are 20 jobs in queue, a factor 0.25 will result in 5 new agents being added (20 * 0.25)
-`scaleInFactor` | `number` | Factor by which to shrink the pool size with when queue is empty The result is always rounded up to nearest integer. Example: when the queue size is zero, a default value of 0.9 will shrink the pool by 10% (current agent count * 0.9)
-`samplePeriodMin` | `integer` | How far back in time to look for job batches (that potentially are in the queue)
-`readyTimeThresholdSec` | `integer` | Time spent in ready state before considered truly waiting for an agent<br>A job batch can be in ready state before getting picked up and executed. This threshold will help ensure only batches that have been waiting longer than this value will be considered.
+Name | Description
+---- | -----------
+`scaleOutFactor` | `number`<br>Factor translating queue size to additional agents to grow the pool with The result is always rounded up to nearest integer. Example: if there are 20 jobs in queue, a factor 0.25 will result in 5 new agents being added (20 * 0.25)
+`scaleInFactor` | `number`<br>Factor by which to shrink the pool size with when queue is empty The result is always rounded up to nearest integer. Example: when the queue size is zero, a default value of 0.9 will shrink the pool by 10% (current agent count * 0.9)
+`samplePeriodMin` | `integer`<br>How far back in time to look for job batches (that potentially are in the queue)
+`readyTimeThresholdSec` | `integer`<br>Time spent in ready state before considered truly waiting for an agent<br>A job batch can be in ready state before getting picked up and executed. This threshold will help ensure only batches that have been waiting longer than this value will be considered.
 
 ## ComputeQueueAwsMetricSettings
 
 Settings for
 
-Name | Type | Description
----- | ---- | -----------
-`computeClusterId` | `string` | Compute cluster ID to observe
-`namespace` | `string` | AWS CloudWatch namespace to write metrics in
+Name | Description
+---- | -----------
+`computeClusterId` | `string`<br>Compute cluster ID to observe
+`namespace` | `string`<br>AWS CloudWatch namespace to write metrics in
 
 ## ProjectCategoryConfig
 
 Information about a category to display for a stream
 
-Name | Type | Description
----- | ---- | -----------
-`name` | `string` | Name of this category
-`row` | `integer` | Index of the row to display this category on
-`showOnNavMenu` | `boolean` | Whether to show this category on the nav menu
-`includePatterns` | `string[]` | Patterns for stream names to include
-`excludePatterns` | `string[]` | Patterns for stream names to exclude
+Name | Description
+---- | -----------
+`name` | `string`<br>Name of this category
+`row` | `integer`<br>Index of the row to display this category on
+`showOnNavMenu` | `boolean`<br>Whether to show this category on the nav menu
+`includePatterns` | `string[]`<br>Patterns for stream names to include
+`excludePatterns` | `string[]`<br>Patterns for stream names to exclude
 
 ## JobOptions
 
 Options for executing a job
 
-Name | Type | Description
----- | ---- | -----------
-`executor` | `string` | Name of the executor to use
-`useNewTempStorage` | `boolean` | Whether to use the new temp storage backend
-`useWine` | `boolean` | Whether to execute using Wine emulation on Linux
-`runInSeparateProcess` | `boolean` | Executes the job lease in a separate process
-`workspaceMaterializer` | `string` | What workspace materializer to use in WorkspaceExecutor. Will override any value from workspace config.
-`container` | [`JobContainerOptions`](#jobcontaineroptions) | Options for executing a job inside a container
+Name | Description
+---- | -----------
+`executor` | `string`<br>Name of the executor to use
+`useNewTempStorage` | `boolean`<br>Whether to use the new temp storage backend
+`useWine` | `boolean`<br>Whether to execute using Wine emulation on Linux
+`runInSeparateProcess` | `boolean`<br>Executes the job lease in a separate process
+`workspaceMaterializer` | `string`<br>What workspace materializer to use in WorkspaceExecutor. Will override any value from workspace config.
+`container` | [`JobContainerOptions`](#jobcontaineroptions)<br>Options for executing a job inside a container
 
 ## JobContainerOptions
 
 Options for executing a job inside a container
 
-Name | Type | Description
----- | ---- | -----------
-`enabled` | `boolean` | Whether to execute job inside a container
-`imageUrl` | `string` | Image URL to container, such as "quay.io/podman/hello"
-`containerEngineExecutable` | `string` | Container engine executable (docker or with full path like /usr/bin/podman)
-`extraArguments` | `string` | Additional arguments to pass to container engine
+Name | Description
+---- | -----------
+`enabled` | `boolean`<br>Whether to execute job inside a container
+`imageUrl` | `string`<br>Image URL to container, such as "quay.io/podman/hello"
+`containerEngineExecutable` | `string`<br>Container engine executable (docker or with full path like /usr/bin/podman)
+`extraArguments` | `string`<br>Additional arguments to pass to container engine
 
 ## AclConfig
 
 Parameters to update an ACL
 
-Name | Type | Description
----- | ---- | -----------
-`entries` | [`AclEntryConfig`](#aclentryconfig)`[]` | Entries to replace the existing ACL
-`profiles` | [`AclProfileConfig`](#aclprofileconfig)`[]` | Defines profiles which allow grouping sets of actions into named collections
-`inherit` | `boolean` | Whether to inherit permissions from the parent ACL
-`exceptions` | `string[]` | List of exceptions to the inherited setting
+Name | Description
+---- | -----------
+`entries` | [`AclEntryConfig`](#aclentryconfig)`[]`<br>Entries to replace the existing ACL
+`profiles` | [`AclProfileConfig`](#aclprofileconfig)`[]`<br>Defines profiles which allow grouping sets of actions into named collections
+`inherit` | `boolean`<br>Whether to inherit permissions from the parent ACL
+`exceptions` | `string[]`<br>List of exceptions to the inherited setting
 
 ## AclEntryConfig
 
 Individual entry in an ACL
 
-Name | Type | Description
----- | ---- | -----------
-`claim` | [`AclClaimConfig`](#aclclaimconfig) | Name of the user or group
-`actions` | `string[]` | Array of actions to allow
-`profiles` | `string[]` | List of profiles to grant
+Name | Description
+---- | -----------
+`claim` | [`AclClaimConfig`](#aclclaimconfig)<br>Name of the user or group
+`actions` | `string[]`<br>Array of actions to allow
+`profiles` | `string[]`<br>List of profiles to grant
 
 ## AclClaimConfig
 
 New claim to create
 
-Name | Type | Description
----- | ---- | -----------
-`type` | `string` | The claim type
-`value` | `string` | The claim value
+Name | Description
+---- | -----------
+`type` | `string`<br>The claim type
+`value` | `string`<br>The claim value
 
 ## AclProfileConfig
 
 Configuration for an ACL profile. This defines a preset group of actions which can be given to a user via an ACL entry.
 
-Name | Type | Description
----- | ---- | -----------
-`id` | `string` | Identifier for this profile
-`actions` | `string[]` | Actions to include
-`excludeActions` | `string[]` | Actions to exclude from the inherited actions
-`extends` | `string[]` | Other profiles to extend from
+Name | Description
+---- | -----------
+`id` | `string`<br>Identifier for this profile
+`actions` | `string[]`<br>Actions to include
+`excludeActions` | `string[]`<br>Actions to exclude from the inherited actions
+`extends` | `string[]`<br>Other profiles to extend from
diff --git a/Engine/Source/Programs/Horde/Docs/Config/Schema/Streams.md b/Engine/Source/Programs/Horde/Docs/Config/Schema/Streams.md
--- a/Engine/Source/Programs/Horde/Docs/Config/Schema/Streams.md
+++ b/Engine/Source/Programs/Horde/Docs/Config/Schema/Streams.md
@@ -4,98 +4,98 @@
 
 Config for a stream
 
-Name | Type | Description
----- | ---- | -----------
-`id` | `string` | Identifier for the stream
-`path` | `string` | Direct include path for the stream config. For backwards compatibility with old config files when including from a ProjectConfig object.
-`include` | [`ConfigInclude`](#configinclude)`[]` | Includes for other configuration files
-`macros` | [`ConfigMacro`](#configmacro)`[]` | Macros within this stream
-`name` | `string` | Name of the stream
-`clusterName` | `string` | The perforce cluster containing the stream
-`order` | `integer` | Order for this stream
-`initialAgentType` | `string` | Default initial agent type for templates
-`notificationChannel` | `string` | Notification channel for all jobs in this stream
-`notificationChannelFilter` | `string` | Notification channel filter for this template. Can be Success|Failure|Warnings
-`triageChannel` | `string` | Channel to post issue triage notifications
-`jobOptions` | [`JobOptions`](#joboptions) | Default settings for executing jobs
-`autoSdkView` | `string[]` | View for the AutoSDK paths to sync. If null, the whole thing will be synced.
-`defaultPreflightTemplate` | `string` | Legacy name for the default preflight template
-`defaultPreflight` | [`DefaultPreflightConfig`](#defaultpreflightconfig) | Default template for running preflights
-`commitTags` | [`CommitTagConfig`](#committagconfig)`[]` | List of tags to apply to commits. Allows fast searching and classification of different commit types (eg. code vs content).
-`tabs` | [`TabConfig`](#tabconfig)`[]` | List of tabs to show for the new stream
-`environment` | `string` `->` `string` | Global environment variables for all agents in this stream
-`agentTypes` | `string` `->` [`AgentConfig`](#agentconfig) | Map of agent name to type
-`workspaceTypes` | `string` `->` [`WorkspaceConfig`](#workspaceconfig) | Map of workspace name to type
-`templates` | [`TemplateRefConfig`](#templaterefconfig)`[]` | List of templates to create
-`acl` | [`AclConfig`](#aclconfig) | Custom permissions for this object
-`pausedUntil` | `string` | Pause stream builds until specified date
-`pauseComment` | `string` | Reason for pausing builds of the stream
-`replicators` | [`ReplicatorConfig`](#replicatorconfig)`[]` | Configuration for workers to replicate commit data into Horde Storage.
-`workflows` | [`WorkflowConfig`](#workflowconfig)`[]` | Workflows for dealing with new issues
-`tokens` | [`TokenConfig`](#tokenconfig)`[]` | Tokens to create for each job step
+Name | Description
+---- | -----------
+`id` | `string`<br>Identifier for the stream
+`path` | `string`<br>Direct include path for the stream config. For backwards compatibility with old config files when including from a ProjectConfig object.
+`include` | [`ConfigInclude`](#configinclude)`[]`<br>Includes for other configuration files
+`macros` | [`ConfigMacro`](#configmacro)`[]`<br>Macros within this stream
+`name` | `string`<br>Name of the stream
+`clusterName` | `string`<br>The perforce cluster containing the stream
+`order` | `integer`<br>Order for this stream
+`initialAgentType` | `string`<br>Default initial agent type for templates
+`notificationChannel` | `string`<br>Notification channel for all jobs in this stream
+`notificationChannelFilter` | `string`<br>Notification channel filter for this template. Can be Success|Failure|Warnings
+`triageChannel` | `string`<br>Channel to post issue triage notifications
+`jobOptions` | [`JobOptions`](#joboptions)<br>Default settings for executing jobs
+`autoSdkView` | `string[]`<br>View for the AutoSDK paths to sync. If null, the whole thing will be synced.
+`defaultPreflightTemplate` | `string`<br>Legacy name for the default preflight template
+`defaultPreflight` | [`DefaultPreflightConfig`](#defaultpreflightconfig)<br>Default template for running preflights
+`commitTags` | [`CommitTagConfig`](#committagconfig)`[]`<br>List of tags to apply to commits. Allows fast searching and classification of different commit types (eg. code vs content).
+`tabs` | [`TabConfig`](#tabconfig)`[]`<br>List of tabs to show for the new stream
+`environment` | `string` `->` `string`<br>Global environment variables for all agents in this stream
+`agentTypes` | `string` `->` [`AgentConfig`](#agentconfig)<br>Map of agent name to type
+`workspaceTypes` | `string` `->` [`WorkspaceConfig`](#workspaceconfig)<br>Map of workspace name to type
+`templates` | [`TemplateRefConfig`](#templaterefconfig)`[]`<br>List of templates to create
+`acl` | [`AclConfig`](#aclconfig)<br>Custom permissions for this object
+`pausedUntil` | `string`<br>Pause stream builds until specified date
+`pauseComment` | `string`<br>Reason for pausing builds of the stream
+`replicators` | [`ReplicatorConfig`](#replicatorconfig)`[]`<br>Configuration for workers to replicate commit data into Horde Storage.
+`workflows` | [`WorkflowConfig`](#workflowconfig)`[]`<br>Workflows for dealing with new issues
+`tokens` | [`TokenConfig`](#tokenconfig)`[]`<br>Tokens to create for each job step
 
 ## ConfigInclude
 
 Directive to merge config data from another source
 
-Name | Type | Description
----- | ---- | -----------
-`path` | `string` | Path to the config data to be included. May be relative to the including file's location.
+Name | Description
+---- | -----------
+`path` | `string`<br>Path to the config data to be included. May be relative to the including file's location.
 
 ## ConfigMacro
 
 Declares a config macro
 
-Name | Type | Description
----- | ---- | -----------
-`name` | `string` | Name of the macro property
-`value` | `string` | Value for the macro property
+Name | Description
+---- | -----------
+`name` | `string`<br>Name of the macro property
+`value` | `string`<br>Value for the macro property
 
 ## JobOptions
 
 Options for executing a job
 
-Name | Type | Description
----- | ---- | -----------
-`executor` | `string` | Name of the executor to use
-`useNewTempStorage` | `boolean` | Whether to use the new temp storage backend
-`useWine` | `boolean` | Whether to execute using Wine emulation on Linux
-`runInSeparateProcess` | `boolean` | Executes the job lease in a separate process
-`workspaceMaterializer` | `string` | What workspace materializer to use in WorkspaceExecutor. Will override any value from workspace config.
-`container` | [`JobContainerOptions`](#jobcontaineroptions) | Options for executing a job inside a container
+Name | Description
+---- | -----------
+`executor` | `string`<br>Name of the executor to use
+`useNewTempStorage` | `boolean`<br>Whether to use the new temp storage backend
+`useWine` | `boolean`<br>Whether to execute using Wine emulation on Linux
+`runInSeparateProcess` | `boolean`<br>Executes the job lease in a separate process
+`workspaceMaterializer` | `string`<br>What workspace materializer to use in WorkspaceExecutor. Will override any value from workspace config.
+`container` | [`JobContainerOptions`](#jobcontaineroptions)<br>Options for executing a job inside a container
 
 ## JobContainerOptions
 
 Options for executing a job inside a container
 
-Name | Type | Description
----- | ---- | -----------
-`enabled` | `boolean` | Whether to execute job inside a container
-`imageUrl` | `string` | Image URL to container, such as "quay.io/podman/hello"
-`containerEngineExecutable` | `string` | Container engine executable (docker or with full path like /usr/bin/podman)
-`extraArguments` | `string` | Additional arguments to pass to container engine
+Name | Description
+---- | -----------
+`enabled` | `boolean`<br>Whether to execute job inside a container
+`imageUrl` | `string`<br>Image URL to container, such as "quay.io/podman/hello"
+`containerEngineExecutable` | `string`<br>Container engine executable (docker or with full path like /usr/bin/podman)
+`extraArguments` | `string`<br>Additional arguments to pass to container engine
 
 ## DefaultPreflightConfig
 
 Specifies defaults for running a preflight
 
-Name | Type | Description
----- | ---- | -----------
-`templateId` | `string` | The template id to query
-`change` | [`ChangeQueryConfig`](#changequeryconfig) | Query for the change to use
+Name | Description
+---- | -----------
+`templateId` | `string`<br>The template id to query
+`change` | [`ChangeQueryConfig`](#changequeryconfig)<br>Query for the change to use
 
 ## ChangeQueryConfig
 
 Query selecting the base changelist to use
 
-Name | Type | Description
----- | ---- | -----------
-`name` | `string` | Name of this query, for display on the dashboard.
-`condition` | `string` | Condition to evaluate before deciding to use this query. May query tags in a preflight.
-`templateId` | `string` | The template id to query
-`target` | `string` | The target to query
-`outcomes` | [`JobStepOutcome`](#jobstepoutcome-enum)`[]` | Whether to match a job that produced warnings
-`commitTag` | `string` | Finds the last commit with this tag
+Name | Description
+---- | -----------
+`name` | `string`<br>Name of this query, for display on the dashboard.
+`condition` | `string`<br>Condition to evaluate before deciding to use this query. May query tags in a preflight.
+`templateId` | `string`<br>The template id to query
+`target` | `string`<br>The target to query
+`outcomes` | [`JobStepOutcome`](#jobstepoutcome-enum)`[]`<br>Whether to match a job that produced warnings
+`commitTag` | `string`<br>Finds the last commit with this tag
 
 ## JobStepOutcome (Enum)
 
@@ -110,26 +110,26 @@ Name | Description
 
 Configuration for custom commit filters
 
-Name | Type | Description
----- | ---- | -----------
-`name` | `string` | Name of the tag
-`base` | `string` | Base tag to copy settings from
-`filter` | `string[]` | List of files to be included in this filter
+Name | Description
+---- | -----------
+`name` | `string`<br>Name of the tag
+`base` | `string`<br>Base tag to copy settings from
+`filter` | `string[]`<br>List of files to be included in this filter
 
 ## TabConfig
 
 Information about a page to display in the dashboard for a stream
 
-Name | Type | Description
----- | ---- | -----------
-`title` | `string` | Title of this page
-`type` | `string` | Type of this tab
-`style` | [`TabStyle`](#tabstyle-enum) | Presentation style for this page
-`showNames` | `boolean` | Whether to show job names on this page
-`showPreflights` | `boolean` | Whether to show all user preflights
-`jobNames` | `string[]` | Names of jobs to include on this page. If there is only one name specified, the name column does not need to be displayed.
-`templates` | `string[]` | List of job template names to show on this page.
-`columns` | [`TabColumnConfig`](#tabcolumnconfig)`[]` | Columns to display for different types of aggregates
+Name | Description
+---- | -----------
+`title` | `string`<br>Title of this page
+`type` | `string`<br>Type of this tab
+`style` | [`TabStyle`](#tabstyle-enum)<br>Presentation style for this page
+`showNames` | `boolean`<br>Whether to show job names on this page
+`showPreflights` | `boolean`<br>Whether to show all user preflights
+`jobNames` | `string[]`<br>Names of jobs to include on this page. If there is only one name specified, the name column does not need to be displayed.
+`templates` | `string[]`<br>List of job template names to show on this page.
+`columns` | [`TabColumnConfig`](#tabcolumnconfig)`[]`<br>Columns to display for different types of aggregates
 
 ## TabStyle (Enum)
 
@@ -144,13 +144,13 @@ Name | Description
 
 Describes a column to display on the jobs page
 
-Name | Type | Description
----- | ---- | -----------
-`type` | [`TabColumnType`](#tabcolumntype-enum) | The type of column
-`heading` | `string` | Heading for this column
-`category` | `string` | Category of aggregates to display in this column. If null, includes any aggregate not matched by another column.
-`parameter` | `string` | Parameter to show in this column
-`relativeWidth` | `integer` | Relative width of this column.
+Name | Description
+---- | -----------
+`type` | [`TabColumnType`](#tabcolumntype-enum)<br>The type of column
+`heading` | `string`<br>Heading for this column
+`category` | `string`<br>Category of aggregates to display in this column. If null, includes any aggregate not matched by another column.
+`parameter` | `string`<br>Parameter to show in this column
+`relativeWidth` | `integer`<br>Relative width of this column.
 
 ## TabColumnType (Enum)
 
@@ -165,112 +165,112 @@ Name | Description
 
 Mapping from a BuildGraph agent type to a set of machines on the farm
 
-Name | Type | Description
----- | ---- | -----------
-`base` | `string` | Base agent config to inherit settings from
-`pool` | `string` | Pool of agents to use for this agent type
-`workspace` | `string` | Name of the workspace to sync
-`tempStorageDir` | `string` | Path to the temporary storage dir
-`environment` | `string` `->` `string` | Environment variables to be set when executing the job
-`tokens` | [`TokenConfig`](#tokenconfig)`[]` | Tokens to allocate for this agent type
+Name | Description
+---- | -----------
+`base` | `string`<br>Base agent config to inherit settings from
+`pool` | `string`<br>Pool of agents to use for this agent type
+`workspace` | `string`<br>Name of the workspace to sync
+`tempStorageDir` | `string`<br>Path to the temporary storage dir
+`environment` | `string` `->` `string`<br>Environment variables to be set when executing the job
+`tokens` | [`TokenConfig`](#tokenconfig)`[]`<br>Tokens to allocate for this agent type
 
 ## TokenConfig
 
 Configuration for allocating access tokens for each job
 
-Name | Type | Description
----- | ---- | -----------
-`url` | `string` | URL to request tokens from
-`clientId` | `string` | Client id to use to request a new token
-`clientSecret` | `string` | Client secret to request a new access token
-`envVar` | `string` | Environment variable to set with the access token
+Name | Description
+---- | -----------
+`url` | `string`<br>URL to request tokens from
+`clientId` | `string`<br>Client id to use to request a new token
+`clientSecret` | `string`<br>Client secret to request a new access token
+`envVar` | `string`<br>Environment variable to set with the access token
 
 ## WorkspaceConfig
 
 Information about a workspace type
 
-Name | Type | Description
----- | ---- | -----------
-`base` | `string` | Base workspace to derive from
-`cluster` | `string` | Name of the Perforce server cluster to use
-`serverAndPort` | `string` | The Perforce server and port (eg. perforce:1666)
-`userName` | `string` | User to log into Perforce with (defaults to buildmachine)
-`password` | `string` | Password to use to log into the workspace
-`identifier` | `string` | Identifier to distinguish this workspace from other workspaces. Defaults to the workspace type name.
-`stream` | `string` | Override for the stream to sync
-`view` | `string[]` | Custom view for the workspace
-`incremental` | `boolean` | Whether to use an incrementally synced workspace
-`useAutoSdk` | `boolean` | Whether to use the AutoSDK
-`autoSdkView` | `string[]` | View for the AutoSDK paths to sync. If null, the whole thing will be synced.
-`method` | `string` | Method to use when syncing/materializing data from Perforce
+Name | Description
+---- | -----------
+`base` | `string`<br>Base workspace to derive from
+`cluster` | `string`<br>Name of the Perforce server cluster to use
+`serverAndPort` | `string`<br>The Perforce server and port (eg. perforce:1666)
+`userName` | `string`<br>User to log into Perforce with (defaults to buildmachine)
+`password` | `string`<br>Password to use to log into the workspace
+`identifier` | `string`<br>Identifier to distinguish this workspace from other workspaces. Defaults to the workspace type name.
+`stream` | `string`<br>Override for the stream to sync
+`view` | `string[]`<br>Custom view for the workspace
+`incremental` | `boolean`<br>Whether to use an incrementally synced workspace
+`useAutoSdk` | `boolean`<br>Whether to use the AutoSDK
+`autoSdkView` | `string[]`<br>View for the AutoSDK paths to sync. If null, the whole thing will be synced.
+`method` | `string`<br>Method to use when syncing/materializing data from Perforce
 
 ## TemplateRefConfig
 
 Parameters to create a template within a stream
 
-Name | Type | Description
----- | ---- | -----------
-`id` | `string` | Optional identifier for this ref. If not specified, an id will be generated from the name.
-`base` | `string` | Base template id to copy from
-`showUgsBadges` | `boolean` | Whether to show badges in UGS for these jobs
-`showUgsAlerts` | `boolean` | Whether to show alerts in UGS for these jobs
-`notificationChannel` | `string` | Notification channel for this template. Overrides the stream channel if set.
-`notificationChannelFilter` | `string` | Notification channel filter for this template. Can be Success|Failure|Warnings
-`triageChannel` | `string` | Triage channel for this template. Overrides the stream channel if set.
-`workflowId` | `string` | Workflow to user for this stream
-`annotations` | `string` `->` `string` | Default annotations to apply to nodes in this template
-`schedule` | [`ScheduleConfig`](#scheduleconfig) | Schedule to execute this template
-`chainedJobs` | [`ChainedJobTemplateConfig`](#chainedjobtemplateconfig)`[]` | List of chained job triggers
-`acl` | [`AclConfig`](#aclconfig) | The ACL for this template
-`name` | `string` | Name for the new template
-`description` | `string` | Description for the template
-`priority` | [`Priority`](#priority-enum) | Default priority for this job
-`allowPreflights` | `boolean` | Whether to allow preflights of this template
-`updateIssues` | `boolean` | Whether issues should be updated for all jobs using this template
-`promoteIssuesByDefault` | `boolean` | Whether issues should be promoted by default for this template, promoted issues will generate user notifications
-`initialAgentType` | `string` | Initial agent type to parse the buildgraph script on
-`submitNewChange` | `string` | Path to a file within the stream to submit to generate a new changelist for jobs
-`submitDescription` | `string` | Description for new changelists
-`defaultChange` | [`ChangeQueryConfig`](#changequeryconfig)`[]` | Default change to build at. Each object has a condition parameter which can evaluated by the server to determine which change to use.
-`arguments` | `string[]` | Fixed arguments for the new job
-`parameters` | [`GroupParameterData`](#groupparameterdata)/[`TextParameterData`](#textparameterdata)/[`ListParameterData`](#listparameterdata)/[`BoolParameterData`](#boolparameterdata)`[]` | Parameters for this template
-`jobOptions` | [`JobOptions`](#joboptions) | Default settings for jobs
+Name | Description
+---- | -----------
+`id` | `string`<br>Optional identifier for this ref. If not specified, an id will be generated from the name.
+`base` | `string`<br>Base template id to copy from
+`showUgsBadges` | `boolean`<br>Whether to show badges in UGS for these jobs
+`showUgsAlerts` | `boolean`<br>Whether to show alerts in UGS for these jobs
+`notificationChannel` | `string`<br>Notification channel for this template. Overrides the stream channel if set.
+`notificationChannelFilter` | `string`<br>Notification channel filter for this template. Can be Success|Failure|Warnings
+`triageChannel` | `string`<br>Triage channel for this template. Overrides the stream channel if set.
+`workflowId` | `string`<br>Workflow to user for this stream
+`annotations` | `string` `->` `string`<br>Default annotations to apply to nodes in this template
+`schedule` | [`ScheduleConfig`](#scheduleconfig)<br>Schedule to execute this template
+`chainedJobs` | [`ChainedJobTemplateConfig`](#chainedjobtemplateconfig)`[]`<br>List of chained job triggers
+`acl` | [`AclConfig`](#aclconfig)<br>The ACL for this template
+`name` | `string`<br>Name for the new template
+`description` | `string`<br>Description for the template
+`priority` | [`Priority`](#priority-enum)<br>Default priority for this job
+`allowPreflights` | `boolean`<br>Whether to allow preflights of this template
+`updateIssues` | `boolean`<br>Whether issues should be updated for all jobs using this template
+`promoteIssuesByDefault` | `boolean`<br>Whether issues should be promoted by default for this template, promoted issues will generate user notifications
+`initialAgentType` | `string`<br>Initial agent type to parse the buildgraph script on
+`submitNewChange` | `string`<br>Path to a file within the stream to submit to generate a new changelist for jobs
+`submitDescription` | `string`<br>Description for new changelists
+`defaultChange` | [`ChangeQueryConfig`](#changequeryconfig)`[]`<br>Default change to build at. Each object has a condition parameter which can evaluated by the server to determine which change to use.
+`arguments` | `string[]`<br>Fixed arguments for the new job
+`parameters` | [`GroupParameterData`](#groupparameterdata)/[`TextParameterData`](#textparameterdata)/[`ListParameterData`](#listparameterdata)/[`BoolParameterData`](#boolparameterdata)`[]`<br>Parameters for this template
+`jobOptions` | [`JobOptions`](#joboptions)<br>Default settings for jobs
 
 ## ScheduleConfig
 
 Parameters to create a new schedule
 
-Name | Type | Description
----- | ---- | -----------
-`claims` | [`AclClaimConfig`](#aclclaimconfig)`[]` | Roles to impersonate for this schedule
-`enabled` | `boolean` | Whether the schedule should be enabled
-`maxActive` | `integer` | Maximum number of builds that can be active at once
-`maxChanges` | `integer` | Maximum number of changes the schedule can fall behind head revision. If greater than zero, builds will be triggered for every submitted changelist until the backlog is this size.
-`requireSubmittedChange` | `boolean` | Whether the build requires a change to be submitted
-`gate` | [`ScheduleGateConfig`](#schedulegateconfig) | Gate allowing the schedule to trigger
-`commits` | `string[]` | Commit tags for this schedule
-`filter` | [`ChangeContentFlags`](#changecontentflags-enum)`[]` | The types of changes to run for
-`files` | `string[]` | Files that should cause the job to trigger
-`templateParameters` | `string` `->` `string` | Parameters for the template
-`patterns` | [`SchedulePatternConfig`](#schedulepatternconfig)`[]` | New patterns for the schedule
+Name | Description
+---- | -----------
+`claims` | [`AclClaimConfig`](#aclclaimconfig)`[]`<br>Roles to impersonate for this schedule
+`enabled` | `boolean`<br>Whether the schedule should be enabled
+`maxActive` | `integer`<br>Maximum number of builds that can be active at once
+`maxChanges` | `integer`<br>Maximum number of changes the schedule can fall behind head revision. If greater than zero, builds will be triggered for every submitted changelist until the backlog is this size.
+`requireSubmittedChange` | `boolean`<br>Whether the build requires a change to be submitted
+`gate` | [`ScheduleGateConfig`](#schedulegateconfig)<br>Gate allowing the schedule to trigger
+`commits` | `string[]`<br>Commit tags for this schedule
+`filter` | [`ChangeContentFlags`](#changecontentflags-enum)`[]`<br>The types of changes to run for
+`files` | `string[]`<br>Files that should cause the job to trigger
+`templateParameters` | `string` `->` `string`<br>Parameters for the template
+`patterns` | [`SchedulePatternConfig`](#schedulepatternconfig)`[]`<br>New patterns for the schedule
 
 ## AclClaimConfig
 
 New claim to create
 
-Name | Type | Description
----- | ---- | -----------
-`type` | `string` | The claim type
-`value` | `string` | The claim value
+Name | Description
+---- | -----------
+`type` | `string`<br>The claim type
+`value` | `string`<br>The claim value
 
 ## ScheduleGateConfig
 
 Gate allowing a schedule to trigger.
 
-Name | Type | Description
----- | ---- | -----------
-`templateId` | `string` | The template containing the dependency
-`target` | `string` | Target to wait for
+Name | Description
+---- | -----------
+`templateId` | `string`<br>The template containing the dependency
+`target` | `string`<br>Target to wait for
 
 ## ChangeContentFlags (Enum)
 
@@ -285,12 +285,12 @@ Name | Description
 
 Parameters to create a new schedule
 
-Name | Type | Description
----- | ---- | -----------
-`daysOfWeek` | [`DayOfWeek`](#dayofweek-enum)`[]` | Days of the week to run this schedule on. If null, the schedule will run every day.
-`minTime` | `string` | Time during the day for the first schedule to trigger. Measured in minutes from midnight.
-`maxTime` | `string` | Time during the day for the last schedule to trigger. Measured in minutes from midnight.
-`interval` | `string` | Interval between each schedule triggering
+Name | Description
+---- | -----------
+`daysOfWeek` | [`DayOfWeek`](#dayofweek-enum)`[]`<br>Days of the week to run this schedule on. If null, the schedule will run every day.
+`minTime` | `string`<br>Time during the day for the first schedule to trigger. Measured in minutes from midnight.
+`maxTime` | `string`<br>Time during the day for the last schedule to trigger. Measured in minutes from midnight.
+`interval` | `string`<br>Interval between each schedule triggering
 
 ## DayOfWeek (Enum)
 
@@ -308,43 +308,43 @@ Name | Description
 
 Trigger for another template
 
-Name | Type | Description
----- | ---- | -----------
-`trigger` | `string` | Name of the target that needs to complete before starting the other template
-`templateId` | `string` | Id of the template to trigger
-`useDefaultChangeForTemplate` | `boolean` | Whether to use the default change for the template rather than the change for the parent job.
+Name | Description
+---- | -----------
+`trigger` | `string`<br>Name of the target that needs to complete before starting the other template
+`templateId` | `string`<br>Id of the template to trigger
+`useDefaultChangeForTemplate` | `boolean`<br>Whether to use the default change for the template rather than the change for the parent job.
 
 ## AclConfig
 
 Parameters to update an ACL
 
-Name | Type | Description
----- | ---- | -----------
-`entries` | [`AclEntryConfig`](#aclentryconfig)`[]` | Entries to replace the existing ACL
-`profiles` | [`AclProfileConfig`](#aclprofileconfig)`[]` | Defines profiles which allow grouping sets of actions into named collections
-`inherit` | `boolean` | Whether to inherit permissions from the parent ACL
-`exceptions` | `string[]` | List of exceptions to the inherited setting
+Name | Description
+---- | -----------
+`entries` | [`AclEntryConfig`](#aclentryconfig)`[]`<br>Entries to replace the existing ACL
+`profiles` | [`AclProfileConfig`](#aclprofileconfig)`[]`<br>Defines profiles which allow grouping sets of actions into named collections
+`inherit` | `boolean`<br>Whether to inherit permissions from the parent ACL
+`exceptions` | `string[]`<br>List of exceptions to the inherited setting
 
 ## AclEntryConfig
 
 Individual entry in an ACL
 
-Name | Type | Description
----- | ---- | -----------
-`claim` | [`AclClaimConfig`](#aclclaimconfig) | Name of the user or group
-`actions` | `string[]` | Array of actions to allow
-`profiles` | `string[]` | List of profiles to grant
+Name | Description
+---- | -----------
+`claim` | [`AclClaimConfig`](#aclclaimconfig)<br>Name of the user or group
+`actions` | `string[]`<br>Array of actions to allow
+`profiles` | `string[]`<br>List of profiles to grant
 
 ## AclProfileConfig
 
 Configuration for an ACL profile. This defines a preset group of actions which can be given to a user via an ACL entry.
 
-Name | Type | Description
----- | ---- | -----------
-`id` | `string` | Identifier for this profile
-`actions` | `string[]` | Actions to include
-`excludeActions` | `string[]` | Actions to exclude from the inherited actions
-`extends` | `string[]` | Other profiles to extend from
+Name | Description
+---- | -----------
+`id` | `string`<br>Identifier for this profile
+`actions` | `string[]`<br>Actions to include
+`excludeActions` | `string[]`<br>Actions to exclude from the inherited actions
+`extends` | `string[]`<br>Other profiles to extend from
 
 ## Priority (Enum)
 
@@ -362,12 +362,12 @@ Name | Description
 
 Used to group a number of other parameters
 
-Name | Type | Description
----- | ---- | -----------
-`type` | Group | Type discriminator
-`label` | `string` | Label to display next to this parameter
-`style` | [`GroupParameterStyle`](#groupparameterstyle-enum) | How to display this group
-`children` | [`GroupParameterData`](#groupparameterdata)/[`TextParameterData`](#textparameterdata)/[`ListParameterData`](#listparameterdata)/[`BoolParameterData`](#boolparameterdata)`[]` | List of child parameters
+Name | Description
+---- | -----------
+`type` | Group<br>Type discriminator
+`label` | `string`<br>Label to display next to this parameter
+`style` | [`GroupParameterStyle`](#groupparameterstyle-enum)<br>How to display this group
+`children` | [`GroupParameterData`](#groupparameterdata)/[`TextParameterData`](#textparameterdata)/[`ListParameterData`](#listparameterdata)/[`BoolParameterData`](#boolparameterdata)`[]`<br>List of child parameters
 
 ## GroupParameterStyle (Enum)
 
@@ -382,29 +382,29 @@ Name | Description
 
 Free-form text entry parameter
 
-Name | Type | Description
----- | ---- | -----------
-`type` | Text | Type discriminator
-`label` | `string` | Name of the parameter associated with this parameter.
-`argument` | `string` | Argument to pass to the executor
-`default` | `string` | Default value for this argument
-`scheduleOverride` | `string` | Override for the default value for this parameter when running a scheduled build
-`hint` | `string` | Hint text for this parameter
-`validation` | `string` | Regex used to validate this parameter
-`validationError` | `string` | Message displayed if validation fails, informing user of valid values.
-`toolTip` | `string` | Tool-tip text to display
+Name | Description
+---- | -----------
+`type` | Text<br>Type discriminator
+`label` | `string`<br>Name of the parameter associated with this parameter.
+`argument` | `string`<br>Argument to pass to the executor
+`default` | `string`<br>Default value for this argument
+`scheduleOverride` | `string`<br>Override for the default value for this parameter when running a scheduled build
+`hint` | `string`<br>Hint text for this parameter
+`validation` | `string`<br>Regex used to validate this parameter
+`validationError` | `string`<br>Message displayed if validation fails, informing user of valid values.
+`toolTip` | `string`<br>Tool-tip text to display
 
 ## ListParameterData
 
 Allows the user to select a value from a constrained list of choices
 
-Name | Type | Description
----- | ---- | -----------
-`type` | List | Type discriminator
-`label` | `string` | Label to display next to this parameter. Defaults to the parameter name.
-`style` | [`ListParameterStyle`](#listparameterstyle-enum) | The type of list parameter
-`items` | [`ListParameterItemData`](#listparameteritemdata)`[]` | List of values to display in the list
-`toolTip` | `string` | Tool tip text to display
+Name | Description
+---- | -----------
+`type` | List<br>Type discriminator
+`label` | `string`<br>Label to display next to this parameter. Defaults to the parameter name.
+`style` | [`ListParameterStyle`](#listparameterstyle-enum)<br>The type of list parameter
+`items` | [`ListParameterItemData`](#listparameteritemdata)`[]`<br>List of values to display in the list
+`toolTip` | `string`<br>Tool tip text to display
 
 ## ListParameterStyle (Enum)
 
@@ -420,77 +420,77 @@ Name | Description
 
 Possible option for a list parameter
 
-Name | Type | Description
----- | ---- | -----------
-`group` | `string` | Optional group heading to display this entry under, if the picker style supports it.
-`text` | `string` | Name of the parameter associated with this list.
-`argumentIfEnabled` | `string` | Argument to pass with this parameter.
-`argumentsIfEnabled` | `string[]` | Arguments to pass with this parameter.
-`argumentIfDisabled` | `string` | Argument to pass with this parameter.
-`argumentsIfDisabled` | `string[]` | Arguments to pass if this parameter is disabled.
-`default` | `boolean` | Whether this item is selected by default
-`scheduleOverride` | `boolean` | Overridden value for this property in schedule builds
+Name | Description
+---- | -----------
+`group` | `string`<br>Optional group heading to display this entry under, if the picker style supports it.
+`text` | `string`<br>Name of the parameter associated with this list.
+`argumentIfEnabled` | `string`<br>Argument to pass with this parameter.
+`argumentsIfEnabled` | `string[]`<br>Arguments to pass with this parameter.
+`argumentIfDisabled` | `string`<br>Argument to pass with this parameter.
+`argumentsIfDisabled` | `string[]`<br>Arguments to pass if this parameter is disabled.
+`default` | `boolean`<br>Whether this item is selected by default
+`scheduleOverride` | `boolean`<br>Overridden value for this property in schedule builds
 
 ## BoolParameterData
 
 Allows the user to toggle an option on or off
 
-Name | Type | Description
----- | ---- | -----------
-`type` | Bool | Type discriminator
-`label` | `string` | Name of the parameter associated with this parameter.
-`argumentIfEnabled` | `string` | Argument to add if this parameter is enabled
-`argumentsIfEnabled` | `string[]` | Argument to add if this parameter is enabled
-`argumentIfDisabled` | `string` | Argument to add if this parameter is enabled
-`argumentsIfDisabled` | `string[]` | Arguments to add if this parameter is disabled
-`default` | `boolean` | Whether this argument is enabled by default
-`scheduleOverride` | `boolean` | Override for this parameter in scheduled builds
-`toolTip` | `string` | Tool tip text to display
+Name | Description
+---- | -----------
+`type` | Bool<br>Type discriminator
+`label` | `string`<br>Name of the parameter associated with this parameter.
+`argumentIfEnabled` | `string`<br>Argument to add if this parameter is enabled
+`argumentsIfEnabled` | `string[]`<br>Argument to add if this parameter is enabled
+`argumentIfDisabled` | `string`<br>Argument to add if this parameter is enabled
+`argumentsIfDisabled` | `string[]`<br>Arguments to add if this parameter is disabled
+`default` | `boolean`<br>Whether this argument is enabled by default
+`scheduleOverride` | `boolean`<br>Override for this parameter in scheduled builds
+`toolTip` | `string`<br>Tool tip text to display
 
 ## ReplicatorConfig
 
 Configuration for a stream replicator
 
-Name | Type | Description
----- | ---- | -----------
-`id` | `string` | Identifier for the replicator within the current stream
-`enabled` | `boolean` | Whether the replicator is enabled
-`minChange` | `integer` | Minimum change number to replicate
-`maxChange` | `integer` | Maximum change number to replicate
+Name | Description
+---- | -----------
+`id` | `string`<br>Identifier for the replicator within the current stream
+`enabled` | `boolean`<br>Whether the replicator is enabled
+`minChange` | `integer`<br>Minimum change number to replicate
+`maxChange` | `integer`<br>Maximum change number to replicate
 
 ## WorkflowConfig
 
 Configuration for an issue workflow
 
-Name | Type | Description
----- | ---- | -----------
-`id` | `string` | Identifier for this workflow
-`reportTimes` | `string[]` | Times of day at which to send a report
-`summaryTab` | `string` | Name of the tab to post summary data to
-`reportChannel` | `string` | Channel to post summary information for these templates.
-`groupIssuesByTemplate` | `boolean` | Whether to group issues by template in the report
-`triageChannel` | `string` | Channel to post threads for triaging new issues
-`triagePrefix` | `string` | Prefix for all triage messages
-`triageSuffix` | `string` | Suffix for all triage messages
-`triageInstructions` | `string` | Instructions posted to triage threads
-`triageAlias` | `string` | User id of a Slack user/alias to ping if there is nobody assigned to an issue by default.
-`triageTypeAliases` | `string` `->` `string` | Slack user/alias to ping for specific issue types (such as Systemic), if there is nobody assigned to an issue by default.
-`escalateAlias` | `string` | Alias to ping if an issue has not been resolved for a certain amount of time
-`escalateTimes` | `integer[]` | Times after an issue has been opened to escalate to the alias above, in minutes. Continues to notify on the last interval once reaching the end of the list.
-`maxMentions` | `integer` | Maximum number of people to mention on a triage thread
-`allowMentions` | `boolean` | Whether to mention people on this thread. Useful to disable for testing.
-`inviteRestrictedUsers` | `boolean` | Uses the admin.conversations.invite API to invite users to the channel
-`skipWhenEmpty` | `boolean` | Skips sending reports when there are no active issues.
-`annotations` | `string` `->` `string` | Additional node annotations implicit in this workflow
-`externalIssues` | [`ExternalIssueConfig`](#externalissueconfig) | External issue tracking configuration for this workflow
-`issueHandlers` | `string[]` | Additional issue handlers enabled for this workflow
+Name | Description
+---- | -----------
+`id` | `string`<br>Identifier for this workflow
+`reportTimes` | `string[]`<br>Times of day at which to send a report
+`summaryTab` | `string`<br>Name of the tab to post summary data to
+`reportChannel` | `string`<br>Channel to post summary information for these templates.
+`groupIssuesByTemplate` | `boolean`<br>Whether to group issues by template in the report
+`triageChannel` | `string`<br>Channel to post threads for triaging new issues
+`triagePrefix` | `string`<br>Prefix for all triage messages
+`triageSuffix` | `string`<br>Suffix for all triage messages
+`triageInstructions` | `string`<br>Instructions posted to triage threads
+`triageAlias` | `string`<br>User id of a Slack user/alias to ping if there is nobody assigned to an issue by default.
+`triageTypeAliases` | `string` `->` `string`<br>Slack user/alias to ping for specific issue types (such as Systemic), if there is nobody assigned to an issue by default.
+`escalateAlias` | `string`<br>Alias to ping if an issue has not been resolved for a certain amount of time
+`escalateTimes` | `integer[]`<br>Times after an issue has been opened to escalate to the alias above, in minutes. Continues to notify on the last interval once reaching the end of the list.
+`maxMentions` | `integer`<br>Maximum number of people to mention on a triage thread
+`allowMentions` | `boolean`<br>Whether to mention people on this thread. Useful to disable for testing.
+`inviteRestrictedUsers` | `boolean`<br>Uses the admin.conversations.invite API to invite users to the channel
+`skipWhenEmpty` | `boolean`<br>Skips sending reports when there are no active issues.
+`annotations` | `string` `->` `string`<br>Additional node annotations implicit in this workflow
+`externalIssues` | [`ExternalIssueConfig`](#externalissueconfig)<br>External issue tracking configuration for this workflow
+`issueHandlers` | `string[]`<br>Additional issue handlers enabled for this workflow
 
 ## ExternalIssueConfig
 
 External issue tracking configuration for a workflow
 
-Name | Type | Description
----- | ---- | -----------
-`projectKey` | `string` | Project key in external issue tracker
-`defaultComponentId` | `string` | Default component id for issues using workflow
-`defaultIssueTypeId` | `string` | Default issue type id for issues using workflow
+Name | Description
+---- | -----------
+`projectKey` | `string`<br>Project key in external issue tracker
+`defaultComponentId` | `string`<br>Default component id for issues using workflow
+`defaultIssueTypeId` | `string`<br>Default issue type id for issues using workflow
diff --git a/Engine/Source/Programs/Horde/Docs/Deployment/AgentSettings.md b/Engine/Source/Programs/Horde/Docs/Deployment/AgentSettings.md
--- a/Engine/Source/Programs/Horde/Docs/Deployment/AgentSettings.md
+++ b/Engine/Source/Programs/Horde/Docs/Deployment/AgentSettings.md
@@ -4,87 +4,87 @@
 
 All Horde-specific settings are stored in a root object called `Horde`. Other .NET functionality may be configured using properties in the root of this file.
 
-Name | Type | Description
----- | ---- | -----------
-`serverProfiles` | `string` `->` [`ServerProfile`](#serverprofile) | Known servers to connect to
-`server` | `string` | The default server, unless overridden from the command line
-`name` | `string` | Name of agent to report as when connecting to server. By default, the computer's hostname will be used.
-`installed` | `boolean` | Whether the server is running in 'installed' mode. In this mode, on Windows, the default data directory will use the common application data folder (C:\ProgramData\Epic\Horde), and configuration data will be read from here and the registry. This setting is overridden to false for local builds from appsettings.Local.json.
-`ephemeral` | `boolean` | Whether agent should register as being ephemeral. Doing so will not persist any long-lived data on the server and once disconnected it's assumed to have been deleted permanently. Ideal for short-lived agents, such as spot instances on AWS EC2.
-`executor` | `string` | The executor to use for jobs. Defaults to the Perforce executor.
-`updateChannel` | `string` | Default channel for downloading updates
-`localExecutor` | [`LocalExecutorSettings`](#localexecutorsettings) | Settings for the local executor
-`perforceExecutor` | [`PerforceExecutorSettings`](#perforceexecutorsettings) | Settings for the perforce executor
-`workingDir` | [`DirectoryReference`](#directoryreference) | Working directory
-`shareMountingEnabled` | `boolean` | Whether to mount the specified list of network shares
-`shares` | [`MountNetworkShare`](#mountnetworkshare)`[]` | List of network shares to mount
-`processNamesToTerminate` | `string[]` | List of process names to terminate after a job
-`processesToTerminate` | [`ProcessToTerminate`](#processtoterminate)`[]` | List of process names to terminate after a lease completes, but not after a job step
-`wineExecutablePath` | `string` | Path to Wine executable. If null, execution under Wine is disabled
-`containerEngineExecutablePath` | `string` | Path to container engine executable, such as /usr/bin/podman. If null, execution of compute workloads inside a container is disabled
-`writeStepOutputToLogger` | `boolean` | Whether to write step output to the logging device
-`enableAwsEc2Support` | `boolean` | Queries information about the current agent through the AWS EC2 interface
-`useLocalStorageClient` | `boolean` | Option to use a local storage client rather than connecting through the server. Primarily for convenience when debugging / iterating locally.
-`computePort` | `integer` | Incoming port for listening for compute work. Needs to be tied with a lease.
-`enableTelemetry` | `boolean` | Whether to send telemetry back to Horde server
-`telemetryReportInterval` | `integer` | How often to report telemetry events to server in milliseconds
-`bundleCacheSize` | `integer` | Maximum size of the bundle cache, in megabytes.
-`properties` | `string` `->` `string` | Key/value properties in addition to those set internally by the agent
+Name | Description
+---- | -----------
+`serverProfiles` | `string` `->` [`ServerProfile`](#serverprofile)<br>Known servers to connect to
+`server` | `string`<br>The default server, unless overridden from the command line
+`name` | `string`<br>Name of agent to report as when connecting to server. By default, the computer's hostname will be used.
+`installed` | `boolean`<br>Whether the server is running in 'installed' mode. In this mode, on Windows, the default data directory will use the common application data folder (C:\ProgramData\Epic\Horde), and configuration data will be read from here and the registry. This setting is overridden to false for local builds from appsettings.Local.json.
+`ephemeral` | `boolean`<br>Whether agent should register as being ephemeral. Doing so will not persist any long-lived data on the server and once disconnected it's assumed to have been deleted permanently. Ideal for short-lived agents, such as spot instances on AWS EC2.
+`executor` | `string`<br>The executor to use for jobs. Defaults to the Perforce executor.
+`updateChannel` | `string`<br>Default channel for downloading updates
+`localExecutor` | [`LocalExecutorSettings`](#localexecutorsettings)<br>Settings for the local executor
+`perforceExecutor` | [`PerforceExecutorSettings`](#perforceexecutorsettings)<br>Settings for the perforce executor
+`workingDir` | [`DirectoryReference`](#directoryreference)<br>Working directory
+`shareMountingEnabled` | `boolean`<br>Whether to mount the specified list of network shares
+`shares` | [`MountNetworkShare`](#mountnetworkshare)`[]`<br>List of network shares to mount
+`processNamesToTerminate` | `string[]`<br>List of process names to terminate after a job
+`processesToTerminate` | [`ProcessToTerminate`](#processtoterminate)`[]`<br>List of process names to terminate after a lease completes, but not after a job step
+`wineExecutablePath` | `string`<br>Path to Wine executable. If null, execution under Wine is disabled
+`containerEngineExecutablePath` | `string`<br>Path to container engine executable, such as /usr/bin/podman. If null, execution of compute workloads inside a container is disabled
+`writeStepOutputToLogger` | `boolean`<br>Whether to write step output to the logging device
+`enableAwsEc2Support` | `boolean`<br>Queries information about the current agent through the AWS EC2 interface
+`useLocalStorageClient` | `boolean`<br>Option to use a local storage client rather than connecting through the server. Primarily for convenience when debugging / iterating locally.
+`computePort` | `integer`<br>Incoming port for listening for compute work. Needs to be tied with a lease.
+`enableTelemetry` | `boolean`<br>Whether to send telemetry back to Horde server
+`telemetryReportInterval` | `integer`<br>How often to report telemetry events to server in milliseconds
+`bundleCacheSize` | `integer`<br>Maximum size of the bundle cache, in megabytes.
+`properties` | `string` `->` `string`<br>Key/value properties in addition to those set internally by the agent
 
 ## ServerProfile
 
 Information about a server to use
 
-Name | Type | Description
----- | ---- | -----------
-`name` | `string` | Name of this server profile
-`environment` | `string` | Name of the environment (currently just used for tracing)
-`url` | `string` | Url of the server
-`token` | `string` | Bearer token to use to initiate the connection
-`thumbprint` | `string` | Thumbprint of a certificate to trust. Allows using self-signed certs for the server.
-`thumbprints` | `string[]` | Thumbprints of certificates to trust. Allows using self-signed certs for the server.
+Name | Description
+---- | -----------
+`name` | `string`<br>Name of this server profile
+`environment` | `string`<br>Name of the environment (currently just used for tracing)
+`url` | `string`<br>Url of the server
+`token` | `string`<br>Bearer token to use to initiate the connection
+`thumbprint` | `string`<br>Thumbprint of a certificate to trust. Allows using self-signed certs for the server.
+`thumbprints` | `string[]`<br>Thumbprints of certificates to trust. Allows using self-signed certs for the server.
 
 ## LocalExecutorSettings
 
 Settings for the local executor
 
-Name | Type | Description
----- | ---- | -----------
-`workspaceDir` | `string` | Path to the local workspace to use with the local executor
-`runSteps` | `boolean` | Whether to actually execute steps, or just do job setup
+Name | Description
+---- | -----------
+`workspaceDir` | `string`<br>Path to the local workspace to use with the local executor
+`runSteps` | `boolean`<br>Whether to actually execute steps, or just do job setup
 
 ## PerforceExecutorSettings
 
 Settings for the perforce executor
 
-Name | Type | Description
----- | ---- | -----------
-`runConform` | `boolean` | Whether to run conform jobs
+Name | Description
+---- | -----------
+`runConform` | `boolean`<br>Whether to run conform jobs
 
 ## DirectoryReference
 
-Name | Type | Description
----- | ---- | -----------
-`parentDirectory` | [`DirectoryReference`](#directoryreference) | 
-`fullName` | `string` | 
+Name | Description
+---- | -----------
+`parentDirectory` | [`DirectoryReference`](#directoryreference)<br>
+`fullName` | `string`<br>
 
 ## MountNetworkShare
 
 Describes a network share to mount
 
-Name | Type | Description
----- | ---- | -----------
-`mountPoint` | `string` | Where the share should be mounted on the local machine. Must be a drive letter for Windows.
-`remotePath` | `string` | Path to the remote resource
+Name | Description
+---- | -----------
+`mountPoint` | `string`<br>Where the share should be mounted on the local machine. Must be a drive letter for Windows.
+`remotePath` | `string`<br>Path to the remote resource
 
 ## ProcessToTerminate
 
 Specifies a process to terminate
 
-Name | Type | Description
----- | ---- | -----------
-`name` | `string` | Name of the process
-`when` | [`TerminateCondition`](#terminatecondition-enum)`[]` | When to terminate this process
+Name | Description
+---- | -----------
+`name` | `string`<br>Name of the process
+`when` | [`TerminateCondition`](#terminatecondition-enum)`[]`<br>When to terminate this process
 
 ## TerminateCondition (Enum)
 
diff --git a/Engine/Source/Programs/Horde/Docs/Deployment/ServerSettings.md b/Engine/Source/Programs/Horde/Docs/Deployment/ServerSettings.md
--- a/Engine/Source/Programs/Horde/Docs/Deployment/ServerSettings.md
+++ b/Engine/Source/Programs/Horde/Docs/Deployment/ServerSettings.md
@@ -4,102 +4,102 @@
 
 All Horde-specific settings are stored in a root object called `Horde`. Other .NET functionality may be configured using properties in the root of this file.
 
-Name | Type | Description
----- | ---- | -----------
-`runModes` | [`RunMode`](#runmode-enum)`[]` | Modes that the server should run in. Runmodes can be used in a multi-server deployment to limit the operations that a particular instance will try to perform.
-`dataDir` | `string` | Override the data directory used by Horde. Defaults to C:\ProgramData\HordeServer on Windows, {AppDir}/Data on other platforms.
-`installed` | `boolean` | Whether the server is running in 'installed' mode. In this mode, on Windows, the default data directory will use the common application data folder (C:\ProgramData\Epic\Horde), and configuration data will be read from here and the registry. This setting is overridden to false for local builds from appsettings.Local.json.
-`httpPort` | `integer` | Main port for serving HTTP.
-`httpsPort` | `integer` | Port for serving HTTP with TLS enabled. Disabled by default.
-`http2Port` | `integer` | Dedicated port for serving only HTTP/2.
-`computeTunnelPort` | `integer` | Port to listen on for tunneling compute sockets to agents
-`computeTunnelAddress` | `string` | What address (host:port) clients should connect to for compute socket tunneling Port may differ from  if Horde server is behind a reverse proxy/firewall
-`databaseConnectionString` | `string` | MongoDB connection string
-`databaseName` | `string` | MongoDB database name
-`databasePublicCert` | `string` | Optional certificate to trust in order to access the database (eg. AWS public cert for TLS)
-`databaseReadOnlyMode` | `boolean` | Access the database in read-only mode (avoids creating indices or updating content) Useful for debugging a local instance of HordeServer against a production database.
-`shutdownMemoryThreshold` | `integer` | Shutdown the current server process if memory usage reaches this threshold (specified in MB)<br>Usually set to 80-90% of available memory to avoid CLR heap using all of it. If a memory leak was to occur, it's usually better to restart the process rather than to let the GC work harder and harder trying to recoup memory.<br>Should only be used when multiple server processes are running behind a load balancer and one can be safely restarted automatically by the underlying process handler (Docker, Kubernetes, AWS ECS, Supervisor etc). The shutdown behaves similar to receiving a SIGTERM and will wait for outstanding requests to finish.
-`serverPrivateCert` | `string` | Optional PFX certificate to use for encrypting agent SSL traffic. This can be a self-signed certificate, as long as it's trusted by agents.
-`authMethod` | [`AuthMethod`](#authmethod-enum) | Issuer for tokens from the auth provider
-`oidcProfileName` | `string` | Optional profile name to report through the /api/v1/server/auth endpoint. Allows sharing auth tokens between providers configured through the same profile name in OidcToken.exe config files.
-`oidcAuthority` | `string` | Issuer for tokens from the auth provider
-`oidcAudience` | `string` | Audience for validating externally issued tokens
-`oidcClientId` | `string` | Client id for the OIDC authority
-`oidcClientSecret` | `string` | Client secret for the OIDC authority
-`oidcSigninRedirect` | `string` | Optional redirect url provided to OIDC login
-`oidcLocalRedirectUrls` | `string[]` | Optional redirect url provided to OIDC login for external tools (typically to a local server)
-`oidcRequestedScopes` | `string[]` | OpenID Connect scopes to request when signing in
-`oidcClaimNameMapping` | `string[]` | List of fields in /userinfo endpoint to try map to the standard name claim (see System.Security.Claims.ClaimTypes.Name)
-`oidcClaimEmailMapping` | `string[]` | List of fields in /userinfo endpoint to try map to the standard email claim (see System.Security.Claims.ClaimTypes.Email)
-`oidcClaimHordeUserMapping` | `string[]` | List of fields in /userinfo endpoint to try map to the Horde user claim (see HordeClaimTypes.User)
-`oidcClaimHordePerforceUserMapping` | `string[]` | List of fields in /userinfo endpoint to try map to the Horde Perforce user claim (see HordeClaimTypes.PerforceUser)
-`serverUrl` | `string` | Name of this machine
-`jwtIssuer` | `string` | Name of the issuer in bearer tokens from the server
-`jwtExpiryTimeHours` | `integer` | Length of time before JWT tokens expire, in hours
-`adminClaimType` | `string` | The claim type for administrators
-`adminClaimValue` | `string` | Value of the claim type for administrators
-`corsEnabled` | `boolean` | Whether to enable Cors, generally for development purposes
-`corsOrigin` | `string` | Allowed Cors origin
-`enableNewAgentsByDefault` | `boolean` | Whether to automatically enable new agents by default. If false, new agents must manually be enabled before they can take on work.
-`testDataRetainMonths` | `integer` | The number of months to retain test data
-`schedulePollingInterval` | `string` | Interval between rebuilding the schedule queue with a DB query.
-`noResourceBackOffTime` | `string` | Interval between polling for new jobs
-`initiateJobBackOffTime` | `string` | Interval between attempting to assign agents to take on jobs
-`unknownErrorBackOffTime` | `string` | Interval between scheduling jobs when an unknown error occurs
-`redisConnectionConfig` | `string` | Config for connecting to Redis server(s). Setting it to null will disable Redis use and connection See format at https://stackexchange.github.io/StackExchange.Redis/Configuration.html
-`logServiceWriteCacheType` | `string` | Type of write cache to use in log service Currently Supported: "InMemory" or "Redis"
-`logStorage` | [`StorageBackendOptions`](#storagebackendoptions) | Settings for artifact storage
-`artifactStorage` | [`StorageBackendOptions`](#storagebackendoptions) | Settings for artifact storage
-`commitStorage` | [`TreeStoreOptions`](#treestoreoptions) | Configuration of tree storage
-`logJsonToStdOut` | `boolean` | Whether to log json to stdout
-`logSessionRequests` | `boolean` | Whether to log requests to the UpdateSession and QueryServerState RPC endpoints
-`enableLogService` | `boolean` | Whether to enable the hosted LogService running background jobs
-`fleetManagerV2` | [`FleetManagerType`](#fleetmanagertype-enum) | Default fleet manager to use (when not specified by pool)
-`fleetManagerV2Config` | `string` | Config for the fleet manager (serialized JSON)
-`awsAutoScalingQueueUrls` | `string[]` | AWS SQS queue URLs where lifecycle events from EC2 auto-scaling are received
-`disableSchedules` | `boolean` | Whether to run scheduled jobs.
-`scheduleTimeZone` | `string` | Timezone for evaluating schedules
-`slackToken` | `string` | Bot token for interacting with Slack (xoxb-*)
-`slackSocketToken` | `string` | Token for opening a socket to slack (xapp-*)
-`slackAdminToken` | `string` | Admin user token for Slack (xoxp-*). This is only required when using the admin endpoints to invite users.
-`slackUsers` | `string` | Filtered list of slack users to send notifications to. Should be Slack user ids, separated by commas.
-`slackErrorPrefix` | `string` | Prefix to use when reporting errors
-`slackWarningPrefix` | `string` | Prefix to use when reporting warnings
-`configNotificationChannel` | `string` | Channel for sending messages related to config update failures
-`updateStreamsNotificationChannel` | `string` | Channel to send stream notification update failures to
-`jobNotificationChannel` | `string` | Slack channel to send job related notifications to. Multiple channels can be specified, separated by ;
-`agentNotificationChannel` | `string` | Slack channel to send agent related notifications to.
-`dashboardUrl` | `string` | The URl to use for generating links back to the dashboard.
-`helpEmailAddress` | `string` | Help email address that users can contact with issues
-`helpSlackChannel` | `string` | Help slack channel that users can use for issues
-`p4SwarmUrl` | `string` | Url of P4 Swarm installation
-`jiraUsername` | `string` | The Jira service account user name
-`jiraApiToken` | `string` | The Jira service account API token
-`jiraUrl` | `string` | The Uri for the Jira installation
-`sharedDeviceCheckoutDays` | `integer` | The number of days shared device checkouts are held
-`deviceProblemCooldownMinutes` | `integer` | The number of cooldown minutes for device problems
-`deviceReportChannel` | `string` | Channel to send device reports to
-`defaultAgentPoolSizeStrategy` | [`PoolSizeStrategy`](#poolsizestrategy-enum) | Default agent pool sizing strategy for pools that doesn't have one explicitly configured
-`agentPoolScaleOutCooldownSeconds` | `integer` | Scale-out cooldown for auto-scaling agent pools (in seconds). Can be overridden by per-pool settings.
-`agentPoolScaleInCooldownSeconds` | `integer` | Scale-in cooldown for auto-scaling agent pools (in seconds). Can be overridden by per-pool settings.
-`globalThreadPoolMinSize` | `integer` | Set the minimum size of the global thread pool This value has been found in need of tweaking to avoid timeouts with the Redis client during bursts of traffic. Default is 16 for .NET Core CLR. The correct value is dependent on the traffic the Horde Server is receiving. For Epic's internal deployment, this is set to 40.
-`withDatadog` | `boolean` | Whether to enable Datadog integration for tracing
-`withAws` | `boolean` | Whether to enable Amazon Web Services (AWS) specific features
-`configPath` | `string` | Path to the root config file. Relative to the server.json file by default.
-`perforce` | [`PerforceConnectionSettings`](#perforceconnectionsettings)`[]` | Perforce connections for use by the Horde server (not agents)
-`useLocalPerforceEnv` | `boolean` | Whether to use the local Perforce environment
-`perforceConnectionPoolSize` | `integer` | Number of pooled perforce connections to keep
-`enableUpgradeTasks` | `boolean` | Whether to enable the upgrade task source.
-`enableConformTasks` | `boolean` | Whether to enable the conform task source.
-`forceConfigUpdateOnStartup` | `boolean` | Forces configuration data to be read and updated as part of appplication startup, rather than on a schedule. Useful when running locally.
-`openBrowser` | `boolean` | Whether to open a browser on startup
-`bundleCacheDir` | `string` | Directory to use for cache data
-`bundleCacheSize` | `integer` | Maximum size of the storage cache on disk, in megabytes
-`featureFlags` | [`FeatureFlagSettings`](#featureflagsettings) | Experimental features to enable on the server.
-`commits` | [`CommitSettings`](#commitsettings) | Options for the commit service
-`telemetry` | [`BaseTelemetryConfig`](#basetelemetryconfig)`[]` | Settings for sending telemetry events to external services (for example Snowflake, ClickHouse etc)
-`bundledTools` | [`BundledToolConfig`](#bundledtoolconfig)`[]` | Tools bundled along with the server. Data for each tool can be produced using the 'bundle create' command, and should be stored in the Tools directory.
-`openTelemetry` | [`OpenTelemetrySettings`](#opentelemetrysettings) | Options for OpenTelemetry
+Name | Description
+---- | -----------
+`runModes` | [`RunMode`](#runmode-enum)`[]`<br>Modes that the server should run in. Runmodes can be used in a multi-server deployment to limit the operations that a particular instance will try to perform.
+`dataDir` | `string`<br>Override the data directory used by Horde. Defaults to C:\ProgramData\HordeServer on Windows, {AppDir}/Data on other platforms.
+`installed` | `boolean`<br>Whether the server is running in 'installed' mode. In this mode, on Windows, the default data directory will use the common application data folder (C:\ProgramData\Epic\Horde), and configuration data will be read from here and the registry. This setting is overridden to false for local builds from appsettings.Local.json.
+`httpPort` | `integer`<br>Main port for serving HTTP.
+`httpsPort` | `integer`<br>Port for serving HTTP with TLS enabled. Disabled by default.
+`http2Port` | `integer`<br>Dedicated port for serving only HTTP/2.
+`computeTunnelPort` | `integer`<br>Port to listen on for tunneling compute sockets to agents
+`computeTunnelAddress` | `string`<br>What address (host:port) clients should connect to for compute socket tunneling Port may differ from  if Horde server is behind a reverse proxy/firewall
+`databaseConnectionString` | `string`<br>MongoDB connection string
+`databaseName` | `string`<br>MongoDB database name
+`databasePublicCert` | `string`<br>Optional certificate to trust in order to access the database (eg. AWS public cert for TLS)
+`databaseReadOnlyMode` | `boolean`<br>Access the database in read-only mode (avoids creating indices or updating content) Useful for debugging a local instance of HordeServer against a production database.
+`shutdownMemoryThreshold` | `integer`<br>Shutdown the current server process if memory usage reaches this threshold (specified in MB)<br>Usually set to 80-90% of available memory to avoid CLR heap using all of it. If a memory leak was to occur, it's usually better to restart the process rather than to let the GC work harder and harder trying to recoup memory.<br>Should only be used when multiple server processes are running behind a load balancer and one can be safely restarted automatically by the underlying process handler (Docker, Kubernetes, AWS ECS, Supervisor etc). The shutdown behaves similar to receiving a SIGTERM and will wait for outstanding requests to finish.
+`serverPrivateCert` | `string`<br>Optional PFX certificate to use for encrypting agent SSL traffic. This can be a self-signed certificate, as long as it's trusted by agents.
+`authMethod` | [`AuthMethod`](#authmethod-enum)<br>Issuer for tokens from the auth provider
+`oidcProfileName` | `string`<br>Optional profile name to report through the /api/v1/server/auth endpoint. Allows sharing auth tokens between providers configured through the same profile name in OidcToken.exe config files.
+`oidcAuthority` | `string`<br>Issuer for tokens from the auth provider
+`oidcAudience` | `string`<br>Audience for validating externally issued tokens
+`oidcClientId` | `string`<br>Client id for the OIDC authority
+`oidcClientSecret` | `string`<br>Client secret for the OIDC authority
+`oidcSigninRedirect` | `string`<br>Optional redirect url provided to OIDC login
+`oidcLocalRedirectUrls` | `string[]`<br>Optional redirect url provided to OIDC login for external tools (typically to a local server)
+`oidcRequestedScopes` | `string[]`<br>OpenID Connect scopes to request when signing in
+`oidcClaimNameMapping` | `string[]`<br>List of fields in /userinfo endpoint to try map to the standard name claim (see System.Security.Claims.ClaimTypes.Name)
+`oidcClaimEmailMapping` | `string[]`<br>List of fields in /userinfo endpoint to try map to the standard email claim (see System.Security.Claims.ClaimTypes.Email)
+`oidcClaimHordeUserMapping` | `string[]`<br>List of fields in /userinfo endpoint to try map to the Horde user claim (see HordeClaimTypes.User)
+`oidcClaimHordePerforceUserMapping` | `string[]`<br>List of fields in /userinfo endpoint to try map to the Horde Perforce user claim (see HordeClaimTypes.PerforceUser)
+`serverUrl` | `string`<br>Name of this machine
+`jwtIssuer` | `string`<br>Name of the issuer in bearer tokens from the server
+`jwtExpiryTimeHours` | `integer`<br>Length of time before JWT tokens expire, in hours
+`adminClaimType` | `string`<br>The claim type for administrators
+`adminClaimValue` | `string`<br>Value of the claim type for administrators
+`corsEnabled` | `boolean`<br>Whether to enable Cors, generally for development purposes
+`corsOrigin` | `string`<br>Allowed Cors origin
+`enableNewAgentsByDefault` | `boolean`<br>Whether to automatically enable new agents by default. If false, new agents must manually be enabled before they can take on work.
+`testDataRetainMonths` | `integer`<br>The number of months to retain test data
+`schedulePollingInterval` | `string`<br>Interval between rebuilding the schedule queue with a DB query.
+`noResourceBackOffTime` | `string`<br>Interval between polling for new jobs
+`initiateJobBackOffTime` | `string`<br>Interval between attempting to assign agents to take on jobs
+`unknownErrorBackOffTime` | `string`<br>Interval between scheduling jobs when an unknown error occurs
+`redisConnectionConfig` | `string`<br>Config for connecting to Redis server(s). Setting it to null will disable Redis use and connection See format at https://stackexchange.github.io/StackExchange.Redis/Configuration.html
+`logServiceWriteCacheType` | `string`<br>Type of write cache to use in log service Currently Supported: "InMemory" or "Redis"
+`logStorage` | [`StorageBackendOptions`](#storagebackendoptions)<br>Settings for artifact storage
+`artifactStorage` | [`StorageBackendOptions`](#storagebackendoptions)<br>Settings for artifact storage
+`commitStorage` | [`TreeStoreOptions`](#treestoreoptions)<br>Configuration of tree storage
+`logJsonToStdOut` | `boolean`<br>Whether to log json to stdout
+`logSessionRequests` | `boolean`<br>Whether to log requests to the UpdateSession and QueryServerState RPC endpoints
+`enableLogService` | `boolean`<br>Whether to enable the hosted LogService running background jobs
+`fleetManagerV2` | [`FleetManagerType`](#fleetmanagertype-enum)<br>Default fleet manager to use (when not specified by pool)
+`fleetManagerV2Config` | `string`<br>Config for the fleet manager (serialized JSON)
+`awsAutoScalingQueueUrls` | `string[]`<br>AWS SQS queue URLs where lifecycle events from EC2 auto-scaling are received
+`disableSchedules` | `boolean`<br>Whether to run scheduled jobs.
+`scheduleTimeZone` | `string`<br>Timezone for evaluating schedules
+`slackToken` | `string`<br>Bot token for interacting with Slack (xoxb-*)
+`slackSocketToken` | `string`<br>Token for opening a socket to slack (xapp-*)
+`slackAdminToken` | `string`<br>Admin user token for Slack (xoxp-*). This is only required when using the admin endpoints to invite users.
+`slackUsers` | `string`<br>Filtered list of slack users to send notifications to. Should be Slack user ids, separated by commas.
+`slackErrorPrefix` | `string`<br>Prefix to use when reporting errors
+`slackWarningPrefix` | `string`<br>Prefix to use when reporting warnings
+`configNotificationChannel` | `string`<br>Channel for sending messages related to config update failures
+`updateStreamsNotificationChannel` | `string`<br>Channel to send stream notification update failures to
+`jobNotificationChannel` | `string`<br>Slack channel to send job related notifications to. Multiple channels can be specified, separated by ;
+`agentNotificationChannel` | `string`<br>Slack channel to send agent related notifications to.
+`dashboardUrl` | `string`<br>The URl to use for generating links back to the dashboard.
+`helpEmailAddress` | `string`<br>Help email address that users can contact with issues
+`helpSlackChannel` | `string`<br>Help slack channel that users can use for issues
+`p4SwarmUrl` | `string`<br>Url of P4 Swarm installation
+`jiraUsername` | `string`<br>The Jira service account user name
+`jiraApiToken` | `string`<br>The Jira service account API token
+`jiraUrl` | `string`<br>The Uri for the Jira installation
+`sharedDeviceCheckoutDays` | `integer`<br>The number of days shared device checkouts are held
+`deviceProblemCooldownMinutes` | `integer`<br>The number of cooldown minutes for device problems
+`deviceReportChannel` | `string`<br>Channel to send device reports to
+`defaultAgentPoolSizeStrategy` | [`PoolSizeStrategy`](#poolsizestrategy-enum)<br>Default agent pool sizing strategy for pools that doesn't have one explicitly configured
+`agentPoolScaleOutCooldownSeconds` | `integer`<br>Scale-out cooldown for auto-scaling agent pools (in seconds). Can be overridden by per-pool settings.
+`agentPoolScaleInCooldownSeconds` | `integer`<br>Scale-in cooldown for auto-scaling agent pools (in seconds). Can be overridden by per-pool settings.
+`globalThreadPoolMinSize` | `integer`<br>Set the minimum size of the global thread pool This value has been found in need of tweaking to avoid timeouts with the Redis client during bursts of traffic. Default is 16 for .NET Core CLR. The correct value is dependent on the traffic the Horde Server is receiving. For Epic's internal deployment, this is set to 40.
+`withDatadog` | `boolean`<br>Whether to enable Datadog integration for tracing
+`withAws` | `boolean`<br>Whether to enable Amazon Web Services (AWS) specific features
+`configPath` | `string`<br>Path to the root config file. Relative to the server.json file by default.
+`perforce` | [`PerforceConnectionSettings`](#perforceconnectionsettings)`[]`<br>Perforce connections for use by the Horde server (not agents)
+`useLocalPerforceEnv` | `boolean`<br>Whether to use the local Perforce environment
+`perforceConnectionPoolSize` | `integer`<br>Number of pooled perforce connections to keep
+`enableUpgradeTasks` | `boolean`<br>Whether to enable the upgrade task source.
+`enableConformTasks` | `boolean`<br>Whether to enable the conform task source.
+`forceConfigUpdateOnStartup` | `boolean`<br>Forces configuration data to be read and updated as part of appplication startup, rather than on a schedule. Useful when running locally.
+`openBrowser` | `boolean`<br>Whether to open a browser on startup
+`bundleCacheDir` | `string`<br>Directory to use for cache data
+`bundleCacheSize` | `integer`<br>Maximum size of the storage cache on disk, in megabytes
+`featureFlags` | [`FeatureFlagSettings`](#featureflagsettings)<br>Experimental features to enable on the server.
+`commits` | [`CommitSettings`](#commitsettings)<br>Options for the commit service
+`telemetry` | [`BaseTelemetryConfig`](#basetelemetryconfig)`[]`<br>Settings for sending telemetry events to external services (for example Snowflake, ClickHouse etc)
+`bundledTools` | [`BundledToolConfig`](#bundledtoolconfig)`[]`<br>Tools bundled along with the server. Data for each tool can be produced using the 'bundle create' command, and should be stored in the Tools directory.
+`openTelemetry` | [`OpenTelemetrySettings`](#opentelemetrysettings)<br>Options for OpenTelemetry
 
 ## RunMode (Enum)
 
@@ -126,18 +126,18 @@ Name | Description
 
 Common settings object for different providers
 
-Name | Type | Description
----- | ---- | -----------
-`type` | [`StorageBackendType`](#storagebackendtype-enum) | The type of storage backend to use
-`baseDir` | `string` | Base directory for filesystem storage
-`awsBucketName` | `string` | Name of the bucket to use
-`awsBucketPath` | `string` | Base path within the bucket
-`awsCredentials` | [`AwsCredentialsType`](#awscredentialstype-enum) | Type of credentials to use
-`awsRole` | `string` | ARN of a role to assume
-`awsProfile` | `string` | The AWS profile to read credentials form
-`awsRegion` | `string` | Region to connect to
-`azureConnectionString` | `string` | Connection string for Azure
-`azureContainerName` | `string` | Name of the container
+Name | Description
+---- | -----------
+`type` | [`StorageBackendType`](#storagebackendtype-enum)<br>The type of storage backend to use
+`baseDir` | `string`<br>Base directory for filesystem storage
+`awsBucketName` | `string`<br>Name of the bucket to use
+`awsBucketPath` | `string`<br>Base path within the bucket
+`awsCredentials` | [`AwsCredentialsType`](#awscredentialstype-enum)<br>Type of credentials to use
+`awsRole` | `string`<br>ARN of a role to assume
+`awsProfile` | `string`<br>The AWS profile to read credentials form
+`awsRegion` | `string`<br>Region to connect to
+`azureConnectionString` | `string`<br>Connection string for Azure
+`azureContainerName` | `string`<br>Name of the container
 
 ## StorageBackendType (Enum)
 
@@ -165,32 +165,32 @@ Name | Description
 
 Options for storing trees
 
-Name | Type | Description
----- | ---- | -----------
-`bundle` | [`BundleOptions`](#bundleoptions) | Options for creating bundles
-`chunking` | [`ChunkingOptions`](#chunkingoptions) | Options for chunking content
-`type` | [`StorageBackendType`](#storagebackendtype-enum) | The type of storage backend to use
-`baseDir` | `string` | Base directory for filesystem storage
-`awsBucketName` | `string` | Name of the bucket to use
-`awsBucketPath` | `string` | Base path within the bucket
-`awsCredentials` | [`AwsCredentialsType`](#awscredentialstype-enum) | Type of credentials to use
-`awsRole` | `string` | ARN of a role to assume
-`awsProfile` | `string` | The AWS profile to read credentials form
-`awsRegion` | `string` | Region to connect to
-`azureConnectionString` | `string` | Connection string for Azure
-`azureContainerName` | `string` | Name of the container
+Name | Description
+---- | -----------
+`bundle` | [`BundleOptions`](#bundleoptions)<br>Options for creating bundles
+`chunking` | [`ChunkingOptions`](#chunkingoptions)<br>Options for chunking content
+`type` | [`StorageBackendType`](#storagebackendtype-enum)<br>The type of storage backend to use
+`baseDir` | `string`<br>Base directory for filesystem storage
+`awsBucketName` | `string`<br>Name of the bucket to use
+`awsBucketPath` | `string`<br>Base path within the bucket
+`awsCredentials` | [`AwsCredentialsType`](#awscredentialstype-enum)<br>Type of credentials to use
+`awsRole` | `string`<br>ARN of a role to assume
+`awsProfile` | `string`<br>The AWS profile to read credentials form
+`awsRegion` | `string`<br>Region to connect to
+`azureConnectionString` | `string`<br>Connection string for Azure
+`azureContainerName` | `string`<br>Name of the container
 
 ## BundleOptions
 
 Options for configuring a bundle serializer
 
-Name | Type | Description
----- | ---- | -----------
-`maxVersion` | [`BundleVersion`](#bundleversion-enum) | Maximum version number of bundles to write
-`maxBlobSize` | `integer` | Maximum payload size fo a blob
-`compressionFormat` | [`BundleCompressionFormat`](#bundlecompressionformat-enum) | Compression format to use
-`minCompressionPacketSize` | `integer` | Minimum size of a block to be compressed
-`maxWriteQueueLength` | `integer` | Maximum amount of data to store in memory. This includes any background writes as well as bundles being built.
+Name | Description
+---- | -----------
+`maxVersion` | [`BundleVersion`](#bundleversion-enum)<br>Maximum version number of bundles to write
+`maxBlobSize` | `integer`<br>Maximum payload size fo a blob
+`compressionFormat` | [`BundleCompressionFormat`](#bundlecompressionformat-enum)<br>Compression format to use
+`minCompressionPacketSize` | `integer`<br>Minimum size of a block to be compressed
+`maxWriteQueueLength` | `integer`<br>Maximum amount of data to store in memory. This includes any background writes as well as bundles being built.
 
 ## BundleVersion (Enum)
 
@@ -225,31 +225,31 @@ Name | Description
 
 Options for creating file nodes
 
-Name | Type | Description
----- | ---- | -----------
-`leafOptions` | [`LeafChunkedDataNodeOptions`](#leafchunkeddatanodeoptions) | Options for creating leaf nodes
-`interiorOptions` | [`InteriorChunkedDataNodeOptions`](#interiorchunkeddatanodeoptions) | Options for creating interior nodes
+Name | Description
+---- | -----------
+`leafOptions` | [`LeafChunkedDataNodeOptions`](#leafchunkeddatanodeoptions)<br>Options for creating leaf nodes
+`interiorOptions` | [`InteriorChunkedDataNodeOptions`](#interiorchunkeddatanodeoptions)<br>Options for creating interior nodes
 
 ## LeafChunkedDataNodeOptions
 
 Options for creating a specific type of file nodes
 
-Name | Type | Description
----- | ---- | -----------
-`minSize` | `integer` | Minimum chunk size
-`maxSize` | `integer` | Maximum chunk size. Chunks will be split on this boundary if another match is not found.
-`targetSize` | `integer` | Target chunk size for content-slicing
+Name | Description
+---- | -----------
+`minSize` | `integer`<br>Minimum chunk size
+`maxSize` | `integer`<br>Maximum chunk size. Chunks will be split on this boundary if another match is not found.
+`targetSize` | `integer`<br>Target chunk size for content-slicing
 
 ## InteriorChunkedDataNodeOptions
 
 Options for creating interior nodes
 
-Name | Type | Description
----- | ---- | -----------
-`minChildCount` | `integer` | Minimum number of children in each node
-`targetChildCount` | `integer` | Target number of children in each node
-`maxChildCount` | `integer` | Maximum number of children in each node
-`sliceThreshold` | `integer` | Threshold hash value for splitting interior nodes
+Name | Description
+---- | -----------
+`minChildCount` | `integer`<br>Minimum number of children in each node
+`targetChildCount` | `integer`<br>Target number of children in each node
+`maxChildCount` | `integer`<br>Maximum number of children in each node
+`sliceThreshold` | `integer`<br>Threshold hash value for splitting interior nodes
 
 ## FleetManagerType (Enum)
 
@@ -280,21 +280,21 @@ Name | Description
 
 Perforce connection information for use by the Horde server (for reading config files, etc...)
 
-Name | Type | Description
----- | ---- | -----------
-`id` | `string` | Identifier for this server
-`serverAndPort` | `string` | Server and port
-`credentials` | [`PerforceCredentials`](#perforcecredentials) | Credentials for the server
+Name | Description
+---- | -----------
+`id` | `string`<br>Identifier for this server
+`serverAndPort` | `string`<br>Server and port
+`credentials` | [`PerforceCredentials`](#perforcecredentials)<br>Credentials for the server
 
 ## PerforceCredentials
 
 Credentials for a Perforce user
 
-Name | Type | Description
----- | ---- | -----------
-`userName` | `string` | The username
-`password` | `string` | Password for the user
-`ticket` | `string` | Login ticket for the user (will be used instead of password if set)
+Name | Description
+---- | -----------
+`userName` | `string`<br>The username
+`password` | `string`<br>Password for the user
+`ticket` | `string`<br>Login ticket for the user (will be used instead of password if set)
 
 ## FeatureFlagSettings
 
@@ -306,21 +306,21 @@ Once a feature is running in its intended state and is stable, the flag should b
 
 Options for the commit service
 
-Name | Type | Description
----- | ---- | -----------
-`replicateMetadata` | `boolean` | Whether to mirror commit metadata to the database
-`replicateContent` | `boolean` | Whether to mirror commit data to storage
-`bundle` | [`BundleOptions`](#bundleoptions) | Options for how objects are packed together
-`chunking` | [`ChunkingOptions`](#chunkingoptions) | Options for how objects are sliced
+Name | Description
+---- | -----------
+`replicateMetadata` | `boolean`<br>Whether to mirror commit metadata to the database
+`replicateContent` | `boolean`<br>Whether to mirror commit data to storage
+`bundle` | [`BundleOptions`](#bundleoptions)<br>Options for how objects are packed together
+`chunking` | [`ChunkingOptions`](#chunkingoptions)<br>Options for how objects are sliced
 
 ## BaseTelemetryConfig
 
 Configuration for the telemetry sink
 
-Name | Type | Description
----- | ---- | -----------
-`id` | `string` | Unique ID for this sink config (any arbitrary string)
-`type` | [`TelemetrySinkType`](#telemetrysinktype-enum) | Type of telemetry sink
+Name | Description
+---- | -----------
+`id` | `string`<br>Unique ID for this sink config (any arbitrary string)
+`type` | [`TelemetrySinkType`](#telemetrysinktype-enum)<br>Type of telemetry sink
 
 ## TelemetrySinkType (Enum)
 
@@ -337,82 +337,82 @@ Name | Description
 
 Configuration for a tool bundled alongsize the server
 
-Name | Type | Description
----- | ---- | -----------
-`version` | `string` | Version string for the current tool data
-`refName` | `string` | Ref name in the tools directory
-`dataDir` | `string` | Directory containing blob data for this tool. If empty, the tools/{id} folder next to the server will be used.
-`id` | `string` | Unique identifier for the tool
-`name` | `string` | Name of the tool
-`description` | `string` | Description for the tool
-`parentId` | `string` | Tool id to nest this tool under
-`public` | `boolean` | Whether this tool should be exposed for download on a public endpoint without authentication
-`showInUgs` | `boolean` | Whether to show this tool for download in the UGS tools menu
-`showInDashboard` | `boolean` | Whether to show this tool for download in the dashboard
-`namespaceId` | `string` | Default namespace for new deployments of this tool
-`acl` | [`AclConfig`](#aclconfig) | Permissions for the tool
+Name | Description
+---- | -----------
+`version` | `string`<br>Version string for the current tool data
+`refName` | `string`<br>Ref name in the tools directory
+`dataDir` | `string`<br>Directory containing blob data for this tool. If empty, the tools/{id} folder next to the server will be used.
+`id` | `string`<br>Unique identifier for the tool
+`name` | `string`<br>Name of the tool
+`description` | `string`<br>Description for the tool
+`parentId` | `string`<br>Tool id to nest this tool under
+`public` | `boolean`<br>Whether this tool should be exposed for download on a public endpoint without authentication
+`showInUgs` | `boolean`<br>Whether to show this tool for download in the UGS tools menu
+`showInDashboard` | `boolean`<br>Whether to show this tool for download in the dashboard
+`namespaceId` | `string`<br>Default namespace for new deployments of this tool
+`acl` | [`AclConfig`](#aclconfig)<br>Permissions for the tool
 
 ## AclConfig
 
 Parameters to update an ACL
 
-Name | Type | Description
----- | ---- | -----------
-`entries` | [`AclEntryConfig`](#aclentryconfig)`[]` | Entries to replace the existing ACL
-`profiles` | [`AclProfileConfig`](#aclprofileconfig)`[]` | Defines profiles which allow grouping sets of actions into named collections
-`inherit` | `boolean` | Whether to inherit permissions from the parent ACL
-`exceptions` | `string[]` | List of exceptions to the inherited setting
+Name | Description
+---- | -----------
+`entries` | [`AclEntryConfig`](#aclentryconfig)`[]`<br>Entries to replace the existing ACL
+`profiles` | [`AclProfileConfig`](#aclprofileconfig)`[]`<br>Defines profiles which allow grouping sets of actions into named collections
+`inherit` | `boolean`<br>Whether to inherit permissions from the parent ACL
+`exceptions` | `string[]`<br>List of exceptions to the inherited setting
 
 ## AclEntryConfig
 
 Individual entry in an ACL
 
-Name | Type | Description
----- | ---- | -----------
-`claim` | [`AclClaimConfig`](#aclclaimconfig) | Name of the user or group
-`actions` | `string[]` | Array of actions to allow
-`profiles` | `string[]` | List of profiles to grant
+Name | Description
+---- | -----------
+`claim` | [`AclClaimConfig`](#aclclaimconfig)<br>Name of the user or group
+`actions` | `string[]`<br>Array of actions to allow
+`profiles` | `string[]`<br>List of profiles to grant
 
 ## AclClaimConfig
 
 New claim to create
 
-Name | Type | Description
----- | ---- | -----------
-`type` | `string` | The claim type
-`value` | `string` | The claim value
+Name | Description
+---- | -----------
+`type` | `string`<br>The claim type
+`value` | `string`<br>The claim value
 
 ## AclProfileConfig
 
 Configuration for an ACL profile. This defines a preset group of actions which can be given to a user via an ACL entry.
 
-Name | Type | Description
----- | ---- | -----------
-`id` | `string` | Identifier for this profile
-`actions` | `string[]` | Actions to include
-`excludeActions` | `string[]` | Actions to exclude from the inherited actions
-`extends` | `string[]` | Other profiles to extend from
+Name | Description
+---- | -----------
+`id` | `string`<br>Identifier for this profile
+`actions` | `string[]`<br>Actions to include
+`excludeActions` | `string[]`<br>Actions to exclude from the inherited actions
+`extends` | `string[]`<br>Other profiles to extend from
 
 ## OpenTelemetrySettings
 
 OpenTelemetry configuration for collection and sending of traces and metrics.
 
-Name | Type | Description
----- | ---- | -----------
-`enabled` | `boolean` | Whether OpenTelemetry exporting is enabled
-`serviceName` | `string` | Service name
-`serviceNamespace` | `string` | Service namespace
-`serviceVersion` | `string` | Service version
-`enableDatadogCompatibility` | `boolean` | Whether to enrich and format telemetry to fit presentation in Datadog
-`attributes` | `string` `->` `string` | Extra attributes to set
-`enableConsoleExporter` | `boolean` | Whether to enable the console exporter (for debugging purposes)
-`protocolExporters` | `string` `->` [`OpenTelemetryProtocolExporterSettings`](#opentelemetryprotocolexportersettings) | Protocol exporters (key is a unique and arbitrary name)
+Name | Description
+---- | -----------
+`enabled` | `boolean`<br>Whether OpenTelemetry exporting is enabled
+`serviceName` | `string`<br>Service name
+`serviceNamespace` | `string`<br>Service namespace
+`serviceVersion` | `string`<br>Service version
+`enableDatadogCompatibility` | `boolean`<br>Whether to enrich and format telemetry to fit presentation in Datadog
+`attributes` | `string` `->` `string`<br>Extra attributes to set
+`enableConsoleExporter` | `boolean`<br>Whether to enable the console exporter (for debugging purposes)
+`protocolExporters` | `string` `->` [`OpenTelemetryProtocolExporterSettings`](#opentelemetryprotocolexportersettings)<br>Protocol exporters (key is a unique and arbitrary name)
 
 ## OpenTelemetryProtocolExporterSettings
 
 Configuration for an OpenTelemetry exporter
 
-Name | Type | Description
----- | ---- | -----------
-`endpoint` | `string` | Endpoint URL. Usually differs depending on protocol used.
-`protocol` | `string` | Protocol for the exporter ('grpc' or 'httpprotobuf')
+Name | Description
+---- | -----------
+`endpoint` | `string`<br>Endpoint URL. Usually differs depending on protocol used.
+`protocol` | `string`<br>Protocol for the exporter ('grpc' or 'httpprotobuf')
