Horde: Fix separate RAM slots being reported separately in agent properties.

diff --git a/Engine/Source/Programs/Horde/Horde.Agent/Services/CapabilitiesService.cs b/Engine/Source/Programs/Horde/Horde.Agent/Services/CapabilitiesService.cs
--- a/Engine/Source/Programs/Horde/Horde.Agent/Services/CapabilitiesService.cs
+++ b/Engine/Source/Programs/Horde/Horde.Agent/Services/CapabilitiesService.cs
@@ -142,9 +142,9 @@ namespace Horde.Agent.Services
 					AddCpuInfo(primaryDevice, cpuNameToCount, totalLogicalCores, totalPhysicalCores);
 
 					// Add RAM info
+					ulong totalCapacity = 0;
 					foreach (CimInstance instance in session.QueryInstances(QueryNamespace, QueryDialect, "select Capacity from Win32_PhysicalMemory"))
 					{
-						ulong totalCapacity = 0;
 						foreach (CimProperty property in instance.CimInstanceProperties)
 						{
 							if (property.Name.Equals("Capacity", StringComparison.OrdinalIgnoreCase) && property.Value is ulong capacity)
@@ -152,12 +152,8 @@ namespace Horde.Agent.Services
 								totalCapacity += capacity;
 							}
 						}
-
-						if (totalCapacity > 0)
-						{
-							primaryDevice.Properties.Add($"RAM={totalCapacity / (1024 * 1024 * 1024)}");
-						}
 					}
+					primaryDevice.Properties.Add($"RAM={totalCapacity / (1024 * 1024 * 1024)}");
 
 					// Add GPU info
 					int index = 0;
