Horde: Log a warning whenever an obsolete endpoint is used.

diff --git a/Engine/Source/Programs/Horde/Horde.Server/Startup.cs b/Engine/Source/Programs/Horde/Horde.Server/Startup.cs
--- a/Engine/Source/Programs/Horde/Horde.Server/Startup.cs
+++ b/Engine/Source/Programs/Horde/Horde.Server/Startup.cs
@@ -99,8 +99,11 @@ using Microsoft.AspNetCore.DataProtection;
 using Microsoft.AspNetCore.Hosting;
 using Microsoft.AspNetCore.Hosting.Server.Features;
 using Microsoft.AspNetCore.Http;
+using Microsoft.AspNetCore.Http.Extensions;
 using Microsoft.AspNetCore.HttpOverrides;
 using Microsoft.AspNetCore.Mvc;
+using Microsoft.AspNetCore.Mvc.Controllers;
+using Microsoft.AspNetCore.Mvc.Filters;
 using Microsoft.AspNetCore.Mvc.ModelBinding;
 using Microsoft.Extensions.Configuration;
 using Microsoft.Extensions.DependencyInjection;
@@ -126,6 +129,7 @@ using StatusCode = Grpc.Core.StatusCode;
 namespace Horde.Server
 {
 	using ContentHash = EpicGames.Core.ContentHash;
+	using ILogger = Microsoft.Extensions.Logging.ILogger;
 	using ReferenceResolver = Horde.Server.Ddc.ReferenceResolver;
 
 	class Startup
@@ -294,6 +298,23 @@ namespace Horde.Server
 			}
 		}
 
+		class ObsoleteLoggingFilter : IActionFilter
+		{
+			public void OnActionExecuted(ActionExecutedContext context)
+			{
+			}
+
+			public void OnActionExecuting(ActionExecutingContext context)
+			{
+				ControllerActionDescriptor? actionDescriptor = context.ActionDescriptor as ControllerActionDescriptor;
+				if (actionDescriptor?.MethodInfo.GetCustomAttribute<ObsoleteAttribute>() != null)
+				{
+					ILogger? logger = context.HttpContext.RequestServices.GetService<ILogger<ObsoleteLoggingFilter>>();
+					logger?.LogWarning("Using obsolete endpoint for request: {Request}", context.HttpContext.Request.GetDisplayUrl());
+				}
+			}
+		}
+
 		static IObjectStore CreateObjectStore(IServiceProvider sp, StorageBackendOptions options)
 		{
 			switch (options.Type ?? StorageBackendType.FileSystem)
@@ -903,7 +924,8 @@ namespace Horde.Server
 			});
 
 			services.AddMvc().AddJsonOptions(options => ConfigureJsonSerializer(options.JsonSerializerOptions));
-			services.AddControllersWithViews().AddRazorRuntimeCompilation();
+			services.AddControllersWithViews(options => options.Filters.Add(new ObsoleteLoggingFilter()))
+				.AddRazorRuntimeCompilation();
 
 			services.AddControllers(options =>
 			{
