Horde: Deprecate API methods which are no longer wanted.

#rnx

diff --git a/Engine/Source/Programs/Horde/Horde.Server/Agents/Enrollment/EnrollmentController.cs b/Engine/Source/Programs/Horde/Horde.Server/Agents/Enrollment/EnrollmentController.cs
--- a/Engine/Source/Programs/Horde/Horde.Server/Agents/Enrollment/EnrollmentController.cs
+++ b/Engine/Source/Programs/Horde/Horde.Server/Agents/Enrollment/EnrollmentController.cs
@@ -1,5 +1,6 @@
 // Copyright Epic Games, Inc. All Rights Reserved.
 
+using System;
 using System.Collections.Generic;
 using System.Threading;
 using System.Threading.Tasks;
@@ -40,7 +41,6 @@ namespace Horde.Server.Agents.Enrollment
 		/// <returns>List of matching agents</returns>
 		[HttpGet]
 		[Route("/api/v1/enrollment")]
-		[Route("/api/v1/registration")]
 		[ProducesResponseType(typeof(GetPendingAgentsResponse), 200)]
 		public async Task<ActionResult<object>> GetPendingAgentsAsync([FromQuery] PropertyFilter? filter = null, CancellationToken cancellationToken = default)
 		{
@@ -55,6 +55,18 @@ namespace Horde.Server.Agents.Enrollment
 			return PropertyFilter.Apply(response, filter);
 		}
 
+		/// <summary>
+		/// Legacy endpoint for enrollment
+		/// </summary>
+		[HttpGet]
+		[Route("/api/v1/registration")]
+		[Obsolete("Use api/v1/enrollment instead")]
+		[ProducesResponseType(typeof(GetPendingAgentsResponse), 200)]
+		public Task<ActionResult<object>> LegacyGetPendingAgentsAsync([FromQuery] PropertyFilter? filter = null, CancellationToken cancellationToken = default)
+		{
+			return GetPendingAgentsAsync(filter, cancellationToken);
+		}
+
 		/// <summary>
 		/// Retrieve information about a specific agent
 		/// </summary>
@@ -62,7 +74,6 @@ namespace Horde.Server.Agents.Enrollment
 		/// <param name="cancellationToken">Cancellation token for the operation</param>
 		[HttpPost]
 		[Route("/api/v1/enrollment")]
-		[Route("/api/v1/registration")]
 		public async Task<ActionResult> ApproveAgentsAsync([FromBody] ApproveAgentsRequest request, CancellationToken cancellationToken)
 		{
 			if (!_globalConfig.Value.Authorize(AgentAclAction.CreateAgent, User))
@@ -77,5 +88,16 @@ namespace Horde.Server.Agents.Enrollment
 
 			return Ok();
 		}
+
+		/// <summary>
+		/// Legacy registration method
+		/// </summary>
+		[HttpPost]
+		[Route("/api/v1/registration")]
+		[Obsolete("Use api/v1/enrollment instead")]
+		public Task<ActionResult> LegacyApproveAgentsAsync([FromBody] ApproveAgentsRequest request, CancellationToken cancellationToken)
+		{
+			return ApproveAgentsAsync(request, cancellationToken);
+		}
 	}
 }
diff --git a/Engine/Source/Programs/Horde/Horde.Server/Artifacts/ArtifactsController.cs b/Engine/Source/Programs/Horde/Horde.Server/Artifacts/ArtifactsController.cs
--- a/Engine/Source/Programs/Horde/Horde.Server/Artifacts/ArtifactsController.cs
+++ b/Engine/Source/Programs/Horde/Horde.Server/Artifacts/ArtifactsController.cs
@@ -199,7 +199,6 @@ namespace Horde.Server.Artifacts
 		/// <returns>Information about all the artifacts</returns>
 		[HttpGet]
 		[Route("/api/v2/artifacts/{id}/blobs/{*locator}")]
-		[Route("/api/v2/artifacts/{id}/bundles/{*locator}")]
 		public async Task<ActionResult> ReadArtifactBlobAsync(ArtifactId id, BlobLocator locator, CancellationToken cancellationToken = default)
 		{
 			IArtifact? artifact = await _artifactCollection.GetAsync(id, cancellationToken);
diff --git a/Engine/Source/Programs/Horde/Horde.Server/Compute/ComputeController.cs b/Engine/Source/Programs/Horde/Horde.Server/Compute/ComputeController.cs
--- a/Engine/Source/Programs/Horde/Horde.Server/Compute/ComputeController.cs
+++ b/Engine/Source/Programs/Horde/Horde.Server/Compute/ComputeController.cs
@@ -24,7 +24,7 @@ namespace Horde.Server.Compute
 	[ApiController]
 	[Authorize]
 	[Route("[controller]")]
-	public class ComputeControllerV2 : HordeControllerBase
+	public class ComputeController : HordeControllerBase
 	{
 		readonly ComputeService _computeService;
 		readonly IOptionsSnapshot<GlobalConfig> _globalConfig;
@@ -32,7 +32,7 @@ namespace Horde.Server.Compute
 		/// <summary>
 		/// Constructor
 		/// </summary>
-		public ComputeControllerV2(ComputeService computeService, IOptionsSnapshot<GlobalConfig> globalConfig)
+		public ComputeController(ComputeService computeService, IOptionsSnapshot<GlobalConfig> globalConfig)
 		{
 			_computeService = computeService;
 			_globalConfig = globalConfig;
diff --git a/Engine/Source/Programs/Horde/Horde.Server/Ddc/BlobsController.cs b/Engine/Source/Programs/Horde/Horde.Server/Ddc/BlobsController.cs
--- a/Engine/Source/Programs/Horde/Horde.Server/Ddc/BlobsController.cs
+++ b/Engine/Source/Programs/Horde/Horde.Server/Ddc/BlobsController.cs
@@ -26,6 +26,7 @@ namespace Horde.Server.Ddc
 	using IDiagnosticContext = Serilog.IDiagnosticContext;
 
 	[ApiController]
+	[Tags("DDC Blobs")]
 	[Route("api/v1/s", Order = 1)]
 	[Route("api/v1/blobs", Order = 0)]
 	[Authorize]
diff --git a/Engine/Source/Programs/Horde/Horde.Server/Ddc/CompressedBlobController.cs b/Engine/Source/Programs/Horde/Horde.Server/Ddc/CompressedBlobController.cs
--- a/Engine/Source/Programs/Horde/Horde.Server/Ddc/CompressedBlobController.cs
+++ b/Engine/Source/Programs/Horde/Horde.Server/Ddc/CompressedBlobController.cs
@@ -13,6 +13,7 @@ using EpicGames.Horde.Storage;
 using EpicGames.Serialization;
 using Horde.Server.Storage;
 using Microsoft.AspNetCore.Authorization;
+using Microsoft.AspNetCore.Http;
 using Microsoft.AspNetCore.Mvc;
 using Microsoft.Extensions.Primitives;
 
@@ -25,6 +26,7 @@ namespace Horde.Server.Ddc
 	[ApiController]
 	[Authorize]
 	[Route("api/v1/compressed-blobs")]
+	[Tags("DDC Compressed Blobs")]
 	public class CompressedBlobController : ControllerBase
 	{
 		private readonly IBlobService _storage;
diff --git a/Engine/Source/Programs/Horde/Horde.Server/Ddc/ContentIdController.cs b/Engine/Source/Programs/Horde/Horde.Server/Ddc/ContentIdController.cs
--- a/Engine/Source/Programs/Horde/Horde.Server/Ddc/ContentIdController.cs
+++ b/Engine/Source/Programs/Horde/Horde.Server/Ddc/ContentIdController.cs
@@ -6,6 +6,7 @@ using EpicGames.AspNet;
 using EpicGames.Horde.Storage;
 using Horde.Server.Storage;
 using Microsoft.AspNetCore.Authorization;
+using Microsoft.AspNetCore.Http;
 using Microsoft.AspNetCore.Mvc;
 
 #pragma warning disable CS1591
@@ -15,6 +16,7 @@ namespace Horde.Server.Ddc
 	[ApiController]
 	[FormatFilter]
 	[Produces(MediaTypeNames.Application.Json, MediaTypeNames.Application.Octet, CustomMediaTypeNames.UnrealCompactBinary)]
+	[Tags("DDC ContentId")]
 	[Route("api/v1/content-id")]
 	[Authorize]
 	public class ContentIdController : ControllerBase
diff --git a/Engine/Source/Programs/Horde/Horde.Server/Ddc/ObjectsController.cs b/Engine/Source/Programs/Horde/Horde.Server/Ddc/ObjectsController.cs
--- a/Engine/Source/Programs/Horde/Horde.Server/Ddc/ObjectsController.cs
+++ b/Engine/Source/Programs/Horde/Horde.Server/Ddc/ObjectsController.cs
@@ -14,6 +14,7 @@ using EpicGames.Horde.Storage;
 using EpicGames.Serialization;
 using Horde.Server.Storage;
 using Microsoft.AspNetCore.Authorization;
+using Microsoft.AspNetCore.Http;
 using Microsoft.AspNetCore.Mvc;
 using Microsoft.Extensions.Logging;
 
@@ -25,6 +26,7 @@ namespace Horde.Server.Ddc
 
 	[ApiController]
 	[Route("api/v1/objects", Order = 0)]
+	[Tags("DDC Objects")]
 	[Authorize]
 	[Produces(CustomMediaTypeNames.UnrealCompactBinary, MediaTypeNames.Application.Json)]
 	public class ObjectsController : ControllerBase
diff --git a/Engine/Source/Programs/Horde/Horde.Server/Ddc/RefsController.cs b/Engine/Source/Programs/Horde/Horde.Server/Ddc/RefsController.cs
--- a/Engine/Source/Programs/Horde/Horde.Server/Ddc/RefsController.cs
+++ b/Engine/Source/Programs/Horde/Horde.Server/Ddc/RefsController.cs
@@ -38,6 +38,7 @@ namespace Horde.Server.Ddc
 	[FormatFilter]
 	[Produces(MediaTypeNames.Application.Json, MediaTypeNames.Application.Octet, CustomMediaTypeNames.UnrealCompactBinary)]
 	[Route("api/v1/refs")]
+	[Tags("DDC Refs")]
 	[Authorize]
 	public class RefsController : ControllerBase
 	{
diff --git a/Engine/Source/Programs/Horde/Horde.Server/Jobs/Artifacts/ArtifactsControllerV1.cs b/Engine/Source/Programs/Horde/Horde.Server/Jobs/Artifacts/ArtifactsControllerV1.cs
--- a/Engine/Source/Programs/Horde/Horde.Server/Jobs/Artifacts/ArtifactsControllerV1.cs
+++ b/Engine/Source/Programs/Horde/Horde.Server/Jobs/Artifacts/ArtifactsControllerV1.cs
@@ -30,6 +30,7 @@ namespace Horde.Server.Jobs.Artifacts
 	/// Controller for the /api/artifacts endpoint
 	/// </summary>
 	[ApiController]
+	[Obsolete("Use /api/v2/artifacts instead")]
 	[Route("[controller]")]
 	public class ArtifactsControllerV1 : ControllerBase
 	{
diff --git a/Engine/Source/Programs/Horde/Horde.Server/Jobs/JobsController.cs b/Engine/Source/Programs/Horde/Horde.Server/Jobs/JobsController.cs
--- a/Engine/Source/Programs/Horde/Horde.Server/Jobs/JobsController.cs
+++ b/Engine/Source/Programs/Horde/Horde.Server/Jobs/JobsController.cs
@@ -916,6 +916,7 @@ namespace Horde.Server.Jobs
 		/// <param name="cancellationToken">Cancellation token for the operation</param>
 		/// <returns>Id of the new job</returns>
 		[HttpPost]
+		[Obsolete("Modifying graph through REST API is not supported")]
 		[Route("/api/v1/jobs/{jobId}/groups")]
 		public async Task<ActionResult> CreateGroupsAsync(JobId jobId, [FromBody] List<NewGroup> requests, CancellationToken cancellationToken = default)
 		{
@@ -956,6 +957,7 @@ namespace Horde.Server.Jobs
 		/// <param name="cancellationToken">Cancellation token for the operation</param>
 		/// <returns>List of nodes to be executed</returns>
 		[HttpGet]
+		[Obsolete("Query entire job instead")]
 		[Route("/api/v1/jobs/{jobId}/groups")]
 		[ProducesResponseType(typeof(List<GetGroupResponse>), 200)]
 		public async Task<ActionResult<List<object>>> GetGroupsAsync(JobId jobId, [FromQuery] PropertyFilter? filter = null, CancellationToken cancellationToken = default)
@@ -989,6 +991,7 @@ namespace Horde.Server.Jobs
 		/// <param name="cancellationToken">Cancellation token for the operation</param>
 		/// <returns>List of nodes to be executed</returns>
 		[HttpGet]
+		[Obsolete("Query entire job instead")]
 		[Route("/api/v1/jobs/{jobId}/groups/{groupIdx}")]
 		[ProducesResponseType(typeof(GetGroupResponse), 200)]
 		public async Task<ActionResult<object>> GetGroupAsync(JobId jobId, int groupIdx, [FromQuery] PropertyFilter? filter = null, CancellationToken cancellationToken = default)
@@ -1025,6 +1028,7 @@ namespace Horde.Server.Jobs
 		/// <param name="groupIdx">Index of the group containing the node to update</param>
 		/// <param name="filter">Filter for the properties to return</param>
 		[HttpGet]
+		[Obsolete("Query entire job instead")]
 		[Route("/api/v1/jobs/{jobId}/groups/{groupIdx}/nodes")]
 		[ProducesResponseType(typeof(List<GetNodeResponse>), 200)]
 		public async Task<ActionResult<List<object>>> GetNodesAsync(JobId jobId, int groupIdx, [FromQuery] PropertyFilter? filter = null)
@@ -1062,6 +1066,7 @@ namespace Horde.Server.Jobs
 		/// <param name="nodeIdx">Index of the node to update</param>
 		/// <param name="filter">Filter for the properties to return</param>
 		[HttpGet]
+		[Obsolete("Query entire job instead")]
 		[Route("/api/v1/jobs/{jobId}/groups/{groupIdx}/nodes/{nodeIdx}")]
 		[ProducesResponseType(typeof(GetNodeResponse), 200)]
 		public async Task<ActionResult<object>> GetNodeAsync(JobId jobId, int groupIdx, int nodeIdx, [FromQuery] PropertyFilter? filter = null)
@@ -1103,6 +1108,7 @@ namespace Horde.Server.Jobs
 		/// <param name="cancellationToken">Cancellation token for the operation</param>
 		/// <returns>List of nodes to be executed</returns>
 		[HttpGet]
+		[Obsolete("Query entire job instead")]
 		[Route("/api/v1/jobs/{jobId}/batches")]
 		[ProducesResponseType(typeof(List<GetBatchResponse>), 200)]
 		public async Task<ActionResult<List<object>>> GetBatchesAsync(JobId jobId, [FromQuery] PropertyFilter? filter = null, CancellationToken cancellationToken = default)
@@ -1141,6 +1147,7 @@ namespace Horde.Server.Jobs
 		/// <param name="batchId">Unique id for the step</param>
 		/// <param name="request">Updates to apply to the node</param>
 		[HttpPut]
+		[Obsolete("Query entire job instead")]
 		[Route("/api/v1/jobs/{jobId}/batches/{batchId}")]
 		public async Task<ActionResult> UpdateBatchAsync(JobId jobId, JobStepBatchId batchId, [FromBody] UpdateBatchRequest request)
 		{
@@ -1181,6 +1188,7 @@ namespace Horde.Server.Jobs
 		/// <param name="cancellationToken">Cancellation token for the operation</param>
 		/// <returns>List of nodes to be executed</returns>
 		[HttpGet]
+		[Obsolete("Query entire job instead")]
 		[Route("/api/v1/jobs/{jobId}/batches/{batchId}")]
 		[ProducesResponseType(typeof(GetBatchResponse), 200)]
 		public async Task<ActionResult<object>> GetBatchAsync(JobId jobId, JobStepBatchId batchId, [FromQuery] PropertyFilter? filter = null, CancellationToken cancellationToken = default)
@@ -1223,6 +1231,7 @@ namespace Horde.Server.Jobs
 		/// <param name="cancellationToken">Cancellation token for the operation</param>
 		/// <returns>List of nodes to be executed</returns>
 		[HttpGet]
+		[Obsolete("Query entire job instead")]
 		[Route("/api/v1/jobs/{jobId}/batches/{batchId}/steps")]
 		[ProducesResponseType(typeof(List<GetStepResponse>), 200)]
 		public async Task<ActionResult<List<object>>> GetStepsAsync(JobId jobId, JobStepBatchId batchId, [FromQuery] PropertyFilter? filter = null, CancellationToken cancellationToken = default)
@@ -1357,6 +1366,7 @@ namespace Horde.Server.Jobs
 		/// <param name="cancellationToken">Cancellation token for the operation</param>
 		/// <returns>List of nodes to be executed</returns>
 		[HttpGet]
+		[Obsolete("Query entire job instead")]
 		[Route("/api/v1/jobs/{jobId}/batches/{batchId}/steps/{stepId}")]
 		[ProducesResponseType(typeof(GetStepResponse), 200)]
 		public async Task<ActionResult<object>> GetStepAsync(JobId jobId, JobStepBatchId batchId, JobStepId stepId, [FromQuery] PropertyFilter? filter = null, CancellationToken cancellationToken = default)
diff --git a/Engine/Source/Programs/Horde/Horde.Server/Parameters/ParametersController.cs b/Engine/Source/Programs/Horde/Horde.Server/Parameters/ParametersController.cs
--- a/Engine/Source/Programs/Horde/Horde.Server/Parameters/ParametersController.cs
+++ b/Engine/Source/Programs/Horde/Horde.Server/Parameters/ParametersController.cs
@@ -28,7 +28,7 @@ namespace Horde.Server.Parameters
 		}
 
 		/// <summary>
-		/// Query all the parameters
+		/// Query side-wide parameters published for automatic configuration of external tools.
 		/// </summary>
 		/// <param name="path">Base path for the object to return</param>
 		/// <param name="filter">Filter for the properties to return</param>
diff --git a/Engine/Source/Programs/Horde/Horde.Server/Perforce/PerforceController.cs b/Engine/Source/Programs/Horde/Horde.Server/Perforce/PerforceController.cs
--- a/Engine/Source/Programs/Horde/Horde.Server/Perforce/PerforceController.cs
+++ b/Engine/Source/Programs/Horde/Horde.Server/Perforce/PerforceController.cs
@@ -8,6 +8,7 @@ using System.Threading;
 using System.Threading.Tasks;
 using Horde.Server.Server;
 using Microsoft.AspNetCore.Authorization;
+using Microsoft.AspNetCore.Http;
 using Microsoft.AspNetCore.Mvc;
 using Microsoft.Extensions.Logging;
 using Microsoft.Extensions.Options;
@@ -251,6 +252,7 @@ namespace Horde.Server.Perforce
 	/// Controller for Perforce triggers and callbacks
 	/// </summary>
 	[ApiController]
+	[Tags("Perforce")]
 	[Route("[controller]")]
 	public class PublicPerforceController : ControllerBase
 	{
diff --git a/Engine/Source/Programs/Horde/Horde.Server/Server/DebugController.cs b/Engine/Source/Programs/Horde/Horde.Server/Server/DebugController.cs
--- a/Engine/Source/Programs/Horde/Horde.Server/Server/DebugController.cs
+++ b/Engine/Source/Programs/Horde/Horde.Server/Server/DebugController.cs
@@ -251,6 +251,7 @@ namespace Horde.Server.Server
 	/// </summary>
 	[ApiController]
 	[Authorize]
+	[Tags("Debug")]
 	public class SecureDebugController : HordeControllerBase
 	{
 		private static readonly Random s_random = new();
@@ -334,7 +335,7 @@ namespace Horde.Server.Server
 		/// Converts all legacy pools into config entries
 		/// </summary>
 		[HttpGet]
-		[Route("/api/v1/server/debug/aclscopes")]
+		[Route("/api/v1/debug/aclscopes")]
 		public ActionResult<object> GetAclScopes()
 		{
 			if (!_globalConfig.Value.Authorize(ServerAclAction.Debug, User))
@@ -349,7 +350,7 @@ namespace Horde.Server.Server
 		/// Returns the fully parsed config object.
 		/// </summary>
 		[HttpGet]
-		[Route("/api/v1/server/debug/appsettings")]
+		[Route("/api/v1/debug/appsettings")]
 		public ActionResult<object> GetAppSettings()
 		{
 			if (!_globalConfig.Value.Authorize(ServerAclAction.Debug, User))
@@ -364,7 +365,7 @@ namespace Horde.Server.Server
 		/// Returns the fully parsed config object.
 		/// </summary>
 		[HttpGet]
-		[Route("/api/v1/server/debug/config")]
+		[Route("/api/v1/debug/config")]
 		public ActionResult<object> GetConfig()
 		{
 			if (!_globalConfig.Value.Authorize(ServerAclAction.Debug, User))
diff --git a/Engine/Source/Programs/Horde/Horde.Server/Server/ServerStatusController.cs b/Engine/Source/Programs/Horde/Horde.Server/Server/ServerStatusController.cs
--- a/Engine/Source/Programs/Horde/Horde.Server/Server/ServerStatusController.cs
+++ b/Engine/Source/Programs/Horde/Horde.Server/Server/ServerStatusController.cs
@@ -6,6 +6,7 @@ using System.Linq;
 using System.Threading.Tasks;
 using EpicGames.Horde.Server;
 using Microsoft.AspNetCore.Authorization;
+using Microsoft.AspNetCore.Http;
 using Microsoft.AspNetCore.Mvc;
 using Microsoft.Extensions.Diagnostics.HealthChecks;
 
@@ -59,6 +60,7 @@ public class ServerStatusUpdatesViewModel
 /// </summary>
 [ApiController]
 [Authorize]
+[Tags("Server")]
 public class ServerStatusController : Controller
 {
 	private readonly ServerStatusService _serverStatus;
diff --git a/Engine/Source/Programs/Horde/Horde.Server/Startup.cs b/Engine/Source/Programs/Horde/Horde.Server/Startup.cs
--- a/Engine/Source/Programs/Horde/Horde.Server/Startup.cs
+++ b/Engine/Source/Programs/Horde/Horde.Server/Startup.cs
@@ -307,7 +307,8 @@ namespace Horde.Server
 			public void OnActionExecuting(ActionExecutingContext context)
 			{
 				ControllerActionDescriptor? actionDescriptor = context.ActionDescriptor as ControllerActionDescriptor;
-				if (actionDescriptor?.MethodInfo.GetCustomAttribute<ObsoleteAttribute>() != null)
+				if (actionDescriptor?.MethodInfo.GetCustomAttribute<ObsoleteAttribute>() != null
+					|| actionDescriptor?.MethodInfo?.DeclaringType?.GetCustomAttribute<ObsoleteAttribute>() != null)
 				{
 					ILogger? logger = context.HttpContext.RequestServices.GetService<ILogger<ObsoleteLoggingFilter>>();
 					logger?.LogWarning("Using obsolete endpoint for request: {Request}", context.HttpContext.Request.GetDisplayUrl());
diff --git a/Engine/Source/Programs/Horde/Horde.Server/Storage/StorageController.cs b/Engine/Source/Programs/Horde/Horde.Server/Storage/StorageController.cs
--- a/Engine/Source/Programs/Horde/Horde.Server/Storage/StorageController.cs
+++ b/Engine/Source/Programs/Horde/Horde.Server/Storage/StorageController.cs
@@ -61,7 +61,6 @@ namespace Horde.Server.Storage
 		/// <param name="cancellationToken">Cancellation token for the operation</param>
 		[HttpPost]
 		[Route("/api/v1/storage/{namespaceId}/blobs")]
-		[Route("/api/v1/storage/{namespaceId}/bundles")]
 		public async Task<ActionResult<WriteBlobResponse>> WriteBlobAsync(NamespaceId namespaceId, IFormFile? file, [FromForm] string? prefix = default, CancellationToken cancellationToken = default)
 		{
 			IStorageBackend? storageBackend = _storageService.TryCreateBackend(namespaceId);
@@ -113,7 +112,6 @@ namespace Horde.Server.Storage
 		/// <param name="cancellationToken">Cancellation token for the operation</param>
 		[HttpGet]
 		[Route("/api/v1/storage/{namespaceId}/blobs/{*locator}")]
-		[Route("/api/v1/storage/{namespaceId}/bundles/{*locator}")]
 		public async Task<ActionResult> ReadBlobAsync(NamespaceId namespaceId, BlobLocator locator, CancellationToken cancellationToken = default)
 		{
 			IStorageBackend? backend = _storageService.TryCreateBackend(namespaceId);
diff --git a/Engine/Source/Programs/Horde/Horde.Server/Tools/ToolsController.cs b/Engine/Source/Programs/Horde/Horde.Server/Tools/ToolsController.cs
--- a/Engine/Source/Programs/Horde/Horde.Server/Tools/ToolsController.cs
+++ b/Engine/Source/Programs/Horde/Horde.Server/Tools/ToolsController.cs
@@ -160,6 +160,7 @@ namespace Horde.Server.Tools
 	/// Public methods available without authorization (or with very custom authorization)
 	/// </summary>
 	[ApiController]
+	[Tags("Tools")]
 	public class PublicToolsController : HordeControllerBase
 	{
 		readonly IToolCollection _toolCollection;
