Horde: Remove support for overriding update channel through agent properties. This is now driven by conditions in the server config file rather than the agent itself.

#rnx

diff --git a/Engine/Source/Programs/Horde/Horde.Agent/AgentSettings.cs b/Engine/Source/Programs/Horde/Horde.Agent/AgentSettings.cs
--- a/Engine/Source/Programs/Horde/Horde.Agent/AgentSettings.cs
+++ b/Engine/Source/Programs/Horde/Horde.Agent/AgentSettings.cs
@@ -3,7 +3,6 @@
 using System.ComponentModel.DataAnnotations;
 using EpicGames.Core;
 using EpicGames.Horde;
-using EpicGames.Horde.Tools;
 using Microsoft.Extensions.Configuration;
 
 namespace Horde.Agent
@@ -206,11 +205,6 @@ namespace Horde.Agent
 		/// </summary>
 		public string Executor { get; set; } = Execution.WorkspaceExecutor.Name;
 
-		/// <summary>
-		/// Default channel for downloading updates
-		/// </summary>
-		public ToolId UpdateChannel { get; set; }
-
 		/// <summary>
 		/// Settings for the local executor
 		/// </summary>
diff --git a/Engine/Source/Programs/Horde/Horde.Agent/Commands/Utilities/SetUpdateChannelCommand.cs b/Engine/Source/Programs/Horde/Horde.Agent/Commands/Utilities/SetUpdateChannelCommand.cs
deleted file mode 100644
diff --git a/Engine/Source/Programs/Horde/Horde.Agent/Services/CapabilitiesService.cs b/Engine/Source/Programs/Horde/Horde.Agent/Services/CapabilitiesService.cs
--- a/Engine/Source/Programs/Horde/Horde.Agent/Services/CapabilitiesService.cs
+++ b/Engine/Source/Programs/Horde/Horde.Agent/Services/CapabilitiesService.cs
@@ -437,12 +437,6 @@ namespace Horde.Agent.Services
 			// Used during the transition period over from multi-platform, non-self-contained agent packages.
 			agent.Properties.Add($"SelfContained={AgentApp.IsSelfContained}");
 
-			// Channel to use for agent updates
-			if (!_settings.UpdateChannel.IsEmpty)
-			{
-				agent.Properties.Add($"UpdateChannel={_settings.UpdateChannel}");
-			}
-
 			// Add any additional properties from the config file
 			agent.Properties.AddRange(_settings.Properties.Select(kvp => $"{kvp.Key}={kvp.Value}"));
 			return agent;
diff --git a/Engine/Source/Programs/Horde/Horde.Server/Agents/IAgent.cs b/Engine/Source/Programs/Horde/Horde.Server/Agents/IAgent.cs
--- a/Engine/Source/Programs/Horde/Horde.Server/Agents/IAgent.cs
+++ b/Engine/Source/Programs/Horde/Horde.Server/Agents/IAgent.cs
@@ -653,16 +653,6 @@ namespace Horde.Server.Agents
 		/// <returns>Identifier for the tool that the agent should be using</returns>
 		public static ToolId GetSoftwareToolId(this IAgent agent, GlobalConfig globalConfig)
 		{
-			string? updateChannel = agent.GetPropertyValues("UpdateChannel").FirstOrDefault();
-			if (!String.IsNullOrEmpty(updateChannel))
-			{
-				ToolId defaultToolId = new ToolId(updateChannel);
-				if (globalConfig.TryGetTool(defaultToolId, out _))
-				{
-					return defaultToolId;
-				}
-			}
-
 			ToolId toolId = AgentToolId;
 
 			if (agent.IsSelfContained())
