Horde: Fix invalid links in docs.

#rnx

diff --git a/Engine/Source/Programs/Horde/Docs/Home.md b/Engine/Source/Programs/Horde/Docs/Home.md
--- a/Engine/Source/Programs/Horde/Docs/Home.md
+++ b/Engine/Source/Programs/Horde/Docs/Home.md
@@ -2,7 +2,7 @@
 
 ## Getting Started
 
-* **[Installing Horde](Docs/Tutorials/InstallHorde.md)**
+* **[Installing Horde](Tutorials/InstallHorde.md)**
 * **[Set up remote C++ compilation](Tutorials/RemoteCompilation.md)**
 * **[Set up build automation](Tutorials/BuildAutomation.md)**
 * **[Set up test automation](Tutorials/TestAutomation.md)**
diff --git a/Engine/Source/Programs/Horde/Docs/Landing.md b/Engine/Source/Programs/Horde/Docs/Landing.md
--- a/Engine/Source/Programs/Horde/Docs/Landing.md
+++ b/Engine/Source/Programs/Horde/Docs/Landing.md
@@ -2,7 +2,7 @@
 
 **Horde** has been installed successfully. To continue setting up your installation, see the following guides:
 
-* **[Installing Horde](Docs/Tutorials/InstallHorde.md)**
+* **[Installing Horde](Tutorials/InstallHorde.md)**
 * **[Set up remote C++ compilation](Tutorials/RemoteCompilation.md)**
 * **[Set up build automation](Tutorials/BuildAutomation.md)**
 * **[Set up test automation](Tutorials/TestAutomation.md)**
