Horde: Fix static analysis warning.

#rnx

diff --git a/Engine/Source/Programs/Horde/Horde.Server.Tests/Jobs/JobsControllerDbTest.cs b/Engine/Source/Programs/Horde/Horde.Server.Tests/Jobs/JobsControllerDbTest.cs
--- a/Engine/Source/Programs/Horde/Horde.Server.Tests/Jobs/JobsControllerDbTest.cs
+++ b/Engine/Source/Programs/Horde/Horde.Server.Tests/Jobs/JobsControllerDbTest.cs
@@ -45,7 +45,9 @@ namespace Horde.Server.Tests.Jobs
 			JobStepBatchId batchId = job.Batches[0].Id;
 			JobStepId stepId = job.Batches[0].Steps[0].Id;
 
+#pragma warning disable CS0618 // Type or member is obsolete
 			object obj = (await JobsController.GetStepAsync(job.Id, batchId, stepId)).Value!;
+#pragma warning restore CS0618 // Type or member is obsolete
 			GetStepResponse stepRes = (obj as GetStepResponse)!;
 			Assert.IsFalse(stepRes.AbortRequested);
 
@@ -54,7 +56,9 @@ namespace Horde.Server.Tests.Jobs
 			await JobsController.UpdateStepAsync(job.Id, batchId, stepId, updateReq);
 			//	        UpdateStepResponse updateRes = (obj as UpdateStepResponse)!;
 
+#pragma warning disable CS0618 // Type or member is obsolete
 			obj = (await JobsController.GetStepAsync(job.Id, batchId, stepId)).Value!;
+#pragma warning restore CS0618 // Type or member is obsolete
 			stepRes = (obj as GetStepResponse)!;
 			Assert.IsTrue(stepRes.AbortRequested);
 			//	        Assert.AreEqual("Anonymous", StepRes.AbortByUser);
