Horde: Remove download agent button and backend api
#rnx

diff --git a/Engine/Source/Programs/Horde/HordeDashboard/src/backend/Backend.ts b/Engine/Source/Programs/Horde/HordeDashboard/src/backend/Backend.ts
--- a/Engine/Source/Programs/Horde/HordeDashboard/src/backend/Backend.ts
+++ b/Engine/Source/Programs/Horde/HordeDashboard/src/backend/Backend.ts
@@ -1336,19 +1336,7 @@ export class Backend {
         }
 
     }
-
-    downloadAgentZip() {
-        try {
-            const url = `${this.serverUrl}/api/v1/agentsoftware/default/zip`;
-            const link = document.createElement('a');
-            link.href = url;
-            document.body.appendChild(link);
-            link.click();
-        } catch (reason) {
-            console.error(reason);
-        }
-    }
-
+    
     updateJobStep(jobId: string, batchId: string, stepId: string, request: UpdateStepRequest): Promise<UpdateStepResponse> {
         return new Promise<UpdateStepResponse>((resolve, reject) => {
             this.backend.put(`api/v1/jobs/${jobId}/batches/${batchId}/steps/${stepId}`, request).then((value) => {
diff --git a/Engine/Source/Programs/Horde/HordeDashboard/src/components/AgentView.tsx b/Engine/Source/Programs/Horde/HordeDashboard/src/components/AgentView.tsx
--- a/Engine/Source/Programs/Horde/HordeDashboard/src/components/AgentView.tsx
+++ b/Engine/Source/Programs/Horde/HordeDashboard/src/components/AgentView.tsx
@@ -1475,7 +1475,6 @@ export const AgentMenuBar: React.FC<{ agentView?: boolean }> = observer(({ agent
          </Stack.Item>
          {!!agentView && <Stack><Stack grow />
             <Stack horizontal tokens={{ childrenGap: 12 }} grow>
-               <PrimaryButton styles={{ root: { fontFamily: "Horde Open Sans SemiBold !important" } }} text="Download Agent" onClick={() => { backend.downloadAgentZip() }} />
                {!!dashboard.user?.dashboardFeatures?.showPoolEditor && <CommandButton
                   onClick={() => { editPoolsModalState.setOpen(); }}
                   iconProps={{ iconName: 'Edit' }}
