Horde: Add notes to UBA tutorial.

#rnx

diff --git a/Engine/Source/Programs/Horde/Docs/Tutorials/RemoteCompilation.md b/Engine/Source/Programs/Horde/Docs/Tutorials/RemoteCompilation.md
--- a/Engine/Source/Programs/Horde/Docs/Tutorials/RemoteCompilation.md
+++ b/Engine/Source/Programs/Horde/Docs/Tutorials/RemoteCompilation.md
@@ -4,7 +4,7 @@
 
 ## Introduction
 
-Horde implements a platform for generic remote execution workloads, allowing clients to leverage idle CPU cycles on
+**Horde** implements a platform for generic remote execution workloads, allowing clients to leverage idle CPU cycles on
 other machines to accelerate workloads that would otherwise be executed locally. With Horde's remote execution platform,
 you can issue explicit commands to remote agents sequentially, such as "upload these files", "run this process",
 "send these files back", and so on.
@@ -15,9 +15,12 @@ initiating machine as required. The remotely executed process behaves as if it's
 on the local machine, seeing the same view of the file system and so on, and files are transferred to and from the
 remote machine behind the scenes as necessary.
 
-Unreal Build Tool can use Unreal Build Accelerator with Horde to offload build tasks to connected agents,
+**Unreal Build Tool** can use Unreal Build Accelerator with Horde to offload compilation tasks to connected agents,
 spreading the workload over multiple machines.
 
+> Note: Unreal Build Accelerator only supports Windows in Unreal Engine 5.4. Support for Mac and Linux are planned for 
+  a future release.
+
 ## Prerequisites
 
 * Horde Server and one or more Horde Agents (see [Getting Started: Install Horde](InstallHorde.md)).
@@ -82,3 +85,10 @@ spreading the workload over multiple machines.
        <bForceBuildAllRemote>true</bForceBuildAllRemote>
    </UnrealBuildAccelerator>
    ```
+
+> Note: It is not recommended to run a Horde Agent on the same machine as the Horde Server for performance reasons.
+
+> Note: When using Horde's build automation functionality, be mindful of mixing pools of agents for UBA and 
+  pools of agents for build automation. Agents used for build automation typically have higher requirements 
+  and are a more scarce resource than compute helpers.
+
