Horde: Propagate artifact keys through to step requests.

#rnx

diff --git a/Engine/Source/Programs/Horde/Horde.Agent/Execution/JobExecutor.cs b/Engine/Source/Programs/Horde/Horde.Agent/Execution/JobExecutor.cs
--- a/Engine/Source/Programs/Horde/Horde.Agent/Execution/JobExecutor.cs
+++ b/Engine/Source/Programs/Horde/Horde.Agent/Execution/JobExecutor.cs
@@ -1161,6 +1161,7 @@ namespace Horde.Agent.Execution
 					artifactRequest.Name = graphArtifact.Name;
 					artifactRequest.Type = graphArtifact.Type;
 					artifactRequest.Description = graphArtifact.Description;
+					artifactRequest.Keys.AddRange(graphArtifact.Keys);
 
 					CreateJobArtifactResponseV2 artifact = await jobRpc.Client.CreateArtifactV2Async(artifactRequest, cancellationToken: cancellationToken);
 					logger.LogInformation("Created artifact {ArtifactId} '{ArtifactName}' ({ArtifactType}) with ref {RefName} ({RefUrl})", artifact.Id, artifactRequest.Name, ArtifactType.StepOutput, artifact.RefName, $"{Session.ServerUrl.ToString().TrimEnd('/')}/api/v1/storage/{artifact.NamespaceId}/refs/{artifact.RefName}");
diff --git a/Engine/Source/Programs/Horde/Horde.Server/Jobs/JobMessages.cs b/Engine/Source/Programs/Horde/Horde.Server/Jobs/JobMessages.cs
--- a/Engine/Source/Programs/Horde/Horde.Server/Jobs/JobMessages.cs
+++ b/Engine/Source/Programs/Horde/Horde.Server/Jobs/JobMessages.cs
@@ -422,6 +422,7 @@ namespace Horde.Server.Jobs
 	/// <param name="Name">Name of the artifact</param>
 	/// <param name="Type">Artifact type</param>
 	/// <param name="Description">Description to display for the artifact on the dashboard</param>
+	/// <param name="Keys">Keys for the artifact</param>
 	/// <param name="StepId">Step producing the artifact</param>
 	public record class GetJobArtifactResponse
 	(
@@ -429,6 +430,7 @@ namespace Horde.Server.Jobs
 		ArtifactName Name,
 		ArtifactType Type,
 		string? Description,
+		List<string> Keys,
 		JobStepId StepId
 	);
 
diff --git a/Engine/Source/Programs/Horde/Horde.Server/Jobs/JobRpcService.cs b/Engine/Source/Programs/Horde/Horde.Server/Jobs/JobRpcService.cs
--- a/Engine/Source/Programs/Horde/Horde.Server/Jobs/JobRpcService.cs
+++ b/Engine/Source/Programs/Horde/Horde.Server/Jobs/JobRpcService.cs
@@ -114,6 +114,7 @@ namespace Horde.Server.Jobs
 			List<string> keys = new List<string>();
 			keys.Add(job.GetArtifactKey());
 			keys.Add(job.GetArtifactKey(step));
+			keys.AddRange(request.Keys);
 
 			if (!_globalConfig.TryGetTemplate(job.StreamId, job.TemplateId, out TemplateRefConfig? templateConfig))
 			{
diff --git a/Engine/Source/Programs/Horde/Horde.Server/Jobs/JobsController.cs b/Engine/Source/Programs/Horde/Horde.Server/Jobs/JobsController.cs
--- a/Engine/Source/Programs/Horde/Horde.Server/Jobs/JobsController.cs
+++ b/Engine/Source/Programs/Horde/Horde.Server/Jobs/JobsController.cs
@@ -507,7 +507,7 @@ namespace Horde.Server.Jobs
 						string? stepKey = artifact.Keys.FirstOrDefault(x => x.StartsWith(artifactStepKeyPrefix, StringComparison.Ordinal));
 						if (stepKey != null && JobStepId.TryParse(stepKey.Substring(artifactStepKeyPrefix.Length), out JobStepId jobStepId))
 						{
-							response.Artifacts.Add(new GetJobArtifactResponse(artifact.Id, artifact.Name, artifact.Type, artifact.Description, jobStepId));
+							response.Artifacts.Add(new GetJobArtifactResponse(artifact.Id, artifact.Name, artifact.Type, artifact.Description, artifact.Keys.ToList(), jobStepId));
 							addedArtifacts.Add((artifact.Name, jobStepId));
 						}
 					}
@@ -530,7 +530,7 @@ namespace Horde.Server.Jobs
 							IGraphArtifact? graphArtifact;
 							if (outputNameToArtifact.TryGetValue(outputName, out graphArtifact) && !addedArtifacts.Contains((graphArtifact.Name, step.Id)))
 							{
-								response.Artifacts.Add(new GetJobArtifactResponse(null, graphArtifact.Name, graphArtifact.Type, graphArtifact.Description, step.Id));
+								response.Artifacts.Add(new GetJobArtifactResponse(null, graphArtifact.Name, graphArtifact.Type, graphArtifact.Description, graphArtifact.Keys.ToList(), step.Id));
 							}
 						}
 					}
