Add example of using virtual streams to reduce sync size
#rnx

diff --git a/Engine/Source/Programs/Horde/Horde.Server/Defaults/ue5-release-5.4.stream.json b/Engine/Source/Programs/Horde/Horde.Server/Defaults/ue5-release-5.4.stream.json
--- a/Engine/Source/Programs/Horde/Horde.Server/Defaults/ue5-release-5.4.stream.json
+++ b/Engine/Source/Programs/Horde/Horde.Server/Defaults/ue5-release-5.4.stream.json
@@ -185,6 +185,9 @@
 	{
 		"Full":
 		{
+			// Use a virtual stream for complex filtering
+			"stream": "//UE5/Release-5.4-HordeSync",
+
 			// Set a unique identifier to prevent overlap with other streams; will cause workspace to be synced to a different folder
 			"identifier": "Demo",
 
@@ -193,6 +196,9 @@
 		},
 		"IncrementalCompileOnly":
 		{
+			// Use a virtual stream for complex filtering
+			"stream": "//UE5/Release-5.4-HordeSync",
+
 			// Set a unique identifier to prevent overlap with other streams; will cause workspace to be synced to a different folder
 			"identifier": "Demo-Inc",
 
