Horde: Merging latest code from UnrealCloudDDC.

#rnx

diff --git a/Engine/Source/Programs/Horde/Horde.Server.Tests/Ddc/Functional/ContentId/CompressedBlobTests.cs b/Engine/Source/Programs/Horde/Horde.Server.Tests/Ddc/Functional/ContentId/CompressedBlobTests.cs
--- a/Engine/Source/Programs/Horde/Horde.Server.Tests/Ddc/Functional/ContentId/CompressedBlobTests.cs
+++ b/Engine/Source/Programs/Horde/Horde.Server.Tests/Ddc/Functional/ContentId/CompressedBlobTests.cs
@@ -9,8 +9,6 @@ using System.Net.Http.Headers;
 using System.Net.Http.Json;
 using System.Net.Mime;
 using System.Text;
-using System.Text.Json;
-using System.Threading;
 using System.Threading.Tasks;
 using Blake3;
 using EpicGames.AspNet;
@@ -21,9 +19,9 @@ using Horde.Server.Ddc;
 using Horde.Server.Server;
 using Horde.Server.Storage;
 using Microsoft.Extensions.DependencyInjection;
-using Microsoft.Extensions.Logging;
 using Microsoft.Extensions.Options;
 using Microsoft.VisualStudio.TestTools.UnitTesting;
+using Serilog;
 
 namespace Horde.Server.Tests.Ddc.FunctionalTests.CompressedBlobs
 {
@@ -54,24 +52,27 @@ namespace Horde.Server.Tests.Ddc.FunctionalTests.CompressedBlobs
 		[TestMethod]
 		[DataTestMethod]
 		[DataRow("7835c353d7dc67e8a0531c88fbc75ddfda10dee4", "7835c353d7dc67e8a0531c88fbc75ddfda10dee4")]
-		/*        [DataRow("4958689fe783e02fb35b13c14b0c3d7beb91e50c", "4958689fe783e02fb35b13c14b0c3d7beb91e50c")]
-				[DataRow("dce31eb416f3dcb4c8250ac545eda3930919d3ff", "dce31eb416f3dcb4c8250ac545eda3930919d3ff")]
-				[DataRow("05d7c699a2668efdecbe48f10db0d621d736f449.uecomp", "05D7C699A2668EFDECBE48F10DB0D621D736F449")]
-				[DataRow("dce31eb416f3dcb4c8250ac545eda3930919d3ff", "dce31eb416f3dcb4c8250ac545eda3930919d3ff")]
-				[DataRow("UncompressedTexture_CAS_dea81b6c3b565bb5089695377c98ce0f1c13b0c3.udd", "DEA81B6C3B565BB5089695377C98CE0F1C13B0C3")]
-				[DataRow("Oodle-f895ea954b37217270e88d8b728bd3c09152689c", "F895EA954B37217270E88D8B728BD3C09152689C")]
-				[DataRow("Oodle-f0b9c675fe21951ca27699f9baab9f9f5040b202", "F0B9C675FE21951CA27699F9BAAB9F9F5040B202")]
-				[DataRow("OodleTexture_CAS_dbda9040e75c4674fcec173f982fddf12b021e24.udd", "DBDA9040E75C4674FCEC173F982FDDF12B021E24")]
-		*/
+		[DataRow("4958689fe783e02fb35b13c14b0c3d7beb91e50c", "4958689fe783e02fb35b13c14b0c3d7beb91e50c")]
+		[DataRow("dce31eb416f3dcb4c8250ac545eda3930919d3ff", "dce31eb416f3dcb4c8250ac545eda3930919d3ff")]
+		[DataRow("05d7c699a2668efdecbe48f10db0d621d736f449.uecomp", "05D7C699A2668EFDECBE48F10DB0D621D736F449")]
+		[DataRow("dce31eb416f3dcb4c8250ac545eda3930919d3ff", "dce31eb416f3dcb4c8250ac545eda3930919d3ff")]
+		[DataRow("UncompressedTexture_CAS_dea81b6c3b565bb5089695377c98ce0f1c13b0c3.udd", "DEA81B6C3B565BB5089695377C98CE0F1C13B0C3")]
+		[DataRow("Oodle-f895ea954b37217270e88d8b728bd3c09152689c", "F895EA954B37217270E88D8B728BD3C09152689C")]
+		[DataRow("Oodle-f0b9c675fe21951ca27699f9baab9f9f5040b202", "F0B9C675FE21951CA27699F9BAAB9F9F5040B202")]
+		[DataRow("OodleTexture_CAS_dbda9040e75c4674fcec173f982fddf12b021e24.udd", "DBDA9040E75C4674FCEC173F982FDDF12B021E24")]
 		public async Task PutPayloadsAsync(string payloadFilename, string uncompressedHash)
 		{
 			byte[] texturePayload = await File.ReadAllBytesAsync($"Ddc/Functional/ContentId/Payloads/{payloadFilename}");
 			BlobId compressedPayloadIdentifier = BlobId.FromBlob(texturePayload);
-			BlobId uncompressedPayloadIdentifier = BlobId.Parse(uncompressedHash);
+			BlobId uncompressedPayloadIdentifier = new BlobId(uncompressedHash);
 
 			using ByteArrayContent content = new(texturePayload);
 			content.Headers.ContentType = new MediaTypeHeaderValue(CustomMediaTypeNames.UnrealCompressedBuffer);
 			HttpResponseMessage result = await Client!.PutAsync(new Uri($"api/v1/compressed-blobs/{TestNamespace}/{uncompressedPayloadIdentifier}", UriKind.Relative), content);
+			if (result.StatusCode == HttpStatusCode.InternalServerError)
+			{
+				Assert.Fail($"Internal server error with message: {await result.Content.ReadAsStringAsync()}");
+			}
 			result.EnsureSuccessStatusCode();
 
 			InsertResponse? response = await result.Content.ReadFromJsonAsync<InsertResponse>();
@@ -85,7 +86,7 @@ namespace Horde.Server.Tests.Ddc.FunctionalTests.CompressedBlobs
 		{
 			byte[] texturePayload = await File.ReadAllBytesAsync("Ddc/Functional/ContentId/Payloads/UncompressedTexture_CAS_dea81b6c3b565bb5089695377c98ce0f1c13b0c3.udd");
 			BlobId compressedPayloadIdentifier = BlobId.FromBlob(texturePayload);
-			ContentId uncompressedPayloadIdentifier = ContentId.Parse("DEA81B6C3B565BB5089695377C98CE0F1C13B0C3");
+			ContentId uncompressedPayloadIdentifier = new ContentId("DEA81B6C3B565BB5089695377C98CE0F1C13B0C3");
 
 			{
 				using ByteArrayContent content = new(texturePayload);
@@ -97,7 +98,7 @@ namespace Horde.Server.Tests.Ddc.FunctionalTests.CompressedBlobs
 				Assert.IsNotNull(response);
 				Assert.IsNotNull(response.Identifier);
 				Assert.AreNotEqual(compressedPayloadIdentifier, response.Identifier);
-				Assert.AreEqual(uncompressedPayloadIdentifier, ContentId.FromBlobId(response.Identifier.Value));
+				Assert.AreEqual(uncompressedPayloadIdentifier, ContentId.FromBlobIdentifier(response.Identifier));
 			}
 
 			{
@@ -135,14 +136,16 @@ namespace Horde.Server.Tests.Ddc.FunctionalTests.CompressedBlobs
 		public async Task PutGetLargeCompressedPayloadAsync()
 		{
 			// we submit a blob so large that it can not fit using the memory blob store
+			IBlobStore? blobStore = ServiceProvider.GetRequiredService<IBlobStore>();
+
 			FileInfo tempOutputFile = new FileInfo(Path.GetTempFileName());
 			FileInfo tempCompressedFile = new FileInfo(Path.GetTempFileName());
 
-			ILogger logger = ServiceProvider.GetRequiredService<ILogger<CompressedBlobTests>>()!;
+			ILogger logger = ServiceProvider.GetRequiredService<ILogger>()!;
 
 			try
 			{
-				logger.LogInformation("Generating large file");
+				logger.Information("Generating large file");
 				int blockSize = 1024 * 1024;
 				// we want a file larger then 6GB, each block is 1 MB
 				int countOfBlocks = 6500;
@@ -168,15 +171,15 @@ namespace Horde.Server.Tests.Ddc.FunctionalTests.CompressedBlobs
 					bufferUtils.CompressContent(fs, OoodleCompressorMethod.Mermaid, OoodleCompressionLevel.HyperFast4, blocksToCompress, blockSize);
 				}
 
-				logger.LogInformation("Hashing generated file");
+				logger.Information("Hashing generated file");
 				BlobId blobIdentifier = BlobId.FromIoHash(uncompressedContentHash);
 				BlobId compressedContentHash;
 				{
 					await using FileStream fs = tempCompressedFile.OpenRead();
-					compressedContentHash = await BlobId.FromStreamAsync(fs, CancellationToken.None);
+					compressedContentHash = await BlobId.FromStreamAsync(fs);
 				}
 
-				logger.LogInformation("Uploading large file");
+				logger.Information("Uploading large file");
 
 				// it takes a long time to upload this content and we will not get any response while it happens so we have to bump the timeout
 				Client!.Timeout = TimeSpan.FromMinutes(5.0);
@@ -192,8 +195,8 @@ namespace Horde.Server.Tests.Ddc.FunctionalTests.CompressedBlobs
 					Assert.AreEqual(blobIdentifier, response.Identifier);
 				}
 
-				logger.LogInformation("Large file uploaded");
-				logger.LogInformation("Downloading large file");
+				logger.Information("Large file uploaded");
+				logger.Information("Downloading large file");
 
 				{
 					// verify we can fetch the blob again
@@ -210,11 +213,11 @@ namespace Horde.Server.Tests.Ddc.FunctionalTests.CompressedBlobs
 
 					await using FileStream downloadedFile = tempOutputFile.OpenRead();
 
-					BlobId downloadedBlobIdentifier = await BlobId.FromStreamAsync(downloadedFile, CancellationToken.None);
+					BlobId downloadedBlobIdentifier = await BlobId.FromStreamAsync(downloadedFile);
 					Assert.AreEqual(compressedContentHash, downloadedBlobIdentifier);
 				}
 
-				logger.LogInformation("Download completed");
+				logger.Information("Download completed");
 			}
 			finally
 			{
@@ -234,7 +237,7 @@ namespace Horde.Server.Tests.Ddc.FunctionalTests.CompressedBlobs
 		[TestMethod]
 		public async Task RecompressionTestAsync()
 		{
-			ContentId uncompressedPayloadIdentifier = ContentId.Parse("A2AC0ECED768698F7413F131D064D36B7EC6F7DA");
+			ContentId uncompressedPayloadIdentifier = new ContentId("A2AC0ECED768698F7413F131D064D36B7EC6F7DA");
 			byte[] texturePayloadSmaller = await File.ReadAllBytesAsync("Ddc/Functional/ContentId/Payloads/smallerfile");
 			BlobId compressedPayloadIdentifierSmaller = BlobId.FromBlob(texturePayloadSmaller);
 
@@ -244,13 +247,13 @@ namespace Horde.Server.Tests.Ddc.FunctionalTests.CompressedBlobs
 				HttpResponseMessage result = await Client!.PutAsync(new Uri($"api/v1/compressed-blobs/{TestNamespace}/{uncompressedPayloadIdentifier}", UriKind.Relative), content);
 				result.EnsureSuccessStatusCode();
 
-				InsertResponse response = JsonSerializer.Deserialize<InsertResponse>(await result.Content.ReadAsByteArrayAsync())!;
+				InsertResponse response = (await result.Content.ReadFromJsonAsync<InsertResponse>())!;
 				Assert.IsNotNull(response.Identifier);
 				Assert.AreNotEqual(compressedPayloadIdentifierSmaller, response.Identifier);
-				Assert.AreEqual(uncompressedPayloadIdentifier, ContentId.FromBlobId(response.Identifier.Value));
+				Assert.AreEqual(uncompressedPayloadIdentifier, ContentId.FromBlobIdentifier(response.Identifier));
 			}
 
-			byte[] texturePayloadLarger = await File.ReadAllBytesAsync("Ddc/Functional/ContentId/Payloads/largerfile");
+			byte[] texturePayloadLarger = await File.ReadAllBytesAsync("ContentId/Payloads/largerfile");
 			BlobId compressedPayloadIdentifierLarger = BlobId.FromBlob(texturePayloadLarger);
 
 			{
@@ -259,10 +262,10 @@ namespace Horde.Server.Tests.Ddc.FunctionalTests.CompressedBlobs
 				HttpResponseMessage result = await Client!.PutAsync(new Uri($"api/v1/compressed-blobs/{TestNamespace}/{uncompressedPayloadIdentifier}", UriKind.Relative), content);
 				result.EnsureSuccessStatusCode();
 
-				InsertResponse response = JsonSerializer.Deserialize<InsertResponse>(await result.Content.ReadAsByteArrayAsync())!;
+				InsertResponse response = (await result.Content.ReadFromJsonAsync<InsertResponse>())!;
 				Assert.IsNotNull(response.Identifier);
 				Assert.AreNotEqual(compressedPayloadIdentifierLarger, response.Identifier);
-				Assert.AreEqual(uncompressedPayloadIdentifier, ContentId.FromBlobId(response.Identifier.Value));
+				Assert.AreEqual(uncompressedPayloadIdentifier, ContentId.FromBlobIdentifier(response.Identifier));
 			}
 
 			{
@@ -298,7 +301,7 @@ namespace Horde.Server.Tests.Ddc.FunctionalTests.CompressedBlobs
 		{
 			byte[] texturePayload = await File.ReadAllBytesAsync("Ddc/Functional/ContentId/Payloads/UncompressedTexture_CAS_dea81b6c3b565bb5089695377c98ce0f1c13b0c3.udd");
 			BlobId compressedPayloadIdentifier = BlobId.FromBlob(texturePayload);
-			BlobId uncompressedPayloadIdentifier = BlobId.Parse("DEA81B6C3B565BB5089695377C98CE0F1C13B0C3");
+			BlobId uncompressedPayloadIdentifier = new BlobId("DEA81B6C3B565BB5089695377C98CE0F1C13B0C3");
 			{
 				using ByteArrayContent content = new(texturePayload);
 				content.Headers.ContentType = new MediaTypeHeaderValue(CustomMediaTypeNames.UnrealCompressedBuffer);
diff --git a/Engine/Source/Programs/Horde/Horde.Server.Tests/Ddc/Functional/Objects/ReferencesTests.cs b/Engine/Source/Programs/Horde/Horde.Server.Tests/Ddc/Functional/Objects/ReferencesTests.cs
--- a/Engine/Source/Programs/Horde/Horde.Server.Tests/Ddc/Functional/Objects/ReferencesTests.cs
+++ b/Engine/Source/Programs/Horde/Horde.Server.Tests/Ddc/Functional/Objects/ReferencesTests.cs
@@ -21,35 +21,16 @@ using Horde.Server.Configuration;
 using Horde.Server.Ddc;
 using Horde.Server.Server;
 using Horde.Server.Storage;
+using Microsoft.AspNetCore.Mvc;
 using Microsoft.Extensions.DependencyInjection;
 using Microsoft.Extensions.Options;
 using Microsoft.VisualStudio.TestTools.UnitTesting;
-using RefNotFoundException = Horde.Server.Ddc.RefNotFoundException;
 
 namespace Horde.Server.Tests.Ddc.FunctionalTests.References
 {
-	public interface IReferencesStore
+	public class InsertResponse
 	{
-		Task<RefRecord> GetAsync(NamespaceId ns, BucketId bucket, RefId key, FieldFlags flags);
-
-		[Flags]
-		public enum FieldFlags
-		{
-			None = 0,
-			IncludePayload = 1,
-			All = IncludePayload
-		}
-
-		Task PutAsync(NamespaceId ns, BucketId bucket, RefId key, BlobId blobHash, byte[] blob, bool isFinalized);
-		Task FinalizeAsync(NamespaceId ns, BucketId bucket, RefId key, BlobId blobIdentifier);
-
-		Task UpdateLastAccessTimeAsync(NamespaceId ns, BucketId bucket, RefId key, DateTime newLastAccessTime);
-		IAsyncEnumerable<(NamespaceId, BucketId, RefId, DateTime)> GetRecordsAsync();
-
-		IAsyncEnumerable<NamespaceId> GetNamespacesAsync();
-		Task<bool> DeleteAsync(NamespaceId ns, BucketId bucket, RefId key);
-		Task<long> DropNamespaceAsync(NamespaceId ns);
-		Task<long> DeleteBucketAsync(NamespaceId ns, BucketId bucket);
+		public BlobId? Identifier { get; set; }
 	}
 
 	[TestClass]
@@ -60,9 +41,13 @@ namespace Horde.Server.Tests.Ddc.FunctionalTests.References
 		private readonly AsyncServiceScope _serviceScope;
 
 		protected IBlobService BlobService => _serviceScope.ServiceProvider.GetRequiredService<IBlobService>();
+		protected IContentIdStore ContentIdStore => _serviceScope.ServiceProvider.GetRequiredService<IContentIdStore>();
 		protected IRefService RefService => _serviceScope.ServiceProvider.GetRequiredService<IRefService>();
 
-		protected IReferencesStore ReferencesStore { get; set; } = null!;
+		private IBlobService Service => BlobService;
+
+		[Obsolete]
+		protected IReferencesStore ReferencesStore => null!;
 		protected NamespaceId TestNamespace { get; } = new NamespaceId("test-namespace");
 
 		public ReferencesTests()
@@ -94,8 +79,8 @@ namespace Horde.Server.Tests.Ddc.FunctionalTests.References
 		[TestMethod]
 		public async Task PutGetBlobAsync()
 		{
-			const string ObjectContents = "This is treated as a opaque blob";
-			byte[] data = Encoding.ASCII.GetBytes(ObjectContents);
+			const string objectContents = $"This is treated as a opaque blob in {nameof(PutGetBlobAsync)}";
+			byte[] data = Encoding.ASCII.GetBytes(objectContents);
 			BlobId objectHash = BlobId.FromBlob(data);
 			RefId key = RefId.FromName("newBlobObject");
 			using HttpContent requestContent = new ByteArrayContent(data);
@@ -114,7 +99,7 @@ namespace Horde.Server.Tests.Ddc.FunctionalTests.References
 				byte[] roundTrippedBuffer = ms.ToArray();
 				string roundTrippedPayload = Encoding.ASCII.GetString(roundTrippedBuffer);
 
-				Assert.AreEqual(ObjectContents, roundTrippedPayload);
+				Assert.AreEqual(objectContents, roundTrippedPayload);
 				CollectionAssert.AreEqual(data, roundTrippedBuffer);
 				Assert.AreEqual(objectHash, BlobId.FromBlob(roundTrippedBuffer));
 			}
@@ -146,7 +131,7 @@ namespace Horde.Server.Tests.Ddc.FunctionalTests.References
 					byte[] roundTrippedBuffer = ms.ToArray();
 					string roundTrippedString = Encoding.ASCII.GetString(roundTrippedBuffer);
 
-					Assert.AreEqual(ObjectContents, roundTrippedString);
+					Assert.AreEqual(objectContents, roundTrippedString);
 					Assert.AreEqual(objectHash, BlobId.FromBlob(roundTrippedBuffer));
 				}
 			}
@@ -161,7 +146,7 @@ namespace Horde.Server.Tests.Ddc.FunctionalTests.References
 				string s = Encoding.ASCII.GetString(roundTrippedBuffer);
 				JsonNode? jsonNode = JsonNode.Parse(s);
 				Assert.IsNotNull(jsonNode);
-				Assert.AreEqual(objectHash, BlobId.Parse(jsonNode["RawHash"]!.GetValue<string>()));
+				Assert.AreEqual(objectHash, new BlobId(jsonNode["RawHash"]!.GetValue<string>()));
 			}
 
 			{
@@ -179,7 +164,7 @@ namespace Horde.Server.Tests.Ddc.FunctionalTests.References
 				string s = Encoding.ASCII.GetString(roundTrippedBuffer);
 				JsonNode? node = JsonNode.Parse(s);
 				Assert.IsNotNull(node);
-				Assert.AreEqual(objectHash, BlobId.Parse(node["RawHash"]!.ToString()));
+				Assert.AreEqual(objectHash, new BlobId(node["RawHash"]!.ToString()));
 			}
 
 			{
@@ -196,22 +181,24 @@ namespace Horde.Server.Tests.Ddc.FunctionalTests.References
 
 				string roundTrippedString = Encoding.ASCII.GetString(roundTrippedBuffer);
 
-				Assert.AreEqual(ObjectContents, roundTrippedString);
+				Assert.AreEqual(objectContents, roundTrippedString);
 				Assert.AreEqual(objectHash, BlobId.FromBlob(roundTrippedBuffer));
 			}
 		}
 
+		[Ignore]
+		[Obsolete]
 		[TestMethod]
 		public async Task PutGetCompactBinaryAsync()
 		{
 			CbWriter writer = new CbWriter();
 			writer.BeginObject();
-			writer.WriteString("stringField", "thisIsAField");
+			writer.WriteString("stringField", nameof(PutGetCompactBinaryAsync));
 			writer.EndObject();
 
 			byte[] objectData = writer.ToByteArray();
 			BlobId objectHash = BlobId.FromBlob(objectData);
-			RefId key = RefId.FromName("newReferenceObject");
+			RefId key = RefId.FromName("newReferenceObjectCb");
 
 			using HttpContent requestContent = new ByteArrayContent(objectData);
 			requestContent.Headers.ContentType = new MediaTypeHeaderValue(CustomMediaTypeNames.UnrealCompactBinary);
@@ -236,14 +223,12 @@ namespace Horde.Server.Tests.Ddc.FunctionalTests.References
 			{
 				BucketId bucket = new BucketId("bucket");
 
-				IRefService refService = _serviceScope.ServiceProvider.GetRequiredService<IRefService>();
-				(RefRecord objectRecord, BlobContents? contents) = await refService.GetAsync(TestNamespace, bucket, key, Array.Empty<string>());
+				RefRecord objectRecord = await ReferencesStore.GetAsync(TestNamespace, bucket, key, IReferencesStore.FieldFlags.IncludePayload, IReferencesStore.OperationFlags.None);
 
 				Assert.IsTrue(objectRecord.IsFinalized);
 				Assert.AreEqual(key, objectRecord.Name);
 				Assert.AreEqual(objectHash, objectRecord.BlobIdentifier);
-				// Note: Horde does not use the InlinePayload field
-				// Assert.IsNotNull(objectRecord.InlinePayload);
+				Assert.IsNotNull(objectRecord.InlinePayload);
 			}
 
 			{
@@ -270,8 +255,8 @@ namespace Horde.Server.Tests.Ddc.FunctionalTests.References
 
 				Assert.AreEqual(1, fields.Count);
 				CbField stringField = fields[0];
-				Assert.AreEqual("stringField", stringField.Name.ToString());
-				Assert.AreEqual("thisIsAField", stringField.AsString());
+				Assert.AreEqual(new Utf8String("stringField"), stringField.Name);
+				Assert.AreEqual(nameof(PutGetCompactBinaryAsync), stringField.AsString());
 			}
 
 			{
@@ -285,21 +270,23 @@ namespace Horde.Server.Tests.Ddc.FunctionalTests.References
 				string s = Encoding.ASCII.GetString(roundTrippedBuffer);
 				JsonNode? node = JsonNode.Parse(s);
 				Assert.IsNotNull(node);
-				Assert.AreEqual("thisIsAField", node["stringField"]!.GetValue<string>());
+				Assert.AreEqual(nameof(PutGetCompactBinaryAsync), node["stringField"]!.GetValue<string>());
 			}
 		}
 
+		[Ignore]
+		[Obsolete]
 		[TestMethod]
 		public async Task PutGetCompactBinaryFilteringAsync()
 		{
 			CbWriter writer = new CbWriter();
 			writer.BeginObject();
-			writer.WriteString("stringField", "thisIsAField");
+			writer.WriteString("stringField", nameof(PutGetCompactBinaryFilteringAsync));
 			writer.EndObject();
 
 			byte[] objectData = writer.ToByteArray();
 			BlobId objectHash = BlobId.FromBlob(objectData);
-			RefId key = RefId.FromName("newReferenceObject");
+			RefId key = RefId.FromName("newReferenceObjectCBFiltering");
 
 			using HttpContent requestContent = new ByteArrayContent(objectData);
 			requestContent.Headers.ContentType = new MediaTypeHeaderValue(CustomMediaTypeNames.UnrealCompactBinary);
@@ -324,8 +311,7 @@ namespace Horde.Server.Tests.Ddc.FunctionalTests.References
 			{
 				BucketId bucket = new BucketId("bucket");
 
-				IRefService refService = _serviceScope.ServiceProvider.GetRequiredService<IRefService>();
-				(RefRecord objectRecord, _) = await refService.GetAsync(TestNamespace, bucket, key, Array.Empty<string>());
+				RefRecord objectRecord = await ReferencesStore.GetAsync(TestNamespace, bucket, key, IReferencesStore.FieldFlags.None, IReferencesStore.OperationFlags.None);
 
 				Assert.IsTrue(objectRecord.IsFinalized);
 				Assert.AreEqual(key, objectRecord.Name);
@@ -344,7 +330,7 @@ namespace Horde.Server.Tests.Ddc.FunctionalTests.References
 				string s = Encoding.ASCII.GetString(roundTrippedBuffer);
 				JsonNode? node = JsonNode.Parse(s);
 				Assert.IsNotNull(node);
-				Assert.AreEqual("thisIsAField", node["stringField"]!.GetValue<string>());
+				Assert.AreEqual(nameof(PutGetCompactBinaryFilteringAsync), node["stringField"]!.GetValue<string>());
 			}
 		}
 
@@ -367,12 +353,14 @@ namespace Horde.Server.Tests.Ddc.FunctionalTests.References
 			Assert.AreEqual(4924, needsField.Count);
 		}
 
+		[Ignore]
+		[Obsolete]
 		[TestMethod]
 		public async Task PutGetCompactBinaryHierarchyAsync()
 		{
 			CbWriter childObjectWriter = new CbWriter();
 			childObjectWriter.BeginObject();
-			childObjectWriter.WriteString("stringField", "thisIsAField");
+			childObjectWriter.WriteString("stringField", nameof(PutGetCompactBinaryHierarchyAsync));
 			childObjectWriter.EndObject();
 			byte[] childObjectData = childObjectWriter.ToByteArray();
 			BlobId childObjectHash = BlobId.FromBlob(childObjectData);
@@ -384,7 +372,7 @@ namespace Horde.Server.Tests.Ddc.FunctionalTests.References
 			byte[] parentObjectData = parentObjectWriter.ToByteArray();
 			BlobId parentObjectHash = BlobId.FromBlob(parentObjectData);
 
-			RefId key = RefId.FromName("newHierarchyObject");
+			RefId key = RefId.FromName("newCbHierarchyObject");
 			// this first upload should fail with the child object missing
 			{
 				using HttpContent requestContent = new ByteArrayContent(parentObjectData);
@@ -454,13 +442,12 @@ namespace Horde.Server.Tests.Ddc.FunctionalTests.References
 			{
 				BucketId bucket = new BucketId("bucket");
 
-				IRefService refService = _serviceScope.ServiceProvider.GetRequiredService<IRefService>();
-				(RefRecord objectRecord, _) = await refService.GetAsync(TestNamespace, bucket, key, Array.Empty<string>());
+				RefRecord objectRecord = await ReferencesStore.GetAsync(TestNamespace, bucket, key, IReferencesStore.FieldFlags.IncludePayload, IReferencesStore.OperationFlags.None);
 
 				Assert.IsTrue(objectRecord.IsFinalized);
 				Assert.AreEqual(key, objectRecord.Name);
 				Assert.AreEqual(parentObjectHash, objectRecord.BlobIdentifier);
-				//                Assert.IsNotNull(objectRecord.InlinePayload);
+				Assert.IsNotNull(objectRecord.InlinePayload);
 			}
 
 			{
@@ -487,7 +474,7 @@ namespace Horde.Server.Tests.Ddc.FunctionalTests.References
 
 				Assert.AreEqual(1, fields.Count);
 				CbField childObjectField = fields[0];
-				Assert.AreEqual("childObject", childObjectField.Name.ToString());
+				Assert.AreEqual(new Utf8String("childObject"), childObjectField.Name);
 				Assert.AreEqual(childObjectHash, BlobId.FromIoHash(childObjectField.AsHash()));
 			}
 
@@ -521,8 +508,8 @@ namespace Horde.Server.Tests.Ddc.FunctionalTests.References
 		[TestMethod]
 		public async Task ExistsChecksAsync()
 		{
-			const string ObjectContents = "This is treated as a opaque blob";
-			byte[] data = Encoding.ASCII.GetBytes(ObjectContents);
+			const string objectContents = $"This is treated as a opaque blob in {nameof(ExistsChecksAsync)}";
+			byte[] data = Encoding.ASCII.GetBytes(objectContents);
 			BlobId objectHash = BlobId.FromBlob(data);
 			RefId key = RefId.FromName("newObject");
 			using HttpContent requestContent = new ByteArrayContent(data);
@@ -551,10 +538,10 @@ namespace Horde.Server.Tests.Ddc.FunctionalTests.References
 		public async Task ExistsChecksMultipleAsync()
 		{
 			BucketId bucket = new BucketId("bucket");
-			RefId existingObject = RefId.FromName("existingObject");
+			RefId existingObject = RefId.FromName("existingObjectMultiple");
 
-			const string ObjectContents = "This is treated as a opaque blob";
-			byte[] data = Encoding.ASCII.GetBytes(ObjectContents);
+			const string objectContents = $"This is treated as a opaque blob in {nameof(ExistsChecksMultipleAsync)}";
+			byte[] data = Encoding.ASCII.GetBytes(objectContents);
 			BlobId objectHash = BlobId.FromBlob(data);
 			using HttpContent requestContent = new ByteArrayContent(data);
 			requestContent.Headers.ContentType = new MediaTypeHeaderValue(MediaTypeNames.Application.Octet);
@@ -565,7 +552,7 @@ namespace Horde.Server.Tests.Ddc.FunctionalTests.References
 				result.EnsureSuccessStatusCode();
 			}
 
-			RefId missingObject = RefId.FromName("missingObject");
+			RefId missingObject = RefId.FromName("missingObjectMultiple");
 
 			string queryString = $"?names={bucket}.{existingObject}&names={bucket}.{missingObject}";
 
@@ -581,18 +568,20 @@ namespace Horde.Server.Tests.Ddc.FunctionalTests.References
 			}
 		}
 
+		[Ignore]
+		[Obsolete]
 		[TestMethod]
 		public async Task PutGetObjectHierarchyAsync()
 		{
-			string blobContents = "This is a string that is referenced as a blob";
+			string blobContents = $"This is a string that is referenced as a blob in {nameof(PutGetObjectHierarchyAsync)}";
 			byte[] blobData = Encoding.ASCII.GetBytes(blobContents);
 			BlobId blobHash = BlobId.FromBlob(blobData);
-			await BlobService.PutObjectAsync(TestNamespace, blobData, blobHash);
+			await Service.PutObjectAsync(TestNamespace, blobData, blobHash);
 
-			string blobContentsChild = "This string is also referenced as a blob but from a child object";
+			string blobContentsChild = $"This string is also referenced as a blob but from a child object in {nameof(PutGetObjectHierarchyAsync)}";
 			byte[] dataChild = Encoding.ASCII.GetBytes(blobContentsChild);
 			BlobId blobHashChild = BlobId.FromBlob(dataChild);
-			await BlobService.PutObjectAsync(TestNamespace, dataChild, blobHashChild);
+			await Service.PutObjectAsync(TestNamespace, dataChild, blobHashChild);
 
 			CbWriter writerChild = new CbWriter();
 			writerChild.BeginObject();
@@ -601,7 +590,7 @@ namespace Horde.Server.Tests.Ddc.FunctionalTests.References
 
 			byte[] childDataObject = writerChild.ToByteArray();
 			BlobId childDataObjectHash = BlobId.FromBlob(childDataObject);
-			await BlobService.PutObjectAsync(TestNamespace, childDataObject, childDataObjectHash);
+			await Service.PutObjectAsync(TestNamespace, childDataObject, childDataObjectHash);
 
 			CbWriter writerParent = new CbWriter();
 			writerParent.BeginObject();
@@ -639,12 +628,12 @@ namespace Horde.Server.Tests.Ddc.FunctionalTests.References
 			{
 				BucketId bucket = new BucketId("bucket");
 
-				(RefRecord objectRecord, _) = await RefService.GetAsync(TestNamespace, bucket, key, Array.Empty<string>());
+				RefRecord objectRecord = await ReferencesStore.GetAsync(TestNamespace, bucket, key, IReferencesStore.FieldFlags.IncludePayload, IReferencesStore.OperationFlags.None);
 
 				Assert.IsTrue(objectRecord.IsFinalized);
 				Assert.AreEqual(key, objectRecord.Name);
 				Assert.AreEqual(objectHash, objectRecord.BlobIdentifier);
-				//                Assert.IsNotNull(objectRecord.InlinePayload);
+				Assert.IsNotNull(objectRecord.InlinePayload);
 			}
 
 			// verify attachments
@@ -735,8 +724,8 @@ namespace Horde.Server.Tests.Ddc.FunctionalTests.References
 				string s = Encoding.ASCII.GetString(roundTrippedBuffer);
 				JsonNode? node = JsonNode.Parse(s);
 				Assert.IsNotNull(node);
-				Assert.AreEqual(blobHash, BlobId.Parse(node["blobAttachment"]!.GetValue<string>()));
-				Assert.AreEqual(childDataObjectHash, BlobId.Parse(node["objectAttachment"]!.GetValue<string>()));
+				Assert.AreEqual(blobHash, new BlobId(node["blobAttachment"]!.GetValue<string>()));
+				Assert.AreEqual(childDataObjectHash, new BlobId(node["objectAttachment"]!.GetValue<string>()));
 			}
 		}
 
@@ -744,11 +733,11 @@ namespace Horde.Server.Tests.Ddc.FunctionalTests.References
 		public async Task PutPartialHierarchyAsync()
 		{
 			// do not submit the content of the blobs, which should be reported in the response of the put
-			string blobContents = "This is a string that is referenced as a blob";
+			string blobContents = $"This is a string that is referenced as a blob in {nameof(PutPartialHierarchyAsync)}";
 			byte[] blobData = Encoding.ASCII.GetBytes(blobContents);
 			BlobId blobHash = BlobId.FromBlob(blobData);
 
-			string blobContentsChild = "This string is also referenced as a blob but from a child object";
+			string blobContentsChild = $"This string is also referenced as a blob but from a child object in {nameof(PutPartialHierarchyAsync)}";
 			byte[] dataChild = Encoding.ASCII.GetBytes(blobContentsChild);
 			BlobId blobHashChild = BlobId.FromBlob(dataChild);
 
@@ -759,7 +748,7 @@ namespace Horde.Server.Tests.Ddc.FunctionalTests.References
 
 			byte[] childDataObject = writerChild.ToByteArray();
 			BlobId childDataObjectHash = BlobId.FromBlob(childDataObject);
-			await BlobService.PutObjectAsync(TestNamespace, childDataObject, childDataObjectHash);
+			await Service.PutObjectAsync(TestNamespace, childDataObject, childDataObjectHash);
 
 			CbWriter writerParent = new CbWriter();
 			writerParent.BeginObject();
@@ -771,7 +760,7 @@ namespace Horde.Server.Tests.Ddc.FunctionalTests.References
 			byte[] objectData = writerParent.ToByteArray();
 			BlobId objectHash = BlobId.FromBlob(objectData);
 
-			RefId key = RefId.FromName("newHierarchyObject");
+			RefId key = RefId.FromName("newPartialHierarchyObject");
 
 			{
 				using HttpContent requestContent = new ByteArrayContent(objectData);
@@ -786,7 +775,6 @@ namespace Horde.Server.Tests.Ddc.FunctionalTests.References
 					await using MemoryStream ms = new MemoryStream();
 					await result.Content.CopyToAsync(ms);
 					byte[] roundTrippedBuffer = ms.ToArray();
-					ReadOnlyMemory<byte> localMemory = new ReadOnlyMemory<byte>(roundTrippedBuffer);
 					CbObject cb = new CbObject(roundTrippedBuffer);
 					CbField needsField = cb["needs"];
 					List<BlobId> missingBlobs = needsField.AsArray().Select(field => BlobId.FromIoHash(field.AsHash())).ToList();
@@ -811,13 +799,10 @@ namespace Horde.Server.Tests.Ddc.FunctionalTests.References
 					byte[] roundTrippedBuffer = ms.ToArray();
 					string s = Encoding.ASCII.GetString(roundTrippedBuffer);
 
-					JsonSerializerOptions jsonSerializerOptions = new JsonSerializerOptions();
-					Startup.ConfigureJsonSerializer(jsonSerializerOptions);
-
-					PutObjectResponse? response = JsonSerializer.Deserialize<PutObjectResponse>(s, jsonSerializerOptions);
+					PutObjectResponse? response = JsonSerializer.Deserialize<PutObjectResponse>(s, JsonTestUtils.DefaultJsonSerializerSettings);
 					Assert.IsNotNull(response);
 
-					BlobId[] missingBlobs = response.Needs.Select(hash => new BlobId(hash)).ToArray();
+					BlobId[] missingBlobs = response.Needs.Select(hash => new BlobId(hash.HashData)).ToArray();
 
 					Assert.AreEqual(2, missingBlobs.Length);
 					Assert.IsTrue(missingBlobs.Contains(blobHash));
@@ -826,18 +811,19 @@ namespace Horde.Server.Tests.Ddc.FunctionalTests.References
 			}
 		}
 
-		[Ignore("Horde does not allow submitting content ids before blobs")]
+		[Ignore]
+		[Obsolete]
 		[TestMethod]
 		public async Task PutContentIdMissingBlobAsync()
 		{
-			IContentIdService? contentIdStore = ServiceProvider.GetService<IContentIdService>();
+			IContentIdStore? contentIdStore = ContentIdStore;
 			Assert.IsNotNull(contentIdStore);
 
 			// submit a object which contains a content id, which exists but points to a blob that does not exist
-			string blobContents = "This is a string that is referenced as a blob";
+			string blobContents = $"This is a string that is referenced as a blob in {nameof(PutContentIdMissingBlobAsync)}";
 			byte[] blobData = Encoding.ASCII.GetBytes(blobContents);
 			BlobId blobHash = BlobId.FromBlob(blobData);
-			ContentId contentId = ContentId.Parse("0000000000000000000000000000000000000000");
+			ContentId contentId = new ContentId("0000000000000000000000AA0000000000000000");
 
 			await contentIdStore.PutAsync(TestNamespace, contentId, blobHash, blobData.Length);
 
@@ -870,7 +856,7 @@ namespace Horde.Server.Tests.Ddc.FunctionalTests.References
 					List<IoHash> missingBlobs = needsField.AsArray().Select(field => field.AsHash()).ToList();
 					Assert.AreEqual(1, missingBlobs.Count);
 
-					Assert.AreNotEqual(blobHash.Hash, missingBlobs[0], "Refs should not be returning the mapped blob identifiers as this is unknown to the client attempting to put a new ref");
+					Assert.AreNotEqual(blobHash.AsIoHash(), missingBlobs[0], "Refs should not be returning the mapped blob identifiers as this is unknown to the client attempting to put a new ref");
 					Assert.AreEqual(contentId.AsBlobIdentifier(), BlobId.FromIoHash(missingBlobs[0]));
 				}
 			}
@@ -889,14 +875,10 @@ namespace Horde.Server.Tests.Ddc.FunctionalTests.References
 					await result.Content.CopyToAsync(ms);
 					byte[] roundTrippedBuffer = ms.ToArray();
 					string s = Encoding.ASCII.GetString(roundTrippedBuffer);
-
-					JsonSerializerOptions jsonSerializerOptions = new JsonSerializerOptions();
-					Startup.ConfigureJsonSerializer(jsonSerializerOptions);
-
-					PutObjectResponse? response = JsonSerializer.Deserialize<PutObjectResponse>(s, jsonSerializerOptions);
+					PutObjectResponse? response = JsonSerializer.Deserialize<PutObjectResponse>(s, JsonTestUtils.DefaultJsonSerializerSettings);
 					Assert.IsNotNull(response);
 
-					BlobId[] missingBlobs = response.Needs.Select(field => new BlobId(field)).ToArray();
+					BlobId[] missingBlobs = response.Needs.Select(field => new BlobId(field.HashData)).ToArray();
 
 					Assert.AreEqual(1, missingBlobs.Length);
 					Assert.AreEqual(contentId.AsBlobIdentifier(), missingBlobs[0]);
@@ -904,6 +886,7 @@ namespace Horde.Server.Tests.Ddc.FunctionalTests.References
 			}
 		}
 
+
 		[TestMethod]
 		public async Task PutMissingAttachmentComplexAsync()
 		{
@@ -951,6 +934,8 @@ namespace Horde.Server.Tests.Ddc.FunctionalTests.References
 			}
 		}
 
+		[Ignore]
+		[Obsolete]
 		[TestMethod]
 		public async Task PutAndFinalizeAsync()
 		{
@@ -958,11 +943,11 @@ namespace Horde.Server.Tests.Ddc.FunctionalTests.References
 			RefId key = RefId.FromName("willFinalizeObject");
 
 			// do not submit the content of the blobs, which should be reported in the response of the put
-			string blobContents = "This is a string that is referenced as a blob";
+			string blobContents = $"This is a string that is referenced as a blob in {nameof(PutAndFinalizeAsync)}";
 			byte[] blobData = Encoding.ASCII.GetBytes(blobContents);
 			BlobId blobHash = BlobId.FromBlob(blobData);
 
-			string blobContentsChild = "This string is also referenced as a blob but from a child object";
+			string blobContentsChild = $"This string is also referenced as a blob but from a child object in {nameof(PutAndFinalizeAsync)}";
 			byte[] dataChild = Encoding.ASCII.GetBytes(blobContentsChild);
 			BlobId blobHashChild = BlobId.FromBlob(dataChild);
 
@@ -973,11 +958,10 @@ namespace Horde.Server.Tests.Ddc.FunctionalTests.References
 
 			byte[] childDataObject = writerChild.ToByteArray();
 			BlobId childDataObjectHash = BlobId.FromBlob(childDataObject);
-			await BlobService.PutObjectAsync(TestNamespace, childDataObject, childDataObjectHash);
+			await Service.PutObjectAsync(TestNamespace, childDataObject, childDataObjectHash);
 
 			CbWriter writerParent = new CbWriter();
 			writerParent.BeginObject();
-
 			writerParent.WriteBinaryAttachment("blobAttachment", blobHash.AsIoHash());
 			writerParent.WriteObjectAttachment("objectAttachment", childDataObjectHash.AsIoHash());
 			writerParent.EndObject();
@@ -1009,13 +993,18 @@ namespace Horde.Server.Tests.Ddc.FunctionalTests.References
 
 			// check that actual internal representation
 			{
-				await Assert.ThrowsExceptionAsync<RefNotFoundException>(() => RefService.GetAsync(TestNamespace, bucket, key, Array.Empty<string>()));
+				RefRecord objectRecord = await ReferencesStore.GetAsync(TestNamespace, bucket, key, IReferencesStore.FieldFlags.IncludePayload, IReferencesStore.OperationFlags.None);
+
+				Assert.IsFalse(objectRecord.IsFinalized);
+				Assert.AreEqual(key, objectRecord.Name);
+				Assert.AreEqual(objectHash, objectRecord.BlobIdentifier);
+				Assert.IsNotNull(objectRecord.InlinePayload);
 			}
 
 			// upload missing pieces
 			{
-				await BlobService.PutObjectAsync(TestNamespace, blobData, blobHash);
-				await BlobService.PutObjectAsync(TestNamespace, dataChild, blobHashChild);
+				await Service.PutObjectAsync(TestNamespace, blobData, blobHash);
+				await Service.PutObjectAsync(TestNamespace, dataChild, blobHashChild);
 			}
 
 			// finalize the object as no pieces is now missing
@@ -1040,7 +1029,7 @@ namespace Horde.Server.Tests.Ddc.FunctionalTests.References
 
 			// check that actual internal representation has updated its state
 			{
-				(RefRecord objectRecord, _) = await RefService.GetAsync(TestNamespace, bucket, key, Array.Empty<string>());
+				RefRecord objectRecord = await ReferencesStore.GetAsync(TestNamespace, bucket, key, IReferencesStore.FieldFlags.None, IReferencesStore.OperationFlags.None);
 
 				Assert.IsTrue(objectRecord.IsFinalized);
 				Assert.AreEqual(key, objectRecord.Name);
@@ -1051,7 +1040,7 @@ namespace Horde.Server.Tests.Ddc.FunctionalTests.References
 		[TestMethod]
 		public async Task GetMissingContentIdRecordAsync()
 		{
-			string blobContents = "This is a blob";
+			string blobContents = $"This is a blob in {nameof(GetMissingContentIdRecordAsync)}";
 			byte[] blobData = Encoding.ASCII.GetBytes(blobContents);
 			BlobId uncompressedHash = BlobId.FromBlob(blobData);
 			RefId key = RefId.FromName("compressedObject");
@@ -1129,7 +1118,7 @@ namespace Horde.Server.Tests.Ddc.FunctionalTests.References
 
 			{
 				// delete the blob referenced by the compressed buffer
-				await BlobService.DeleteObjectAsync(TestNamespace, compressedHash);
+				await Service.DeleteObjectAsync(TestNamespace, compressedHash);
 			}
 
 			{
@@ -1233,7 +1222,7 @@ namespace Horde.Server.Tests.Ddc.FunctionalTests.References
 
 			{
 				// delete the blob referenced by the compressed buffer
-				await BlobService.DeleteObjectAsync(TestNamespace, cbAttachmentHash);
+				await Service.DeleteObjectAsync(TestNamespace, cbAttachmentHash);
 			}
 
 			{
@@ -1265,10 +1254,10 @@ namespace Horde.Server.Tests.Ddc.FunctionalTests.References
 		[TestMethod]
 		public async Task GetMissingBlobRecordAsync()
 		{
-			string blobContents = "This is a blob";
+			string blobContents = $"This is a blob in {nameof(GetMissingBlobRecordAsync)}";
 			byte[] blobData = Encoding.ASCII.GetBytes(blobContents);
 			BlobId blobHash = BlobId.FromBlob(blobData);
-			RefId key = RefId.FromName("newReferenceObject");
+			RefId key = RefId.FromName("newReferenceObjectMissingBlobs");
 
 			using HttpContent requestContent = new ByteArrayContent(blobData);
 			requestContent.Headers.ContentType = new MediaTypeHeaderValue(MediaTypeNames.Application.Octet);
@@ -1321,7 +1310,7 @@ namespace Horde.Server.Tests.Ddc.FunctionalTests.References
 
 			{
 				// delete the blob 
-				await BlobService.DeleteObjectAsync(TestNamespace, blobHash);
+				await Service.DeleteObjectAsync(TestNamespace, blobHash);
 			}
 
 			{
@@ -1341,6 +1330,10 @@ namespace Horde.Server.Tests.Ddc.FunctionalTests.References
 				HttpResponseMessage getResponse = await _httpClient.GetAsync(new Uri($"api/v1/refs/{TestNamespace}/bucket/{key}.raw", UriKind.Relative));
 				Assert.AreEqual(HttpStatusCode.NotFound, getResponse.StatusCode);
 				Assert.AreEqual("application/problem+json", getResponse.Content.Headers.ContentType!.MediaType);
+				string s = await getResponse.Content.ReadAsStringAsync();
+				ProblemDetails? problem = JsonSerializer.Deserialize<ProblemDetails>(s, JsonTestUtils.DefaultJsonSerializerSettings);
+				Assert.IsNotNull(problem);
+				Assert.AreEqual($"Object {blobHash} in {TestNamespace} not found", problem.Title);
 			}
 
 			{
@@ -1366,8 +1359,8 @@ namespace Horde.Server.Tests.Ddc.FunctionalTests.References
 		[TestMethod]
 		public async Task DeleteObjectAsync()
 		{
-			const string ObjectContents = "This is treated as a opaque blob";
-			byte[] data = Encoding.ASCII.GetBytes(ObjectContents);
+			const string objectContents = $"This is treated as a opaque blob in {nameof(DeleteObjectAsync)}";
+			byte[] data = Encoding.ASCII.GetBytes(objectContents);
 			BlobId objectHash = BlobId.FromBlob(data);
 
 			using HttpContent requestContent = new ByteArrayContent(data);
@@ -1399,28 +1392,29 @@ namespace Horde.Server.Tests.Ddc.FunctionalTests.References
 				Assert.AreEqual(HttpStatusCode.NotFound, result.StatusCode);
 			}
 
-			// delete the object again which doesn't exist anymore
+			// delete the object again which doesn't exist anymore, but we do not return that in the api (as it causes us to need to do extra work for some backends)
 			{
 				HttpResponseMessage result = await _httpClient!.DeleteAsync(new Uri($"api/v1/refs/{TestNamespace}/bucket/{key}", UriKind.Relative));
-				Assert.AreEqual(HttpStatusCode.NotFound, result.StatusCode);
+				result.EnsureSuccessStatusCode();
 			}
 		}
 
-		[Ignore("Horde does not support dropping buckets")]
+		[Ignore]
+		[Obsolete]
 		[TestMethod]
 		public async Task DropBucketAsync()
 		{
 			const string BucketToDelete = "delete-bucket";
 
-			const string ObjectContents = "This is treated as a opaque blob";
-			byte[] data = Encoding.ASCII.GetBytes(ObjectContents);
+			const string objectContents = $"This is treated as a opaque blob in {nameof(DropBucketAsync)}";
+			byte[] data = Encoding.ASCII.GetBytes(objectContents);
 			BlobId objectHash = BlobId.FromBlob(data);
 
 			using HttpContent requestContent = new ByteArrayContent(data);
 			requestContent.Headers.ContentType = new MediaTypeHeaderValue(MediaTypeNames.Application.Octet);
 			requestContent.Headers.Add(CommonHeaders.HashHeaderName, objectHash.ToString());
 
-			RefId key = RefId.FromName("deletableObject");
+			RefId key = RefId.FromName("deletableObjectBucket");
 
 			// submit the object into multiple buckets
 			{
@@ -1464,21 +1458,19 @@ namespace Horde.Server.Tests.Ddc.FunctionalTests.References
 			}
 		}
 
-		[Ignore("Horde does not support deleting namespaces")]
+#if false
 		[TestMethod]
 		public async Task DeleteNamespaceAsync()
 		{
-			const string NamespaceToBeDeleted = "test-delete-namespace";
-
-			const string ObjectContents = "This is treated as a opaque blob";
-			byte[] data = Encoding.ASCII.GetBytes(ObjectContents);
+			const string objectContents = $"This is treated as a opaque blob in {nameof(DeleteNamespaceAsync)}";
+			byte[] data = Encoding.ASCII.GetBytes(objectContents);
 			BlobId objectHash = BlobId.FromBlob(data);
 
 			using HttpContent requestContent = new ByteArrayContent(data);
 			requestContent.Headers.ContentType = new MediaTypeHeaderValue(MediaTypeNames.Application.Octet);
 			requestContent.Headers.Add(CommonHeaders.HashHeaderName, objectHash.ToString());
 
-			RefId key = RefId.FromName("deletableObject");
+			RefId key = RefId.FromName("deletableObjectNamespace");
 			// submit the object into multiple namespaces
 			{
 				HttpResponseMessage result = await _httpClient!.PutAsync(new Uri($"api/v1/refs/{TestNamespace}/bucket/{key}", UriKind.Relative), requestContent);
@@ -1531,12 +1523,11 @@ namespace Horde.Server.Tests.Ddc.FunctionalTests.References
 			}
 		}
 
-		[Ignore("Horde does not support enumerating namespaces")]
 		[TestMethod]
 		public async Task ListNamespacesAsync()
 		{
-			const string ObjectContents = "This is treated as a opaque blob";
-			byte[] data = Encoding.ASCII.GetBytes(ObjectContents);
+			const string objectContents = $"This is treated as a opaque blob in {nameof(ListNamespacesAsync)}";
+			byte[] data = Encoding.ASCII.GetBytes(objectContents);
 			BlobId objectHash = BlobId.FromBlob(data);
 			RefId key = RefId.FromName("notUsedObject");
 
@@ -1559,32 +1550,7 @@ namespace Horde.Server.Tests.Ddc.FunctionalTests.References
 				Assert.IsTrue(response.Namespaces.Contains(TestNamespace));
 			}
 		}
-
-		[Ignore("Horde does not support enumerating all blobs")]
-		[TestMethod]
-		public async Task GetOldRecordsAsync()
-		{
-			const string ObjectContents = "This is treated as a opaque blob";
-			byte[] data = Encoding.ASCII.GetBytes(ObjectContents);
-			BlobId objectHash = BlobId.FromBlob(data);
-
-			using HttpContent requestContent = new ByteArrayContent(data);
-			requestContent.Headers.ContentType = new MediaTypeHeaderValue(MediaTypeNames.Application.Octet);
-			requestContent.Headers.Add(CommonHeaders.HashHeaderName, objectHash.ToString());
-
-			RefId key = RefId.FromName("oldRecord");
-			// submit some contents
-			{
-				HttpResponseMessage result = await _httpClient!.PutAsync(new Uri($"api/v1/refs/{TestNamespace}/bucket/{key}", UriKind.Relative), requestContent);
-				result.EnsureSuccessStatusCode();
-			}
-
-			List<(NamespaceId, BucketId, RefId, DateTime)> records = await ReferencesStore.GetRecordsAsync().ToListAsync();
-
-			(NamespaceId oldNs, BucketId oldBucket, RefId oldName, DateTime oldDate) = records.Where(tuple => tuple.Item1 == TestNamespace).First(record => record.Item3 == key);
-			Assert.AreEqual(key, oldName);
-			Assert.AreEqual("bucket", oldBucket.ToString());
-		}
+#endif
 
 		[TestMethod]
 		public async Task BatchJsonRequestAsync()
@@ -1613,7 +1579,7 @@ namespace Horde.Server.Tests.Ddc.FunctionalTests.References
 			CbWriter getObjectOp = new CbWriter();
 			getObjectOp.BeginObject();
 			getObjectOp.WriteInteger("opId", 0);
-			getObjectOp.WriteString("op", RefBatchOps.RefBatchOp.RefOperation.GET.ToString());
+			getObjectOp.WriteString("op", BatchOps.BatchOp.Operation.GET.ToString());
 			getObjectOp.WriteString("bucket", bucket.ToString());
 			getObjectOp.WriteString("key", newBlobObjectKey.ToString());
 			getObjectOp.EndObject();
@@ -1621,7 +1587,7 @@ namespace Horde.Server.Tests.Ddc.FunctionalTests.References
 			CbWriter putObjectOp = new CbWriter();
 			putObjectOp.BeginObject();
 			putObjectOp.WriteInteger("opId", 1);
-			putObjectOp.WriteString("op", RefBatchOps.RefBatchOp.RefOperation.PUT.ToString());
+			putObjectOp.WriteString("op", BatchOps.BatchOp.Operation.PUT.ToString());
 			putObjectOp.WriteString("bucket", bucket.ToString());
 			putObjectOp.WriteString("key", putObjectKey.ToString());
 			putObjectOp.WriteObject("payload", ref1);
@@ -1663,7 +1629,7 @@ namespace Horde.Server.Tests.Ddc.FunctionalTests.References
 				Assert.IsNotNull(op0.Response);
 				Assert.AreEqual(404, op0.StatusCode);
 				Assert.IsTrue(!op0.Response["title"].Equals(CbField.Empty));
-				Assert.AreEqual("Object not found 29911b58b3c970ba39d9690f2dca66839dd6f5d9 in bucket bucket namespace test-namespace", op0.Response["title"].AsString());
+				Assert.AreEqual($"Object not found 29911b58b3c970ba39d9690f2dca66839dd6f5d9 in bucket bucket namespace {TestNamespace}", op0.Response["title"].AsString());
 
 				BatchOpsResponse.OpResponses op1 = response.Results.First(r => r.OpId == 1);
 				Assert.IsNotNull(op1.Response);
@@ -1678,8 +1644,8 @@ namespace Horde.Server.Tests.Ddc.FunctionalTests.References
 		{
 			// seed some data
 			BucketId bucket = new BucketId("bucket");
-			RefId newBlobObjectKey = RefId.FromName("newBlobObject");
-			CbObject newBlobObject = CbObject.Build(writer => writer.WriteString("String", "this-has-contents"));
+			RefId newBlobObjectKey = RefId.FromName("newBlobObjectBatch");
+			CbObject newBlobObject = CbObject.Build(writer => writer.WriteString("String", $"this-has-contents-in-{nameof(BatchGetOperationsAsync)}"));
 
 			{
 				byte[] cbObjectBytes = newBlobObject.GetView().ToArray();
@@ -1697,9 +1663,9 @@ namespace Horde.Server.Tests.Ddc.FunctionalTests.References
 				Assert.IsTrue(response.Needs.Length == 0);
 			}
 
-			byte[] blobContents = Encoding.ASCII.GetBytes("This is a attached blob");
+			byte[] blobContents = Encoding.ASCII.GetBytes($"This is a attached blob in {nameof(BatchGetOperationsAsync)}");
 			CbObject newReferenceObject = CbObject.Build(writer => writer.WriteBinaryAttachment("Attachment", IoHash.Compute(blobContents)));
-			RefId newReferenceObjectKey = RefId.FromName("newReferenceObject");
+			RefId newReferenceObjectKey = RefId.FromName("newReferenceObjectBatch");
 
 			{
 				BlobId blobHash = BlobId.FromBlob(blobContents);
@@ -1746,7 +1712,7 @@ namespace Horde.Server.Tests.Ddc.FunctionalTests.References
 			CbWriter getObjectOp = new CbWriter();
 			getObjectOp.BeginObject();
 			getObjectOp.WriteInteger("opId", 0);
-			getObjectOp.WriteString("op", RefBatchOps.RefBatchOp.RefOperation.GET.ToString());
+			getObjectOp.WriteString("op", BatchOps.BatchOp.Operation.GET.ToString());
 			getObjectOp.WriteString("bucket", bucket.ToString());
 			getObjectOp.WriteString("key", newBlobObjectKey.ToString());
 			getObjectOp.EndObject();
@@ -1754,15 +1720,16 @@ namespace Horde.Server.Tests.Ddc.FunctionalTests.References
 			CbWriter getObjectOp2 = new CbWriter();
 			getObjectOp2.BeginObject();
 			getObjectOp2.WriteInteger("opId", 1);
-			getObjectOp2.WriteString("op", RefBatchOps.RefBatchOp.RefOperation.GET.ToString());
+			getObjectOp2.WriteString("op", BatchOps.BatchOp.Operation.GET.ToString());
 			getObjectOp2.WriteString("bucket", bucket.ToString());
 			getObjectOp2.WriteString("key", newReferenceObjectKey.ToString());
 			getObjectOp2.EndObject();
 
+
 			CbWriter getObjectOp3 = new CbWriter();
 			getObjectOp3.BeginObject();
 			getObjectOp3.WriteInteger("opId", 2);
-			getObjectOp3.WriteString("op", RefBatchOps.RefBatchOp.RefOperation.GET.ToString());
+			getObjectOp3.WriteString("op", BatchOps.BatchOp.Operation.GET.ToString());
 			getObjectOp3.WriteString("bucket", bucket.ToString());
 			getObjectOp3.WriteString("key", missingAttachmentKey.ToString());
 			getObjectOp3.WriteBool("resolveAttachments", true);
@@ -1822,8 +1789,8 @@ namespace Horde.Server.Tests.Ddc.FunctionalTests.References
 		{
 			// seed some data
 			BucketId bucket = new BucketId("bucket");
-			RefId newBlobObjectKey = RefId.FromName("newBlobObject");
-			CbObject newBlobObject = CbObject.Build(writer => writer.WriteString("String", "this-has-contents"));
+			RefId newBlobObjectKey = RefId.FromName("newBlobObjectBatchHead");
+			CbObject newBlobObject = CbObject.Build(writer => writer.WriteString("String", $"this-has-contents-in-{nameof(BatchHeadOperationsAsync)}"));
 
 			{
 				byte[] cbObjectBytes = newBlobObject.GetView().ToArray();
@@ -1841,9 +1808,9 @@ namespace Horde.Server.Tests.Ddc.FunctionalTests.References
 				Assert.IsTrue(response.Needs.Length == 0);
 			}
 
-			byte[] blobContents = Encoding.ASCII.GetBytes("This is a attached blob");
+			byte[] blobContents = Encoding.ASCII.GetBytes($"This is a attached blob in {nameof(BatchHeadOperationsAsync)}");
 			CbObject newReferenceObject = CbObject.Build(writer => writer.WriteBinaryAttachment("Attachment", IoHash.Compute(blobContents)));
-			RefId newReferenceObjectKey = RefId.FromName("newReferenceObject");
+			RefId newReferenceObjectKey = RefId.FromName("newReferenceObjectBatchHead");
 
 			{
 				BlobId blobHash = BlobId.FromBlob(blobContents);
@@ -1865,8 +1832,6 @@ namespace Horde.Server.Tests.Ddc.FunctionalTests.References
 				requestContent.Headers.Add(CommonHeaders.HashHeaderName, blobHash.ToString());
 
 				HttpResponseMessage result = await _httpClient!.PutAsync(new Uri($"api/v1/refs/{TestNamespace}/{bucket}/{newReferenceObjectKey}.uecb", UriKind.Relative), requestContent);
-				result.EnsureSuccessStatusCode();
-
 				result.EnsureSuccessStatusCode();
 				byte[] content = await result.Content.ReadAsByteArrayAsync();
 				PutObjectResponse response = CbSerializer.Deserialize<PutObjectResponse>(content);
@@ -1875,7 +1840,7 @@ namespace Horde.Server.Tests.Ddc.FunctionalTests.References
 			CbWriter getObjectOp = new CbWriter();
 			getObjectOp.BeginObject();
 			getObjectOp.WriteInteger("opId", 0);
-			getObjectOp.WriteString("op", RefBatchOps.RefBatchOp.RefOperation.HEAD.ToString());
+			getObjectOp.WriteString("op", BatchOps.BatchOp.Operation.HEAD.ToString());
 			getObjectOp.WriteString("bucket", bucket.ToString());
 			getObjectOp.WriteString("key", newBlobObjectKey.ToString());
 			getObjectOp.EndObject();
@@ -1883,7 +1848,7 @@ namespace Horde.Server.Tests.Ddc.FunctionalTests.References
 			CbWriter getObjectOp2 = new CbWriter();
 			getObjectOp2.BeginObject();
 			getObjectOp2.WriteInteger("opId", 1);
-			getObjectOp2.WriteString("op", RefBatchOps.RefBatchOp.RefOperation.HEAD.ToString());
+			getObjectOp2.WriteString("op", BatchOps.BatchOp.Operation.HEAD.ToString());
 			getObjectOp2.WriteString("bucket", bucket.ToString());
 			getObjectOp2.WriteString("key", newReferenceObjectKey.ToString());
 			getObjectOp2.EndObject();
@@ -1946,7 +1911,7 @@ namespace Horde.Server.Tests.Ddc.FunctionalTests.References
 			CbWriter object0 = new CbWriter();
 			object0.BeginObject();
 			object0.WriteInteger("opId", 0);
-			object0.WriteString("op", RefBatchOps.RefBatchOp.RefOperation.PUT.ToString());
+			object0.WriteString("op", BatchOps.BatchOp.Operation.PUT.ToString());
 			object0.WriteString("bucket", bucket.ToString());
 			object0.WriteString("key", ref0name.ToString());
 			object0.WriteObject("payload", ref0);
@@ -1956,7 +1921,7 @@ namespace Horde.Server.Tests.Ddc.FunctionalTests.References
 			CbWriter object1 = new CbWriter();
 			object1.BeginObject();
 			object1.WriteInteger("opId", 1);
-			object1.WriteString("op", RefBatchOps.RefBatchOp.RefOperation.PUT.ToString());
+			object1.WriteString("op", BatchOps.BatchOp.Operation.PUT.ToString());
 			object1.WriteString("bucket", bucket.ToString());
 			object1.WriteString("key", ref1name.ToString());
 			object1.WriteObject("payload", ref1);
@@ -2014,10 +1979,10 @@ namespace Horde.Server.Tests.Ddc.FunctionalTests.References
 			// seed some data
 			BucketId bucket = new BucketId("bucket");
 			RefId getObjectKey = RefId.FromName("getBlobObject");
-			CbObject getObject = CbObject.Build(writer => writer.WriteString("String", "this-has-contents"));
+			CbObject getObject = CbObject.Build(writer => writer.WriteString("String", $"this-has-contents-in-{nameof(BatchMixedOperationsAsync)}-0"));
 
 			RefId putObjectKey = RefId.FromName("putBlobObject");
-			CbObject putObject = CbObject.Build(writer => writer.WriteString("String", "this-has-contents"));
+			CbObject putObject = CbObject.Build(writer => writer.WriteString("String", $"this-has-contents-in-{nameof(BatchMixedOperationsAsync)}-1"));
 
 			RefId missingObjectKey = RefId.FromName("thisKeyDoesNotExist");
 
@@ -2040,7 +2005,7 @@ namespace Horde.Server.Tests.Ddc.FunctionalTests.References
 			CbWriter getObjectOp = new CbWriter();
 			getObjectOp.BeginObject();
 			getObjectOp.WriteInteger("opId", 0);
-			getObjectOp.WriteString("op", RefBatchOps.RefBatchOp.RefOperation.GET.ToString());
+			getObjectOp.WriteString("op", BatchOps.BatchOp.Operation.GET.ToString());
 			getObjectOp.WriteString("bucket", bucket.ToString());
 			getObjectOp.WriteString("key", getObjectKey.ToString());
 			getObjectOp.EndObject();
@@ -2048,7 +2013,7 @@ namespace Horde.Server.Tests.Ddc.FunctionalTests.References
 			CbWriter putObjectOp = new CbWriter();
 			putObjectOp.BeginObject();
 			putObjectOp.WriteInteger("opId", 1);
-			putObjectOp.WriteString("op", RefBatchOps.RefBatchOp.RefOperation.PUT.ToString());
+			putObjectOp.WriteString("op", BatchOps.BatchOp.Operation.PUT.ToString());
 			putObjectOp.WriteString("bucket", bucket.ToString());
 			putObjectOp.WriteString("key", putObjectKey.ToString());
 			putObjectOp.WriteObject("payload", putObject);
@@ -2058,7 +2023,7 @@ namespace Horde.Server.Tests.Ddc.FunctionalTests.References
 			CbWriter errorObjectOp = new CbWriter();
 			errorObjectOp.BeginObject();
 			errorObjectOp.WriteInteger("opId", 2);
-			errorObjectOp.WriteString("op", RefBatchOps.RefBatchOp.RefOperation.GET.ToString());
+			errorObjectOp.WriteString("op", BatchOps.BatchOp.Operation.GET.ToString());
 			errorObjectOp.WriteString("bucket", bucket.ToString());
 			errorObjectOp.WriteString("key", missingObjectKey.ToString());
 			errorObjectOp.EndObject();
@@ -2066,7 +2031,7 @@ namespace Horde.Server.Tests.Ddc.FunctionalTests.References
 			CbWriter headObjectOp = new CbWriter();
 			headObjectOp.BeginObject();
 			headObjectOp.WriteInteger("opId", 3);
-			headObjectOp.WriteString("op", RefBatchOps.RefBatchOp.RefOperation.HEAD.ToString());
+			headObjectOp.WriteString("op", BatchOps.BatchOp.Operation.HEAD.ToString());
 			headObjectOp.WriteString("bucket", bucket.ToString());
 			headObjectOp.WriteString("key", getObjectKey.ToString());
 			headObjectOp.EndObject();
@@ -2119,7 +2084,7 @@ namespace Horde.Server.Tests.Ddc.FunctionalTests.References
 				Assert.IsNotNull(errorOp.Response);
 				Assert.AreEqual(404, errorOp.StatusCode);
 				Assert.IsTrue(!errorOp.Response["title"].Equals(CbField.Empty));
-				Assert.AreEqual("Object not found cbc3db15b9c8253f6106158962325c8fd848daef in bucket bucket namespace test-namespace", errorOp.Response["title"].AsString());
+				Assert.AreEqual($"Object not found cbc3db15b9c8253f6106158962325c8fd848daef in bucket bucket namespace {TestNamespace}", errorOp.Response["title"].AsString());
 				Assert.AreEqual(404, errorOp.Response["status"].AsInt32());
 
 				BatchOpsResponse.OpResponses headOp = response.Results.First(r => r.OpId == 3);
diff --git a/Engine/Source/Programs/Horde/Horde.Server.Tests/Ddc/Functional/Utils.cs b/Engine/Source/Programs/Horde/Horde.Server.Tests/Ddc/Functional/Utils.cs
new file mode 100644
--- /dev/null
+++ b/Engine/Source/Programs/Horde/Horde.Server.Tests/Ddc/Functional/Utils.cs
@@ -0,0 +1,18 @@
+// Copyright Epic Games, Inc. All Rights Reserved.
+
+using System.Text.Json;
+
+namespace Horde.Server.Tests.Ddc.FunctionalTests
+{
+	public static class JsonTestUtils
+	{
+		public static readonly JsonSerializerOptions DefaultJsonSerializerSettings = ConfigureJsonOptions();
+
+		private static JsonSerializerOptions ConfigureJsonOptions()
+		{
+			JsonSerializerOptions options = new JsonSerializerOptions();
+			Startup.ConfigureJsonSerializer(options);
+			return options;
+		}
+	}
+}
diff --git a/Engine/Source/Programs/Horde/Horde.Server.Tests/Ddc/UnitTests/CompressedBufferTests.cs b/Engine/Source/Programs/Horde/Horde.Server.Tests/Ddc/UnitTests/CompressedBufferTests.cs
--- a/Engine/Source/Programs/Horde/Horde.Server.Tests/Ddc/UnitTests/CompressedBufferTests.cs
+++ b/Engine/Source/Programs/Horde/Horde.Server.Tests/Ddc/UnitTests/CompressedBufferTests.cs
@@ -2,11 +2,12 @@
 
 using System.IO;
 using System.Text;
-using System.Threading;
 using System.Threading.Tasks;
 using EpicGames.Core;
 using Horde.Server.Ddc;
+using Microsoft.Extensions.Options;
 using Microsoft.VisualStudio.TestTools.UnitTesting;
+using Moq;
 using OpenTelemetry.Trace;
 
 namespace Horde.Server.Tests.Ddc.UnitTests
@@ -20,15 +21,21 @@ namespace Horde.Server.Tests.Ddc.UnitTests
 		{
 			byte[] bytes = Encoding.UTF8.GetBytes("this is a test string");
 
-			CompressedBufferUtils bufferUtils = new(TracerProvider.Default.GetTracer("TestTracer"));
+			Tracer tracer = TracerProvider.Default.GetTracer("TestTracer");
+			BufferedPayloadOptions bufferedPayloadOptions = new BufferedPayloadOptions()
+			{
+			};
+			IOptionsMonitor<BufferedPayloadOptions> payloadOptionsMock = Mock.Of<IOptionsMonitor<BufferedPayloadOptions>>(_ => _.CurrentValue == bufferedPayloadOptions);
+			BufferedPayloadFactory bufferedPayloadFactory = new BufferedPayloadFactory(payloadOptionsMock, tracer);
+			CompressedBufferUtils bufferUtils = new(tracer, bufferedPayloadFactory);
 
 			using MemoryStream ms = new MemoryStream();
 			IoHash uncompressedHash = bufferUtils.CompressContent(ms, OoodleCompressorMethod.Mermaid, OoodleCompressionLevel.VeryFast, bytes);
 			ms.Position = 0;
 
-			BufferedPayload bufferedPayload = await bufferUtils.DecompressContentAsync(ms, (ulong)ms.Length, CancellationToken.None);
-
-			byte[] roundTrippedBytes = await bufferedPayload.GetStream().ReadAllBytesAsync();
+			using IBufferedPayload bufferedPayload = await bufferUtils.DecompressContentAsync(ms, (ulong)ms.Length);
+			await using Stream s = bufferedPayload.GetStream();
+			byte[] roundTrippedBytes = await s.ReadAllBytesAsync();
 			CollectionAssert.AreEqual(bytes, roundTrippedBytes);
 		}
 	}
diff --git a/Engine/Source/Programs/Horde/Horde.Server/Ddc/.editorconfig b/Engine/Source/Programs/Horde/Horde.Server/Ddc/.editorconfig
new file mode 100644
--- /dev/null
+++ b/Engine/Source/Programs/Horde/Horde.Server/Ddc/.editorconfig
@@ -0,0 +1,9 @@
+# EditorConfig is awesome: http://EditorConfig.org
+
+[*.cs]
+
+# MA0032: Use an overload with a CancellationToken
+dotnet_diagnostic.MA0032.severity = none
+
+# VSTHRD200: Use "Async" suffix in names of methods that return an awaitable type
+dotnet_diagnostic.VSTHRD200.severity = none
diff --git a/Engine/Source/Programs/Horde/Horde.Server/Ddc/BlobContents.cs b/Engine/Source/Programs/Horde/Horde.Server/Ddc/BlobContents.cs
--- a/Engine/Source/Programs/Horde/Horde.Server/Ddc/BlobContents.cs
+++ b/Engine/Source/Programs/Horde/Horde.Server/Ddc/BlobContents.cs
@@ -11,6 +11,7 @@ namespace Horde.Server.Ddc
 	public sealed class BlobContents : IDisposable, IAsyncDisposable
 	{
 		private readonly Stream? _stream;
+		private readonly IBufferedPayload? _bufferedPayload;
 
 		public BlobContents(Stream stream, long length, string? localPath = null)
 		{
@@ -27,6 +28,14 @@ namespace Horde.Server.Ddc
 			LocalPath = null;
 		}
 
+		public BlobContents(IBufferedPayload payload)
+		{
+			_bufferedPayload = payload;
+			_stream = _bufferedPayload.GetStream();
+			Length = payload.Length;
+			LocalPath = null;
+		}
+
 		public BlobContents(Uri redirectUri)
 		{
 			RedirectUri = redirectUri;
@@ -52,7 +61,9 @@ namespace Horde.Server.Ddc
 		public void Dispose()
 		{
 			_stream?.Dispose();
+			_bufferedPayload?.Dispose();
 		}
+
 		public ValueTask DisposeAsync()
 		{
 			_stream?.Dispose();
diff --git a/Engine/Source/Programs/Horde/Horde.Server/Ddc/BlobId.cs b/Engine/Source/Programs/Horde/Horde.Server/Ddc/BlobId.cs
--- a/Engine/Source/Programs/Horde/Horde.Server/Ddc/BlobId.cs
+++ b/Engine/Source/Programs/Horde/Horde.Server/Ddc/BlobId.cs
@@ -1,86 +1,172 @@
 ﻿// Copyright Epic Games, Inc. All Rights Reserved.
 
 using System;
+using System.Buffers;
 using System.ComponentModel;
 using System.Globalization;
 using System.IO;
 using System.Text.Json;
 using System.Text.Json.Serialization;
-using System.Threading;
 using System.Threading.Tasks;
+using Blake3;
 using EpicGames.Core;
 using EpicGames.Serialization;
 
+#pragma warning disable CS1591
+
 namespace Horde.Server.Ddc
 {
-	/// <summary>
-	/// Identifier for a blob in the store
-	/// </summary>
-	[TypeConverter(typeof(BlobIdTypeConverter))]
-	[JsonConverter(typeof(BlobIdJsonConverter))]
-	[CbConverter(typeof(BlobIdCbConverter))]
-	public readonly struct BlobId : IEquatable<BlobId>
+	[TypeConverter(typeof(BlobIdentifierTypeConverter))]
+	[JsonConverter(typeof(BlobIdentifierJsonConverter))]
+	[CbConverter(typeof(BlobIdentifierCbConverter))]
+	public class BlobId : ContentHash, IEquatable<BlobId>
 	{
-		/// <summary>
-		/// Hash of the blob
-		/// </summary>
-		public IoHash Hash { get; }
-
-		/// <summary>
-		/// Constructor
-		/// </summary>
-		public BlobId(IoHash hash) => Hash = hash;
-
-		/// <summary>
-		/// Parses a BlobId from a string
-		/// </summary>
-		public static BlobId Parse(string text) => new BlobId(IoHash.Parse(text));
-
-		/// <summary>
-		/// Creates a BlobId from a block of data
-		/// </summary>
-		public static BlobId FromBlob(ReadOnlySpan<byte> span) => new BlobId(IoHash.Compute(span));
+		// multi thread the hashing for blobs larger then this size
+		private const int MultiThreadedSize = 1_000_000;
+		private string? _stringIdentifier;
 
-		/// <inheritdoc/>
-		public override readonly int GetHashCode() => Hash.GetHashCode();
+		public BlobId(byte[] identifier) : base(identifier)
+		{
+		}
 
-		/// <inheritdoc/>
-		public readonly bool Equals(BlobId other) => Hash.Equals(other.Hash);
+		[JsonConstructor]
+		public BlobId(string identifier) : base(identifier)
+		{
 
-		/// <inheritdoc/>
-		public override readonly bool Equals(object? obj) => obj is BlobId blobId && Equals(blobId);
+		}
 
-		/// <inheritdoc/>
-		public override readonly string ToString() => Hash.ToString();
+		public override int GetHashCode()
+		{
+			return Comparer.GetHashCode(Identifier);
+		}
 
-		/// <inheritdoc cref="IoHash.op_Equality"/>
-		public static bool operator ==(BlobId left, BlobId right) => left.Hash == right.Hash;
+		public bool Equals(BlobId? other)
+		{
+			if (other == null)
+			{
+				return false;
+			}
 
-		/// <inheritdoc cref="IoHash.op_Inequality"/>
-		public static bool operator !=(BlobId left, BlobId right) => !(left == right);
+			return Comparer.Equals(Identifier, other.Identifier);
+		}
 
-		/// <summary>
-		/// Constructs a BlobId from a hash
-		/// </summary>
-		public static BlobId FromIoHash(IoHash hash) => new BlobId(hash);
+		public override bool Equals(object? obj)
+		{
+			if (ReferenceEquals(null, obj))
+			{
+				return false;
+			}
 
-		/// <summary>
-		/// Creates a BlobId from a data stream
-		/// </summary>
-		public static async Task<BlobId> FromStreamAsync(Stream stream, CancellationToken cancellationToken) => FromIoHash(await IoHash.ComputeAsync(stream, cancellationToken));
+			if (ReferenceEquals(this, obj))
+			{
+				return true;
+			}
+
+			if (obj.GetType() != GetType())
+			{
+				return false;
+			}
+
+			return Equals((BlobId)obj);
+		}
+
+		public override string ToString()
+		{
+			if (_stringIdentifier == null)
+			{
+				_stringIdentifier = StringUtils.FormatAsHexString(Identifier);
+			}
 
-		/// <summary>
-		/// Converts a BlobId to an IoHash
-		/// </summary>
-		public IoHash AsIoHash() => Hash;
+			return _stringIdentifier;
+		}
+
+		public static new BlobId FromBlob(byte[] blobMemory)
+		{
+			Hash blake3Hash;
+			if (blobMemory.Length < MultiThreadedSize)
+			{
+				using Hasher hasher = Hasher.New();
+				hasher.Update(blobMemory);
+				blake3Hash = hasher.Finalize();
+			}
+			else
+			{
+				using Hasher hasher = Hasher.New();
+				hasher.UpdateWithJoin(blobMemory);
+				blake3Hash = hasher.Finalize();
+			}
+
+			// we only keep the first 20 bytes of the Blake3 hash
+			Span<byte> hash = blake3Hash.AsSpanUnsafe().Slice(0, 20);
+			return new BlobId(hash.ToArray());
+		}
+
+		public static BlobId FromBlob(in Memory<byte> blobMemory)
+		{
+			Hash blake3Hash;
+			if (blobMemory.Length < MultiThreadedSize)
+			{
+				using Hasher hasher = Hasher.New();
+				hasher.Update(blobMemory.Span);
+				blake3Hash = hasher.Finalize();
+			}
+			else
+			{
+				using Hasher hasher = Hasher.New();
+				hasher.UpdateWithJoin(blobMemory.Span);
+				blake3Hash = hasher.Finalize();
+			}
+
+			// we only keep the first 20 bytes of the Blake3 hash
+			Span<byte> hash = blake3Hash.AsSpanUnsafe().Slice(0, 20);
+			return new BlobId(hash.ToArray());
+		}
+
+		public static BlobId FromContentHash(ContentHash testObjectHash)
+		{
+			return new BlobId(testObjectHash.HashData);
+		}
+
+		public static async Task<BlobId> FromStreamAsync(Stream stream)
+		{
+			using Hasher hasher = Hasher.New();
+
+			const int BufferSize = 1024 * 1024 * 5;
+			byte[] buffer = ArrayPool<byte>.Shared.Rent(BufferSize);
+
+			try
+			{
+				int read = await stream.ReadAsync(buffer, 0, buffer.Length);
+				while (read > 0)
+				{
+					hasher.UpdateWithJoin(new ReadOnlySpan<byte>(buffer, 0, read));
+					read = await stream.ReadAsync(buffer, 0, buffer.Length);
+				}
+				Hash blake3Hash = hasher.Finalize();
+
+				// we only keep the first 20 bytes of the Blake3 hash
+				byte[] hash = blake3Hash.AsSpanUnsafe().Slice(0, 20).ToArray();
+				return new BlobId(hash);
+			}
+			finally
+			{
+				ArrayPool<byte>.Shared.Return(buffer);
+			}
+		}
+
+		public static BlobId FromIoHash(IoHash blobIdentifier)
+		{
+			return new BlobId(blobIdentifier.ToByteArray());
+		}
+
+		public IoHash AsIoHash()
+		{
+			return new IoHash(HashData);
+		}
 	}
 
-	/// <summary>
-	/// Converts from <see cref="BlobId"/> instances to other types
-	/// </summary>
-	public class BlobIdTypeConverter : TypeConverter
+	public class BlobIdentifierTypeConverter : TypeConverter
 	{
-		/// <inheritdoc/>
 		public override bool CanConvertFrom(ITypeDescriptorContext? context, Type sourceType)
 		{
 			if (sourceType == typeof(string))
@@ -90,18 +176,16 @@ namespace Horde.Server.Ddc
 			return base.CanConvertFrom(context, sourceType);
 		}
 
-		/// <inheritdoc/>
 		public override object? ConvertFrom(ITypeDescriptorContext? context, CultureInfo? culture, object value)
 		{
 			if (value is string s)
 			{
-				return BlobId.Parse(s);
+				return new BlobId(s);
 			}
 
 			return base.ConvertFrom(context, culture, value);
 		}
 
-		/// <inheritdoc/>
 		public override bool CanConvertTo(ITypeDescriptorContext? context, Type? destinationType)
 		{
 			if (destinationType == typeof(string))
@@ -111,25 +195,20 @@ namespace Horde.Server.Ddc
 			return base.CanConvertTo(context, destinationType);
 		}
 
-		/// <inheritdoc/>
 		public override object? ConvertTo(ITypeDescriptorContext? context, CultureInfo? culture, object? value, Type destinationType)
 		{
 			if (destinationType == typeof(string))
 			{
-				return value?.ToString();
+				BlobId? identifier = (BlobId?)value;
+				return identifier?.ToString();
 			}
-
 			return base.ConvertTo(context, culture, value, destinationType);
 		}
 	}
 
-	/// <summary>
-	/// Converts from <see cref="BlobId"/> instances to other types
-	/// </summary>
-	public class BlobIdJsonConverter : JsonConverter<BlobId>
+	public class BlobIdentifierJsonConverter : JsonConverter<BlobId>
 	{
-		/// <inheritdoc/>
-		public override BlobId Read(ref Utf8JsonReader reader, Type typeToConvert, JsonSerializerOptions options)
+		public override BlobId? Read(ref Utf8JsonReader reader, Type typeToConvert, JsonSerializerOptions options)
 		{
 			string? str = reader.GetString();
 			if (str == null)
@@ -137,28 +216,23 @@ namespace Horde.Server.Ddc
 				throw new InvalidDataException("Unable to parse blob identifier");
 			}
 
-			return BlobId.Parse(str);
+			return new BlobId(str);
 		}
 
-		/// <inheritdoc/>
 		public override void Write(Utf8JsonWriter writer, BlobId value, JsonSerializerOptions options)
 		{
 			writer.WriteStringValue(value.ToString());
 		}
 	}
 
-	/// <summary>
-	/// Serializes <see cref="BlobId"/> instances to compact binary
-	/// </summary>
-	public class BlobIdCbConverter : CbConverter<BlobId>
+	public class BlobIdentifierCbConverter : CbConverter<BlobId>
 	{
-		/// <inheritdoc/>
-		public override BlobId Read(CbField field) => new BlobId(field.AsHash());
+		public override BlobId Read(CbField field) => new BlobId(field.AsHash().ToByteArray());
 
 		/// <inheritdoc/>
-		public override void Write(CbWriter writer, BlobId value) => writer.WriteBinaryAttachmentValue(value.Hash);
+		public override void Write(CbWriter writer, BlobId value) => writer.WriteBinaryAttachmentValue(new IoHash(value.HashData));
 
 		/// <inheritdoc/>
-		public override void WriteNamed(CbWriter writer, CbFieldName name, BlobId value) => writer.WriteBinaryAttachment(name, value.Hash);
+		public override void WriteNamed(CbWriter writer, CbFieldName name, BlobId value) => writer.WriteBinaryAttachment(name, new IoHash(value.HashData));
 	}
 }
diff --git a/Engine/Source/Programs/Horde/Horde.Server/Ddc/BlobService.cs b/Engine/Source/Programs/Horde/Horde.Server/Ddc/BlobService.cs
--- a/Engine/Source/Programs/Horde/Horde.Server/Ddc/BlobService.cs
+++ b/Engine/Source/Programs/Horde/Horde.Server/Ddc/BlobService.cs
@@ -1,9 +1,9 @@
 // Copyright Epic Games, Inc. All Rights Reserved.
 
 using System;
+using System.Collections.Concurrent;
 using System.Collections.Generic;
 using System.IO;
-using System.Threading;
 using System.Threading.Tasks;
 using EpicGames.Core;
 using EpicGames.Horde.Storage;
@@ -26,20 +26,23 @@ namespace Horde.Server.Ddc
 
 		public static string GetAlias(BlobId blobId) => $"ddc:{blobId}";
 
-		public async Task VerifyContentMatchesHashAsync(Stream stream, IoHash expectedHash, CancellationToken cancellationToken)
+		public async Task<ContentHash> VerifyContentMatchesHashAsync(Stream content, ContentHash identifier)
 		{
-			IoHash hash;
-			using (TelemetrySpan _ = _tracer.StartActiveSpan("web.hash").SetAttribute("operation.name", "web.hash"))
+			ContentHash blobHash;
 			{
-				hash = await IoHash.ComputeAsync(stream, cancellationToken);
+				using TelemetrySpan _ = _tracer.StartActiveSpan("web.hash").SetAttribute("operation.name", "web.hash");
+				blobHash = await BlobId.FromStreamAsync(content);
 			}
-			if (hash != expectedHash)
+
+			if (!identifier.Equals(blobHash))
 			{
-				throw new HashMismatchException(hash, expectedHash);
+				throw new HashMismatchException(identifier, blobHash);
 			}
+
+			return identifier;
 		}
 
-		public async Task<BlobId> PutObjectAsync(NamespaceId ns, BufferedPayload payload, BlobId identifier, CancellationToken cancellationToken)
+		public async Task<BlobId> PutObjectAsync(NamespaceId ns, IBufferedPayload payload, BlobId identifier)
 		{
 			using TelemetrySpan scope = _tracer.StartActiveSpan("put_blob")
 				.SetAttribute("operation.name", "put_blob")
@@ -47,44 +50,44 @@ namespace Horde.Server.Ddc
 				.SetAttribute("Content-Length", payload.Length.ToString());
 
 			await using Stream hashStream = payload.GetStream();
-			await VerifyContentMatchesHashAsync(hashStream, identifier.Hash, cancellationToken);
+			await VerifyContentMatchesHashAsync(hashStream, identifier);
 
-			await PutObjectKnownHashAsync(ns, payload, identifier, cancellationToken);
+			await PutObjectKnownHashAsync(ns, payload, identifier);
 			return identifier;
 		}
 
-		public Task<BlobId> PutObjectAsync(NamespaceId ns, ReadOnlyMemory<byte> payload, BlobId identifier, CancellationToken cancellationToken)
+		public Task<BlobId> PutObjectAsync(NamespaceId ns, byte[] payload, BlobId identifier)
 		{
 			using MemoryBufferedPayload bufferedPayload = new MemoryBufferedPayload(payload);
-			return PutObjectAsync(ns, bufferedPayload, identifier, cancellationToken);
+			return PutObjectAsync(ns, bufferedPayload, identifier);
 		}
 
-		public async Task<bool> ExistsAsync(NamespaceId ns, BlobId blob, List<string>? storageLayers, CancellationToken cancellationToken)
+		public async Task<bool> ExistsAsync(NamespaceId ns, BlobId blob, List<string>? storageLayers = null)
 		{
 			using IStorageClient storageClient = _storageClientFactory.CreateClient(ns);
-			return await storageClient.FindAliasAsync(GetAlias(blob), cancellationToken) != null;
+			return await storageClient.FindAliasAsync(GetAlias(blob)) != null;
 		}
 
-		public async Task DeleteObjectAsync(NamespaceId ns, BlobId blob, CancellationToken cancellationToken)
+		public async Task DeleteObjectAsync(NamespaceId ns, BlobId blob)
 		{
 			using IStorageClient storageClient = _storageClientFactory.CreateClient(ns);
 			string aliasName = GetAlias(blob);
 
-			BlobAlias[] aliases = await storageClient.FindAliasesAsync(aliasName, cancellationToken: cancellationToken);
+			BlobAlias[] aliases = await storageClient.FindAliasesAsync(aliasName);
 			foreach (BlobAlias alias in aliases)
 			{
-				await storageClient.RemoveAliasAsync(aliasName, alias.Target, cancellationToken);
+				await storageClient.RemoveAliasAsync(aliasName, alias.Target);
 			}
 		}
 
-		public async Task<BlobId[]> FilterOutKnownBlobsAsync(NamespaceId ns, IEnumerable<BlobId> blobIds, CancellationToken cancellationToken)
+		public async Task<BlobId[]> FilterOutKnownBlobsAsync(NamespaceId ns, IEnumerable<BlobId> blobIds)
 		{
 			using IStorageClient storageClient = _storageClientFactory.CreateClient(ns);
 
 			List<BlobId> unknownBlobIds = new List<BlobId>();
 			foreach (BlobId blobId in blobIds)
 			{
-				if (await storageClient.FindAliasAsync(GetAlias(blobId), cancellationToken) == null)
+				if (await storageClient.FindAliasAsync(GetAlias(blobId)) == null)
 				{
 					unknownBlobIds.Add(blobId);
 				}
@@ -93,27 +96,61 @@ namespace Horde.Server.Ddc
 			return unknownBlobIds.ToArray();
 		}
 
-		public async Task<BlobContents> GetObjectAsync(NamespaceId ns, BlobId blob, List<string>? storageLayers, bool supportsRedirectUri, CancellationToken cancellationToken)
+		public async Task<BlobId[]> FilterOutKnownBlobsAsync(NamespaceId ns, IAsyncEnumerable<BlobId> blobs)
+		{
+			ConcurrentBag<BlobId> missingBlobs = new ConcurrentBag<BlobId>();
+
+			try
+			{
+				await Parallel.ForEachAsync(blobs, async (identifier, ctx) =>
+				{
+					bool exists = await ExistsAsync(ns, identifier);
+
+					if (!exists)
+					{
+						missingBlobs.Add(identifier);
+					}
+				});
+			}
+			catch (AggregateException e)
+			{
+				if (e.InnerException is PartialReferenceResolveException)
+				{
+					throw e.InnerException;
+				}
+
+				throw;
+			}
+
+			return missingBlobs.ToArray();
+		}
+
+		public async Task<BlobContents> GetObjectAsync(NamespaceId ns, BlobId blob, List<string>? storageLayers, bool supportsRedirectUri, bool allowOndemandReplication = true)
 		{
 			using IStorageClient storageClient = _storageClientFactory.CreateClient(ns);
 
-			BlobAlias? alias = await storageClient.FindAliasAsync(GetAlias(blob), cancellationToken);
+			BlobAlias? alias = await storageClient.FindAliasAsync(GetAlias(blob));
 			if (alias == null)
 			{
 				throw new BlobNotFoundException(ns, blob);
 			}
 
-			using BlobData data = await alias.Target.ReadBlobDataAsync(cancellationToken);
+			using BlobData data = await alias.Target.ReadBlobDataAsync();
 			return new BlobContents(data.Data.ToArray());
 		}
 
-		public async Task<BlobContents> GetObjectsAsync(NamespaceId ns, BlobId[] blobs, CancellationToken cancellationToken)
+		public Task<BlobMetadata> GetObjectMetadataAsync(NamespaceId ns, BlobId blobId)
+		{
+			throw new NotImplementedException();
+		}
+
+		public async Task<BlobContents> GetObjectsAsync(NamespaceId ns, BlobId[] blobs)
 		{
 			using TelemetrySpan _ = _tracer.StartActiveSpan("blob.combine").SetAttribute("operation.name", "blob.combine");
 			Task<BlobContents>[] tasks = new Task<BlobContents>[blobs.Length];
 			for (int i = 0; i < blobs.Length; i++)
 			{
-				tasks[i] = GetObjectAsync(ns, blobs[i], storageLayers: null, supportsRedirectUri: false, cancellationToken);
+				tasks[i] = GetObjectAsync(ns, blobs[i], storageLayers: null, supportsRedirectUri: false);
 			}
 
 			MemoryStream ms = new MemoryStream();
@@ -121,7 +158,7 @@ namespace Horde.Server.Ddc
 			{
 				BlobContents blob = await task;
 				await using Stream s = blob.Stream;
-				await s.CopyToAsync(ms, cancellationToken);
+				await s.CopyToAsync(ms);
 			}
 
 			ms.Seek(0, SeekOrigin.Begin);
@@ -129,17 +166,17 @@ namespace Horde.Server.Ddc
 			return new BlobContents(ms, ms.Length);
 		}
 
-		public Task<Uri?> GetObjectWithRedirectAsync(NamespaceId ns, BlobId blobIdentifier, List<string>? storageLayers, CancellationToken cancellationToken)
+		public Task<Uri?> GetObjectWithRedirectAsync(NamespaceId ns, BlobId blobIdentifier, List<string>? storageLayers)
 		{
 			return Task.FromResult<Uri?>(null);
 		}
 
-		public Task<Uri?> MaybePutObjectWithRedirectAsync(NamespaceId ns, BlobId identifier, CancellationToken cancellationToken)
+		public Task<Uri?> MaybePutObjectWithRedirectAsync(NamespaceId ns, BlobId identifier)
 		{
 			return Task.FromResult<Uri?>(null);
 		}
 
-		public async Task<BlobId> PutObjectKnownHashAsync(NamespaceId ns, BufferedPayload content, BlobId identifier, CancellationToken cancellationToken)
+		public async Task<BlobId> PutObjectKnownHashAsync(NamespaceId ns, IBufferedPayload content, BlobId identifier)
 		{
 			using IStorageClient storageClient = _storageClientFactory.CreateClient(ns);
 
@@ -150,15 +187,30 @@ namespace Horde.Server.Ddc
 				Memory<byte> memory = writer.GetMemory(contentLength).Slice(0, contentLength);
 
 				using Stream stream = content.GetStream();
-				await stream.ReadFixedLengthBytesAsync(memory, cancellationToken);
+				await stream.ReadFixedLengthBytesAsync(memory);
 				writer.Advance(contentLength);
 
-				blobRef = await writer.CompleteAsync(s_rawBlobType, cancellationToken);
-				await writer.FlushAsync(cancellationToken);
+				blobRef = await writer.CompleteAsync(s_rawBlobType);
+				await writer.FlushAsync();
 			}
 
-			await storageClient.AddAliasAsync(GetAlias(identifier), blobRef, data: identifier.Hash.ToByteArray(), cancellationToken: cancellationToken);
+			await storageClient.AddAliasAsync(GetAlias(identifier), blobRef, data: identifier.AsIoHash().ToByteArray());
 			return identifier;
 		}
+
+		public Task<BlobContents> ReplicateObjectAsync(NamespaceId ns, BlobId blob, bool force = false)
+			=> throw new NotSupportedException();
+
+		public Task<bool> ExistsInRootStoreAsync(NamespaceId ns, BlobId blob)
+			=> Task.FromResult(false);
+
+		public Task DeleteNamespaceAsync(NamespaceId ns)
+			=> throw new NotSupportedException();
+
+		public IAsyncEnumerable<(BlobId, DateTime)> ListObjectsAsync(NamespaceId ns)
+			=> throw new NotSupportedException();
+
+		public bool ShouldFetchBlobOnDemand(NamespaceId ns)
+			=> false;
 	}
 }
diff --git a/Engine/Source/Programs/Horde/Horde.Server/Ddc/BlobsController.cs b/Engine/Source/Programs/Horde/Horde.Server/Ddc/BlobsController.cs
--- a/Engine/Source/Programs/Horde/Horde.Server/Ddc/BlobsController.cs
+++ b/Engine/Source/Programs/Horde/Horde.Server/Ddc/BlobsController.cs
@@ -8,13 +8,11 @@ using System.IO;
 using System.Linq;
 using System.Net;
 using System.Net.Mime;
-using System.Threading;
 using System.Threading.Tasks;
 using EpicGames.AspNet;
 using EpicGames.Horde.Acls;
 using EpicGames.Horde.Storage;
 using EpicGames.Serialization;
-using Horde.Server.Storage;
 using Microsoft.AspNetCore.Authorization;
 using Microsoft.AspNetCore.Http;
 using Microsoft.AspNetCore.Mvc;
@@ -52,9 +50,10 @@ namespace Horde.Server.Ddc
 		public async Task<IActionResult> GetAsync(
 			[Required] NamespaceId ns,
 			[Required] BlobId id,
-			[FromQuery] List<string>? storageLayers = null)
+			[FromQuery] List<string>? storageLayers = null,
+			[FromQuery] bool allowOndemandReplication = true)
 		{
-			ActionResult? result = await _requestHelper.HasAccessToNamespaceAsync(User, Request, ns, new[] { StorageAclAction.ReadBlobs });
+			ActionResult? result = await _requestHelper.HasAccessToNamespaceAsync(User, Request, ns, new[] { JupiterAclAction.ReadObject });
 			if (result != null)
 			{
 				return result;
@@ -62,7 +61,7 @@ namespace Horde.Server.Ddc
 
 			try
 			{
-				BlobContents blobContents = await GetImplAsync(ns, id, storageLayers, supportsRedirectUri: true);
+				BlobContents blobContents = await GetImplAsync(ns, id, storageLayers, supportsRedirectUri: true, allowOndemandReplication: allowOndemandReplication);
 
 				if (blobContents.RedirectUri != null)
 				{
@@ -91,7 +90,7 @@ namespace Horde.Server.Ddc
 			[Required] BlobId id,
 			[FromQuery] List<string>? storageLayers = null)
 		{
-			ActionResult? result = await _requestHelper.HasAccessToNamespaceAsync(User, Request, ns, new[] { StorageAclAction.ReadBlobs });
+			ActionResult? result = await _requestHelper.HasAccessToNamespaceAsync(User, Request, ns, new[] { JupiterAclAction.ReadObject });
 			if (result != null)
 			{
 				return result;
@@ -112,7 +111,7 @@ namespace Horde.Server.Ddc
 			[Required] NamespaceId ns,
 			[Required][FromQuery] List<BlobId> id)
 		{
-			ActionResult? result = await _requestHelper.HasAccessToNamespaceAsync(User, Request, ns, new[] { StorageAclAction.ReadBlobs });
+			ActionResult? result = await _requestHelper.HasAccessToNamespaceAsync(User, Request, ns, new[] { JupiterAclAction.ReadObject });
 			if (result != null)
 			{
 				return result;
@@ -138,7 +137,7 @@ namespace Horde.Server.Ddc
 			[Required] NamespaceId ns,
 			[FromBody] BlobId[] bodyIds)
 		{
-			ActionResult? result = await _requestHelper.HasAccessToNamespaceAsync(User, Request, ns, new[] { StorageAclAction.ReadBlobs });
+			ActionResult? result = await _requestHelper.HasAccessToNamespaceAsync(User, Request, ns, new[] { JupiterAclAction.ReadObject });
 			if (result != null)
 			{
 				return result;
@@ -158,9 +157,21 @@ namespace Horde.Server.Ddc
 			return Ok(new HeadMultipleResponse { Needs = missingBlobs.ToArray() });
 		}
 
-		private async Task<BlobContents> GetImplAsync(NamespaceId ns, BlobId blob, List<string>? storageLayers = null, bool supportsRedirectUri = false)
+		private async Task<BlobContents> GetImplAsync(NamespaceId ns, BlobId blob, List<string>? storageLayers = null, bool supportsRedirectUri = false, bool allowOndemandReplication = true)
 		{
-			return await _storage.GetObjectAsync(ns, blob, storageLayers, supportsRedirectUri);
+			try
+			{
+				return await _storage.GetObjectAsync(ns, blob, storageLayers, supportsRedirectUri, allowOndemandReplication);
+			}
+			catch (BlobNotFoundException)
+			{
+				if (!_storage.ShouldFetchBlobOnDemand(ns) || !allowOndemandReplication)
+				{
+					throw;
+				}
+
+				return await _storage.ReplicateObjectAsync(ns, blob);
+			}
 		}
 
 		[HttpPut("{ns}/{id}")]
@@ -168,10 +179,9 @@ namespace Horde.Server.Ddc
 		[DisableRequestSizeLimit]
 		public async Task<IActionResult> PutAsync(
 			[Required] NamespaceId ns,
-			[Required] BlobId id,
-			CancellationToken cancellationToken)
+			[Required] BlobId id)
 		{
-			ActionResult? result = await _requestHelper.HasAccessToNamespaceAsync(User, Request, ns, new[] { StorageAclAction.WriteBlobs });
+			ActionResult? result = await _requestHelper.HasAccessToNamespaceAsync(User, Request, ns, new[] { JupiterAclAction.WriteObject });
 			if (result != null)
 			{
 				return result;
@@ -181,7 +191,7 @@ namespace Horde.Server.Ddc
 
 			try
 			{
-				Uri? uri = await _storage.MaybePutObjectWithRedirectAsync(ns, id, cancellationToken);
+				Uri? uri = await _storage.MaybePutObjectWithRedirectAsync(ns, id);
 				if (uri != null)
 				{
 					return Ok(new
@@ -190,9 +200,9 @@ namespace Horde.Server.Ddc
 						RedirectUri = uri,
 					});
 				}
-				using BufferedPayload payload = await _bufferedPayloadFactory.CreateFromRequestAsync(Request);
+				using IBufferedPayload payload = await _bufferedPayloadFactory.CreateFromRequest(Request);
 
-				BlobId identifier = await _storage.PutObjectAsync(ns, payload, id, cancellationToken);
+				BlobId identifier = await _storage.PutObjectAsync(ns, payload, id);
 				return Ok(new
 				{
 					Identifier = identifier.ToString()
@@ -212,10 +222,9 @@ namespace Horde.Server.Ddc
 		[RequiredContentType(MediaTypeNames.Application.Octet)]
 		[DisableRequestSizeLimit]
 		public async Task<IActionResult> PostAsync(
-			[Required] NamespaceId ns,
-			CancellationToken cancellationToken)
+			[Required] NamespaceId ns)
 		{
-			ActionResult? result = await _requestHelper.HasAccessToNamespaceAsync(User, Request, ns, new[] { StorageAclAction.WriteBlobs });
+			ActionResult? result = await _requestHelper.HasAccessToNamespaceAsync(User, Request, ns, new[] { JupiterAclAction.WriteObject });
 			if (result != null)
 			{
 				return result;
@@ -224,12 +233,12 @@ namespace Horde.Server.Ddc
 			_diagnosticContext.Set("Content-Length", Request.ContentLength ?? -1);
 			try
 			{
-				using BufferedPayload payload = await _bufferedPayloadFactory.CreateFromRequestAsync(Request);
+				using IBufferedPayload payload = await _bufferedPayloadFactory.CreateFromRequest(Request);
 
 				await using Stream stream = payload.GetStream();
 
-				BlobId id = await BlobId.FromStreamAsync(stream, cancellationToken);
-				await _storage.PutObjectKnownHashAsync(ns, payload, id, cancellationToken);
+				BlobId id = await BlobId.FromStreamAsync(stream);
+				await _storage.PutObjectKnownHashAsync(ns, payload, id);
 
 				return Ok(new
 				{
@@ -242,6 +251,42 @@ namespace Horde.Server.Ddc
 			}
 		}
 
+		[HttpDelete("{ns}/{id}")]
+		public async Task<IActionResult> DeleteAsync(
+			[Required] NamespaceId ns,
+			[Required] BlobId id)
+		{
+			ActionResult? result = await _requestHelper.HasAccessToNamespaceAsync(User, Request, ns, new[] { JupiterAclAction.DeleteObject });
+			if (result != null)
+			{
+				return result;
+			}
+
+			await DeleteImplAsync(ns, id);
+
+			return NoContent();
+		}
+
+		[HttpDelete("{ns}")]
+		public async Task<IActionResult> DeleteNamespaceAsync(
+			[Required] NamespaceId ns)
+		{
+			ActionResult? result = await _requestHelper.HasAccessToNamespaceAsync(User, Request, ns, new[] { JupiterAclAction.DeleteNamespace });
+			if (result != null)
+			{
+				return result;
+			}
+
+			await _storage.DeleteNamespaceAsync(ns);
+
+			return NoContent();
+		}
+
+		private async Task DeleteImplAsync(NamespaceId ns, BlobId id)
+		{
+			await _storage.DeleteObjectAsync(ns, id);
+		}
+
 		// ReSharper disable UnusedAutoPropertyAccessor.Global
 		// ReSharper disable once ClassNeverInstantiated.Global
 		[System.Diagnostics.CodeAnalysis.SuppressMessage("Design", "CA1034:Nested types should not be visible", Justification = "Used by serialization")]
@@ -253,6 +298,7 @@ namespace Horde.Server.Ddc
 				INVALID,
 				GET,
 				PUT,
+				DELETE,
 				HEAD
 			}
 
@@ -273,17 +319,19 @@ namespace Horde.Server.Ddc
 		// ReSharper restore UnusedAutoPropertyAccessor.Global
 
 		[HttpPost("")]
-		public async Task<IActionResult> PostAsync([FromBody] BatchCall batch, CancellationToken cancellationToken)
+		public async Task<IActionResult> PostAsync([FromBody] BatchCall batch)
 		{
-			static AclAction MapToAclAction(BatchOp.Operation op)
+			AclAction MapToAclAction(BatchOp.Operation op)
 			{
 				switch (op)
 				{
 					case BatchOp.Operation.GET:
 					case BatchOp.Operation.HEAD:
-						return StorageAclAction.ReadBlobs;
+						return JupiterAclAction.ReadObject;
 					case BatchOp.Operation.PUT:
-						return StorageAclAction.WriteBlobs;
+						return JupiterAclAction.WriteObject;
+					case BatchOp.Operation.DELETE:
+						return JupiterAclAction.DeleteObject;
 					default:
 						throw new ArgumentOutOfRangeException(nameof(op), op, null);
 				}
@@ -320,7 +368,7 @@ namespace Horde.Server.Ddc
 							return BadRequest();
 						}
 
-						tasks[index] = GetImplAsync(op.Namespace.Value, op.Id.Value).ContinueWith((t, _) =>
+						tasks[index] = GetImplAsync(op.Namespace.Value, op.Id).ContinueWith((t, _) =>
 						{
 							// TODO: This is very allocation heavy but given that the end result is a json object we can not really stream this anyway
 							using BlobContents blobContents = t.Result;
@@ -338,7 +386,7 @@ namespace Horde.Server.Ddc
 							return BadRequest();
 						}
 
-						tasks[index] = _storage.ExistsAsync(op.Namespace.Value, op.Id.Value, cancellationToken: cancellationToken)
+						tasks[index] = _storage.ExistsAsync(op.Namespace.Value, op.Id)
 							.ContinueWith((t, _) => t.Result ? (object?)null : op.Id, null, TaskScheduler.Current);
 						break;
 					case BatchOp.Operation.PUT:
@@ -354,9 +402,17 @@ namespace Horde.Server.Ddc
 							}
 
 							using MemoryBufferedPayload payload = new MemoryBufferedPayload(op.Content);
-							tasks[index] = _storage.PutObjectAsync(op.Namespace.Value, payload, op.Id.Value, cancellationToken).ContinueWith((t, _) => (object?)t.Result, null, TaskScheduler.Current);
+							tasks[index] = _storage.PutObjectAsync(op.Namespace.Value, payload, op.Id).ContinueWith((t, _) => (object?)t.Result, null, TaskScheduler.Current);
 							break;
 						}
+					case BatchOp.Operation.DELETE:
+						if (op.Id == null)
+						{
+							return BadRequest();
+						}
+
+						tasks[index] = DeleteImplAsync(op.Namespace.Value, op.Id).ContinueWith((t, _) => (object?)null, null, TaskScheduler.Current);
+						break;
 					default:
 						throw new NotImplementedException($"{op.Op} is not a support op type");
 				}
diff --git a/Engine/Source/Programs/Horde/Horde.Server/Ddc/BucketId.cs b/Engine/Source/Programs/Horde/Horde.Server/Ddc/BucketId.cs
--- a/Engine/Source/Programs/Horde/Horde.Server/Ddc/BucketId.cs
+++ b/Engine/Source/Programs/Horde/Horde.Server/Ddc/BucketId.cs
@@ -58,7 +58,7 @@ namespace Horde.Server.Ddc
 	sealed class BucketIdJsonConverter : JsonConverter<BucketId>
 	{
 		/// <inheritdoc/>
-		public override BucketId Read(ref Utf8JsonReader reader, Type typeToConvert, JsonSerializerOptions options) => new BucketId(reader.GetString() ?? String.Empty);
+		public override BucketId Read(ref Utf8JsonReader reader, Type typeToConvert, JsonSerializerOptions options) => new BucketId(reader.GetString() ?? string.Empty);
 
 		/// <inheritdoc/>
 		public override void Write(Utf8JsonWriter writer, BucketId value, JsonSerializerOptions options) => writer.WriteStringValue(value.ToString());
diff --git a/Engine/Source/Programs/Horde/Horde.Server/Ddc/BufferedPayload.cs b/Engine/Source/Programs/Horde/Horde.Server/Ddc/BufferedPayload.cs
--- a/Engine/Source/Programs/Horde/Horde.Server/Ddc/BufferedPayload.cs
+++ b/Engine/Source/Programs/Horde/Horde.Server/Ddc/BufferedPayload.cs
@@ -2,95 +2,66 @@
 
 using System;
 using System.IO;
-using System.Threading;
 using System.Threading.Tasks;
-using EpicGames.Core;
 using Microsoft.AspNetCore.Http;
 using Microsoft.Extensions.Options;
 using OpenTelemetry.Trace;
 
+#pragma warning disable CS1591
+
 namespace Horde.Server.Ddc
 {
-	/// <summary>
-	/// Base class for a payload stream that supports seeking
-	/// </summary>
-	public abstract class BufferedPayload : IDisposable
+	public interface IBufferedPayload : IDisposable
 	{
-		/// <summary>
-		/// Length of the payload
-		/// </summary>
-		public long Length { get; }
-
-		/// <summary>
-		/// Constructor
-		/// </summary>
-		protected BufferedPayload(long length)
-		{
-			Length = length;
-		}
-
-		/// <inheritdoc/>
-		public void Dispose()
-		{
-			Dispose(true);
-			GC.SuppressFinalize(this);
-		}
-
-		/// <summary>
-		/// Overridable implementation of <see cref="Dispose()"/>
-		/// </summary>
-		protected virtual void Dispose(bool disposing)
-		{
-		}
-
-		/// <summary>
-		/// Opens a stream to the payload
-		/// </summary>
-		public abstract Stream GetStream();
+		Stream GetStream();
+		long Length { get; }
 	}
 
 	/// <summary>
 	/// Streaming request that is streamed into memory
 	/// </summary>
-	public sealed class MemoryBufferedPayload : BufferedPayload
+	public sealed class MemoryBufferedPayload : IBufferedPayload
 	{
-		private readonly ReadOnlyMemory<byte> _buffer;
+		private readonly byte[] _buffer;
 
-		/// <summary>
-		/// Constructor
-		/// </summary>
-		public MemoryBufferedPayload(ReadOnlyMemory<byte> source)
-			: base(source.Length)
+		public MemoryBufferedPayload(byte[] source)
 		{
 			_buffer = source;
 		}
 
-		/// <summary>
-		/// Create a buffered payload from a stream
-		/// </summary>
-		public static async Task<MemoryBufferedPayload> CreateAsync(Tracer tracer, Stream s, CancellationToken cancellationToken = default)
+		public static async Task<MemoryBufferedPayload> CreateAsync(Tracer tracer, Stream s)
 		{
 			using TelemetrySpan scope = tracer.StartActiveSpan("payload.buffer")
 				.SetAttribute("operation.name", "payload.buffer")
 				.SetAttribute("bufferType", "Memory");
-			MemoryBufferedPayload payload = new MemoryBufferedPayload(await s.ToByteArrayAsync(cancellationToken));
+			MemoryBufferedPayload payload = new MemoryBufferedPayload(await s.ToByteArrayAsync());
+
 			return payload;
 		}
 
-		/// <inheritdoc/>
-		public override Stream GetStream() => new ReadOnlyMemoryStream(_buffer);
+		public void Dispose()
+		{
+
+		}
+
+		public Stream GetStream()
+		{
+			return new MemoryStream(_buffer);
+		}
+
+		public long Length => _buffer.LongLength;
 	}
 
 	/// <summary>
 	/// Helper to generate a filesystem buffered payload from a stream 
 	/// </summary>
-	sealed class FilesystemBufferedPayloadWriter : IDisposable
+	public sealed class FilesystemBufferedPayloadWriter : IDisposable
 	{
 		private FileInfo? _tempFile;
 
-		public FilesystemBufferedPayloadWriter()
+		private FilesystemBufferedPayloadWriter(string filesystemRoot)
 		{
-			_tempFile = new FileInfo(Path.GetTempFileName());
+			_tempFile = new FileInfo(Path.Combine(filesystemRoot, Path.GetRandomFileName()));
 		}
 
 		public void Dispose()
@@ -123,88 +94,96 @@ namespace Horde.Server.Ddc
 
 			return _tempFile.OpenWrite();
 		}
+
+		public static FilesystemBufferedPayloadWriter Create(string filesystemTempPayloadRoot)
+		{
+			return new FilesystemBufferedPayloadWriter(filesystemTempPayloadRoot);
+		}
 	}
 
 	/// <summary>
 	/// A streaming request backed by a temporary file on disk
 	/// </summary>
-	public sealed class FilesystemBufferedPayload : BufferedPayload
+	public sealed class FilesystemBufferedPayload : IBufferedPayload
 	{
 		private readonly FileInfo _tempFile;
+		private long _length;
 
-		internal FilesystemBufferedPayload(FileInfo tempFile)
-			: base(tempFile.Length)
+		public FileInfo TempFile => _tempFile;
+
+		private FilesystemBufferedPayload(string filesystemRoot)
 		{
-			_tempFile = tempFile;
+			_tempFile = new FileInfo(Path.Combine(filesystemRoot, Path.GetRandomFileName()));
 		}
 
-		/// <summary>
-		/// Create a new payload instance backed by the filesystem
-		/// </summary>
-		public static async Task<FilesystemBufferedPayload> CreateAsync(Tracer tracer, Stream s, CancellationToken cancellationToken)
+		internal FilesystemBufferedPayload(FileInfo bufferFile)
+		{
+			_tempFile = bufferFile;
+			_tempFile.Refresh();
+			_length = _tempFile.Length;
+		}
+
+		public static async Task<FilesystemBufferedPayload> CreateAsync(Tracer tracer, Stream s, string filesystemRoot)
 		{
-			FileInfo tempFile = new FileInfo(Path.GetTempFileName());
+			FilesystemBufferedPayload payload = new FilesystemBufferedPayload(filesystemRoot);
 
 			{
 				using TelemetrySpan? scope = tracer.StartActiveSpan("payload.buffer")
 					.SetAttribute("operation.name", "payload.buffer")
 					.SetAttribute("bufferType", "Filesystem");
-				await using FileStream fs = tempFile.OpenWrite();
-				await s.CopyToAsync(fs, cancellationToken);
+				await using FileStream fs = payload._tempFile.OpenWrite();
+				await s.CopyToAsync(fs);
 			}
 
-			tempFile.Refresh();
+			payload._tempFile.Refresh();
+			payload._length = payload._tempFile.Length;
 
-			return new FilesystemBufferedPayload(tempFile);
+			return payload;
 		}
 
-		/// <inheritdoc/>
-		protected override void Dispose(bool disposing)
+		public void Dispose()
 		{
-			base.Dispose(disposing);
-
 			if (_tempFile.Exists)
 			{
 				_tempFile.Delete();
 			}
 		}
 
-		/// <inheritdoc/>
-		public override Stream GetStream() => _tempFile.OpenRead();
+		public Stream GetStream()
+		{
+			return _tempFile.OpenRead();
+		}
+
+		public long Length => _length;
 	}
 
-	/// <summary>
-	/// Options for creating <see cref="BufferedPayload"/> instances
-	/// </summary>
 	public class BufferedPayloadOptions
 	{
 		/// <summary>
 		/// If the request is smaller then MemoryBufferSize we buffer it in memory rather then as a file
 		/// </summary>
-		public long MemoryBufferSize { get; set; } = 128 * 1024 * 1024;
+		public long MemoryBufferSize { get; set; } = int.MaxValue;
+
+		/// <summary>
+		/// The default root to create temporary buffered files under, defaults to %TEMP% or /tmp
+		/// </summary>
+		public string FilesystemTempPayloadRoot { get; set; } = Path.GetTempPath();
 	}
 
-	/// <summary>
-	/// Factory for creating <see cref="BufferedPayload"/> instances
-	/// </summary>
 	public class BufferedPayloadFactory
 	{
 		private readonly IOptionsMonitor<BufferedPayloadOptions> _options;
 		private readonly Tracer _tracer;
 
-		/// <summary>
-		/// Constructor
-		/// </summary>
 		public BufferedPayloadFactory(IOptionsMonitor<BufferedPayloadOptions> options, Tracer tracer)
 		{
 			_options = options;
 			_tracer = tracer;
+
+			Directory.CreateDirectory(options.CurrentValue.FilesystemTempPayloadRoot);
 		}
 
-		/// <summary>
-		/// Create a new buffered payload from an HTTP request
-		/// </summary>
-		public Task<BufferedPayload> CreateFromRequestAsync(HttpRequest request)
+		public Task<IBufferedPayload> CreateFromRequest(HttpRequest request)
 		{
 			long? contentLength = request.ContentLength;
 
@@ -216,18 +195,25 @@ namespace Horde.Server.Ddc
 			return CreateFromStreamAsync(request.Body, contentLength.Value);
 		}
 
-		/// <summary>
-		/// Create a new buffered payload instance from a stream
-		/// </summary>
-		public async Task<BufferedPayload> CreateFromStreamAsync(Stream s, long contentLength, CancellationToken cancellationToken = default)
+		public async Task<IBufferedPayload> CreateFromStreamAsync(Stream s, long contentLength)
 		{
 			// blob is small enough to fit into memory we just read it as is
 			if (contentLength < _options.CurrentValue.MemoryBufferSize)
 			{
-				return await MemoryBufferedPayload.CreateAsync(_tracer, s, cancellationToken);
+				return await MemoryBufferedPayload.CreateAsync(_tracer, s);
 			}
 
-			return await FilesystemBufferedPayload.CreateAsync(_tracer, s, cancellationToken);
+			return await FilesystemBufferedPayload.CreateAsync(_tracer, s, _options.CurrentValue.FilesystemTempPayloadRoot);
+		}
+
+		public async Task<IBufferedPayload> CreateFilesystemBufferedPayloadAsync(Stream s)
+		{
+			return await FilesystemBufferedPayload.CreateAsync(_tracer, s, _options.CurrentValue.FilesystemTempPayloadRoot);
+		}
+
+		public FilesystemBufferedPayloadWriter CreateFilesystemBufferedPayloadWriter()
+		{
+			return FilesystemBufferedPayloadWriter.Create(_options.CurrentValue.FilesystemTempPayloadRoot);
 		}
 	}
 }
diff --git a/Engine/Source/Programs/Horde/Horde.Server/Ddc/CompressedBlobController.cs b/Engine/Source/Programs/Horde/Horde.Server/Ddc/CompressedBlobController.cs
--- a/Engine/Source/Programs/Horde/Horde.Server/Ddc/CompressedBlobController.cs
+++ b/Engine/Source/Programs/Horde/Horde.Server/Ddc/CompressedBlobController.cs
@@ -6,14 +6,11 @@ using System.ComponentModel.DataAnnotations;
 using System.Linq;
 using System.Net;
 using System.Net.Mime;
-using System.Threading;
 using System.Threading.Tasks;
 using EpicGames.AspNet;
 using EpicGames.Horde.Storage;
 using EpicGames.Serialization;
-using Horde.Server.Storage;
 using Microsoft.AspNetCore.Authorization;
-using Microsoft.AspNetCore.Http;
 using Microsoft.AspNetCore.Mvc;
 using Microsoft.Extensions.Primitives;
 
@@ -26,17 +23,16 @@ namespace Horde.Server.Ddc
 	[ApiController]
 	[Authorize]
 	[Route("api/v1/compressed-blobs")]
-	[Tags("DDC Compressed Blobs")]
 	public class CompressedBlobController : ControllerBase
 	{
 		private readonly IBlobService _storage;
-		private readonly IContentIdService _contentIdStore;
+		private readonly IContentIdStore _contentIdStore;
 		private readonly IDiagnosticContext _diagnosticContext;
 		private readonly IRequestHelper _requestHelper;
 		private readonly BufferedPayloadFactory _bufferedPayloadFactory;
 		private readonly NginxRedirectHelper _nginxRedirectHelper;
 
-		public CompressedBlobController(IBlobService storage, IContentIdService contentIdStore, IDiagnosticContext diagnosticContext, IRequestHelper requestHelper, BufferedPayloadFactory bufferedPayloadFactory, NginxRedirectHelper nginxRedirectHelper)
+		public CompressedBlobController(IBlobService storage, IContentIdStore contentIdStore, IDiagnosticContext diagnosticContext, IRequestHelper requestHelper, BufferedPayloadFactory bufferedPayloadFactory, NginxRedirectHelper nginxRedirectHelper)
 		{
 			_storage = storage;
 			_contentIdStore = contentIdStore;
@@ -55,7 +51,7 @@ namespace Horde.Server.Ddc
 			[Required] NamespaceId ns,
 			[Required] ContentId id)
 		{
-			ActionResult? result = await _requestHelper.HasAccessToNamespaceAsync(User, Request, ns, new[] { StorageAclAction.ReadBlobs });
+			ActionResult? result = await _requestHelper.HasAccessToNamespaceAsync(User, Request, ns, new[] { JupiterAclAction.ReadObject });
 			if (result != null)
 			{
 				return result;
@@ -99,7 +95,7 @@ namespace Horde.Server.Ddc
 			[Required] NamespaceId ns,
 			[Required] ContentId id)
 		{
-			ActionResult? result = await _requestHelper.HasAccessToNamespaceAsync(User, Request, ns, new[] { StorageAclAction.ReadBlobs });
+			ActionResult? result = await _requestHelper.HasAccessToNamespaceAsync(User, Request, ns, new[] { JupiterAclAction.ReadObject });
 			if (result != null)
 			{
 				return result;
@@ -135,7 +131,7 @@ namespace Horde.Server.Ddc
 			[Required] NamespaceId ns,
 			[Required][FromQuery] List<ContentId> id)
 		{
-			ActionResult? result = await _requestHelper.HasAccessToNamespaceAsync(User, Request, ns, new[] { StorageAclAction.ReadBlobs });
+			ActionResult? result = await _requestHelper.HasAccessToNamespaceAsync(User, Request, ns, new[] { JupiterAclAction.ReadObject });
 			if (result != null)
 			{
 				return result;
@@ -177,7 +173,7 @@ namespace Horde.Server.Ddc
 			[Required] NamespaceId ns,
 			[FromBody] ContentId[] bodyIds)
 		{
-			ActionResult? result = await _requestHelper.HasAccessToNamespaceAsync(User, Request, ns, new[] { StorageAclAction.ReadBlobs });
+			ActionResult? result = await _requestHelper.HasAccessToNamespaceAsync(User, Request, ns, new[] { JupiterAclAction.ReadObject });
 			if (result != null)
 			{
 				return result;
@@ -218,10 +214,9 @@ namespace Horde.Server.Ddc
 		[RequiredContentType(CustomMediaTypeNames.UnrealCompressedBuffer)]
 		public async Task<IActionResult> PutAsync(
 			[Required] NamespaceId ns,
-			[Required] ContentId id,
-			CancellationToken cancellationToken)
+			[Required] ContentId id)
 		{
-			ActionResult? result = await _requestHelper.HasAccessToNamespaceAsync(User, Request, ns, new[] { StorageAclAction.WriteBlobs });
+			ActionResult? result = await _requestHelper.HasAccessToNamespaceAsync(User, Request, ns, new[] { JupiterAclAction.WriteObject });
 			if (result != null)
 			{
 				return result;
@@ -231,9 +226,9 @@ namespace Horde.Server.Ddc
 
 			try
 			{
-				using BufferedPayload payload = await _bufferedPayloadFactory.CreateFromRequestAsync(Request);
+				using IBufferedPayload payload = await _bufferedPayloadFactory.CreateFromRequest(Request);
 
-				ContentId identifier = await _storage.PutCompressedObjectAsync(ns, payload, id, HttpContext.RequestServices, cancellationToken);
+				ContentId identifier = await _storage.PutCompressedObjectAsync(ns, payload, id, HttpContext.RequestServices);
 
 				return Ok(new { Identifier = identifier.ToString() });
 			}
@@ -255,10 +250,9 @@ namespace Horde.Server.Ddc
 		[DisableRequestSizeLimit]
 		[RequiredContentType(CustomMediaTypeNames.UnrealCompressedBuffer)]
 		public async Task<IActionResult> PostAsync(
-			[Required] NamespaceId ns,
-			CancellationToken cancellationToken)
+			[Required] NamespaceId ns)
 		{
-			ActionResult? result = await _requestHelper.HasAccessToNamespaceAsync(User, Request, ns, new[] { StorageAclAction.WriteBlobs });
+			ActionResult? result = await _requestHelper.HasAccessToNamespaceAsync(User, Request, ns, new[] { JupiterAclAction.WriteObject });
 			if (result != null)
 			{
 				return result;
@@ -268,9 +262,9 @@ namespace Horde.Server.Ddc
 
 			try
 			{
-				using BufferedPayload payload = await _bufferedPayloadFactory.CreateFromRequestAsync(Request);
+				using IBufferedPayload payload = await _bufferedPayloadFactory.CreateFromRequest(Request);
 
-				ContentId identifier = await _storage.PutCompressedObjectAsync(ns, payload, null, HttpContext.RequestServices, cancellationToken);
+				ContentId identifier = await _storage.PutCompressedObjectAsync(ns, payload, null, HttpContext.RequestServices);
 
 				return Ok(new
 				{
diff --git a/Engine/Source/Programs/Horde/Horde.Server/Ddc/CompressedBufferUtils.cs b/Engine/Source/Programs/Horde/Horde.Server/Ddc/CompressedBufferUtils.cs
--- a/Engine/Source/Programs/Horde/Horde.Server/Ddc/CompressedBufferUtils.cs
+++ b/Engine/Source/Programs/Horde/Horde.Server/Ddc/CompressedBufferUtils.cs
@@ -6,7 +6,6 @@ using System.Collections.Generic;
 using System.IO;
 using System.Linq;
 using System.Text;
-using System.Threading;
 using System.Threading.Tasks;
 using Blake3;
 using EpicGames.Compression;
@@ -19,7 +18,7 @@ using OpenTelemetry.Trace;
 
 namespace Horde.Server.Ddc
 {
-	class CompressedBufferHeader
+	public class CompressedBufferHeader
 	{
 		[System.Diagnostics.CodeAnalysis.SuppressMessage("Design", "CA1028:Enum Storage should be Int32", Justification = "Interop requires byte")]
 		public enum CompressionMethod : byte
@@ -74,10 +73,12 @@ namespace Horde.Server.Ddc
 	public class CompressedBufferUtils
 	{
 		private readonly Tracer _tracer;
+		private readonly BufferedPayloadFactory _payloadFactory;
 
-		public CompressedBufferUtils(Tracer tracer)
+		public CompressedBufferUtils(Tracer tracer, BufferedPayloadFactory payloadFactory)
 		{
 			_tracer = tracer;
+			_payloadFactory = payloadFactory;
 		}
 
 		private static (CompressedBufferHeader, uint[]) ExtractHeader(BinaryReader br)
@@ -152,7 +153,7 @@ namespace Horde.Server.Ddc
 			return (header, blocks);
 		}
 
-		static void WriteHeader(CompressedBufferHeader header, BinaryWriter writer)
+		public static void WriteHeader(CompressedBufferHeader header, BinaryWriter writer)
 		{
 			// the header is always stored big endian
 			bool needsByteSwap = BitConverter.IsLittleEndian;
@@ -183,7 +184,7 @@ namespace Horde.Server.Ddc
 			}
 		}
 
-		public async Task<BufferedPayload> DecompressContentAsync(Stream sourceStream, ulong streamSize, CancellationToken cancellationToken)
+		public async Task<IBufferedPayload> DecompressContentAsync(Stream sourceStream, ulong streamSize)
 		{
 			using BinaryReader br = new BinaryReader(sourceStream);
 			(CompressedBufferHeader header, uint[] compressedBlockSizes) = ExtractHeader(br);
@@ -193,7 +194,7 @@ namespace Horde.Server.Ddc
 				throw new Exception($"Expected stream to be {header.TotalCompressedSize} but it was {streamSize}");
 			}
 
-			using FilesystemBufferedPayloadWriter bufferedPayloadWriter = new FilesystemBufferedPayloadWriter();
+			using FilesystemBufferedPayloadWriter bufferedPayloadWriter = _payloadFactory.CreateFilesystemBufferedPayloadWriter();
 
 			{
 				await using Stream targetStream = bufferedPayloadWriter.GetWritableStream();
@@ -226,45 +227,36 @@ namespace Horde.Server.Ddc
 				}
 				else
 				{
-					await sourceStream.CopyToAsync(targetStream, cancellationToken);
+					await sourceStream.CopyToAsync(targetStream);
 				}
 			}
 
 			// not using the buffered payload as we transfer the ownership to the caller of this method
-#pragma warning disable CA2000 // Dispose objects before losing scope
 			FilesystemBufferedPayload finalizedBufferedPayload = bufferedPayloadWriter.Done();
-#pragma warning restore CA2000 // Dispose objects before losing scope
-			try
+
+			if (header.TotalRawSize != (ulong)finalizedBufferedPayload.Length)
 			{
-				if (header.TotalRawSize != (ulong)finalizedBufferedPayload.Length)
-				{
-					throw new Exception("Did not decompress the full payload");
-				}
+				throw new Exception("Did not decompress the full payload");
+			}
 
-				{
-					using TelemetrySpan _ = _tracer.StartActiveSpan("web.hash").SetAttribute("operation.name", "web.hash");
+			{
+				using TelemetrySpan _ = _tracer.StartActiveSpan("web.hash").SetAttribute("operation.name", "web.hash");
 
-					// only read the first 20 bytes of the hash field as IoHashes are 20 bytes and not 32 bytes
-					byte[] slicedHash = new byte[20];
-					Array.Copy(header.RawHash, 0, slicedHash, 0, 20);
+				// only read the first 20 bytes of the hash field as IoHashes are 20 bytes and not 32 bytes
+				byte[] slicedHash = new byte[20];
+				Array.Copy(header.RawHash, 0, slicedHash, 0, 20);
 
-					BlobId headerIdentifier = new BlobId(new IoHash(slicedHash));
-					await using Stream hashStream = finalizedBufferedPayload.GetStream();
-					BlobId contentHash = await BlobId.FromStreamAsync(hashStream, cancellationToken);
+				BlobId headerIdentifier = new BlobId(slicedHash);
+				await using Stream hashStream = finalizedBufferedPayload.GetStream();
+				BlobId contentHash = await BlobId.FromStreamAsync(hashStream);
 
-					if (!headerIdentifier.Equals(contentHash))
-					{
-						throw new Exception($"Payload was expected to be {headerIdentifier} but was {contentHash}");
-					}
+				if (!headerIdentifier.Equals(contentHash))
+				{
+					throw new Exception($"Payload was expected to be {headerIdentifier} but was {contentHash}");
 				}
-
-				return finalizedBufferedPayload;
-			}
-			catch
-			{
-				finalizedBufferedPayload.Dispose();
-				throw;
 			}
+
+			return finalizedBufferedPayload;
 		}
 
 		private static int DecompressPayload(ReadOnlySpan<byte> compressedPayload, CompressedBufferHeader header, ulong rawBlockSize, Stream target)
@@ -386,7 +378,7 @@ namespace Horde.Server.Ddc
 			return hash;
 		}
 
-		static byte[] WriteHeaderToBuffer(CompressedBufferHeader header, uint[] compressedBlockLengths)
+		public static byte[] WriteHeaderToBuffer(CompressedBufferHeader header, uint[] compressedBlockLengths)
 		{
 			uint blockCount = header.BlockCount;
 			uint blocksByteUsed = blockCount * sizeof(uint);
@@ -445,6 +437,7 @@ namespace Horde.Server.Ddc
 	// from OodleDataCompression.h , we define our own enums for oodle compressions used and convert to the ones expected in the oodle api
 #pragma warning disable CA1028 // Enum Storage should be Int32
 	public enum OoodleCompressorMethod : byte
+
 	{
 		NotSet = 0,
 		Selkie = 1,
@@ -470,7 +463,7 @@ namespace Horde.Server.Ddc
 		Optimal4 = 8,
 	}
 #pragma warning restore CA1028 // Enum Storage should be Int32
-	static class OodleUtils
+	public static class OodleUtils
 	{
 		public static OodleCompressionLevel ToOodleApiCompressionLevel(OoodleCompressionLevel compressionLevel)
 		{
diff --git a/Engine/Source/Programs/Horde/Horde.Server/Ddc/ContentHash.cs b/Engine/Source/Programs/Horde/Horde.Server/Ddc/ContentHash.cs
new file mode 100644
--- /dev/null
+++ b/Engine/Source/Programs/Horde/Horde.Server/Ddc/ContentHash.cs
@@ -0,0 +1,185 @@
+﻿// Copyright Epic Games, Inc. All Rights Reserved.
+
+using System;
+using System.ComponentModel;
+using System.Globalization;
+using System.IO;
+using System.Text.Json;
+using System.Text.Json.Serialization;
+using Blake3;
+using EpicGames.Core;
+using EpicGames.Serialization;
+
+#pragma warning disable CS1591
+
+namespace Horde.Server.Ddc
+{
+	[TypeConverter(typeof(ContentHashTypeConverter))]
+	[JsonConverter(typeof(ContentHashJsonConverter))]
+	[CbConverter(typeof(ContentHashCbConverter))]
+	public class ContentHash : IEquatable<ContentHash>, IEquatable<byte[]>
+	{
+		protected ByteArrayComparer Comparer { get; } = new ByteArrayComparer();
+		protected byte[] Identifier { get; init; }
+		public const int HashLength = 20;
+
+		public byte[] HashData => Identifier;
+		public ContentHash(byte[] identifier)
+		{
+			Identifier = identifier;
+			/*if (identifier.Length != HashLength)
+			{
+				throw new ArgumentException("Supplied identifier was not 20 bytes, this is not a valid identifier", nameof(identifier));
+			}*/
+		}
+
+		[JsonConstructor]
+		public ContentHash(string identifier)
+		{
+			if (identifier == null)
+			{
+				throw new ArgumentNullException(nameof(identifier));
+			}
+
+			byte[] byteIdentifier = StringUtils.ToHashFromHexString(identifier);
+			Identifier = byteIdentifier;
+			/*if (byteIdentifier.Length != HashLength)
+			{
+				throw new ArgumentException("Supplied identifier was not 20 bytes, this is not a valid identifier", nameof(identifier));
+			}*/
+		}
+
+		public override int GetHashCode()
+		{
+			return Comparer.GetHashCode(Identifier);
+		}
+
+		public bool Equals(ContentHash? other)
+		{
+			if (other == null)
+			{
+				return false;
+			}
+
+			return Comparer.Equals(Identifier, other.Identifier);
+		}
+
+		public bool Equals(byte[]? other)
+		{
+			if (other == null)
+			{
+				return false;
+			}
+
+			return Comparer.Equals(Identifier, other);
+		}
+
+		public override bool Equals(object? obj)
+		{
+			if (ReferenceEquals(null, obj))
+			{
+				return false;
+			}
+
+			if (ReferenceEquals(this, obj))
+			{
+				return true;
+			}
+
+			if (obj.GetType() != GetType())
+			{
+				return false;
+			}
+
+			return Equals((ContentHash)obj);
+		}
+
+		public override string ToString()
+		{
+			return StringUtils.FormatAsHexString(Identifier);
+		}
+
+		public static ContentHash FromBlob(byte[] blobMemory)
+		{
+			using Hasher hasher = Hasher.New();
+			hasher.UpdateWithJoin(blobMemory);
+			Hash blake3Hash = hasher.Finalize();
+
+			// we only keep the first 20 bytes of the Blake3 hash
+			Span<byte> hash = blake3Hash.AsSpanUnsafe().Slice(0, HashLength);
+
+			return new ContentHash(hash.ToArray());
+		}
+	}
+
+	public class ContentHashCbConverter : CbConverter<ContentHash>
+	{
+		public override ContentHash Read(CbField field) => new ContentHash(field.AsHash().ToByteArray());
+
+		/// <inheritdoc/>
+		public override void Write(CbWriter writer, ContentHash value) => writer.WriteHashValue(new IoHash(value.HashData));
+
+		/// <inheritdoc/>
+		public override void WriteNamed(CbWriter writer, CbFieldName name, ContentHash value) => writer.WriteHash(name, new IoHash(value.HashData));
+	}
+
+	public class ContentHashJsonConverter : JsonConverter<ContentHash>
+	{
+		public override ContentHash Read(ref Utf8JsonReader reader, Type typeToConvert, JsonSerializerOptions options)
+		{
+			string? str = reader.GetString();
+			if (str == null)
+			{
+				throw new InvalidDataException("Unable to parse content hash");
+			}
+
+			return new ContentHash(str);
+		}
+
+		public override void Write(Utf8JsonWriter writer, ContentHash value, JsonSerializerOptions options)
+		{
+			writer.WriteStringValue(value.ToString());
+		}
+	}
+
+	public class ContentHashTypeConverter : TypeConverter
+	{
+		public override bool CanConvertFrom(ITypeDescriptorContext? context, Type sourceType)
+		{
+			if (sourceType == typeof(string))
+			{
+				return true;
+			}
+			return base.CanConvertFrom(context, sourceType);
+		}
+
+		public override object? ConvertFrom(ITypeDescriptorContext? context, CultureInfo? culture, object value)
+		{
+			if (value is string s)
+			{
+				return new ContentHash(s);
+			}
+
+			return base.ConvertFrom(context, culture, value);
+		}
+
+		public override bool CanConvertTo(ITypeDescriptorContext? context, Type? destinationType)
+		{
+			if (destinationType == typeof(string))
+			{
+				return true;
+			}
+			return base.CanConvertTo(context, destinationType);
+		}
+
+		public override object? ConvertTo(ITypeDescriptorContext? context, CultureInfo? culture, object? value, Type destinationType)
+		{
+			if (destinationType == typeof(string))
+			{
+				ContentHash? identifier = (ContentHash?)value;
+				return identifier?.ToString();
+			}
+			return base.ConvertTo(context, culture, value, destinationType);
+		}
+	}
+}
diff --git a/Engine/Source/Programs/Horde/Horde.Server/Ddc/ContentId.cs b/Engine/Source/Programs/Horde/Horde.Server/Ddc/ContentId.cs
--- a/Engine/Source/Programs/Horde/Horde.Server/Ddc/ContentId.cs
+++ b/Engine/Source/Programs/Horde/Horde.Server/Ddc/ContentId.cs
@@ -3,162 +3,106 @@
 using System;
 using System.ComponentModel;
 using System.Globalization;
-using System.IO;
-using System.Text.Json;
 using System.Text.Json.Serialization;
 using EpicGames.Core;
-using EpicGames.Serialization;
 
 #pragma warning disable CS1591
 
 namespace Horde.Server.Ddc
 {
-	/// <summary>
-	/// Identifier for a compressed buffer in the store
-	/// </summary>
 	[TypeConverter(typeof(ContentIdTypeConverter))]
-	[JsonConverter(typeof(ContentIdJsonConverter))]
-	[CbConverter(typeof(ContentIdCbConverter))]
-	public readonly struct ContentId : IEquatable<ContentId>
+	public class ContentId : ContentHash, IEquatable<ContentId>
 	{
-		/// <summary>
-		/// Hash of the blob
-		/// </summary>
-		public IoHash Hash { get; }
-
-		/// <summary>
-		/// Constructor
-		/// </summary>
-		public ContentId(IoHash hash) => Hash = hash;
-
-		/// <summary>
-		/// Parses a <see cref="ContentId"/> from a string
-		/// </summary>
-		public static ContentId Parse(string text) => new ContentId(IoHash.Parse(text));
-
-		/// <inheritdoc/>
-		public override readonly int GetHashCode() => Hash.GetHashCode();
-
-		/// <inheritdoc/>
-		public readonly bool Equals(ContentId other) => Hash.Equals(other.Hash);
-
-		/// <inheritdoc/>
-		public override readonly bool Equals(object? obj) => obj is ContentId contentId && Equals(contentId);
-
-		/// <inheritdoc/>
-		public override readonly string ToString() => Hash.ToString();
-
-		/// <inheritdoc cref="IoHash.op_Equality"/>
-		public static bool operator ==(ContentId left, ContentId right) => left.Hash == right.Hash;
-
-		/// <inheritdoc cref="IoHash.op_Inequality"/>
-		public static bool operator !=(ContentId left, ContentId right) => !(left == right);
-
-		/// <summary>
-		/// Constructs a ContentId from an IoHash
-		/// </summary>
-		public static ContentId FromIoHash(IoHash hash) => new ContentId(hash);
-
-		/// <summary>
-		/// Constructs a ContentId from an IoHash
-		/// </summary>
-		public static ContentId FromBlobId(BlobId blobId) => FromIoHash(blobId.AsIoHash());
-
-		/// <summary>
-		/// Converts a ContentId to IoHash
-		/// </summary>
-		public IoHash AsIoHash() => Hash;
-
-		/// <summary>
-		/// Converts a ContentId to BlobId
-		/// </summary>
-		public BlobId AsBlobIdentifier() => BlobId.FromIoHash(Hash);
-	}
+		public ContentId(byte[] identifier) : base(identifier)
+		{
+		}
 
-	/// <summary>
-	/// Converts from <see cref="ContentId"/> instances to other types
-	/// </summary>
-	public class ContentIdTypeConverter : TypeConverter
-	{
-		/// <inheritdoc/>
-		public override bool CanConvertFrom(ITypeDescriptorContext? context, Type sourceType)
+		[JsonConstructor]
+		public ContentId(string identifier) : base(identifier)
 		{
-			if (sourceType == typeof(string))
-			{
-				return true;
-			}
-			return base.CanConvertFrom(context, sourceType);
+
 		}
 
-		/// <inheritdoc/>
-		public override object? ConvertFrom(ITypeDescriptorContext? context, CultureInfo? culture, object value)
+		public override int GetHashCode()
 		{
-			if (value is string s)
+			return Comparer.GetHashCode(Identifier);
+		}
+
+		public bool Equals(ContentId? other)
+		{
+			if (other == null)
 			{
-				return ContentId.Parse(s);
+				return false;
 			}
 
-			return base.ConvertFrom(context, culture, value);
+			return Comparer.Equals(Identifier, other.Identifier);
 		}
 
-		/// <inheritdoc/>
-		public override bool CanConvertTo(ITypeDescriptorContext? context, Type? destinationType)
+		public override bool Equals(object? obj)
 		{
-			if (destinationType == typeof(string))
+			if (ReferenceEquals(null, obj))
+			{
+				return false;
+			}
+
+			if (ReferenceEquals(this, obj))
 			{
 				return true;
 			}
-			return base.CanConvertTo(context, destinationType);
-		}
 
-		/// <inheritdoc/>
-		public override object? ConvertTo(ITypeDescriptorContext? context, CultureInfo? culture, object? value, Type destinationType)
-		{
-			if (destinationType == typeof(string))
+			if (obj.GetType() != GetType())
 			{
-				return value?.ToString();
+				return false;
 			}
 
-			return base.ConvertTo(context, culture, value, destinationType);
+			return Equals((ContentId) obj);
 		}
-	}
 
-	/// <summary>
-	/// Converts from <see cref="ContentId"/> instances to other types
-	/// </summary>
-	public class ContentIdJsonConverter : JsonConverter<ContentId>
-	{
-		/// <inheritdoc/>
-		public override ContentId Read(ref Utf8JsonReader reader, Type typeToConvert, JsonSerializerOptions options)
+		public static ContentId FromContentHash(ContentHash contentHash)
 		{
-			string? str = reader.GetString();
-			if (str == null)
-			{
-				throw new InvalidDataException("Unable to parse content id");
-			}
+			return new ContentId(contentHash.HashData);
+		}
+
+		public static ContentId FromBlobIdentifier(BlobId blobIdentifier)
+		{
+			return new ContentId(blobIdentifier.HashData);
+		}
+
+		public BlobId AsBlobIdentifier()
+		{
+			return new BlobId(HashData);
+		}
 
-			return ContentId.Parse(str);
+		public static ContentId FromIoHash(IoHash ioHash)
+		{
+			return new ContentId(ioHash.ToByteArray());
 		}
 
-		/// <inheritdoc/>
-		public override void Write(Utf8JsonWriter writer, ContentId value, JsonSerializerOptions options)
+		public IoHash AsIoHash()
 		{
-			writer.WriteStringValue(value.ToString());
+			return new IoHash(HashData);
 		}
 	}
 
-	/// <summary>
-	/// Serializes <see cref="ContentId"/> instances to compact binary
-	/// </summary>
-	public class ContentIdCbConverter : CbConverter<ContentId>
+	public class ContentIdTypeConverter : TypeConverter
 	{
-		/// <inheritdoc/>
-		public override ContentId Read(CbField field) => new ContentId(field.AsHash());
-
-		/// <inheritdoc/>
-		public override void Write(CbWriter writer, ContentId value) => writer.WriteBinaryAttachmentValue(value.Hash);
+		public override bool CanConvertFrom(ITypeDescriptorContext? context, Type sourceType)
+		{  
+			if (sourceType == typeof(string))  
+			{  
+				return true;
+			}  
+			return base.CanConvertFrom(context, sourceType);
+		}  
+  
+		public override object? ConvertFrom(ITypeDescriptorContext? context, CultureInfo? culture, object value)  
+		{
+			if (value is string s)
+			{
+				return new ContentId(s);
+			}
 
-		/// <inheritdoc/>
-		public override void WriteNamed(CbWriter writer, CbFieldName name, ContentId value) => writer.WriteBinaryAttachment(name, value.Hash);
+			return base.ConvertFrom(context, culture, value);  
+		}  
 	}
 }
diff --git a/Engine/Source/Programs/Horde/Horde.Server/Ddc/ContentIdController.cs b/Engine/Source/Programs/Horde/Horde.Server/Ddc/ContentIdController.cs
--- a/Engine/Source/Programs/Horde/Horde.Server/Ddc/ContentIdController.cs
+++ b/Engine/Source/Programs/Horde/Horde.Server/Ddc/ContentIdController.cs
@@ -4,7 +4,6 @@ using System.Net.Mime;
 using System.Threading.Tasks;
 using EpicGames.AspNet;
 using EpicGames.Horde.Storage;
-using Horde.Server.Storage;
 using Microsoft.AspNetCore.Authorization;
 using Microsoft.AspNetCore.Http;
 using Microsoft.AspNetCore.Mvc;
@@ -22,9 +21,9 @@ namespace Horde.Server.Ddc
 	public class ContentIdController : ControllerBase
 	{
 		private readonly IRequestHelper _requestHelper;
-		private readonly IContentIdService _contentIdStore;
+		private readonly IContentIdStore _contentIdStore;
 
-		public ContentIdController(IRequestHelper requestHelper, IContentIdService contentIdStore)
+		public ContentIdController(IRequestHelper requestHelper, IContentIdStore contentIdStore)
 		{
 			_requestHelper = requestHelper;
 			_contentIdStore = contentIdStore;
@@ -40,7 +39,7 @@ namespace Horde.Server.Ddc
 		[ProducesResponseType(type: typeof(ProblemDetails), 400)]
 		public async Task<IActionResult> ResolveAsync(NamespaceId ns, ContentId contentId)
 		{
-			ActionResult? result = await _requestHelper.HasAccessToNamespaceAsync(User, Request, ns, new[] { StorageAclAction.ReadBlobs });
+			ActionResult? result = await _requestHelper.HasAccessToNamespaceAsync(User, Request, ns, new[] { JupiterAclAction.ReadObject });
 			if (result != null)
 			{
 				return result;
@@ -67,7 +66,7 @@ namespace Horde.Server.Ddc
 		[ProducesResponseType(type: typeof(ProblemDetails), 400)]
 		public async Task<IActionResult> UpdateContentIdMappingAsync(NamespaceId ns, ContentId contentId, BlobId blobIdentifier, int contentWeight)
 		{
-			ActionResult? result = await _requestHelper.HasAccessToNamespaceAsync(User, Request, ns, new[] { StorageAclAction.WriteBlobs });
+			ActionResult? result = await _requestHelper.HasAccessToNamespaceAsync(User, Request, ns, new[] { JupiterAclAction.WriteObject });
 			if (result != null)
 			{
 				return result;
diff --git a/Engine/Source/Programs/Horde/Horde.Server/Ddc/ContentIdService.cs b/Engine/Source/Programs/Horde/Horde.Server/Ddc/ContentIdService.cs
deleted file mode 100644
diff --git a/Engine/Source/Programs/Horde/Horde.Server/Ddc/ContentIdStore.cs b/Engine/Source/Programs/Horde/Horde.Server/Ddc/ContentIdStore.cs
new file mode 100644
--- /dev/null
+++ b/Engine/Source/Programs/Horde/Horde.Server/Ddc/ContentIdStore.cs
@@ -0,0 +1,57 @@
+// Copyright Epic Games, Inc. All Rights Reserved.
+
+using System.Threading;
+using System.Threading.Tasks;
+using EpicGames.Core;
+using EpicGames.Horde.Storage;
+
+namespace Horde.Server.Ddc
+{
+	class ContentIdStore : IContentIdStore
+	{
+		readonly IStorageClientFactory _storageClientFactory;
+
+		public ContentIdStore(IStorageClientFactory storageService)
+		{
+			_storageClientFactory = storageService;
+		}
+
+		static string GetAlias(BlobId blobId) => BlobService.GetAlias(blobId);
+		static string GetAlias(ContentId contentId) => $"cid:{contentId}";
+
+		public async Task<BlobId[]?> ResolveAsync(NamespaceId ns, ContentId contentId, bool mustBeContentId = false)
+		{
+			CancellationToken cancellationToken = CancellationToken.None;
+
+			using IStorageClient storageClient = _storageClientFactory.CreateClient(ns);
+
+			BlobAlias? blobAlias = await storageClient.FindAliasAsync(GetAlias(contentId), cancellationToken);
+			if (blobAlias == null && !mustBeContentId)
+			{
+				blobAlias = await storageClient.FindAliasAsync(GetAlias(contentId.AsBlobIdentifier()), cancellationToken);
+			}
+			if (blobAlias == null)
+			{
+				return null;
+			}
+
+			IoHash hash = new IoHash(blobAlias.Data.Span);
+			return new[] { BlobId.FromIoHash(hash) };
+		}
+
+		public async Task PutAsync(NamespaceId ns, ContentId contentId, BlobId blobId, int contentWeight)
+		{
+			CancellationToken cancellationToken = CancellationToken.None;
+
+			using IStorageClient storageClient = _storageClientFactory.CreateClient(ns);
+
+			BlobAlias? blobAlias = await storageClient.FindAliasAsync(GetAlias(blobId), cancellationToken);
+			if (blobAlias == null)
+			{
+				throw new BlobNotFoundException(ns, blobId);
+			}
+
+			await storageClient.AddAliasAsync(GetAlias(contentId), blobAlias.Target, -contentWeight, blobAlias.Data, cancellationToken: cancellationToken);
+		}
+	}
+}
diff --git a/Engine/Source/Programs/Horde/Horde.Server/Ddc/FormatResolver.cs b/Engine/Source/Programs/Horde/Horde.Server/Ddc/FormatResolver.cs
--- a/Engine/Source/Programs/Horde/Horde.Server/Ddc/FormatResolver.cs
+++ b/Engine/Source/Programs/Horde/Horde.Server/Ddc/FormatResolver.cs
@@ -18,10 +18,11 @@ namespace Horde.Server.Ddc
 		private readonly IOptionsMonitor<MvcOptions> _mvcOptions;
 
 		private readonly string[] _validContentTypes = {
-			MediaTypeNames.Application.Octet,
-			MediaTypeNames.Application.Json,
-			CustomMediaTypeNames.UnrealCompactBinary,
+			MediaTypeNames.Application.Octet, 
+			MediaTypeNames.Application.Json, 
+			CustomMediaTypeNames.UnrealCompactBinary, 
 			CustomMediaTypeNames.JupiterInlinedPayload,
+			CustomMediaTypeNames.UnrealCompressedBuffer,
 			CustomMediaTypeNames.UnrealCompactBinaryPackage
 		};
 
@@ -43,7 +44,7 @@ namespace Horde.Server.Ddc
 
 				return typeMapping;
 			}
-
+			
 			StringValues acceptHeader = request.Headers["Accept"];
 
 			if (acceptHeader.Count == 0)
@@ -60,7 +61,12 @@ namespace Horde.Server.Ddc
 
 			foreach (string? header in acceptHeader)
 			{
-				if (header != null && _validContentTypes.Contains(header, StringComparer.OrdinalIgnoreCase))
+				if (header == null)
+			{
+					continue;
+				}
+
+				if (_validContentTypes.Contains(header, StringComparer.OrdinalIgnoreCase))
 				{
 					return header;
 				}
diff --git a/Engine/Source/Programs/Horde/Horde.Server/Ddc/IBlobService.cs b/Engine/Source/Programs/Horde/Horde.Server/Ddc/IBlobService.cs
--- a/Engine/Source/Programs/Horde/Horde.Server/Ddc/IBlobService.cs
+++ b/Engine/Source/Programs/Horde/Horde.Server/Ddc/IBlobService.cs
@@ -4,196 +4,165 @@ using System;
 using System.Collections.Generic;
 using System.IO;
 using System.Net.Mime;
-using System.Threading;
 using System.Threading.Tasks;
 using EpicGames.AspNet;
-using EpicGames.Core;
 using EpicGames.Horde.Storage;
 using Microsoft.Extensions.DependencyInjection;
 using OpenTelemetry.Trace;
 
 #pragma warning disable CS1591
 
-namespace Horde.Server.Ddc
+namespace Horde.Server.Ddc;
+
+public interface IBlobService
 {
-	/// <summary>
-	/// Interface for the DDC blobs service
-	/// </summary>
-	public interface IBlobService
-	{
-		Task VerifyContentMatchesHashAsync(Stream stream, IoHash expectedHash, CancellationToken cancellationToken = default);
-		Task<BlobId> PutObjectAsync(NamespaceId ns, BufferedPayload payload, BlobId identifier, CancellationToken cancellationToken = default);
-		Task<BlobId> PutObjectAsync(NamespaceId ns, ReadOnlyMemory<byte> payload, BlobId identifier, CancellationToken cancellationToken = default);
-		Task<BlobId> PutObjectKnownHashAsync(NamespaceId ns, BufferedPayload content, BlobId identifier, CancellationToken cancellationToken = default);
-		Task<Uri?> MaybePutObjectWithRedirectAsync(NamespaceId ns, BlobId identifier, CancellationToken cancellationToken = default);
+	Task<ContentHash> VerifyContentMatchesHashAsync(Stream content, ContentHash identifier);
+	Task<BlobId> PutObjectKnownHashAsync(NamespaceId ns, IBufferedPayload content, BlobId identifier);
+	Task<BlobId> PutObjectAsync(NamespaceId ns, IBufferedPayload payload, BlobId identifier);
+	Task<BlobId> PutObjectAsync(NamespaceId ns, byte[] payload, BlobId identifier);
+	Task<Uri?> MaybePutObjectWithRedirectAsync(NamespaceId ns, BlobId identifier);
 
-		Task<BlobContents> GetObjectAsync(NamespaceId ns, BlobId blob, List<string>? storageLayers = null, bool supportsRedirectUri = false, CancellationToken cancellationToken = default);
+	Task<BlobContents> GetObjectAsync(NamespaceId ns, BlobId blob, List<string>? storageLayers = null, bool supportsRedirectUri = false, bool allowOndemandReplication = true);
 
-		Task<Uri?> GetObjectWithRedirectAsync(NamespaceId ns, BlobId blobIdentifier, List<string>? storageLayers = null, CancellationToken cancellationToken = default);
+	Task<Uri?> GetObjectWithRedirectAsync(NamespaceId ns, BlobId blobIdentifier, List<string>? storageLayers = null);
 
-		Task<bool> ExistsAsync(NamespaceId ns, BlobId blob, List<string>? storageLayers = null, CancellationToken cancellationToken = default);
+	Task<BlobMetadata> GetObjectMetadataAsync(NamespaceId ns, BlobId blobId);
 
-		Task DeleteObjectAsync(NamespaceId ns, BlobId blob, CancellationToken cancellationToken = default);
+	Task<BlobContents> ReplicateObjectAsync(NamespaceId ns, BlobId blob, bool force = false);
 
-		Task<BlobId[]> FilterOutKnownBlobsAsync(NamespaceId ns, IEnumerable<BlobId> blobs, CancellationToken cancellationToken = default);
-		Task<BlobContents> GetObjectsAsync(NamespaceId ns, BlobId[] refRequestBlobReferences, CancellationToken cancellationToken = default);
-	}
+	Task<bool> ExistsAsync(NamespaceId ns, BlobId blob, List<string>? storageLayers = null);
 
-	public class BlobNotFoundException : Exception
-	{
-		public NamespaceId Ns { get; }
-		public BlobId Blob { get; }
+	/// <summary>
+	/// Checks that the blob exists in the root store, the store which is last in the list and thus is intended to have every blob in it
+	/// </summary>
+	/// <param name="ns">The namespace</param>
+	/// <param name="blob">The identifier of the blob</param>
+	/// <returns></returns>
+	Task<bool> ExistsInRootStoreAsync(NamespaceId ns, BlobId blob);
 
-		public BlobNotFoundException(NamespaceId ns, BlobId blob) : base($"No Blob in Namespace {ns} with id {blob}")
-		{
-			Ns = ns;
-			Blob = blob;
-		}
+	// Delete a object
+	Task DeleteObjectAsync(NamespaceId ns, BlobId blob);
 
-		public BlobNotFoundException(NamespaceId ns, BlobId blob, string message) : base(message)
-		{
-			Ns = ns;
-			Blob = blob;
-		}
-	}
+	// delete the whole namespace
+	Task DeleteNamespaceAsync(NamespaceId ns);
 
-	public class BlobReplicationException : BlobNotFoundException
+	IAsyncEnumerable<(BlobId, DateTime)> ListObjectsAsync(NamespaceId ns);
+	Task<BlobId[]> FilterOutKnownBlobsAsync(NamespaceId ns, IEnumerable<BlobId> blobs);
+	Task<BlobId[]> FilterOutKnownBlobsAsync(NamespaceId ns, IAsyncEnumerable<BlobId> blobs);
+	Task<BlobContents> GetObjectsAsync(NamespaceId ns, BlobId[] refRequestBlobReferences);
+
+	bool ShouldFetchBlobOnDemand(NamespaceId ns);
+}
+
+public class BlobMetadata
+{
+	public BlobMetadata(long length, DateTime creationTime)
 	{
-		public BlobReplicationException(NamespaceId ns, BlobId blob, string message) : base(ns, blob, message)
-		{
-		}
+		Length = length;
+		CreationTime = creationTime;
 	}
 
-	public class BlobTooLargeException : Exception
+	public long Length { get; set; }
+	public DateTime CreationTime { get; set; }
+}
+
+public static class BlobServiceExtensions
+{
+	public static async Task<ContentId> PutCompressedObjectAsync(this IBlobService blobService, NamespaceId ns, IBufferedPayload payload, ContentId? id, IServiceProvider provider)
 	{
-		public BlobId Blob { get; }
+		IContentIdStore contentIdStore = provider.GetService<IContentIdStore>()!;
+		CompressedBufferUtils compressedBufferUtils = provider.GetService<CompressedBufferUtils>()!;
+		Tracer tracer = provider.GetService<Tracer>()!;
 
-		public BlobTooLargeException(BlobId blob) : base($"Blob {blob} was to large to cache")
-		{
-			Blob = blob;
-		}
-	}
+		// decompress the content and generate a identifier from it to verify the identifier we got
+		await using Stream decompressStream = payload.GetStream();
 
-	public class ResourceHasToManyRequestsException : Exception
-	{
-		public ResourceHasToManyRequestsException(Exception originalException) : base($"To many requests to resource", originalException)
+		using IBufferedPayload bufferedPayload = await compressedBufferUtils.DecompressContentAsync(decompressStream, (ulong)payload.Length);
+		await using Stream decompressedStream = bufferedPayload.GetStream();
+
+		ContentId identifierDecompressedPayload;
+		if (id != null)
 		{
+			identifierDecompressedPayload = ContentId.FromContentHash(await blobService.VerifyContentMatchesHashAsync(decompressedStream, id));
 		}
-	}
-
-	public class NamespaceNotFoundException : Exception
-	{
-		public NamespaceId Namespace { get; }
-
-		public NamespaceNotFoundException(NamespaceId @namespace) : base($"Could not find namespace {@namespace}")
+		else
 		{
-			Namespace = @namespace;
+			ContentHash blobHash;
+			{
+				using TelemetrySpan _ = tracer.StartActiveSpan("web.hash").SetAttribute("operation.name", "web.hash");
+				blobHash = await BlobId.FromStreamAsync(decompressedStream);
+			}
+
+			identifierDecompressedPayload = ContentId.FromContentHash(blobHash);
 		}
 
-		public NamespaceNotFoundException(NamespaceId @namespace, string message) : base(message)
+		BlobId identifierCompressedPayload;
 		{
-			Namespace = @namespace;
+			using TelemetrySpan _ = tracer.StartActiveSpan("web.hash").SetAttribute("operation.name", "web.hash");
+			await using Stream hashStream = payload.GetStream();
+			identifierCompressedPayload = await BlobId.FromStreamAsync(hashStream);
 		}
-	}
 
-	public static class BlobServiceExtensions
-	{
-		public static async Task<ContentId> PutCompressedObjectAsync(this IBlobService blobService, NamespaceId ns, BufferedPayload payload, ContentId? id, IServiceProvider provider, CancellationToken cancellationToken)
+		// we still commit the compressed buffer to the object store using the hash of the compressed content
 		{
-			IContentIdService contentIdStore = provider.GetService<IContentIdService>()!;
-			CompressedBufferUtils compressedBufferUtils = provider.GetService<CompressedBufferUtils>()!;
-			Tracer tracer = provider.GetService<Tracer>()!;
-
-			// decompress the content and generate a identifier from it to verify the identifier we got
-			await using Stream decompressStream = payload.GetStream();
-
-			using BufferedPayload bufferedPayload = await compressedBufferUtils.DecompressContentAsync(decompressStream, (ulong)payload.Length, cancellationToken);
-			await using Stream decompressedStream = bufferedPayload.GetStream();
-
-			ContentId identifierDecompressedPayload;
-			if (id != null)
-			{
-				await blobService.VerifyContentMatchesHashAsync(decompressedStream, id.Value.AsIoHash(), cancellationToken);
-				identifierDecompressedPayload = id.Value;
-			}
-			else
-			{
-				BlobId blobHash;
-				{
-					using TelemetrySpan _ = tracer.StartActiveSpan("web.hash").SetAttribute("operation.name", "web.hash");
-					blobHash = await BlobId.FromStreamAsync(decompressedStream, cancellationToken);
-				}
-
-				identifierDecompressedPayload = ContentId.FromBlobId(blobHash);
-			}
+			await blobService.PutObjectKnownHashAsync(ns, payload, identifierCompressedPayload);
+		}
 
-			BlobId identifierCompressedPayload;
-			{
-				using TelemetrySpan _ = tracer.StartActiveSpan("web.hash").SetAttribute("operation.name", "web.hash");
-				await using Stream hashStream = payload.GetStream();
-				identifierCompressedPayload = await BlobId.FromStreamAsync(hashStream, cancellationToken);
-			}
+		// commit the mapping from the decompressed hash to the compressed hash, we run this in parallel with the blob store submit
+		// TODO: let users specify weight of the blob compared to previously submitted content ids
+		int contentIdWeight = (int)payload.Length;
+		Task contentIdStoreTask = contentIdStore.PutAsync(ns, identifierDecompressedPayload, identifierCompressedPayload, contentIdWeight);
 
-			// commit the mapping from the decompressed hash to the compressed hash, we run this in parallel with the blob store submit
-			// TODO: let users specify weight of the blob compared to previously submitted content ids
-			int contentIdWeight = (int)payload.Length;
+		await contentIdStoreTask;
 
-			// we still commit the compressed buffer to the object store using the hash of the compressed content
-			{
-				await blobService.PutObjectKnownHashAsync(ns, payload, identifierCompressedPayload, cancellationToken);
-			}
+		return identifierDecompressedPayload;
+	}
 
-			await contentIdStore.PutAsync(ns, identifierDecompressedPayload, identifierCompressedPayload, contentIdWeight);
+	public static async Task<(BlobContents, string)> GetCompressedObjectAsync(this IBlobService blobService, NamespaceId ns, ContentId contentId, IServiceProvider provider, bool supportsRedirectUri = false)
+	{
+		IContentIdStore contentIdStore = provider.GetService<IContentIdStore>()!;
+		Tracer tracer = provider.GetService<Tracer>()!;
 
-			return identifierDecompressedPayload;
+		BlobId[]? chunks = await contentIdStore.ResolveAsync(ns, contentId, mustBeContentId: false);
+		if (chunks == null || chunks.Length == 0)
+		{
+			throw new ContentIdResolveException(contentId);
 		}
 
-		public static async Task<(BlobContents, string)> GetCompressedObjectAsync(this IBlobService blobService, NamespaceId ns, ContentId contentId, IServiceProvider provider, bool supportsRedirectUri = false, CancellationToken cancellationToken = default)
+		// single chunk, we just return that chunk
+		if (chunks.Length == 1)
 		{
-			IContentIdService contentIdStore = provider.GetService<IContentIdService>()!;
-			Tracer tracer = provider.GetService<Tracer>()!;
-
-			BlobId[]? chunks = await contentIdStore.ResolveAsync(ns, contentId, mustBeContentId: false);
-			if (chunks == null || chunks.Length == 0)
+			BlobId blobToReturn = chunks[0];
+			string mimeType = CustomMediaTypeNames.UnrealCompressedBuffer;
+			if (contentId.Equals(blobToReturn))
 			{
-				throw new ContentIdResolveException(contentId);
+				// this was actually the unmapped blob, meaning its not a compressed buffer
+				mimeType = MediaTypeNames.Application.Octet;
 			}
 
-			// single chunk, we just return that chunk
-			if (chunks.Length == 1)
-			{
-				BlobId blobToReturn = chunks[0];
-				string mimeType = CustomMediaTypeNames.UnrealCompressedBuffer;
-				if (contentId.Hash == blobToReturn.Hash)
-				{
-					// this was actually the unmapped blob, meaning its not a compressed buffer
-					mimeType = MediaTypeNames.Application.Octet;
-				}
-
-				return (await blobService.GetObjectAsync(ns, blobToReturn, supportsRedirectUri: supportsRedirectUri, cancellationToken: cancellationToken), mimeType);
-			}
+			return (await blobService.GetObjectAsync(ns, blobToReturn, supportsRedirectUri: supportsRedirectUri), mimeType);
+		}
 
-			// chunked content, combine the chunks into a single stream
-			using TelemetrySpan _ = tracer.StartActiveSpan("blob.combine").SetAttribute("operation.name", "blob.combine");
-			Task<BlobContents>[] tasks = new Task<BlobContents>[chunks.Length];
-			for (int i = 0; i < chunks.Length; i++)
-			{
-				// even if it was requested to support redirect, since we need to combine the chunks using redirects is not possible
-				tasks[i] = blobService.GetObjectAsync(ns, chunks[i], supportsRedirectUri: false, cancellationToken: cancellationToken);
-			}
+		// chunked content, combine the chunks into a single stream
+		using TelemetrySpan _ = tracer.StartActiveSpan("blob.combine").SetAttribute("operation.name", "blob.combine");
+		Task<BlobContents>[] tasks = new Task<BlobContents>[chunks.Length];
+		for (int i = 0; i < chunks.Length; i++)
+		{
+			// even if it was requested to support redirect, since we need to combine the chunks using redirects is not possible
+			tasks[i] = blobService.GetObjectAsync(ns, chunks[i], supportsRedirectUri: false);
+		}
 
-			MemoryStream ms = new MemoryStream();
-			foreach (Task<BlobContents> task in tasks)
-			{
-				BlobContents blob = await task;
-				await using Stream s = blob.Stream;
-				await s.CopyToAsync(ms, cancellationToken);
-			}
+		MemoryStream ms = new MemoryStream();
+		foreach (Task<BlobContents> task in tasks)
+		{
+			BlobContents blob = await task;
+			await using Stream s = blob.Stream;
+			await s.CopyToAsync(ms);
+		}
 
-			ms.Seek(0, SeekOrigin.Begin);
+		ms.Seek(0, SeekOrigin.Begin);
 
-			// chunking could not have happened for a non compressed buffer so assume it is compressed
-			return (new BlobContents(ms, ms.Length), CustomMediaTypeNames.UnrealCompressedBuffer);
-		}
+		// chunking could not have happened for a non compressed buffer so assume it is compressed
+		return (new BlobContents(ms, ms.Length), CustomMediaTypeNames.UnrealCompressedBuffer);
 	}
 }
diff --git a/Engine/Source/Programs/Horde/Horde.Server/Ddc/IBlobStore.cs b/Engine/Source/Programs/Horde/Horde.Server/Ddc/IBlobStore.cs
new file mode 100644
--- /dev/null
+++ b/Engine/Source/Programs/Horde/Horde.Server/Ddc/IBlobStore.cs
@@ -0,0 +1,91 @@
+﻿// Copyright Epic Games, Inc. All Rights Reserved.
+
+using System;
+using System.Collections.Generic;
+using System.IO;
+using System.Threading;
+using System.Threading.Tasks;
+using EpicGames.Horde.Storage;
+
+#pragma warning disable CS1591
+
+namespace Horde.Server.Ddc
+{
+	public enum LastAccessTrackingFlags
+	{
+		DoTracking = 0,
+		SkipTracking = 1,
+	}
+
+	public interface IBlobStore
+	{
+		Task<BlobId> PutObjectAsync(NamespaceId ns, byte[] blob, BlobId identifier);
+		Task<BlobId> PutObjectAsync(NamespaceId ns, ReadOnlyMemory<byte> blob, BlobId identifier);
+		Task<BlobId> PutObjectAsync(NamespaceId ns, Stream content, BlobId identifier);
+
+		Task<BlobContents> GetObjectAsync(NamespaceId ns, BlobId blob, LastAccessTrackingFlags flags = LastAccessTrackingFlags.DoTracking, bool supportsRedirectUri = false);
+		Task<bool> ExistsAsync(NamespaceId ns, BlobId blob, bool forceCheck = false);
+
+		// Delete a object
+		Task DeleteObjectAsync(NamespaceId ns, BlobId blob);
+
+		// delete the whole namespace
+		Task DeleteNamespaceAsync(NamespaceId ns);
+
+		IAsyncEnumerable<(BlobId, DateTime)> ListObjectsAsync(NamespaceId ns);
+		Task<Uri?> PutObjectWithRedirectAsync(NamespaceId ns, BlobId identifier);
+		Task<Uri?> GetObjectByRedirectAsync(NamespaceId ns, BlobId blob);
+		Task<BlobMetadata> GetObjectMetadataAsync(NamespaceId ns, BlobId blobId);
+	}
+
+	public interface IStorageBackend
+	{
+		Task WriteAsync(string path, Stream content, CancellationToken cancellationToken = default);
+		Task<BlobContents?> TryReadAsync(string path, LastAccessTrackingFlags flags = LastAccessTrackingFlags.DoTracking, CancellationToken cancellationToken = default);
+		Task<bool> ExistsAsync(string path, CancellationToken cancellationToken = default);
+		Task DeleteAsync(string path, CancellationToken cancellationToken = default);
+		IAsyncEnumerable<(string, DateTime)> ListAsync(CancellationToken cancellationToken = default);
+	}
+
+	public class BlobNotFoundException : Exception
+	{
+		public NamespaceId Ns { get; }
+		public BlobId Blob { get; }
+
+		public BlobNotFoundException(NamespaceId ns, BlobId blob) : base($"No Blob in Namespace {ns} with id {blob}")
+		{
+			Ns = ns;
+			Blob = blob;
+		}
+
+		public BlobNotFoundException(NamespaceId ns, BlobId blob, string message) : base(message)
+		{
+			Ns = ns;
+			Blob = blob;
+		}
+	}
+
+	public class BlobReplicationException : BlobNotFoundException
+	{
+		public BlobReplicationException(NamespaceId ns, BlobId blob, string message) : base(ns, blob, message)
+		{
+		}
+	}
+
+	public class BlobToLargeException : Exception
+	{
+		public BlobId Blob { get; }
+
+		public BlobToLargeException(BlobId blob) : base($"Blob {blob} was to large to cache")
+		{
+			Blob = blob;
+		}
+	}
+
+	public class ResourceHasToManyRequestsException : Exception
+	{
+		public ResourceHasToManyRequestsException(Exception originalException) : base($"To many requests to resource", originalException)
+		{
+		}
+	}
+}
diff --git a/Engine/Source/Programs/Horde/Horde.Server/Ddc/IContentIdService.cs b/Engine/Source/Programs/Horde/Horde.Server/Ddc/IContentIdService.cs
--- a/Engine/Source/Programs/Horde/Horde.Server/Ddc/IContentIdService.cs
+++ b/Engine/Source/Programs/Horde/Horde.Server/Ddc/IContentIdService.cs
@@ -1,14 +1,14 @@
 // Copyright Epic Games, Inc. All Rights Reserved.
 
+using EpicGames.Horde.Storage;
 using System;
 using System.Threading.Tasks;
-using EpicGames.Horde.Storage;
 
 #pragma warning disable CS1591
 
 namespace Horde.Server.Ddc
 {
-	public interface IContentIdService
+	public interface IContentIdStore
 	{
 		/// <summary>
 		/// Resolve a content id from its hash into the actual blob (that can in turn be chunked into a set of blobs)
diff --git a/Engine/Source/Programs/Horde/Horde.Server/Ddc/IRefService.cs b/Engine/Source/Programs/Horde/Horde.Server/Ddc/IRefService.cs
--- a/Engine/Source/Programs/Horde/Horde.Server/Ddc/IRefService.cs
+++ b/Engine/Source/Programs/Horde/Horde.Server/Ddc/IRefService.cs
@@ -2,7 +2,6 @@
 
 using System;
 using System.Collections.Generic;
-using System.Threading;
 using System.Threading.Tasks;
 using EpicGames.Horde.Storage;
 using EpicGames.Serialization;
@@ -13,50 +12,18 @@ namespace Horde.Server.Ddc
 {
 	public interface IRefService
 	{
-		Task<(RefRecord, BlobContents?)> GetAsync(NamespaceId ns, BucketId bucket, RefId key, string[] fields, bool doLastAccessTracking = true, CancellationToken cancellationToken = default);
-		Task<(ContentId[], BlobId[])> PutAsync(NamespaceId ns, BucketId bucket, RefId key, BlobId blobHash, CbObject payload, CancellationToken cancellationToken = default);
-		Task<(ContentId[], BlobId[])> FinalizeAsync(NamespaceId ns, BucketId bucket, RefId key, BlobId blobHash, CancellationToken cancellationToken = default);
+		Task<(RefRecord, BlobContents?)> GetAsync(NamespaceId ns, BucketId bucket, RefId key, string[] fields, bool doLastAccessTracking = true);
+		Task<(ContentId[], BlobId[])> PutAsync(NamespaceId ns, BucketId bucket, RefId key, BlobId blobHash, CbObject payload);
+		Task<(ContentId[], BlobId[])> FinalizeAsync(NamespaceId ns, BucketId bucket, RefId key, BlobId blobHash);
 
-		Task<bool> DeleteAsync(NamespaceId ns, BucketId bucket, RefId key, CancellationToken cancellationToken = default);
+		IAsyncEnumerable<NamespaceId> GetNamespacesAsync();
 
-		Task<bool> ExistsAsync(NamespaceId ns, BucketId bucket, RefId key, CancellationToken cancellationToken = default);
-		Task<List<BlobId>> GetReferencedBlobsAsync(NamespaceId ns, BucketId bucket, RefId key, CancellationToken cancellationToken = default);
-	}
-
-	public class RefRecord
-	{
-		public RefRecord(NamespaceId ns, BucketId bucket, RefId name, DateTime lastAccess, byte[]? inlinePayload, BlobId blobIdentifier, bool isFinalized)
-		{
-			Namespace = ns;
-			Bucket = bucket;
-			Name = name;
-			LastAccess = lastAccess;
-			InlinePayload = inlinePayload;
-			BlobIdentifier = blobIdentifier;
-			IsFinalized = isFinalized;
-		}
-
-		public NamespaceId Namespace { get; }
-		public BucketId Bucket { get; }
-		public RefId Name { get; }
-		public DateTime LastAccess { get; }
-		public byte[]? InlinePayload { get; set; }
-		public BlobId BlobIdentifier { get; set; }
-		public bool IsFinalized { get; }
-	}
-
-	public class RefNotFoundException : Exception
-	{
-		public RefNotFoundException(NamespaceId ns, BucketId bucket, RefId key) : base($"Object not found {key} in bucket {bucket} namespace {ns}")
-		{
-			Namespace = ns;
-			Bucket = bucket;
-			Key = key;
-		}
+		Task<bool> DeleteAsync(NamespaceId ns, BucketId bucket, RefId key);
+		Task<long> DropNamespaceAsync(NamespaceId ns);
+		Task<long> DeleteBucketAsync(NamespaceId ns, BucketId bucket);
 
-		public NamespaceId Namespace { get; }
-		public BucketId Bucket { get; }
-		public RefId Key { get; }
+		Task<bool> ExistsAsync(NamespaceId ns, BucketId bucket, RefId key);
+		Task<List<BlobId>> GetReferencedBlobsAsync(NamespaceId ns, BucketId bucket, RefId key, bool ignoreMissingBlobs);
 	}
 
 	public class ObjectHashMismatchException : Exception
diff --git a/Engine/Source/Programs/Horde/Horde.Server/Ddc/IReferenceResolver.cs b/Engine/Source/Programs/Horde/Horde.Server/Ddc/IReferenceResolver.cs
--- a/Engine/Source/Programs/Horde/Horde.Server/Ddc/IReferenceResolver.cs
+++ b/Engine/Source/Programs/Horde/Horde.Server/Ddc/IReferenceResolver.cs
@@ -3,8 +3,6 @@
 using System;
 using System.Collections.Generic;
 using System.Linq;
-using System.Runtime.CompilerServices;
-using System.Threading;
 using System.Threading.Tasks;
 using EpicGames.Core;
 using EpicGames.Horde.Storage;
@@ -76,9 +74,9 @@ namespace Horde.Server.Ddc
 		/// </summary>
 		/// <param name="ns">The namespace to check</param>
 		/// <param name="cb">The compact binary object to resolve references for</param>
-		/// <param name="cancellationToken">Cancellation token for the operation</param>
+		/// <param name="ignoreMissingBlobs">Set to true to always returned the blobs found, ignoring anything that is missing rather then throwing</param>
 		/// <returns></returns>
-		IAsyncEnumerable<BlobId> GetReferencedBlobsAsync(NamespaceId ns, CbObject cb, CancellationToken cancellationToken = default);
+		IAsyncEnumerable<BlobId> GetReferencedBlobs(NamespaceId ns, CbObject cb, bool ignoreMissingBlobs = false);
 
 		/// <summary>
 		/// Returns which attachments exist in the cb object or any children
@@ -86,25 +84,24 @@ namespace Horde.Server.Ddc
 		/// </summary>
 		/// <param name="ns">The namespace to check</param>
 		/// <param name="cb">The compact binary object to resolve references for</param>
-		/// <param name="cancellationToken">Cancellation token for the operation</param>
 		/// <returns></returns>
-		IAsyncEnumerable<Attachment> GetAttachmentsAsync(NamespaceId ns, CbObject cb, CancellationToken cancellationToken = default);
+		IAsyncEnumerable<Attachment> GetAttachments(NamespaceId ns, CbObject cb);
 	}
 
 	public class ReferenceResolver : IReferenceResolver
 	{
 		private readonly IBlobService _blobStore;
-		private readonly IContentIdService _contentIdStore;
+		private readonly IContentIdStore _contentIdStore;
 		private readonly Tracer _tracer;
 
-		public ReferenceResolver(IBlobService blobStore, IContentIdService contentIdStore, Tracer tracer)
+		public ReferenceResolver(IBlobService blobStore, IContentIdStore contentIdStore, Tracer tracer)
 		{
 			_blobStore = blobStore;
 			_contentIdStore = contentIdStore;
 			_tracer = tracer;
 		}
 
-		public async IAsyncEnumerable<Attachment> GetAttachmentsAsync(NamespaceId ns, CbObject cb, [EnumeratorCancellation] CancellationToken cancellationToken)
+		public async IAsyncEnumerable<Attachment> GetAttachments(NamespaceId ns, CbObject cb)
 		{
 			Queue<CbObject> objectsToVisit = new Queue<CbObject>();
 			objectsToVisit.Enqueue(cb);
@@ -134,7 +131,7 @@ namespace Horde.Server.Ddc
 						else if (field.IsObjectAttachment())
 						{
 							attachments.Add(new ObjectAttachment(blobIdentifier));
-							pendingCompactBinaryAttachments.Add(ParseCompactBinaryAttachmentAsync(ns, blobIdentifier, cancellationToken));
+							pendingCompactBinaryAttachments.Add(ParseCompactBinaryAttachmentAsync(ns, blobIdentifier));
 						}
 						else
 						{
@@ -149,9 +146,9 @@ namespace Horde.Server.Ddc
 				{
 					if (pendingContentIdResolve.IsCompleted)
 					{
-						ContentId contentId = default;
+						ContentId? contentId = null;
 						BlobId[]? resolvedBlobs = null;
-						BlobId blobIdentifier = default;
+						BlobId? blobIdentifier = null;
 						bool wasContentId = false;
 						try
 						{
@@ -170,11 +167,11 @@ namespace Horde.Server.Ddc
 
 						if (wasContentId && resolvedBlobs != null)
 						{
-							attachments.Add(new ContentIdAttachment(contentId, resolvedBlobs));
+							attachments.Add(new ContentIdAttachment(contentId!, resolvedBlobs));
 						}
 						else
 						{
-							attachments.Add(new BlobAttachment(blobIdentifier));
+							attachments.Add(new BlobAttachment(blobIdentifier!));
 						}
 
 						finishedContentIdResolves.Add(pendingContentIdResolve);
@@ -238,7 +235,7 @@ namespace Horde.Server.Ddc
 			}
 		}
 
-		public async IAsyncEnumerable<BlobId> GetReferencedBlobsAsync(NamespaceId ns, CbObject cb, [EnumeratorCancellation] CancellationToken cancellationToken)
+		public async IAsyncEnumerable<BlobId> GetReferencedBlobs(NamespaceId ns, CbObject cb, bool ignoreMissingBlobs = false)
 		{
 			List<Task<(BlobId, bool)>> pendingBlobExistsChecks = new();
 			List<Task<(ContentIdAttachment, bool)>> pendingContentIdChecks = new();
@@ -246,21 +243,21 @@ namespace Horde.Server.Ddc
 			List<BlobId> unresolvedBlobReferences = new List<BlobId>();
 
 			// Resolve all the attachments
-			await foreach (Attachment attachment in GetAttachmentsAsync(ns, cb, cancellationToken))
+			await foreach (Attachment attachment in GetAttachments(ns, cb))
 			{
 				if (attachment is BlobAttachment blobAttachment)
 				{
-					pendingBlobExistsChecks.Add(CheckBlobExistsAsync(ns, blobAttachment.Identifier, cancellationToken));
+					pendingBlobExistsChecks.Add(CheckBlobExistsAsync(ns, blobAttachment.Identifier));
 				}
 				else if (attachment is ContentIdAttachment contentIdAttachment)
 				{
 					// If we find a content id we resolve that into the actual blobs it references
-					pendingContentIdChecks.Add(CheckContentIdExistsAsync(ns, contentIdAttachment, cancellationToken));
+					pendingContentIdChecks.Add(CheckContentIdExistsAsync(ns, contentIdAttachment));
 				}
 				else if (attachment is ObjectAttachment objectAttachment)
 				{
 					// a object just references the same blob, traversing the object attachment is done in GetAttachments
-					pendingBlobExistsChecks.Add(CheckBlobExistsAsync(ns, objectAttachment.Identifier, cancellationToken));
+					pendingBlobExistsChecks.Add(CheckBlobExistsAsync(ns, objectAttachment.Identifier));
 				}
 				else
 				{
@@ -299,23 +296,23 @@ namespace Horde.Server.Ddc
 			}
 
 			// if there were any content ids we did not recognize we throw a partial reference exception
-			if (unresolvedContentIdReferences.Count != 0)
+			if (!ignoreMissingBlobs && unresolvedContentIdReferences.Count != 0)
 			{
 				throw new PartialReferenceResolveException(unresolvedContentIdReferences);
 			}
 			// if there were any blobs missing we throw a partial reference exception
-			if (unresolvedBlobReferences.Count != 0)
+			if (!ignoreMissingBlobs && unresolvedBlobReferences.Count != 0)
 			{
 				throw new ReferenceIsMissingBlobsException(unresolvedBlobReferences);
 			}
 		}
 
-		private async Task<(ContentIdAttachment, bool)> CheckContentIdExistsAsync(NamespaceId ns, ContentIdAttachment contentIdAttachment, CancellationToken cancellationToken)
+		private async Task<(ContentIdAttachment, bool)> CheckContentIdExistsAsync(NamespaceId ns, ContentIdAttachment contentIdAttachment)
 		{
 			bool allBlobsExist = true;
 			foreach (BlobId b in contentIdAttachment.ReferencedBlobs)
 			{
-				(BlobId _, bool exists) = await CheckBlobExistsAsync(ns, b, cancellationToken);
+				(BlobId _, bool exists) = await CheckBlobExistsAsync(ns, b);
 
 				if (!exists)
 				{
@@ -327,10 +324,10 @@ namespace Horde.Server.Ddc
 			return (contentIdAttachment, allBlobsExist);
 		}
 
-		private async Task<CbObject> ParseCompactBinaryAttachmentAsync(NamespaceId ns, BlobId blobIdentifier, CancellationToken cancellationToken)
+		private async Task<CbObject> ParseCompactBinaryAttachmentAsync(NamespaceId ns, BlobId blobIdentifier)
 		{
-			BlobContents contents = await _blobStore.GetObjectAsync(ns, blobIdentifier, cancellationToken: cancellationToken);
-			byte[] data = await contents.Stream.ToByteArrayAsync(cancellationToken);
+			BlobContents contents = await _blobStore.GetObjectAsync(ns, blobIdentifier);
+			byte[] data = await contents.Stream.ToByteArrayAsync();
 			CbObject childBinaryObject = new CbObject(data);
 
 			return childBinaryObject;
@@ -345,9 +342,9 @@ namespace Horde.Server.Ddc
 			return (contentId, resolvedBlobs);
 		}
 
-		private async Task<(BlobId, bool)> CheckBlobExistsAsync(NamespaceId ns, BlobId blob, CancellationToken cancellationToken)
+		private async Task<(BlobId, bool)> CheckBlobExistsAsync(NamespaceId ns, BlobId blob)
 		{
-			return (blob, await _blobStore.ExistsAsync(ns, blob, cancellationToken: cancellationToken));
+			return (blob, await _blobStore.ExistsAsync(ns, blob));
 		}
 	}
 
@@ -355,7 +352,7 @@ namespace Horde.Server.Ddc
 	{
 		public List<ContentId> UnresolvedReferences { get; }
 
-		public PartialReferenceResolveException(List<ContentId> unresolvedReferences) : base($"References missing: {String.Join(',', unresolvedReferences)}")
+		public PartialReferenceResolveException(List<ContentId> unresolvedReferences) : base($"References missing: {string.Join(',', unresolvedReferences)}")
 		{
 			UnresolvedReferences = unresolvedReferences;
 		}
@@ -365,7 +362,7 @@ namespace Horde.Server.Ddc
 	{
 		public List<BlobId> MissingBlobs { get; }
 
-		public ReferenceIsMissingBlobsException(List<BlobId> missingBlobs) : base($"References is missing these blobs: {String.Join(',', missingBlobs)}")
+		public ReferenceIsMissingBlobsException(List<BlobId> missingBlobs) : base($"References is missing these blobs: {string.Join(',', missingBlobs)}")
 		{
 			MissingBlobs = missingBlobs;
 		}
diff --git a/Engine/Source/Programs/Horde/Horde.Server/Ddc/IReferencesStore.cs b/Engine/Source/Programs/Horde/Horde.Server/Ddc/IReferencesStore.cs
new file mode 100644
--- /dev/null
+++ b/Engine/Source/Programs/Horde/Horde.Server/Ddc/IReferencesStore.cs
@@ -0,0 +1,84 @@
+// Copyright Epic Games, Inc. All Rights Reserved.
+
+using System;
+using System.Collections.Generic;
+using System.Threading.Tasks;
+using EpicGames.Horde.Storage;
+
+#pragma warning disable CS1591
+
+namespace Horde.Server.Ddc
+{
+	public interface IReferencesStore
+	{
+		Task<RefRecord> GetAsync(NamespaceId ns, BucketId bucket, RefId key, FieldFlags fieldFlags, OperationFlags opFlags);
+
+		[Flags]
+		public enum FieldFlags
+		{
+			None = 0,
+			IncludePayload = 1,
+			All = IncludePayload
+		}
+
+		[Flags]
+		public enum OperationFlags
+		{
+			None = 0,
+			BypassCache = 1
+		}
+
+		Task PutAsync(NamespaceId ns, BucketId bucket, RefId key, BlobId blobHash, byte[] blob, bool isFinalized);
+		Task FinalizeAsync(NamespaceId ns, BucketId bucket, RefId key, BlobId blobIdentifier);
+
+		Task<DateTime?> GetLastAccessTimeAsync(NamespaceId ns, BucketId bucket, RefId key);
+		Task UpdateLastAccessTimeAsync(NamespaceId ns, BucketId bucket, RefId key, DateTime newLastAccessTime);
+		IAsyncEnumerable<(NamespaceId, BucketId, RefId, DateTime)> GetRecordsAsync();
+
+		IAsyncEnumerable<(NamespaceId, BucketId, RefId)> GetRecordsWithoutAccessTimeAsync();
+
+		IAsyncEnumerable<(RefId, BlobId)> GetRecordsInBucketAsync(NamespaceId ns, BucketId bucket);
+
+		IAsyncEnumerable<NamespaceId> GetNamespacesAsync();
+		IAsyncEnumerable<BucketId> GetBuckets(NamespaceId ns);
+		Task<bool> DeleteAsync(NamespaceId ns, BucketId bucket, RefId key);
+		Task<long> DropNamespaceAsync(NamespaceId ns);
+		Task<long> DeleteBucketAsync(NamespaceId ns, BucketId bucket);
+	}
+
+	public class RefRecord
+	{
+		public RefRecord(NamespaceId ns, BucketId bucket, RefId name, DateTime lastAccess, byte[]? inlinePayload, BlobId blobIdentifier, bool isFinalized)
+		{
+			Namespace = ns;
+			Bucket = bucket;
+			Name = name;
+			LastAccess = lastAccess;
+			InlinePayload = inlinePayload;
+			BlobIdentifier = blobIdentifier;
+			IsFinalized = isFinalized;
+		}
+
+		public NamespaceId Namespace { get; }
+		public BucketId Bucket { get; }
+		public RefId Name { get; }
+		public DateTime LastAccess { get; }
+		public byte[]? InlinePayload { get; set; }
+		public BlobId BlobIdentifier { get; set; }
+		public bool IsFinalized { get; }
+	}
+
+	public class RefNotFoundException : Exception
+	{
+		public RefNotFoundException(NamespaceId ns, BucketId bucket, RefId key) : base($"Object not found {key} in bucket {bucket} namespace {ns}")
+		{
+			Namespace = ns;
+			Bucket = bucket;
+			Key = key;
+		}
+
+		public NamespaceId Namespace { get; }
+		public BucketId Bucket { get; }
+		public RefId Key { get; }
+	}
+}
diff --git a/Engine/Source/Programs/Horde/Horde.Server/Ddc/IRequestHelper.cs b/Engine/Source/Programs/Horde/Horde.Server/Ddc/IRequestHelper.cs
--- a/Engine/Source/Programs/Horde/Horde.Server/Ddc/IRequestHelper.cs
+++ b/Engine/Source/Programs/Horde/Horde.Server/Ddc/IRequestHelper.cs
@@ -15,6 +15,7 @@ namespace Horde.Server.Ddc
 	public interface IRequestHelper
 	{
 		public Task<ActionResult?> HasAccessToNamespaceAsync(ClaimsPrincipal user, HttpRequest request, NamespaceId ns, AclAction[] aclActions);
+		public Task<ActionResult?> HasAccessForGlobalOperationsAsync(ClaimsPrincipal user, AclAction[] aclActions);
 	}
 
 	public class AuthorizationException : Exception
diff --git a/Engine/Source/Programs/Horde/Horde.Server/Ddc/JupiterAclAction.cs b/Engine/Source/Programs/Horde/Horde.Server/Ddc/JupiterAclAction.cs
new file mode 100644
--- /dev/null
+++ b/Engine/Source/Programs/Horde/Horde.Server/Ddc/JupiterAclAction.cs
@@ -0,0 +1,52 @@
+﻿// Copyright Epic Games, Inc. All Rights Reserved.
+
+using EpicGames.Horde.Acls;
+
+namespace Horde.Server.Ddc
+{
+	/// <summary>
+	/// ACL actions for DDC
+	/// </summary>
+	public static class JupiterAclAction
+	{
+		/// <summary>
+		/// General read access to refs / blobs and so on
+		/// </summary>
+		public static AclAction ReadObject { get; } = new AclAction("DdcReadObject");
+
+		/// <summary>
+		/// General write access to upload refs / blobs etc
+		/// </summary>
+		public static AclAction WriteObject { get; } = new AclAction("DdcWriteObject");
+
+		/// <summary>
+		/// Access to delete blobs / refs etc
+		/// </summary>
+		public static AclAction DeleteObject { get; } = new AclAction("DdcDeleteObject");
+
+		/// <summary>
+		/// Access to delete a particular bucket
+		/// </summary>
+		public static AclAction DeleteBucket { get; } = new AclAction("DdcDeleteBucket");
+
+		/// <summary>
+		/// Access to delete a whole namespace
+		/// </summary>
+		public static AclAction DeleteNamespace { get; } = new AclAction("DdcDeleteNamespace");
+
+		/// <summary>
+		/// Access to read the transaction log
+		/// </summary>
+		public static AclAction ReadTransactionLog { get; } = new AclAction("DdcReadTransactionLog");
+
+		/// <summary>
+		/// Access to write the transaction log
+		/// </summary>
+		public static AclAction WriteTransactionLog { get; } = new AclAction("DdcWriteTransactionLog");
+
+		/// <summary>
+		/// Access to perform administrative task
+		/// </summary>
+		public static AclAction AdminAction { get; } = new AclAction("DdcAdminAction");
+	}
+}
diff --git a/Engine/Source/Programs/Horde/Horde.Server/Ddc/NamespaceNotFoundException.cs b/Engine/Source/Programs/Horde/Horde.Server/Ddc/NamespaceNotFoundException.cs
new file mode 100644
--- /dev/null
+++ b/Engine/Source/Programs/Horde/Horde.Server/Ddc/NamespaceNotFoundException.cs
@@ -0,0 +1,24 @@
+// Copyright Epic Games, Inc. All Rights Reserved.
+
+using System;
+using EpicGames.Horde.Storage;
+
+#pragma warning disable CS1591
+
+namespace Horde.Server.Ddc
+{
+	public class NamespaceNotFoundException : Exception
+	{
+		public NamespaceId Namespace { get; }
+
+		public NamespaceNotFoundException(NamespaceId @namespace) : base($"Could not find namespace {@namespace}")
+		{
+			Namespace = @namespace;
+		}
+
+		public NamespaceNotFoundException(NamespaceId @namespace, string message) : base(message)
+		{
+			Namespace = @namespace;
+		}
+	}
+}
diff --git a/Engine/Source/Programs/Horde/Horde.Server/Ddc/ObjectsController.cs b/Engine/Source/Programs/Horde/Horde.Server/Ddc/ObjectsController.cs
--- a/Engine/Source/Programs/Horde/Horde.Server/Ddc/ObjectsController.cs
+++ b/Engine/Source/Programs/Horde/Horde.Server/Ddc/ObjectsController.cs
@@ -7,14 +7,11 @@ using System.ComponentModel.DataAnnotations;
 using System.Linq;
 using System.Net;
 using System.Net.Mime;
-using System.Threading;
 using System.Threading.Tasks;
 using EpicGames.AspNet;
 using EpicGames.Horde.Storage;
 using EpicGames.Serialization;
-using Horde.Server.Storage;
 using Microsoft.AspNetCore.Authorization;
-using Microsoft.AspNetCore.Http;
 using Microsoft.AspNetCore.Mvc;
 using Microsoft.Extensions.Logging;
 
@@ -26,19 +23,19 @@ namespace Horde.Server.Ddc
 
 	[ApiController]
 	[Route("api/v1/objects", Order = 0)]
-	[Tags("DDC Objects")]
 	[Authorize]
 	[Produces(CustomMediaTypeNames.UnrealCompactBinary, MediaTypeNames.Application.Json)]
-	public class ObjectsController : ControllerBase
+	public class ObjectController : ControllerBase
 	{
 		private readonly IBlobService _storage;
 		private readonly IDiagnosticContext _diagnosticContext;
 		private readonly IRequestHelper _requestHelper;
 		private readonly IReferenceResolver _referenceResolver;
 		private readonly BufferedPayloadFactory _bufferedPayloadFactory;
+
 		private readonly ILogger _logger;
 
-		public ObjectsController(IBlobService storage, IDiagnosticContext diagnosticContext, IRequestHelper requestHelper, IReferenceResolver referenceResolver, BufferedPayloadFactory bufferedPayloadFactory, ILogger<ObjectsController> logger)
+		public ObjectController(IBlobService storage, IDiagnosticContext diagnosticContext, IRequestHelper requestHelper, IReferenceResolver referenceResolver, BufferedPayloadFactory bufferedPayloadFactory, ILogger<ObjectController> logger)
 		{
 			_storage = storage;
 			_diagnosticContext = diagnosticContext;
@@ -54,7 +51,7 @@ namespace Horde.Server.Ddc
 			[Required] NamespaceId ns,
 			[Required] BlobId id)
 		{
-			ActionResult? result = await _requestHelper.HasAccessToNamespaceAsync(User, Request, ns, new[] { StorageAclAction.ReadBlobs });
+			ActionResult? result = await _requestHelper.HasAccessToNamespaceAsync(User, Request, ns, new[] { JupiterAclAction.ReadObject });
 			if (result != null)
 			{
 				return result;
@@ -78,7 +75,7 @@ namespace Horde.Server.Ddc
 			[Required] NamespaceId ns,
 			[Required] BlobId id)
 		{
-			ActionResult? result = await _requestHelper.HasAccessToNamespaceAsync(User, Request, ns, new[] { StorageAclAction.ReadBlobs });
+			ActionResult? result = await _requestHelper.HasAccessToNamespaceAsync(User, Request, ns, new[] { JupiterAclAction.ReadObject });
 			if (result != null)
 			{
 				return result;
@@ -100,7 +97,7 @@ namespace Horde.Server.Ddc
 			[Required] NamespaceId ns,
 			[Required][FromQuery] List<BlobId> id)
 		{
-			ActionResult? result = await _requestHelper.HasAccessToNamespaceAsync(User, Request, ns, new[] { StorageAclAction.ReadBlobs });
+			ActionResult? result = await _requestHelper.HasAccessToNamespaceAsync(User, Request, ns, new[] { JupiterAclAction.ReadObject });
 			if (result != null)
 			{
 				return result;
@@ -126,7 +123,7 @@ namespace Horde.Server.Ddc
 			[Required] NamespaceId ns,
 			[FromBody] BlobId[] bodyIds)
 		{
-			ActionResult? result = await _requestHelper.HasAccessToNamespaceAsync(User, Request, ns, new[] { StorageAclAction.ReadBlobs });
+			ActionResult? result = await _requestHelper.HasAccessToNamespaceAsync(User, Request, ns, new[] { JupiterAclAction.ReadObject });
 			if (result != null)
 			{
 				return result;
@@ -150,10 +147,9 @@ namespace Horde.Server.Ddc
 		[RequiredContentType(CustomMediaTypeNames.UnrealCompactBinary)]
 		public async Task<IActionResult> PutAsync(
 			[Required] NamespaceId ns,
-			[Required] BlobId id,
-			CancellationToken cancellationToken)
+			[Required] BlobId id)
 		{
-			ActionResult? result = await _requestHelper.HasAccessToNamespaceAsync(User, Request, ns, new[] { StorageAclAction.WriteBlobs });
+			ActionResult? result = await _requestHelper.HasAccessToNamespaceAsync(User, Request, ns, new[] { JupiterAclAction.WriteObject });
 			if (result != null)
 			{
 				return result;
@@ -162,9 +158,9 @@ namespace Horde.Server.Ddc
 			_diagnosticContext.Set("Content-Length", Request.ContentLength ?? -1);
 			try
 			{
-				using BufferedPayload payload = await _bufferedPayloadFactory.CreateFromRequestAsync(Request);
+				using IBufferedPayload payload = await _bufferedPayloadFactory.CreateFromRequest(Request);
 
-				BlobId identifier = await _storage.PutObjectAsync(ns, payload, id, cancellationToken);
+				BlobId identifier = await _storage.PutObjectAsync(ns, payload, id);
 				return Ok(new PutBlobResponse(identifier));
 			}
 			catch (ClientSendSlowException e)
@@ -178,7 +174,7 @@ namespace Horde.Server.Ddc
 			[Required] NamespaceId ns,
 			[Required] BlobId id)
 		{
-			ActionResult? result = await _requestHelper.HasAccessToNamespaceAsync(User, Request, ns, new[] { StorageAclAction.ReadBlobs });
+			ActionResult? result = await _requestHelper.HasAccessToNamespaceAsync(User, Request, ns, new[] { JupiterAclAction.ReadObject });
 			if (result != null)
 			{
 				return result;
@@ -194,7 +190,7 @@ namespace Horde.Server.Ddc
 				return NotFound(new ValidationProblemDetails { Title = $"Object {e.Blob} not found" });
 			}
 
-			byte[] blobContents = await blob.Stream.ToByteArrayAsync(HttpContext.RequestAborted);
+			byte[] blobContents = await blob.Stream.ToByteArrayAsync();
 			if (blobContents.Length == 0)
 			{
 				_logger.LogWarning("0 byte object found for {Id} {Namespace}", id, ns);
@@ -212,24 +208,59 @@ namespace Horde.Server.Ddc
 
 			try
 			{
-				BlobId[] references = await _referenceResolver.GetReferencedBlobsAsync(ns, compactBinaryObject).ToArrayAsync();
+				BlobId[] references = await _referenceResolver.GetReferencedBlobs(ns, compactBinaryObject).ToArrayAsync();
 				return Ok(new ResolvedReferencesResult(references));
 			}
 			catch (PartialReferenceResolveException e)
 			{
-				return BadRequest(new ValidationProblemDetails { Title = $"Object {id} is missing content ids", Detail = $"Following content ids are invalid: {String.Join(",", e.UnresolvedReferences)}" });
+				return BadRequest(new ValidationProblemDetails { Title = $"Object {id} is missing content ids", Detail = $"Following content ids are invalid: {string.Join(",", e.UnresolvedReferences)}" });
 			}
 			catch (ReferenceIsMissingBlobsException e)
 			{
-				return BadRequest(new ValidationProblemDetails { Title = $"Object {id} is missing blobs", Detail = $"Following blobs are missing: {String.Join(",", e.MissingBlobs)}" });
+				return BadRequest(new ValidationProblemDetails { Title = $"Object {id} is missing blobs", Detail = $"Following blobs are missing: {string.Join(",", e.MissingBlobs)}" });
 			}
 		}
+
+		[HttpDelete("{ns}/{id}")]
+		public async Task<IActionResult> DeleteAsync(
+			[Required] NamespaceId ns,
+			[Required] BlobId id)
+		{
+			ActionResult? result = await _requestHelper.HasAccessToNamespaceAsync(User, Request, ns, new[] { JupiterAclAction.DeleteObject });
+			if (result != null)
+			{
+				return result;
+			}
+
+			await _storage.DeleteObjectAsync(ns, id);
+
+			return Ok(new DeletedResponse
+			{
+				DeletedCount = 1
+			});
+		}
+
+		[HttpDelete("{ns}")]
+		public async Task<IActionResult> DeleteNamespaceAsync(
+			[Required] NamespaceId ns)
+		{
+			ActionResult? result = await _requestHelper.HasAccessToNamespaceAsync(User, Request, ns, new[] { JupiterAclAction.DeleteNamespace });
+			if (result != null)
+			{
+				return result;
+			}
+
+			await _storage.DeleteNamespaceAsync(ns);
+
+			return Ok();
+		}
 	}
 
 	public class PutBlobResponse
 	{
 		public PutBlobResponse()
 		{
+			Identifier = null!;
 		}
 
 		public PutBlobResponse(BlobId identifier)
diff --git a/Engine/Source/Programs/Horde/Horde.Server/Ddc/RefId.cs b/Engine/Source/Programs/Horde/Horde.Server/Ddc/RefId.cs
--- a/Engine/Source/Programs/Horde/Horde.Server/Ddc/RefId.cs
+++ b/Engine/Source/Programs/Horde/Horde.Server/Ddc/RefId.cs
@@ -3,91 +3,136 @@
 using System;
 using System.ComponentModel;
 using System.Globalization;
+using System.IO;
 using System.Text;
 using System.Text.Json;
 using System.Text.Json.Serialization;
 using EpicGames.Core;
 using EpicGames.Serialization;
 
+#pragma warning disable CS1591
+
 namespace Horde.Server.Ddc
 {
-	/// <summary>
-	/// Identifier for a DDC ref. Wrapper around an IoHash.
-	/// </summary>
-	[JsonConverter(typeof(RefIdJsonConverter))]
-	[TypeConverter(typeof(RefIdTypeConverter))]
-	[CbConverter(typeof(RefIdCbConverter))]
+	[TypeConverter(typeof(IoHashKeyTypeConverter))]
+	[JsonConverter(typeof(IoHashKeyJsonConverter))]
+	[CbConverter(typeof(IoHashKeyCbConverter))]
 	public readonly struct RefId : IEquatable<RefId>
 	{
-		/// <summary>
-		/// Hash of the ref name
-		/// </summary>
-		public IoHash Hash { get; }
-
-		/// <summary>
-		/// Constructor
-		/// </summary>
-		public RefId(IoHash hash) => Hash = hash;
-
-		/// <summary>
-		/// Parses a ref id from a string
-		/// </summary>
-		public static RefId Parse(string str) => new RefId(IoHash.Parse(str));
-
-		/// <inheritdoc/>
-		public override readonly bool Equals(object? obj) => obj is RefId id && Hash.Equals(id.Hash);
-
-		/// <inheritdoc/>
-		public override readonly int GetHashCode() => Hash.GetHashCode();
-
-		/// <inheritdoc/>
-		public readonly bool Equals(RefId other) => Hash.Equals(other.Hash);
-
-		/// <inheritdoc/>
-		public override readonly string ToString() => Hash.ToString();
-
-		/// <inheritdoc cref="IoHash.op_Equality"/>
-		public static bool operator ==(RefId left, RefId right) => left.Hash == right.Hash;
-
-		/// <inheritdoc cref="IoHash.op_Inequality"/>
-		public static bool operator !=(RefId left, RefId right) => left.Hash != right.Hash;
-
-		/// <summary>
-		/// Creates a RefId by hashing a string
-		/// </summary>
-		public static RefId FromName(string name) => new RefId(IoHash.Compute(Encoding.UTF8.GetBytes(name)));
+		public RefId(string key)
+		{
+			_text = key.ToLower();
+
+			if (_text.Length != 40)
+			{
+				throw new ArgumentException("IoHashKeys must be exactly 40 bytes.");
+			}
+
+			for (int idx = 0; idx < _text.Length; idx++)
+			{
+				if (!IsValidCharacter(_text[idx]))
+				{
+					throw new ArgumentException($"{_text} is not a valid namespace id");
+				}
+			}
+		}
+
+		static bool IsValidCharacter(char c)
+		{
+			if (c >= 'a' && c <= 'z')
+			{
+				return true;
+			}
+			if (c >= '0' && c <= '9')
+			{
+				return true;
+			}
+			return false;
+		}
+		readonly string _text;
+
+		public bool Equals(RefId other)
+		{
+			return string.Equals(_text, other._text, StringComparison.Ordinal);
+		}
+
+		public override bool Equals(object? obj)
+		{
+			return obj is RefId other && Equals(other);
+		}
+
+		public override int GetHashCode()
+		{
+			return _text.GetHashCode(StringComparison.Ordinal);
+		}
+
+		public override string ToString()
+		{
+			return _text;
+		}
+
+		public static bool operator ==(RefId left, RefId right)
+		{
+			return left.Equals(right);
+		}
+
+		public static bool operator !=(RefId left, RefId right)
+		{
+			return !left.Equals(right);
+		}
+
+		public static RefId FromName(string s)
+		{
+			return new RefId(ContentHash.FromBlob(Encoding.UTF8.GetBytes(s)).ToString());
+		}
 	}
 
-	/// <summary>
-	/// Type converter for RefId to and from JSON
-	/// </summary>
-	sealed class RefIdJsonConverter : JsonConverter<RefId>
+	public sealed class IoHashKeyTypeConverter : TypeConverter
 	{
 		/// <inheritdoc/>
-		public override RefId Read(ref Utf8JsonReader reader, Type typeToConvert, JsonSerializerOptions options) => RefId.Parse(reader.GetString() ?? String.Empty);
-
-		/// <inheritdoc/>
-		public override void Write(Utf8JsonWriter writer, RefId value, JsonSerializerOptions options) => writer.WriteStringValue(value.ToString());
+		public override bool CanConvertFrom(ITypeDescriptorContext? context, Type sourceType)
+		{
+			if (sourceType == typeof(string))
+			{
+				return true;
+			}
+
+			return base.CanConvertFrom(context, sourceType);
+		}
+
+		public override object? ConvertFrom(ITypeDescriptorContext? context, CultureInfo? culture, object value)
+		{
+			if (value is string s)
+			{
+				return new RefId(s);
+			}
+
+			return base.ConvertFrom(context, culture, value);
+		}
 	}
 
-	/// <summary>
-	/// Type converter from strings to RefId objects
-	/// </summary>
-	sealed class RefIdTypeConverter : TypeConverter
+	public class IoHashKeyJsonConverter : JsonConverter<RefId>
 	{
-		/// <inheritdoc/>
-		public override bool CanConvertFrom(ITypeDescriptorContext? context, Type? sourceType) => sourceType == typeof(string);
-
-		/// <inheritdoc/>
-		public override object ConvertFrom(ITypeDescriptorContext? context, CultureInfo? culture, object? value) => RefId.Parse((string)value!);
+		public override RefId Read(ref Utf8JsonReader reader, Type typeToConvert, JsonSerializerOptions options)
+		{
+			string? str = reader.GetString();
+			if (str == null)
+			{
+				throw new InvalidDataException("Unable to parse io hash key");
+			}
+
+			return new RefId(str);
+		}
+
+		public override void Write(Utf8JsonWriter writer, RefId value, JsonSerializerOptions options)
+		{
+			writer.WriteStringValue(value.ToString());
+		}
 	}
 
-	/// <summary>
-	/// Type converter to compact binary objects
-	/// </summary>
-	sealed class RefIdCbConverter : CbConverter<RefId>
+	sealed class IoHashKeyCbConverter : CbConverter<RefId>
 	{
-		public override RefId Read(CbField field) => RefId.Parse(field.AsString());
+		public override RefId Read(CbField field) => new RefId(field.AsString());
 
 		/// <inheritdoc/>
 		public override void Write(CbWriter writer, RefId value) => writer.WriteStringValue(value.ToString());
@@ -96,3 +141,4 @@ namespace Horde.Server.Ddc
 		public override void WriteNamed(CbWriter writer, CbFieldName name, RefId value) => writer.WriteString(name, value.ToString());
 	}
 }
+
diff --git a/Engine/Source/Programs/Horde/Horde.Server/Ddc/RefService.cs b/Engine/Source/Programs/Horde/Horde.Server/Ddc/RefService.cs
--- a/Engine/Source/Programs/Horde/Horde.Server/Ddc/RefService.cs
+++ b/Engine/Source/Programs/Horde/Horde.Server/Ddc/RefService.cs
@@ -3,7 +3,6 @@
 using System;
 using System.Collections.Generic;
 using System.Linq;
-using System.Threading;
 using System.Threading.Tasks;
 using EpicGames.Horde.Storage;
 using EpicGames.Serialization;
@@ -25,30 +24,30 @@ namespace Horde.Server.Ddc
 
 		static RefName GetRefName(BucketId bucketId, RefId refId) => $"{bucketId}/{refId}";
 
-		public async Task<bool> DeleteAsync(NamespaceId ns, BucketId bucket, RefId key, CancellationToken cancellationToken)
+		public async Task<bool> DeleteAsync(NamespaceId ns, BucketId bucket, RefId key)
 		{
 			using IStorageClient storageClient = _storageClientFactory.CreateClient(ns);
-			return await storageClient.DeleteRefAsync(GetRefName(bucket, key), cancellationToken);
+			return await storageClient.DeleteRefAsync(GetRefName(bucket, key));
 		}
 
-		public async Task<bool> ExistsAsync(NamespaceId ns, BucketId bucket, RefId key, CancellationToken cancellationToken)
+		public async Task<bool> ExistsAsync(NamespaceId ns, BucketId bucket, RefId key)
 		{
 			using IStorageClient storageClient = _storageClientFactory.CreateClient(ns);
-			return await storageClient.RefExistsAsync(GetRefName(bucket, key), cancellationToken: cancellationToken);
+			return await storageClient.RefExistsAsync(GetRefName(bucket, key));
 		}
 
-		public async Task<(ContentId[], BlobId[])> FinalizeAsync(NamespaceId ns, BucketId bucket, RefId key, BlobId blobHash, CancellationToken cancellationToken)
+		public async Task<(ContentId[], BlobId[])> FinalizeAsync(NamespaceId ns, BucketId bucket, RefId key, BlobId blobHash)
 		{
 			using IStorageClient storageClient = _storageClientFactory.CreateClient(ns);
 
-			BlobAlias? blobAlias = await storageClient.FindAliasAsync(BlobService.GetAlias(blobHash), cancellationToken);
+			BlobAlias? blobAlias = await storageClient.FindAliasAsync(BlobService.GetAlias(blobHash));
 			if (blobAlias == null)
 			{
 				throw new BlobNotFoundException(ns, blobHash);
 			}
 
 			IBlobHandle blobHandle = blobAlias.Target;
-			using BlobData blobContents = await blobHandle.ReadBlobDataAsync(cancellationToken);
+			using BlobData blobContents = await blobHandle.ReadBlobDataAsync();
 			CbObject payload = new CbObject(blobContents.Data);
 
 			BlobId[] referencedBlobs = Array.Empty<BlobId>();
@@ -60,7 +59,7 @@ namespace Horde.Server.Ddc
 			{
 				try
 				{
-					referencedBlobs = await _referenceResolver.GetReferencedBlobsAsync(ns, payload, cancellationToken).ToArrayAsync(cancellationToken);
+					referencedBlobs = await _referenceResolver.GetReferencedBlobs(ns, payload).ToArrayAsync();
 				}
 				catch (PartialReferenceResolveException e)
 				{
@@ -80,17 +79,17 @@ namespace Horde.Server.Ddc
 				IBlobRef<DdcRefNode> refNodeRef;
 				await using (IBlobWriter writer = storageClient.CreateBlobWriter(refName))
 				{
-					DdcRefNode refNode = new DdcRefNode(blobHash.Hash);
-					refNode.References.Add(BlobRef.Create(blobHash.Hash, blobHandle));
+					DdcRefNode refNode = new DdcRefNode(blobHash.AsIoHash());
+					refNode.References.Add(BlobRef.Create(blobHash.AsIoHash(), blobHandle));
 					foreach (BlobId referencedBlob in referencedBlobs)
 					{
-						BlobAlias? alias = await storageClient.FindAliasAsync(BlobService.GetAlias(referencedBlob), cancellationToken);
-						refNode.References.Add(BlobRef.Create(referencedBlob.Hash, alias!.Target));
+						BlobAlias? alias = await storageClient.FindAliasAsync(BlobService.GetAlias(referencedBlob));
+						refNode.References.Add(BlobRef.Create(referencedBlob.AsIoHash(), alias!.Target));
 					}
-					refNodeRef = await writer.WriteBlobAsync(refNode, cancellationToken: cancellationToken);
+					refNodeRef = await writer.WriteBlobAsync(refNode);
 				}
 
-				await storageClient.WriteRefAsync(refName, refNodeRef, cancellationToken: cancellationToken);
+				await storageClient.WriteRefAsync(refName, refNodeRef);
 			}
 
 			return (missingReferences, missingBlobs);
@@ -127,40 +126,60 @@ namespace Horde.Server.Ddc
 			return payload.Any(FieldHasAttachments);
 		}
 
-		public async Task<(RefRecord, BlobContents?)> GetAsync(NamespaceId ns, BucketId bucket, RefId key, string[] fields, bool doLastAccessTracking, CancellationToken cancellationToken)
+		public async Task<(RefRecord, BlobContents?)> GetAsync(NamespaceId ns, BucketId bucket, RefId key, string[] fields, bool doLastAccessTracking)
 		{
 			using IStorageClient storageClient = _storageClientFactory.CreateClient(ns);
 
-			DdcRefNode? node = await storageClient.TryReadRefTargetAsync<DdcRefNode>(GetRefName(bucket, key), cancellationToken: cancellationToken);
+			DdcRefNode? node = await storageClient.TryReadRefTargetAsync<DdcRefNode>(GetRefName(bucket, key));
 			if (node == null)
 			{
 				throw new RefNotFoundException(ns, bucket, key);
 			}
 
-			BlobData data = await node.References.First(x => x.Hash == node.RootHash).ReadBlobDataAsync(cancellationToken);
+			BlobData data = await node.References.First(x => x.Hash == node.RootHash).ReadBlobDataAsync();
 			BlobContents contents = new BlobContents(data.Data.ToArray());
 
-			RefRecord record = new RefRecord(ns, bucket, key, DateTime.UtcNow, null, new BlobId(node.RootHash), true);
+			RefRecord record = new RefRecord(ns, bucket, key, DateTime.UtcNow, null, BlobId.FromIoHash(node.RootHash), true);
 			return (record, contents);
 		}
 
-		public async Task<List<BlobId>> GetReferencedBlobsAsync(NamespaceId ns, BucketId bucket, RefId key, CancellationToken cancellationToken)
+		public async Task<List<BlobId>> GetReferencedBlobsAsync(NamespaceId ns, BucketId bucket, RefId key)
 		{
 			using IStorageClient storageClient = _storageClientFactory.CreateClient(ns);
 
-			DdcRefNode? node = await storageClient.TryReadRefTargetAsync<DdcRefNode>(GetRefName(bucket, key), cancellationToken: cancellationToken);
+			DdcRefNode? node = await storageClient.TryReadRefTargetAsync<DdcRefNode>(GetRefName(bucket, key));
 			if (node == null)
 			{
 				throw new RefNotFoundException(ns, bucket, key);
 			}
 
-			return node.References.Select(x => new BlobId(x.Hash)).ToList();
+			return node.References.Select(x => BlobId.FromIoHash(x.Hash)).ToList();
 		}
 
-		public async Task<(ContentId[], BlobId[])> PutAsync(NamespaceId ns, BucketId bucket, RefId key, BlobId blobHash, CbObject payload, CancellationToken cancellationToken)
+		public async Task<(ContentId[], BlobId[])> PutAsync(NamespaceId ns, BucketId bucket, RefId key, BlobId blobHash, CbObject payload)
 		{
-			await _blobService.PutObjectAsync(ns, payload.GetView(), blobHash, cancellationToken);
-			return await FinalizeAsync(ns, bucket, key, blobHash, cancellationToken);
+			await _blobService.PutObjectAsync(ns, payload.GetView().ToArray(), blobHash);
+			return await FinalizeAsync(ns, bucket, key, blobHash);
+		}
+
+		public IAsyncEnumerable<NamespaceId> GetNamespacesAsync()
+		{
+			throw new NotImplementedException();
+		}
+
+		public Task<long> DropNamespaceAsync(NamespaceId ns)
+		{
+			throw new NotImplementedException();
+		}
+
+		public Task<long> DeleteBucketAsync(NamespaceId ns, BucketId bucket)
+		{
+			throw new NotImplementedException();
+		}
+
+		public Task<List<BlobId>> GetReferencedBlobsAsync(NamespaceId ns, BucketId bucket, RefId key, bool ignoreMissingBlobs)
+		{
+			throw new NotImplementedException();
 		}
 	}
 }
diff --git a/Engine/Source/Programs/Horde/Horde.Server/Ddc/RefsController.cs b/Engine/Source/Programs/Horde/Horde.Server/Ddc/RefsController.cs
--- a/Engine/Source/Programs/Horde/Horde.Server/Ddc/RefsController.cs
+++ b/Engine/Source/Programs/Horde/Horde.Server/Ddc/RefsController.cs
@@ -18,7 +18,6 @@ using EpicGames.Core;
 using EpicGames.Horde.Acls;
 using EpicGames.Horde.Storage;
 using EpicGames.Serialization;
-using Horde.Server.Storage;
 using JetBrains.Annotations;
 using Microsoft.AspNetCore.Authorization;
 using Microsoft.AspNetCore.Http;
@@ -26,6 +25,7 @@ using Microsoft.AspNetCore.Http.Extensions;
 using Microsoft.AspNetCore.Mvc;
 using Microsoft.Extensions.DependencyInjection;
 using Microsoft.Extensions.Logging;
+using Microsoft.Extensions.Primitives;
 using OpenTelemetry.Trace;
 
 #pragma warning disable CS1591
@@ -40,7 +40,7 @@ namespace Horde.Server.Ddc
 	[Route("api/v1/refs")]
 	[Tags("DDC Refs")]
 	[Authorize]
-	public class RefsController : ControllerBase
+	public class ReferencesController : ControllerBase
 	{
 		private readonly IDiagnosticContext _diagnosticContext;
 		private readonly FormatResolver _formatResolver;
@@ -54,9 +54,9 @@ namespace Horde.Server.Ddc
 		private readonly IRefService _refService;
 		private readonly IBlobService _blobStore;
 
-		public RefsController(IRefService objectService, IBlobService blobStore, IDiagnosticContext diagnosticContext, FormatResolver formatResolver, BufferedPayloadFactory bufferedPayloadFactory, IReferenceResolver referenceResolver, NginxRedirectHelper nginxRedirectHelper, IRequestHelper requestHelper, Tracer tracer, ILogger<RefsController> logger)
+		public ReferencesController(IRefService refService, IBlobService blobStore, IDiagnosticContext diagnosticContext, FormatResolver formatResolver, BufferedPayloadFactory bufferedPayloadFactory, IReferenceResolver referenceResolver, NginxRedirectHelper nginxRedirectHelper, IRequestHelper requestHelper, Tracer tracer, ILogger<ReferencesController> logger)
 		{
-			_refService = objectService;
+			_refService = refService;
 			_blobStore = blobStore;
 			_diagnosticContext = diagnosticContext;
 			_formatResolver = formatResolver;
@@ -68,6 +68,31 @@ namespace Horde.Server.Ddc
 			_logger = logger;
 		}
 
+		/// <summary>
+		/// Returns all the known namespace the token has access to
+		/// </summary>
+		/// <returns></returns>
+		[HttpGet("")]
+		[ProducesDefaultResponseType]
+		[ProducesResponseType(type: typeof(ProblemDetails), 400)]
+		public async Task<IActionResult> GetNamespacesAsync()
+		{
+			NamespaceId[] namespaces = await _refService.GetNamespacesAsync().ToArrayAsync();
+
+			// filter namespaces down to only the namespaces the user has access to
+			List<NamespaceId> namespacesWithAccess = new();
+			foreach (NamespaceId ns in namespaces)
+			{
+				ActionResult? accessResult = await _requestHelper.HasAccessToNamespaceAsync(User, Request, ns, new[] { JupiterAclAction.ReadObject });
+				if (accessResult == null)
+				{
+					namespacesWithAccess.Add(ns);
+				}
+			}
+
+			return Ok(new GetNamespacesResponse(namespacesWithAccess.ToArray()));
+		}
+
 		/// <summary>
 		/// Returns a refs key
 		/// </summary>
@@ -83,7 +108,7 @@ namespace Horde.Server.Ddc
 			[FromRoute][Required] RefId key,
 			[FromRoute] string? format = null)
 		{
-			ActionResult? accessResult = await _requestHelper.HasAccessToNamespaceAsync(User, Request, ns, new[] { StorageAclAction.ReadRefs });
+			ActionResult? accessResult = await _requestHelper.HasAccessToNamespaceAsync(User, Request, ns, new[] { JupiterAclAction.ReadObject });
 			if (accessResult != null)
 			{
 				return accessResult;
@@ -146,7 +171,7 @@ namespace Horde.Server.Ddc
 						}
 					case MediaTypeNames.Application.Octet:
 						{
-							byte[] blobMemory = await blob.Stream.ToByteArrayAsync(HttpContext.RequestAborted);
+							byte[] blobMemory = await blob.Stream.ToByteArrayAsync();
 							CbObject cb = new CbObject(blobMemory);
 
 							(int, CbField?) CountFields(CbObject o)
@@ -195,7 +220,7 @@ namespace Horde.Server.Ddc
 							byte[] blobMemory;
 							{
 								using TelemetrySpan scope = _tracer.StartActiveSpan("json.readblob").SetAttribute("operation.name", "json.readblob");
-								blobMemory = await blob.Stream.ToByteArrayAsync(HttpContext.RequestAborted);
+								blobMemory = await blob.Stream.ToByteArrayAsync();
 							}
 							CbObject cb = new CbObject(blobMemory);
 							string s = cb.ToJson();
@@ -207,10 +232,10 @@ namespace Horde.Server.Ddc
 					case CustomMediaTypeNames.UnrealCompactBinaryPackage:
 						{
 							using TelemetrySpan packageScope = _tracer.StartActiveSpan("cbpackage.fetch").SetAttribute("operation.name", "cbpackage.fetch");
-							byte[] blobMemory = await blob.Stream.ToByteArrayAsync(HttpContext.RequestAborted);
+							byte[] blobMemory = await blob.Stream.ToByteArrayAsync();
 							CbObject cb = new CbObject(blobMemory);
 
-							IAsyncEnumerable<Attachment> attachments = _referenceResolver.GetAttachmentsAsync(ns, cb);
+							IAsyncEnumerable<Attachment> attachments = _referenceResolver.GetAttachments(ns, cb);
 
 							using CbPackageBuilder writer = new CbPackageBuilder();
 							writer.AddAttachment(objectRecord.BlobIdentifier.AsIoHash(), CbPackageAttachmentFlags.IsObject, blobMemory);
@@ -226,19 +251,19 @@ namespace Horde.Server.Ddc
 									if (attachment is BlobAttachment blobAttachment)
 									{
 										BlobId referencedBlob = blobAttachment.Identifier;
-										attachmentContents = await _blobStore.GetObjectAsync(ns, referencedBlob, cancellationToken: token);
+										attachmentContents = await _blobStore.GetObjectAsync(ns, referencedBlob);
 									}
 									else if (attachment is ObjectAttachment objectAttachment)
 									{
 										flags |= CbPackageAttachmentFlags.IsObject;
 										BlobId referencedBlob = objectAttachment.Identifier;
-										attachmentContents = await _blobStore.GetObjectAsync(ns, referencedBlob, cancellationToken: token);
+										attachmentContents = await _blobStore.GetObjectAsync(ns, referencedBlob);
 									}
 									else if (attachment is ContentIdAttachment contentIdAttachment)
 									{
 
 										ContentId contentId = contentIdAttachment.Identifier;
-										(attachmentContents, string mime) = await _blobStore.GetCompressedObjectAsync(ns, contentId, HttpContext.RequestServices, cancellationToken: token);
+										(attachmentContents, string mime) = await _blobStore.GetCompressedObjectAsync(ns, contentId, HttpContext.RequestServices);
 										if (mime == CustomMediaTypeNames.UnrealCompressedBuffer)
 										{
 											flags |= CbPackageAttachmentFlags.IsCompressed;
@@ -277,7 +302,7 @@ namespace Horde.Server.Ddc
 						}
 					case CustomMediaTypeNames.JupiterInlinedPayload:
 						{
-							byte[] blobMemory = await blob.Stream.ToByteArrayAsync(HttpContext.RequestAborted);
+							byte[] blobMemory = await blob.Stream.ToByteArrayAsync();
 							CbObject cb = new CbObject(blobMemory);
 
 							static (int, int) CountFields(CbObject o)
@@ -310,7 +335,7 @@ namespace Horde.Server.Ddc
 								List<BlobId> referencedBlobs;
 								try
 								{
-									IAsyncEnumerable<BlobId> referencedBlobsEnumerable = _referenceResolver.GetReferencedBlobsAsync(ns, cb);
+									IAsyncEnumerable<BlobId> referencedBlobsEnumerable = _referenceResolver.GetReferencedBlobs(ns, cb);
 									referencedBlobs = await referencedBlobsEnumerable.ToListAsync();
 								}
 								catch (PartialReferenceResolveException)
@@ -342,6 +367,8 @@ namespace Horde.Server.Ddc
 									}
 									catch (Exception ex)
 									{
+										Tracer.CurrentSpan.SetStatus(Status.Error);
+										Tracer.CurrentSpan.RecordException(ex);
 										_logger.LogError(ex, "Unknown exception encountered while writing body for jupiter inlined payload.");
 										throw;
 									}
@@ -412,7 +439,7 @@ namespace Horde.Server.Ddc
 			[FromRoute][Required] RefId key,
 			[FromQuery] string[] fields)
 		{
-			ActionResult? accessResult = await _requestHelper.HasAccessToNamespaceAsync(User, Request, ns, new[] { StorageAclAction.ReadRefs });
+			ActionResult? accessResult = await _requestHelper.HasAccessToNamespaceAsync(User, Request, ns, new[] { JupiterAclAction.ReadObject });
 			if (accessResult != null)
 			{
 				return accessResult;
@@ -453,7 +480,7 @@ namespace Horde.Server.Ddc
 			[FromRoute][Required] BucketId bucket,
 			[FromRoute][Required] RefId key)
 		{
-			ActionResult? accessResult = await _requestHelper.HasAccessToNamespaceAsync(User, Request, ns, new[] { StorageAclAction.ReadRefs });
+			ActionResult? accessResult = await _requestHelper.HasAccessToNamespaceAsync(User, Request, ns, new[] { JupiterAclAction.ReadObject });
 			if (accessResult != null)
 			{
 				return accessResult;
@@ -473,10 +500,10 @@ namespace Horde.Server.Ddc
 
 				// we have to verify the blobs are available locally, as the record of the key is replicated a head of the content
 				// TODO: Once we support inline replication this step is not needed as at least one region as this blob, just maybe not this current one
-				byte[] blobContents = await blob.Stream.ToByteArrayAsync(HttpContext.RequestAborted);
+				byte[] blobContents = await blob.Stream.ToByteArrayAsync();
 				CbObject compactBinaryObject = new CbObject(blobContents);
 				// the reference resolver will throw if any blob is missing, so no need to do anything other then process each reference
-				IAsyncEnumerable<BlobId> references = _referenceResolver.GetReferencedBlobsAsync(ns, compactBinaryObject);
+				IAsyncEnumerable<BlobId> references = _referenceResolver.GetReferencedBlobs(ns, compactBinaryObject);
 				List<BlobId>? _ = await references.ToListAsync();
 
 				// we have to verify the blobs are available locally, as the record of the key is replicated a head of the content
@@ -517,7 +544,7 @@ namespace Horde.Server.Ddc
 			[FromRoute][Required] NamespaceId ns,
 			[FromQuery][Required] List<string> names)
 		{
-			ActionResult? accessResult = await _requestHelper.HasAccessToNamespaceAsync(User, Request, ns, new[] { StorageAclAction.ReadRefs });
+			ActionResult? accessResult = await _requestHelper.HasAccessToNamespaceAsync(User, Request, ns, new[] { JupiterAclAction.ReadObject });
 			if (accessResult != null)
 			{
 				return accessResult;
@@ -535,7 +562,7 @@ namespace Horde.Server.Ddc
 				}
 
 				BucketId bucket = new BucketId(name.Substring(0, separatorIndex));
-				RefId key = RefId.Parse(name.Substring(separatorIndex + 1));
+				RefId key = new RefId(name.Substring(separatorIndex + 1));
 				requestedNames.Add((bucket, key));
 			}
 
@@ -551,10 +578,10 @@ namespace Horde.Server.Ddc
 
 					// we have to verify the blobs are available locally, as the record of the key is replicated a head of the content
 					// TODO: Once we support inline replication this step is not needed as at least one region as this blob, just maybe not this current one
-					byte[] blobContents = await blob.Stream.ToByteArrayAsync(HttpContext.RequestAborted);
+					byte[] blobContents = await blob.Stream.ToByteArrayAsync();
 					CbObject cb = new CbObject(blobContents);
 					// the reference resolver will throw if any blob is missing, so no need to do anything other then process each reference
-					IAsyncEnumerable<BlobId> references = _referenceResolver.GetReferencedBlobsAsync(ns, cb);
+					IAsyncEnumerable<BlobId> references = _referenceResolver.GetReferencedBlobs(ns, cb);
 					List<BlobId>? _ = await references.ToListAsync();
 				}
 				catch (RefNotFoundException)
@@ -584,10 +611,9 @@ namespace Horde.Server.Ddc
 		public async Task<IActionResult> PutObjectAsync(
 			[FromRoute][Required] NamespaceId ns,
 			[FromRoute][Required] BucketId bucket,
-			[FromRoute][Required] RefId key,
-			CancellationToken cancellationToken)
+			[FromRoute][Required] RefId key)
 		{
-			ActionResult? accessResult = await _requestHelper.HasAccessToNamespaceAsync(User, Request, ns, new[] { StorageAclAction.WriteRefs });
+			ActionResult? accessResult = await _requestHelper.HasAccessToNamespaceAsync(User, Request, ns, new[] { JupiterAclAction.WriteObject });
 			if (accessResult != null)
 			{
 				return accessResult;
@@ -600,18 +626,22 @@ namespace Horde.Server.Ddc
 
 			try
 			{
-				using BufferedPayload payload = await _bufferedPayloadFactory.CreateFromRequestAsync(Request);
-
-				string? hashHeaderValue = null;
-				if (Request.Headers.ContainsKey(CommonHeaders.HashHeaderName))
-				{
-					hashHeaderValue = Request.Headers[CommonHeaders.HashHeaderName];
-				}
+				using IBufferedPayload payload = await _bufferedPayloadFactory.CreateFromRequest(Request);
 
 				BlobId headerHash;
-				if (!String.IsNullOrEmpty(hashHeaderValue))
+				if (Request.Headers.TryGetValue(CommonHeaders.HashHeaderName, out StringValues headers))
 				{
-					headerHash = BlobId.Parse(hashHeaderValue);
+					if (!StringValues.IsNullOrEmpty(headers))
+					{
+						headerHash = new BlobId(headers.ToString());
+					}
+					else
+					{
+						return BadRequest(new ProblemDetails
+						{
+							Title = $"Header {CommonHeaders.HashHeaderName} was empty"
+						});
+					}
 				}
 				else
 				{
@@ -629,7 +659,7 @@ namespace Horde.Server.Ddc
 						{
 							// TODO: define a scheme for how a json object specifies references
 
-							blobHeader = await _blobStore.PutObjectAsync(ns, payload, headerHash, cancellationToken);
+							blobHeader = await _blobStore.PutObjectAsync(ns, payload, headerHash);
 
 							// TODO: convert the json object into a compact binary instead
 							CbWriter writer = new CbWriter();
@@ -646,13 +676,13 @@ namespace Horde.Server.Ddc
 						{
 							await using MemoryStream ms = new MemoryStream();
 							await using Stream payloadStream = payload.GetStream();
-							await payloadStream.CopyToAsync(ms, cancellationToken);
+							await payloadStream.CopyToAsync(ms);
 							payloadObject = new CbObject(ms.ToArray());
 							break;
 						}
 					case MediaTypeNames.Application.Octet:
 						{
-							blobHeader = await _blobStore.PutObjectAsync(ns, payload, headerHash, cancellationToken);
+							blobHeader = await _blobStore.PutObjectAsync(ns, payload, headerHash);
 
 							CbWriter writer = new CbWriter();
 							writer.BeginObject();
@@ -681,12 +711,17 @@ namespace Horde.Server.Ddc
 				return Problem(e.Message, null, (int)HttpStatusCode.RequestTimeout);
 			}
 
-			(ContentId[] missingReferences, BlobId[] missingBlobs) = await _refService.PutAsync(ns, bucket, key, blobHeader, payloadObject, cancellationToken);
+			(ContentId[] missingReferences, BlobId[] missingBlobs) = await _refService.PutAsync(ns, bucket, key, blobHeader, payloadObject);
 
-			List<IoHash> missingHashes = new List<IoHash>();
-			missingHashes.AddRange(missingReferences.Select(x => x.Hash));
-			missingHashes.AddRange(missingBlobs.Select(x => x.Hash));
-			return Ok(new PutObjectResponse(missingHashes.ToArray()));
+			{
+				using TelemetrySpan scope = _tracer.StartActiveSpan("ref.put").SetAttribute("operation.name", "ref.put");
+
+				List<ContentHash> missingHashes = new List<ContentHash>(missingReferences);
+				missingHashes.AddRange(missingBlobs);
+				ContentHash[] missingArray = missingHashes.ToArray();
+				scope.SetAttribute("NeedsCount", missingArray.Length);
+				return Ok(new PutObjectResponse(missingArray));
+			}
 		}
 
 		[HttpPut("{ns}/{bucket}/{key}", Order = 300)]
@@ -695,10 +730,9 @@ namespace Horde.Server.Ddc
 		public async Task<IActionResult> PutPackageAsync(
 			[FromRoute][Required] NamespaceId ns,
 			[FromRoute][Required] BucketId bucket,
-			[FromRoute][Required] RefId key,
-			CancellationToken cancellationToken)
+			[FromRoute][Required] RefId key)
 		{
-			ActionResult? accessResult = await _requestHelper.HasAccessToNamespaceAsync(User, Request, ns, new[] { StorageAclAction.WriteRefs });
+			ActionResult? accessResult = await _requestHelper.HasAccessToNamespaceAsync(User, Request, ns, new[] { JupiterAclAction.WriteObject });
 			if (accessResult != null)
 			{
 				return accessResult;
@@ -706,7 +740,7 @@ namespace Horde.Server.Ddc
 
 			_diagnosticContext.Set("Content-Length", Request.ContentLength ?? -1);
 
-			byte[] b = await ReadRawBodyAsync(Request);
+			byte[] b = await RequestUtil.ReadRawBodyAsync(Request);
 			CbPackageReader packageReader = await CbPackageReader.Create(new MemoryStream(b));
 
 			try
@@ -722,14 +756,14 @@ namespace Horde.Server.Ddc
 					}
 					if (entry.Flags.HasFlag(CbPackageAttachmentFlags.IsCompressed))
 					{
-#pragma warning disable CA2000 // False positive for not disposing payload
+#pragma warning disable CA2000 // Dispose objects before losing scope
 						using MemoryBufferedPayload payload = new MemoryBufferedPayload(blob);
-#pragma warning restore CA2000
-						await _blobStore.PutCompressedObjectAsync(ns, payload, ContentId.FromIoHash(entry.AttachmentHash), HttpContext.RequestServices, cancellationToken);
+#pragma warning restore CA2000 // Dispose objects before losing scope
+						await _blobStore.PutCompressedObjectAsync(ns, payload, ContentId.FromIoHash(entry.AttachmentHash), HttpContext.RequestServices);
 					}
 					else
 					{
-						await _blobStore.PutObjectAsync(ns, blob, BlobId.FromIoHash(entry.AttachmentHash), cancellationToken);
+						await _blobStore.PutObjectAsync(ns, blob, BlobId.FromIoHash(entry.AttachmentHash));
 					}
 				}
 			}
@@ -744,34 +778,11 @@ namespace Horde.Server.Ddc
 			CbObject rootObject = packageReader.RootObject;
 			BlobId rootObjectHash = BlobId.FromIoHash(packageReader.RootHash);
 
-			(ContentId[] missingReferences, BlobId[] missingBlobs) = await _refService.PutAsync(ns, bucket, key, rootObjectHash, rootObject, cancellationToken);
-			return Ok(new PutObjectResponse(missingReferences, missingBlobs));
-		}
-
-		static async Task<byte[]> ReadRawBodyAsync(HttpRequest request)
-		{
-			long? contentLength = request.GetTypedHeaders().ContentLength;
-
-			if (contentLength == null)
-			{
-				throw new Exception("Expected content-length on all raw body requests");
-			}
+			(ContentId[] missingReferences, BlobId[] missingBlobs) = await _refService.PutAsync(ns, bucket, key, rootObjectHash, rootObject);
 
-			Tracer? tracer = request.HttpContext.RequestServices.GetService<Tracer>();
-			using TelemetrySpan? scope = tracer?.StartActiveSpan("readbody")
-				.SetAttribute("operation.name", "readbody")
-				.SetAttribute("Content-Length", contentLength.Value);
-
-			await using MemoryStream ms = new MemoryStream((int)contentLength);
-			DateTime readStart = DateTime.Now;
-			await request.Body.CopyToAsync(ms);
-			TimeSpan duration = DateTime.Now - readStart;
-
-#if WITH_DOGSTATSD
-			double rate = contentLength.Value / duration.TotalSeconds;
-            StatsdClient.DogStatsd.Histogram("jupiter.stream_throughput", rate, tags: new string[] {"sourceIdentifier:" + "readbody"});
-#endif
-			return ms.GetBuffer();
+			List<ContentHash> missingHashes = new List<ContentHash>(missingReferences);
+			missingHashes.AddRange(missingBlobs);
+			return Ok(new PutObjectResponse(missingHashes.ToArray()));
 		}
 
 		[HttpPost("{ns}/{bucket}/{key}/finalize/{hash}.{format?}")]
@@ -781,7 +792,7 @@ namespace Horde.Server.Ddc
 			[FromRoute][Required] RefId key,
 			[FromRoute][Required] BlobId hash)
 		{
-			ActionResult? accessResult = await _requestHelper.HasAccessToNamespaceAsync(User, Request, ns, new[] { StorageAclAction.WriteRefs });
+			ActionResult? accessResult = await _requestHelper.HasAccessToNamespaceAsync(User, Request, ns, new[] { JupiterAclAction.WriteObject });
 			if (accessResult != null)
 			{
 				return accessResult;
@@ -790,7 +801,10 @@ namespace Horde.Server.Ddc
 			try
 			{
 				(ContentId[] missingReferences, BlobId[] missingBlobs) = await _refService.FinalizeAsync(ns, bucket, key, hash);
-				return Ok(new PutObjectResponse(missingReferences, missingBlobs));
+				List<ContentHash> missingHashes = new List<ContentHash>(missingReferences);
+				missingHashes.AddRange(missingBlobs);
+
+				return Ok(new PutObjectResponse(missingHashes.ToArray()));
 			}
 			catch (ObjectHashMismatchException e)
 			{
@@ -805,20 +819,21 @@ namespace Horde.Server.Ddc
 		[HttpPost("{ns}")]
 		[Consumes(CustomMediaTypeNames.UnrealCompactBinary)]
 		[Produces(CustomMediaTypeNames.UnrealCompactBinary)]
+		[ApiExplorerSettings(IgnoreApi = true)]
 		public async Task<IActionResult> BatchAsync(
 			[FromRoute][Required] NamespaceId ns,
-			[FromBody][Required] RefBatchOps ops)
+			[FromBody][Required] BatchOps ops)
 		{
-			AclAction ActionForOp(RefBatchOps.RefBatchOp.RefOperation op)
+			AclAction ActionForOp(BatchOps.BatchOp.Operation op)
 			{
 				switch (op)
 				{
-					case RefBatchOps.RefBatchOp.RefOperation.GET:
-						return StorageAclAction.ReadRefs;
-					case RefBatchOps.RefBatchOp.RefOperation.PUT:
-						return StorageAclAction.WriteRefs;
-					case RefBatchOps.RefBatchOp.RefOperation.HEAD:
-						return StorageAclAction.ReadRefs;
+					case BatchOps.BatchOp.Operation.GET:
+						return JupiterAclAction.ReadObject;
+					case BatchOps.BatchOp.Operation.PUT:
+						return JupiterAclAction.WriteObject;
+					case BatchOps.BatchOp.Operation.HEAD:
+						return JupiterAclAction.ReadObject;
 					default:
 						throw new ArgumentOutOfRangeException(nameof(op), op, null);
 				}
@@ -833,7 +848,7 @@ namespace Horde.Server.Ddc
 			}
 
 			HashSet<uint> usedOpIds = new HashSet<uint>();
-			foreach (RefBatchOps.RefBatchOp batchOp in ops.Ops)
+			foreach (BatchOps.BatchOp batchOp in ops.Ops)
 			{
 				bool added = usedOpIds.Add(batchOp.OpId);
 				if (!added)
@@ -843,7 +858,7 @@ namespace Horde.Server.Ddc
 			}
 			ConcurrentDictionary<uint, (CbObject, HttpStatusCode)> results = new();
 
-			async Task<(CbObject, HttpStatusCode)> BatchGetOp(RefBatchOps.RefBatchOp op)
+			async Task<(CbObject, HttpStatusCode)> BatchGetOp(BatchOps.BatchOp op)
 			{
 				try
 				{
@@ -851,7 +866,7 @@ namespace Horde.Server.Ddc
 
 					if (!objectRecord.IsFinalized)
 					{
-						throw new Exception("Object is not finalized");
+						return ToErrorResult("object not finalized", HttpStatusCode.BadRequest);
 					}
 
 					if (blob == null)
@@ -859,11 +874,11 @@ namespace Horde.Server.Ddc
 						throw new Exception();
 					}
 
-					CbObject cb = new CbObject(await blob.Stream.ToByteArrayAsync(HttpContext.RequestAborted));
+					CbObject cb = new CbObject(await blob.Stream.ToByteArrayAsync());
 
 					if (op.ResolveAttachments ?? false)
 					{
-						IAsyncEnumerable<BlobId> references = _referenceResolver.GetReferencedBlobsAsync(ns, cb);
+						IAsyncEnumerable<BlobId> references = _referenceResolver.GetReferencedBlobs(ns, cb);
 						List<BlobId>? _ = await references.ToListAsync();
 					}
 
@@ -875,11 +890,13 @@ namespace Horde.Server.Ddc
 				}
 				catch (Exception e)
 				{
+					Tracer.CurrentSpan.SetStatus(Status.Error);
+					Tracer.CurrentSpan.RecordException(e);
 					return ToErrorResult(e);
 				}
 			}
 
-			async Task<(CbObject, HttpStatusCode)> BatchHeadOp(RefBatchOps.RefBatchOp op)
+			async Task<(CbObject, HttpStatusCode)> BatchHeadOp(BatchOps.BatchOp op)
 			{
 				try
 				{
@@ -894,10 +911,10 @@ namespace Horde.Server.Ddc
 
 					if (op.ResolveAttachments ?? false)
 					{
-						byte[] blobContents = await blob.Stream.ToByteArrayAsync(HttpContext.RequestAborted);
+						byte[] blobContents = await blob.Stream.ToByteArrayAsync();
 						CbObject cb = new CbObject(blobContents);
 						// the reference resolver will throw if any blob is missing, so no need to do anything other then process each reference
-						IAsyncEnumerable<BlobId> references = _referenceResolver.GetReferencedBlobsAsync(ns, cb);
+						IAsyncEnumerable<BlobId> references = _referenceResolver.GetReferencedBlobs(ns, cb);
 						List<BlobId>? _ = await references.ToListAsync();
 					}
 
@@ -918,7 +935,7 @@ namespace Horde.Server.Ddc
 				}
 			}
 
-			async Task<(CbObject, HttpStatusCode)> BatchPutOp(RefBatchOps.RefBatchOp op)
+			async Task<(CbObject, HttpStatusCode)> BatchPutOp(BatchOps.BatchOp op)
 			{
 				try
 				{
@@ -931,17 +948,18 @@ namespace Horde.Server.Ddc
 					{
 						throw new Exception($"Missing payload hash for operation: {op.OpId}");
 					}
-
-					IoHash headerHash = op.PayloadHash.Value;
-					IoHash objectHash = IoHash.Compute(op.Payload.GetView().Span);
+					BlobId headerHash = BlobId.FromContentHash(op.PayloadHash);
+					BlobId objectHash = BlobId.FromBlob(op.Payload.GetView().ToArray());
 
 					if (!headerHash.Equals(objectHash))
 					{
-						throw new HashMismatchException(op.PayloadHash.Value, objectHash);
+						throw new HashMismatchException(headerHash, objectHash);
 					}
 
-					(ContentId[] missingReferences, BlobId[] missingBlobs) = await _refService.PutAsync(ns, op.Bucket, op.Key, new BlobId(objectHash), op.Payload);
-					return (CbSerializer.Serialize(new PutObjectResponse(missingReferences, missingBlobs)), HttpStatusCode.OK);
+					(ContentId[] missingReferences, BlobId[] missingBlobs) = await _refService.PutAsync(ns, op.Bucket, op.Key, objectHash, op.Payload);
+					List<ContentHash> missingHashes = new List<ContentHash>(missingReferences);
+
+					return (CbSerializer.Serialize(new PutObjectResponse(missingHashes.ToArray())), HttpStatusCode.OK);
 				}
 				catch (Exception e)
 				{
@@ -953,16 +971,16 @@ namespace Horde.Server.Ddc
 			{
 				switch (op.Op)
 				{
-					case RefBatchOps.RefBatchOp.RefOperation.GET:
+					case BatchOps.BatchOp.Operation.GET:
 						results.TryAdd(op.OpId, await BatchGetOp(op));
 						break;
-					case RefBatchOps.RefBatchOp.RefOperation.PUT:
+					case BatchOps.BatchOp.Operation.PUT:
 						results.TryAdd(op.OpId, await BatchPutOp(op));
 						break;
-					case RefBatchOps.RefBatchOp.RefOperation.HEAD:
+					case BatchOps.BatchOp.Operation.HEAD:
 						results.TryAdd(op.OpId, await BatchHeadOp(op));
 						break;
-					case RefBatchOps.RefBatchOp.RefOperation.INVALID:
+					case BatchOps.BatchOp.Operation.INVALID:
 					default:
 						throw new NotImplementedException($"Unknown op type {op.Op}");
 				}
@@ -999,6 +1017,74 @@ namespace Horde.Server.Ddc
 			return (writer.ToObject(), statusCode);
 		}
 
+		private static (CbObject, HttpStatusCode) ToErrorResult(string message, HttpStatusCode statusCode = HttpStatusCode.InternalServerError)
+		{
+			CbWriter writer = new CbWriter();
+			writer.BeginObject();
+			writer.WriteString("title", message);
+			writer.WriteInteger("status", (int)statusCode);
+			writer.EndObject();
+			return (writer.ToObject(), statusCode);
+		}
+
+		/// <summary>
+		/// Drop all refs records in the namespace
+		/// </summary>
+		/// <param name="ns">Namespace. Each namespace is completely separated from each other. Use for different types of data that is never expected to be similar (between two different games for instance)</param>
+		[HttpDelete("{ns}", Order = 500)]
+		[ProducesResponseType(204)]
+		public async Task<IActionResult> DeleteNamespaceAsync(
+			[FromRoute][Required] NamespaceId ns
+		)
+		{
+			ActionResult? accessResult = await _requestHelper.HasAccessToNamespaceAsync(User, Request, ns, new[] { JupiterAclAction.DeleteNamespace });
+			if (accessResult != null)
+			{
+				return accessResult;
+			}
+
+			try
+			{
+				await _refService.DropNamespaceAsync(ns);
+			}
+			catch (NamespaceNotFoundException e)
+			{
+				return NotFound(new ProblemDetails { Title = $"Namespace {e.Namespace} did not exist" });
+			}
+
+			return NoContent();
+		}
+
+		/// <summary>
+		/// Drop all refs records in the bucket
+		/// </summary>
+		/// <param name="ns">Namespace. Each namespace is completely separated from each other. Use for different types of data that is never expected to be similar (between two different games for instance)</param>
+		/// <param name="bucket">The category/type of record you are caching. Is a clustered key together with the actual key, but all records in the same bucket can be dropped easily.</param>
+		[HttpDelete("{ns}/{bucket}", Order = 500)]
+		[ProducesResponseType(200)]
+		public async Task<IActionResult> DeleteBucketAsync(
+			[FromRoute][Required] NamespaceId ns,
+			[FromRoute][Required] BucketId bucket)
+		{
+			ActionResult? accessResult = await _requestHelper.HasAccessToNamespaceAsync(User, Request, ns, new[] { JupiterAclAction.DeleteBucket });
+			if (accessResult != null)
+			{
+				return accessResult;
+			}
+
+			long countOfDeletedRecords;
+			try
+			{
+				countOfDeletedRecords = await _refService.DeleteBucketAsync(ns, bucket);
+			}
+			catch (NamespaceNotFoundException e)
+			{
+				return NotFound(new ProblemDetails { Title = $"Namespace {e.Namespace} did not exist" });
+			}
+
+			return Ok(new BucketDeletedResponse(countOfDeletedRecords));
+		}
+
 		/// <summary>
 		/// Delete a individual refs key
 		/// </summary>
@@ -1013,7 +1099,7 @@ namespace Horde.Server.Ddc
 			[FromRoute][Required] BucketId bucket,
 			[FromRoute][Required] RefId key)
 		{
-			ActionResult? accessResult = await _requestHelper.HasAccessToNamespaceAsync(User, Request, ns, new[] { StorageAclAction.DeleteRefs });
+			ActionResult? accessResult = await _requestHelper.HasAccessToNamespaceAsync(User, Request, ns, new[] { JupiterAclAction.DeleteObject });
 			if (accessResult != null)
 			{
 				return accessResult;
@@ -1022,10 +1108,6 @@ namespace Horde.Server.Ddc
 			try
 			{
 				bool deleted = await _refService.DeleteAsync(ns, bucket, key);
-				if (!deleted)
-				{
-					return NotFound(new ProblemDetails { Title = $"Object {key} in bucket {bucket} and namespace {ns} did not exist" });
-				}
 				return Ok(new RefDeletedResponse(deleted ? 1 : 0));
 			}
 			catch (NamespaceNotFoundException e)
@@ -1034,7 +1116,7 @@ namespace Horde.Server.Ddc
 			}
 			catch (RefNotFoundException)
 			{
-				return NotFound(new ProblemDetails { Title = $"Object {key} in bucket {bucket} and namespace {ns} did not exist" });
+				return Ok(new RefDeletedResponse(0));
 			}
 		}
 	}
@@ -1071,23 +1153,23 @@ namespace Horde.Server.Ddc
 		public long CountOfDeletedRecords { get; set; }
 	}
 
-	public class RefBatchOps
+	public class BatchOps
 	{
-		public RefBatchOps()
+		public BatchOps()
 		{
-			Ops = Array.Empty<RefBatchOp>();
+			Ops = Array.Empty<BatchOp>();
 		}
 
 		[System.Diagnostics.CodeAnalysis.SuppressMessage("Design", "CA1034:Nested types should not be visible", Justification = "For serialization only")]
-		public class RefBatchOp
+		public class BatchOp
 		{
-			public RefBatchOp()
+			public BatchOp()
 			{
 				Payload = null;
 				PayloadHash = null;
 			}
 
-			public enum RefOperation
+			public enum Operation
 			{
 				INVALID,
 				GET,
@@ -1105,11 +1187,11 @@ namespace Horde.Server.Ddc
 			public string OpString
 			{
 				get => Op.ToString();
-				set => Op = Enum.Parse<RefOperation>(value);
+				set => Op = Enum.Parse<Operation>(value);
 			}
 
 			[Required]
-			public RefOperation Op { get; set; } = RefOperation.INVALID;
+			public Operation Op { get; set; } = Operation.INVALID;
 
 			[Required]
 			[CbField("bucket")]
@@ -1126,11 +1208,11 @@ namespace Horde.Server.Ddc
 			public CbObject? Payload { get; set; } = null;
 
 			[CbField("payloadHash")]
-			public IoHash? PayloadHash { get; set; } = null;
+			public ContentHash? PayloadHash { get; set; } = null;
 		}
 
 		[CbField("ops")]
-		public RefBatchOp[] Ops { get; set; }
+		public BatchOp[] Ops { get; set; }
 	}
 
 	public class BatchOpsResponse
@@ -1166,6 +1248,7 @@ namespace Horde.Server.Ddc
 	{
 		public RefMetadataResponse()
 		{
+			PayloadIdentifier = null!;
 			InlinePayload = null!;
 		}
 
@@ -1221,18 +1304,13 @@ namespace Horde.Server.Ddc
 			Needs = null!;
 		}
 
-		public PutObjectResponse(IoHash[] missingReferences)
+		public PutObjectResponse(ContentHash[] missingReferences)
 		{
 			Needs = missingReferences;
 		}
 
-		public PutObjectResponse(IEnumerable<ContentId> missingContentIds, IEnumerable<BlobId> missingBlobIds)
-			: this(Enumerable.Concat(missingContentIds.Select(x => x.Hash), missingBlobIds.Select(x => x.Hash)).ToArray())
-		{
-		}
-
 		[CbField("needs")]
-		public IoHash[] Needs { get; set; }
+		public ContentHash[] Needs { get; set; }
 	}
 
 	public class ExistCheckMultipleRefsResponse
@@ -1284,10 +1362,10 @@ namespace Horde.Server.Ddc
 
 	public class HashMismatchException : Exception
 	{
-		public IoHash SuppliedHash { get; }
-		public IoHash ContentHash { get; }
+		public ContentHash SuppliedHash { get; }
+		public ContentHash ContentHash { get; }
 
-		public HashMismatchException(IoHash suppliedHash, IoHash contentHash) : base($"ID was not a hash of the content uploaded. Supplied hash was: {suppliedHash} but hash of content was {contentHash}")
+		public HashMismatchException(ContentHash suppliedHash, ContentHash contentHash) : base($"ID was not a hash of the content uploaded. Supplied hash was: {suppliedHash} but hash of content was {contentHash}")
 		{
 			SuppliedHash = suppliedHash;
 			ContentHash = contentHash;
diff --git a/Engine/Source/Programs/Horde/Horde.Server/Ddc/RequestHelper.cs b/Engine/Source/Programs/Horde/Horde.Server/Ddc/RequestHelper.cs
--- a/Engine/Source/Programs/Horde/Horde.Server/Ddc/RequestHelper.cs
+++ b/Engine/Source/Programs/Horde/Horde.Server/Ddc/RequestHelper.cs
@@ -24,6 +24,11 @@ namespace Horde.Server.Ddc
 			_globalConfig = globalConfig.Value;
 		}
 
+		public Task<ActionResult?> HasAccessForGlobalOperationsAsync(ClaimsPrincipal user, AclAction[] aclActions)
+		{
+			throw new System.NotImplementedException();
+		}
+
 		public Task<ActionResult?> HasAccessToNamespaceAsync(ClaimsPrincipal user, HttpRequest request, NamespaceId ns, AclAction[] aclActions)
 		{
 			using IStorageClient? storageClient = _storageClientFactory.TryCreateClient(ns);
diff --git a/Engine/Source/Programs/Horde/Horde.Server/Ddc/RequestUtil.cs b/Engine/Source/Programs/Horde/Horde.Server/Ddc/RequestUtil.cs
new file mode 100644
--- /dev/null
+++ b/Engine/Source/Programs/Horde/Horde.Server/Ddc/RequestUtil.cs
@@ -0,0 +1,41 @@
+﻿// Copyright Epic Games, Inc. All Rights Reserved.
+
+using System;
+using System.IO;
+using System.Threading.Tasks;
+using Microsoft.AspNetCore.Http;
+using Microsoft.Extensions.DependencyInjection;
+using OpenTelemetry.Trace;
+
+#pragma warning disable CS1591
+
+namespace Horde.Server.Ddc
+{
+	public static class RequestUtil
+	{
+		public static async Task<byte[]> ReadRawBodyAsync(HttpRequest request)
+		{
+			long? contentLength = request.GetTypedHeaders().ContentLength;
+
+			if (contentLength == null)
+			{
+				throw new Exception("Expected content-length on all raw body requests");
+			}
+
+			Tracer? tracer = request.HttpContext.RequestServices.GetService<Tracer>();
+			using TelemetrySpan? scope = tracer?.StartActiveSpan("readbody")
+				.SetAttribute("operation.name", "readbody")
+				.SetAttribute("Content-Length", contentLength.Value);
+
+			await using MemoryStream ms = new MemoryStream((int)contentLength);
+			DateTime readStart = DateTime.Now;
+			await request.Body.CopyToAsync(ms);
+			TimeSpan duration = DateTime.Now - readStart;
+
+			//			double rate = contentLength.Value / duration.TotalSeconds;
+			//			StatsdClient.DogStatsd.Histogram("jupiter.stream_throughput", rate, tags: new string[] { "sourceIdentifier:" + "readbody" });
+
+			return ms.GetBuffer();
+		}
+	}
+}
diff --git a/Engine/Source/Programs/Horde/Horde.Server/Ddc/RequiredContentTypeAttribute.cs b/Engine/Source/Programs/Horde/Horde.Server/Ddc/RequiredContentTypeAttribute.cs
--- a/Engine/Source/Programs/Horde/Horde.Server/Ddc/RequiredContentTypeAttribute.cs
+++ b/Engine/Source/Programs/Horde/Horde.Server/Ddc/RequiredContentTypeAttribute.cs
@@ -12,22 +12,20 @@ namespace Horde.Server.Ddc
 	[AttributeUsage(AttributeTargets.Method)]
 	public sealed class RequiredContentTypeAttribute : Attribute, IActionConstraint
 	{
-		private readonly string _mediaTypeName;
+		public string[] MediaTypeNames { get; }
 
-		public RequiredContentTypeAttribute(string mediaTypeName)
+		public RequiredContentTypeAttribute(params string[] mediaTypeNames)
 		{
-			_mediaTypeName = mediaTypeName;
+			MediaTypeNames = mediaTypeNames;
 		}
 
-		public string MediaTypeName => _mediaTypeName;
-
 		public int Order => 0;
 
 		public bool Accept(ActionConstraintContext context)
 		{
 			StringValues contentTypeHeader = context.RouteContext.HttpContext.Request.Headers["Content-Type"];
 
-			bool valid = contentTypeHeader.ToList().Contains(_mediaTypeName, StringComparer.InvariantCultureIgnoreCase);
+			bool valid = contentTypeHeader.Any(s => MediaTypeNames.Contains(s, StringComparer.InvariantCultureIgnoreCase));
 			return valid;
 		}
 	}
diff --git a/Engine/Source/Programs/Horde/Horde.Server/Ddc/Sha1Utils.cs b/Engine/Source/Programs/Horde/Horde.Server/Ddc/Sha1Utils.cs
new file mode 100644
--- /dev/null
+++ b/Engine/Source/Programs/Horde/Horde.Server/Ddc/Sha1Utils.cs
@@ -0,0 +1,78 @@
+﻿// Copyright Epic Games, Inc. All Rights Reserved.
+
+using System.Collections.Generic;
+using System.Security.Cryptography;
+using System.Threading.Tasks;
+
+#pragma warning disable CS1591
+#pragma warning disable CA5350 // Do not use weak crypto algorithms
+
+namespace Horde.Server.Ddc
+{
+	public static class Sha1Utils
+	{
+		public static async Task<string> GetSHA1AsStringAsync(byte[] payload)
+		{
+			return await Task.Run(() => GetSHA1AsString(payload));
+		}
+
+		public static byte[] GetSHA1(byte[] payload)
+		{
+			using SHA1 sha1 = SHA1.Create();
+			return sha1.ComputeHash(payload);
+		}
+
+		public static byte[] GetSHA1(byte[] payload, int offset, int count)
+		{
+			using SHA1 sha1 = SHA1.Create();
+			return sha1.ComputeHash(payload, offset, count);
+		}
+
+		public static string GetSHA1AsString(byte[] payload)
+		{
+			byte[] sha1 = GetSHA1(payload);
+
+			return StringUtils.FormatAsHexString(sha1);
+		}
+	}
+
+	public class ByteArrayComparer : EqualityComparer<byte[]>
+	{
+		public override bool Equals(byte[]? a, byte[]? b)
+		{
+			if (a == null)
+			{
+				return b == null;
+			}
+
+			if (b == null)
+			{
+				return false;
+			}
+
+			if (a.Length != b.Length)
+			{
+				return false;
+			}
+
+			for (int i = 0; i < a.Length; i++)
+			{
+				if (a[i] != b[i])
+				{
+					return false;
+				}
+			}
+			return true;
+		}
+		public override int GetHashCode(byte[] a)
+		{
+			uint b = 0;
+			foreach (byte t in a)
+			{
+				b = ((b << 23) | (b >> 9)) ^ t;
+			}
+
+			return unchecked((int)b);
+		}
+	}
+}
diff --git a/Engine/Source/Programs/Horde/Horde.Server/Ddc/StreamUtils.cs b/Engine/Source/Programs/Horde/Horde.Server/Ddc/StreamUtils.cs
--- a/Engine/Source/Programs/Horde/Horde.Server/Ddc/StreamUtils.cs
+++ b/Engine/Source/Programs/Horde/Horde.Server/Ddc/StreamUtils.cs
@@ -1,23 +1,23 @@
 ﻿// Copyright Epic Games, Inc. All Rights Reserved.
 
 using System.IO;
-using System.Threading;
 using System.Threading.Tasks;
 using EpicGames.AspNet;
 using Microsoft.AspNetCore.Http;
 
 #pragma warning disable CS1591
+#pragma warning disable MA0032 // Use an overload with a CancellationToken
 
 namespace Horde.Server.Ddc
 {
 	public static class StreamUtils
 	{
-		public static async Task<byte[]> ToByteArrayAsync(this Stream s, CancellationToken cancellationToken)
+		public static async Task<byte[]> ToByteArrayAsync(this Stream s)
 		{
 			try
 			{
 				await using MemoryStream ms = new MemoryStream();
-				await s.CopyToAsync(ms, cancellationToken);
+				await s.CopyToAsync(ms);
 				return ms.ToArray();
 			}
 			catch (BadHttpRequestException e)
diff --git a/Engine/Source/Programs/Horde/Horde.Server/Ddc/StringUtils.cs b/Engine/Source/Programs/Horde/Horde.Server/Ddc/StringUtils.cs
new file mode 100644
--- /dev/null
+++ b/Engine/Source/Programs/Horde/Horde.Server/Ddc/StringUtils.cs
@@ -0,0 +1,58 @@
+// Copyright Epic Games, Inc. All Rights Reserved.
+
+using System;
+using System.Text;
+
+#pragma warning disable CS1591
+
+namespace Horde.Server.Ddc
+{
+	public static class StringUtils
+	{
+		public static int GetHexVal(char hex)
+		{
+			return hex - (hex < 58 ? 48 : hex < 97 ? 55 : 87);
+		}
+
+		public static string FormatAsHexString(byte[] data)
+		{
+			StringBuilder sb = new StringBuilder(data.Length * 2);
+			foreach (byte b in data)
+			{
+				string hex = b.ToString("X2");
+				sb.Append(hex);
+			}
+
+			return sb.ToString();
+		}
+
+		public static string FormatAsHexLowerString(byte[] data)
+		{
+			StringBuilder sb = new StringBuilder(data.Length * 2);
+			foreach (byte b in data)
+			{
+				string hex = b.ToString("x2");
+				sb.Append(hex);
+			}
+
+			return sb.ToString();
+		}
+
+		public static byte[] ToHashFromHexString(string hex)
+		{
+			if (hex.Length % 2 == 1)
+			{
+				throw new Exception("The binary key cannot have an odd number of digits");
+			}
+
+			byte[] arr = new byte[hex.Length >> 1];
+
+			for (int i = 0; i < hex.Length >> 1; ++i)
+			{
+				arr[i] = (byte)((StringUtils.GetHexVal(hex[i << 1]) << 4) + StringUtils.GetHexVal(hex[(i << 1) + 1]));
+			}
+
+			return arr;
+		}
+	}
+}
diff --git a/Engine/Source/Programs/Horde/Horde.Server/Startup.cs b/Engine/Source/Programs/Horde/Horde.Server/Startup.cs
--- a/Engine/Source/Programs/Horde/Horde.Server/Startup.cs
+++ b/Engine/Source/Programs/Horde/Horde.Server/Startup.cs
@@ -619,7 +619,7 @@ namespace Horde.Server
 			services.AddScoped<IBlobService, BlobService>();
 			services.AddScoped<IRefService, RefService>();
 			services.AddScoped<IReferenceResolver, ReferenceResolver>();
-			services.AddScoped<IContentIdService, ContentIdService>();
+			services.AddScoped<IContentIdStore, ContentIdStore>();
 			services.AddSingleton<BufferedPayloadFactory>();
 			services.AddSingleton<NginxRedirectHelper>();
 			services.AddSingleton<FormatResolver>();
diff --git a/Engine/Source/Programs/Horde/Horde.Server/Storage/StorageController.cs b/Engine/Source/Programs/Horde/Horde.Server/Storage/StorageController.cs
--- a/Engine/Source/Programs/Horde/Horde.Server/Storage/StorageController.cs
+++ b/Engine/Source/Programs/Horde/Horde.Server/Storage/StorageController.cs
@@ -16,7 +16,6 @@ using EpicGames.Horde.Acls;
 using EpicGames.Horde.Storage;
 using EpicGames.Horde.Storage.Nodes;
 using EpicGames.Serialization;
-using Horde.Server.Ddc;
 using Horde.Server.Server;
 using Horde.Server.Utilities;
 using Microsoft.AspNetCore.Authorization;
@@ -27,6 +26,8 @@ using Microsoft.Extensions.Options;
 
 namespace Horde.Server.Storage
 {
+	using DdcRefNode = Horde.Server.Ddc.DdcRefNode;
+
 	/// <summary>
 	/// Controller for the /api/v1/storage endpoint
 	/// </summary>
diff --git a/Engine/Source/Programs/Horde/Horde.Server/Utilities/RawOutputFormatter.cs b/Engine/Source/Programs/Horde/Horde.Server/Utilities/RawOutputFormatter.cs
--- a/Engine/Source/Programs/Horde/Horde.Server/Utilities/RawOutputFormatter.cs
+++ b/Engine/Source/Programs/Horde/Horde.Server/Utilities/RawOutputFormatter.cs
@@ -10,9 +10,11 @@ using Microsoft.AspNetCore.Mvc.Formatters;
 using Microsoft.Extensions.Logging;
 using Microsoft.Net.Http.Headers;
 
+#pragma warning disable CS1591
+
 namespace Horde.Server.Utilities
 {
-	class RawOutputFormatter : OutputFormatter
+	public class RawOutputFormatter : OutputFormatter
 	{
 		private readonly ILogger _logger;
 
