Horde: Merging latest Jupiter code again.

#rnx

diff --git a/Engine/Source/Programs/Horde/Horde.Server/Ddc/BlobsController.cs b/Engine/Source/Programs/Horde/Horde.Server/Ddc/BlobsController.cs
--- a/Engine/Source/Programs/Horde/Horde.Server/Ddc/BlobsController.cs
+++ b/Engine/Source/Programs/Horde/Horde.Server/Ddc/BlobsController.cs
@@ -200,7 +200,7 @@ namespace Horde.Server.Ddc
 						RedirectUri = uri,
 					});
 				}
-				using IBufferedPayload payload = await _bufferedPayloadFactory.CreateFromRequest(Request);
+				using IBufferedPayload payload = await _bufferedPayloadFactory.CreateFromRequestAsync(Request);
 
 				BlobId identifier = await _storage.PutObjectAsync(ns, payload, id);
 				return Ok(new
@@ -233,7 +233,7 @@ namespace Horde.Server.Ddc
 			_diagnosticContext.Set("Content-Length", Request.ContentLength ?? -1);
 			try
 			{
-				using IBufferedPayload payload = await _bufferedPayloadFactory.CreateFromRequest(Request);
+				using IBufferedPayload payload = await _bufferedPayloadFactory.CreateFromRequestAsync(Request);
 
 				await using Stream stream = payload.GetStream();
 
diff --git a/Engine/Source/Programs/Horde/Horde.Server/Ddc/BufferedPayload.cs b/Engine/Source/Programs/Horde/Horde.Server/Ddc/BufferedPayload.cs
--- a/Engine/Source/Programs/Horde/Horde.Server/Ddc/BufferedPayload.cs
+++ b/Engine/Source/Programs/Horde/Horde.Server/Ddc/BufferedPayload.cs
@@ -183,7 +183,7 @@ namespace Horde.Server.Ddc
 			Directory.CreateDirectory(options.CurrentValue.FilesystemTempPayloadRoot);
 		}
 
-		public Task<IBufferedPayload> CreateFromRequest(HttpRequest request)
+		public Task<IBufferedPayload> CreateFromRequestAsync(HttpRequest request)
 		{
 			long? contentLength = request.ContentLength;
 
diff --git a/Engine/Source/Programs/Horde/Horde.Server/Ddc/CompressedBlobController.cs b/Engine/Source/Programs/Horde/Horde.Server/Ddc/CompressedBlobController.cs
--- a/Engine/Source/Programs/Horde/Horde.Server/Ddc/CompressedBlobController.cs
+++ b/Engine/Source/Programs/Horde/Horde.Server/Ddc/CompressedBlobController.cs
@@ -226,7 +226,7 @@ namespace Horde.Server.Ddc
 
 			try
 			{
-				using IBufferedPayload payload = await _bufferedPayloadFactory.CreateFromRequest(Request);
+				using IBufferedPayload payload = await _bufferedPayloadFactory.CreateFromRequestAsync(Request);
 
 				ContentId identifier = await _storage.PutCompressedObjectAsync(ns, payload, id, HttpContext.RequestServices);
 
@@ -262,7 +262,7 @@ namespace Horde.Server.Ddc
 
 			try
 			{
-				using IBufferedPayload payload = await _bufferedPayloadFactory.CreateFromRequest(Request);
+				using IBufferedPayload payload = await _bufferedPayloadFactory.CreateFromRequestAsync(Request);
 
 				ContentId identifier = await _storage.PutCompressedObjectAsync(ns, payload, null, HttpContext.RequestServices);
 
diff --git a/Engine/Source/Programs/Horde/Horde.Server/Ddc/CompressedBufferUtils.cs b/Engine/Source/Programs/Horde/Horde.Server/Ddc/CompressedBufferUtils.cs
--- a/Engine/Source/Programs/Horde/Horde.Server/Ddc/CompressedBufferUtils.cs
+++ b/Engine/Source/Programs/Horde/Horde.Server/Ddc/CompressedBufferUtils.cs
@@ -232,31 +232,42 @@ namespace Horde.Server.Ddc
 			}
 
 			// not using the buffered payload as we transfer the ownership to the caller of this method
-			FilesystemBufferedPayload finalizedBufferedPayload = bufferedPayloadWriter.Done();
-
-			if (header.TotalRawSize != (ulong)finalizedBufferedPayload.Length)
+			FilesystemBufferedPayload? finalizedBufferedPayload = null;
+			try
 			{
-				throw new Exception("Did not decompress the full payload");
-			}
+#pragma warning disable CA2000
+				finalizedBufferedPayload = bufferedPayloadWriter.Done();
+#pragma warning restore CA2000
 
-			{
-				using TelemetrySpan _ = _tracer.StartActiveSpan("web.hash").SetAttribute("operation.name", "web.hash");
+				if (header.TotalRawSize != (ulong)finalizedBufferedPayload.Length)
+				{
+					throw new Exception("Did not decompress the full payload");
+				}
 
-				// only read the first 20 bytes of the hash field as IoHashes are 20 bytes and not 32 bytes
-				byte[] slicedHash = new byte[20];
-				Array.Copy(header.RawHash, 0, slicedHash, 0, 20);
+				{
+					using TelemetrySpan _ = _tracer.StartActiveSpan("web.hash").SetAttribute("operation.name", "web.hash");
 
-				BlobId headerIdentifier = new BlobId(slicedHash);
-				await using Stream hashStream = finalizedBufferedPayload.GetStream();
-				BlobId contentHash = await BlobId.FromStreamAsync(hashStream);
+					// only read the first 20 bytes of the hash field as IoHashes are 20 bytes and not 32 bytes
+					byte[] slicedHash = new byte[20];
+					Array.Copy(header.RawHash, 0, slicedHash, 0, 20);
 
-				if (!headerIdentifier.Equals(contentHash))
-				{
-					throw new Exception($"Payload was expected to be {headerIdentifier} but was {contentHash}");
+					BlobId headerIdentifier = new BlobId(slicedHash);
+					await using Stream hashStream = finalizedBufferedPayload.GetStream();
+					BlobId contentHash = await BlobId.FromStreamAsync(hashStream);
+
+					if (!headerIdentifier.Equals(contentHash))
+					{
+						throw new Exception($"Payload was expected to be {headerIdentifier} but was {contentHash}");
+					}
 				}
-			}
 
-			return finalizedBufferedPayload;
+				return finalizedBufferedPayload;
+			}
+			catch
+			{
+				finalizedBufferedPayload?.Dispose();
+				throw;
+			}
 		}
 
 		private static int DecompressPayload(ReadOnlySpan<byte> compressedPayload, CompressedBufferHeader header, ulong rawBlockSize, Stream target)
diff --git a/Engine/Source/Programs/Horde/Horde.Server/Ddc/IReferenceResolver.cs b/Engine/Source/Programs/Horde/Horde.Server/Ddc/IReferenceResolver.cs
--- a/Engine/Source/Programs/Horde/Horde.Server/Ddc/IReferenceResolver.cs
+++ b/Engine/Source/Programs/Horde/Horde.Server/Ddc/IReferenceResolver.cs
@@ -76,7 +76,7 @@ namespace Horde.Server.Ddc
 		/// <param name="cb">The compact binary object to resolve references for</param>
 		/// <param name="ignoreMissingBlobs">Set to true to always returned the blobs found, ignoring anything that is missing rather then throwing</param>
 		/// <returns></returns>
-		IAsyncEnumerable<BlobId> GetReferencedBlobs(NamespaceId ns, CbObject cb, bool ignoreMissingBlobs = false);
+		IAsyncEnumerable<BlobId> GetReferencedBlobsAsync(NamespaceId ns, CbObject cb, bool ignoreMissingBlobs = false);
 
 		/// <summary>
 		/// Returns which attachments exist in the cb object or any children
@@ -85,7 +85,7 @@ namespace Horde.Server.Ddc
 		/// <param name="ns">The namespace to check</param>
 		/// <param name="cb">The compact binary object to resolve references for</param>
 		/// <returns></returns>
-		IAsyncEnumerable<Attachment> GetAttachments(NamespaceId ns, CbObject cb);
+		IAsyncEnumerable<Attachment> GetAttachmentsAsync(NamespaceId ns, CbObject cb);
 	}
 
 	public class ReferenceResolver : IReferenceResolver
@@ -101,7 +101,7 @@ namespace Horde.Server.Ddc
 			_tracer = tracer;
 		}
 
-		public async IAsyncEnumerable<Attachment> GetAttachments(NamespaceId ns, CbObject cb)
+		public async IAsyncEnumerable<Attachment> GetAttachmentsAsync(NamespaceId ns, CbObject cb)
 		{
 			Queue<CbObject> objectsToVisit = new Queue<CbObject>();
 			objectsToVisit.Enqueue(cb);
@@ -235,7 +235,7 @@ namespace Horde.Server.Ddc
 			}
 		}
 
-		public async IAsyncEnumerable<BlobId> GetReferencedBlobs(NamespaceId ns, CbObject cb, bool ignoreMissingBlobs = false)
+		public async IAsyncEnumerable<BlobId> GetReferencedBlobsAsync(NamespaceId ns, CbObject cb, bool ignoreMissingBlobs = false)
 		{
 			List<Task<(BlobId, bool)>> pendingBlobExistsChecks = new();
 			List<Task<(ContentIdAttachment, bool)>> pendingContentIdChecks = new();
@@ -243,7 +243,7 @@ namespace Horde.Server.Ddc
 			List<BlobId> unresolvedBlobReferences = new List<BlobId>();
 
 			// Resolve all the attachments
-			await foreach (Attachment attachment in GetAttachments(ns, cb))
+			await foreach (Attachment attachment in GetAttachmentsAsync(ns, cb))
 			{
 				if (attachment is BlobAttachment blobAttachment)
 				{
diff --git a/Engine/Source/Programs/Horde/Horde.Server/Ddc/IReferencesStore.cs b/Engine/Source/Programs/Horde/Horde.Server/Ddc/IReferencesStore.cs
--- a/Engine/Source/Programs/Horde/Horde.Server/Ddc/IReferencesStore.cs
+++ b/Engine/Source/Programs/Horde/Horde.Server/Ddc/IReferencesStore.cs
@@ -40,7 +40,7 @@ namespace Horde.Server.Ddc
 		IAsyncEnumerable<(RefId, BlobId)> GetRecordsInBucketAsync(NamespaceId ns, BucketId bucket);
 
 		IAsyncEnumerable<NamespaceId> GetNamespacesAsync();
-		IAsyncEnumerable<BucketId> GetBuckets(NamespaceId ns);
+		IAsyncEnumerable<BucketId> GetBucketsAsync(NamespaceId ns);
 		Task<bool> DeleteAsync(NamespaceId ns, BucketId bucket, RefId key);
 		Task<long> DropNamespaceAsync(NamespaceId ns);
 		Task<long> DeleteBucketAsync(NamespaceId ns, BucketId bucket);
diff --git a/Engine/Source/Programs/Horde/Horde.Server/Ddc/ObjectsController.cs b/Engine/Source/Programs/Horde/Horde.Server/Ddc/ObjectsController.cs
--- a/Engine/Source/Programs/Horde/Horde.Server/Ddc/ObjectsController.cs
+++ b/Engine/Source/Programs/Horde/Horde.Server/Ddc/ObjectsController.cs
@@ -158,7 +158,7 @@ namespace Horde.Server.Ddc
 			_diagnosticContext.Set("Content-Length", Request.ContentLength ?? -1);
 			try
 			{
-				using IBufferedPayload payload = await _bufferedPayloadFactory.CreateFromRequest(Request);
+				using IBufferedPayload payload = await _bufferedPayloadFactory.CreateFromRequestAsync(Request);
 
 				BlobId identifier = await _storage.PutObjectAsync(ns, payload, id);
 				return Ok(new PutBlobResponse(identifier));
@@ -208,7 +208,7 @@ namespace Horde.Server.Ddc
 
 			try
 			{
-				BlobId[] references = await _referenceResolver.GetReferencedBlobs(ns, compactBinaryObject).ToArrayAsync();
+				BlobId[] references = await _referenceResolver.GetReferencedBlobsAsync(ns, compactBinaryObject).ToArrayAsync();
 				return Ok(new ResolvedReferencesResult(references));
 			}
 			catch (PartialReferenceResolveException e)
diff --git a/Engine/Source/Programs/Horde/Horde.Server/Ddc/RefsController.cs b/Engine/Source/Programs/Horde/Horde.Server/Ddc/RefsController.cs
--- a/Engine/Source/Programs/Horde/Horde.Server/Ddc/RefsController.cs
+++ b/Engine/Source/Programs/Horde/Horde.Server/Ddc/RefsController.cs
@@ -235,7 +235,7 @@ namespace Horde.Server.Ddc
 							byte[] blobMemory = await blob.Stream.ToByteArrayAsync();
 							CbObject cb = new CbObject(blobMemory);
 
-							IAsyncEnumerable<Attachment> attachments = _referenceResolver.GetAttachments(ns, cb);
+							IAsyncEnumerable<Attachment> attachments = _referenceResolver.GetAttachmentsAsync(ns, cb);
 
 							using CbPackageBuilder writer = new CbPackageBuilder();
 							writer.AddAttachment(objectRecord.BlobIdentifier.AsIoHash(), CbPackageAttachmentFlags.IsObject, blobMemory);
@@ -335,7 +335,7 @@ namespace Horde.Server.Ddc
 								List<BlobId> referencedBlobs;
 								try
 								{
-									IAsyncEnumerable<BlobId> referencedBlobsEnumerable = _referenceResolver.GetReferencedBlobs(ns, cb);
+									IAsyncEnumerable<BlobId> referencedBlobsEnumerable = _referenceResolver.GetReferencedBlobsAsync(ns, cb);
 									referencedBlobs = await referencedBlobsEnumerable.ToListAsync();
 								}
 								catch (PartialReferenceResolveException)
@@ -503,7 +503,7 @@ namespace Horde.Server.Ddc
 				byte[] blobContents = await blob.Stream.ToByteArrayAsync();
 				CbObject compactBinaryObject = new CbObject(blobContents);
 				// the reference resolver will throw if any blob is missing, so no need to do anything other then process each reference
-				IAsyncEnumerable<BlobId> references = _referenceResolver.GetReferencedBlobs(ns, compactBinaryObject);
+				IAsyncEnumerable<BlobId> references = _referenceResolver.GetReferencedBlobsAsync(ns, compactBinaryObject);
 				List<BlobId>? _ = await references.ToListAsync();
 
 				// we have to verify the blobs are available locally, as the record of the key is replicated a head of the content
@@ -581,7 +581,7 @@ namespace Horde.Server.Ddc
 					byte[] blobContents = await blob.Stream.ToByteArrayAsync();
 					CbObject cb = new CbObject(blobContents);
 					// the reference resolver will throw if any blob is missing, so no need to do anything other then process each reference
-					IAsyncEnumerable<BlobId> references = _referenceResolver.GetReferencedBlobs(ns, cb);
+					IAsyncEnumerable<BlobId> references = _referenceResolver.GetReferencedBlobsAsync(ns, cb);
 					List<BlobId>? _ = await references.ToListAsync();
 				}
 				catch (RefNotFoundException)
@@ -626,7 +626,7 @@ namespace Horde.Server.Ddc
 
 			try
 			{
-				using IBufferedPayload payload = await _bufferedPayloadFactory.CreateFromRequest(Request);
+				using IBufferedPayload payload = await _bufferedPayloadFactory.CreateFromRequestAsync(Request);
 
 				BlobId headerHash;
 				if (Request.Headers.TryGetValue(CommonHeaders.HashHeaderName, out StringValues headers))
@@ -878,7 +878,7 @@ namespace Horde.Server.Ddc
 
 					if (op.ResolveAttachments ?? false)
 					{
-						IAsyncEnumerable<BlobId> references = _referenceResolver.GetReferencedBlobs(ns, cb);
+						IAsyncEnumerable<BlobId> references = _referenceResolver.GetReferencedBlobsAsync(ns, cb);
 						List<BlobId>? _ = await references.ToListAsync();
 					}
 
@@ -914,7 +914,7 @@ namespace Horde.Server.Ddc
 						byte[] blobContents = await blob.Stream.ToByteArrayAsync();
 						CbObject cb = new CbObject(blobContents);
 						// the reference resolver will throw if any blob is missing, so no need to do anything other then process each reference
-						IAsyncEnumerable<BlobId> references = _referenceResolver.GetReferencedBlobs(ns, cb);
+						IAsyncEnumerable<BlobId> references = _referenceResolver.GetReferencedBlobsAsync(ns, cb);
 						List<BlobId>? _ = await references.ToListAsync();
 					}
 
