Horde: Missing file.

#rnx

diff --git a/Engine/Source/Programs/Horde/Horde.Server/Ddc/RefService.cs b/Engine/Source/Programs/Horde/Horde.Server/Ddc/RefService.cs
--- a/Engine/Source/Programs/Horde/Horde.Server/Ddc/RefService.cs
+++ b/Engine/Source/Programs/Horde/Horde.Server/Ddc/RefService.cs
@@ -59,7 +59,7 @@ namespace Horde.Server.Ddc
 			{
 				try
 				{
-					referencedBlobs = await _referenceResolver.GetReferencedBlobs(ns, payload).ToArrayAsync();
+					referencedBlobs = await _referenceResolver.GetReferencedBlobsAsync(ns, payload).ToArrayAsync();
 				}
 				catch (PartialReferenceResolveException e)
 				{
