Horde: Fix padding to code elements in markdown renderer
#rnx

diff --git a/Engine/Source/Programs/Horde/HordeDashboard/src/base/components/Markdown/Markdown.tsx b/Engine/Source/Programs/Horde/HordeDashboard/src/base/components/Markdown/Markdown.tsx
--- a/Engine/Source/Programs/Horde/HordeDashboard/src/base/components/Markdown/Markdown.tsx
+++ b/Engine/Source/Programs/Horde/HordeDashboard/src/base/components/Markdown/Markdown.tsx
@@ -38,7 +38,7 @@ const getStyles: IStyleFunction<IMarkdownStyleProps, IMarkdownStyles> = () => {
          fontSize: "13px", fontFamily: "Horde Open Sans Regular", selectors: {
             'li': {
                margin: "4px 0px 8px 0px",
-               lineHeight:"1.6"
+               lineHeight: "1.6"
             },
             'ul': {
                margin: "8px 0px 16px 0px"
@@ -48,11 +48,15 @@ const getStyles: IStyleFunction<IMarkdownStyleProps, IMarkdownStyles> = () => {
                padding: 16,
                backgroundColor: dashboard.darktheme ? modeColors.header : modeColors.crumbs
             },
+            'p > code': {
+
+               padding: "3px 6px !important"
+
+            },
             'code': {
                whiteSpace: "pre-wrap",
                fontFamily: "Horde Cousine Regular",
                fontSize: "12px",
-               padding: "3px 6px",
                backgroundColor: dashboard.darktheme ? "#25282A" : "#f2f2f2"
             }
 
