Horde: Update docs for deploying server on Mac.

diff --git a/Engine/Source/Programs/Horde/Docs/Deployment/Server.md b/Engine/Source/Programs/Horde/Docs/Deployment/Server.md
--- a/Engine/Source/Programs/Horde/Docs/Deployment/Server.md
+++ b/Engine/Source/Programs/Horde/Docs/Deployment/Server.md
@@ -51,6 +51,8 @@ install all the prerequistes using [Homebrew](https://brew.sh/).
 2. Install MongoDB
 
     ```bash
+    brew tap mongodb/brew
+    brew update
     brew install mongodb-community
     brew services start mongodb-community
     ```
@@ -68,7 +70,7 @@ install all the prerequistes using [Homebrew](https://brew.sh/).
     ```bash
     export Horde__DatabaseConnectionString=mongodb://localhost:27017
     export Horde__HttpPort=37107
-    export Horde__Http2Port=37107
+    export Horde__Http2Port=37109
 
     cd Engine/Source/Programs/Horde/Horde.Server
     dotnet run
