Horde: Add cancellation tokens to async DDC methods.

#rnx

diff --git a/Engine/Source/Programs/Horde/Horde.Server.Tests/Ddc/Functional/ContentId/CompressedBlobTests.cs b/Engine/Source/Programs/Horde/Horde.Server.Tests/Ddc/Functional/ContentId/CompressedBlobTests.cs
--- a/Engine/Source/Programs/Horde/Horde.Server.Tests/Ddc/Functional/ContentId/CompressedBlobTests.cs
+++ b/Engine/Source/Programs/Horde/Horde.Server.Tests/Ddc/Functional/ContentId/CompressedBlobTests.cs
@@ -9,6 +9,7 @@ using System.Net.Http.Headers;
 using System.Net.Http.Json;
 using System.Net.Mime;
 using System.Text;
+using System.Threading;
 using System.Threading.Tasks;
 using Blake3;
 using EpicGames.AspNet;
@@ -135,8 +136,9 @@ namespace Horde.Server.Tests.Ddc.FunctionalTests.CompressedBlobs
 		[TestCategory("SlowTests")]
 		public async Task PutGetLargeCompressedPayloadAsync()
 		{
+			CancellationToken cancellationToken = CancellationToken.None;
+
 			// we submit a blob so large that it can not fit using the memory blob store
-			IBlobStore? blobStore = ServiceProvider.GetRequiredService<IBlobStore>();
 
 			FileInfo tempOutputFile = new FileInfo(Path.GetTempFileName());
 			FileInfo tempCompressedFile = new FileInfo(Path.GetTempFileName());
@@ -176,7 +178,7 @@ namespace Horde.Server.Tests.Ddc.FunctionalTests.CompressedBlobs
 				BlobId compressedContentHash;
 				{
 					await using FileStream fs = tempCompressedFile.OpenRead();
-					compressedContentHash = await BlobId.FromStreamAsync(fs);
+					compressedContentHash = await BlobId.FromStreamAsync(fs, cancellationToken);
 				}
 
 				logger.Information("Uploading large file");
@@ -213,7 +215,7 @@ namespace Horde.Server.Tests.Ddc.FunctionalTests.CompressedBlobs
 
 					await using FileStream downloadedFile = tempOutputFile.OpenRead();
 
-					BlobId downloadedBlobIdentifier = await BlobId.FromStreamAsync(downloadedFile);
+					BlobId downloadedBlobIdentifier = await BlobId.FromStreamAsync(downloadedFile, cancellationToken);
 					Assert.AreEqual(compressedContentHash, downloadedBlobIdentifier);
 				}
 
diff --git a/Engine/Source/Programs/Horde/Horde.Server.Tests/Ddc/Functional/Objects/ReferencesTests.cs b/Engine/Source/Programs/Horde/Horde.Server.Tests/Ddc/Functional/Objects/ReferencesTests.cs
--- a/Engine/Source/Programs/Horde/Horde.Server.Tests/Ddc/Functional/Objects/ReferencesTests.cs
+++ b/Engine/Source/Programs/Horde/Horde.Server.Tests/Ddc/Functional/Objects/ReferencesTests.cs
@@ -12,6 +12,7 @@ using System.Net.Mime;
 using System.Text;
 using System.Text.Json;
 using System.Text.Json.Nodes;
+using System.Threading;
 using System.Threading.Tasks;
 using EpicGames.AspNet;
 using EpicGames.Core;
@@ -46,8 +47,7 @@ namespace Horde.Server.Tests.Ddc.FunctionalTests.References
 
 		private IBlobService Service => BlobService;
 
-		[Obsolete]
-		protected IReferencesStore ReferencesStore => null!;
+		protected IReferencesStore ReferencesStore => throw new NotImplementedException();
 		protected NamespaceId TestNamespace { get; } = new NamespaceId("test-namespace");
 
 		public ReferencesTests()
@@ -79,8 +79,8 @@ namespace Horde.Server.Tests.Ddc.FunctionalTests.References
 		[TestMethod]
 		public async Task PutGetBlobAsync()
 		{
-			const string objectContents = $"This is treated as a opaque blob in {nameof(PutGetBlobAsync)}";
-			byte[] data = Encoding.ASCII.GetBytes(objectContents);
+			const string ObjectContents = $"This is treated as a opaque blob in {nameof(PutGetBlobAsync)}";
+			byte[] data = Encoding.ASCII.GetBytes(ObjectContents);
 			BlobId objectHash = BlobId.FromBlob(data);
 			RefId key = RefId.FromName("newBlobObject");
 			using HttpContent requestContent = new ByteArrayContent(data);
@@ -99,7 +99,7 @@ namespace Horde.Server.Tests.Ddc.FunctionalTests.References
 				byte[] roundTrippedBuffer = ms.ToArray();
 				string roundTrippedPayload = Encoding.ASCII.GetString(roundTrippedBuffer);
 
-				Assert.AreEqual(objectContents, roundTrippedPayload);
+				Assert.AreEqual(ObjectContents, roundTrippedPayload);
 				CollectionAssert.AreEqual(data, roundTrippedBuffer);
 				Assert.AreEqual(objectHash, BlobId.FromBlob(roundTrippedBuffer));
 			}
@@ -131,7 +131,7 @@ namespace Horde.Server.Tests.Ddc.FunctionalTests.References
 					byte[] roundTrippedBuffer = ms.ToArray();
 					string roundTrippedString = Encoding.ASCII.GetString(roundTrippedBuffer);
 
-					Assert.AreEqual(objectContents, roundTrippedString);
+					Assert.AreEqual(ObjectContents, roundTrippedString);
 					Assert.AreEqual(objectHash, BlobId.FromBlob(roundTrippedBuffer));
 				}
 			}
@@ -181,13 +181,12 @@ namespace Horde.Server.Tests.Ddc.FunctionalTests.References
 
 				string roundTrippedString = Encoding.ASCII.GetString(roundTrippedBuffer);
 
-				Assert.AreEqual(objectContents, roundTrippedString);
+				Assert.AreEqual(ObjectContents, roundTrippedString);
 				Assert.AreEqual(objectHash, BlobId.FromBlob(roundTrippedBuffer));
 			}
 		}
 
 		[Ignore]
-		[Obsolete]
 		[TestMethod]
 		public async Task PutGetCompactBinaryAsync()
 		{
@@ -275,7 +274,6 @@ namespace Horde.Server.Tests.Ddc.FunctionalTests.References
 		}
 
 		[Ignore]
-		[Obsolete]
 		[TestMethod]
 		public async Task PutGetCompactBinaryFilteringAsync()
 		{
@@ -354,7 +352,6 @@ namespace Horde.Server.Tests.Ddc.FunctionalTests.References
 		}
 
 		[Ignore]
-		[Obsolete]
 		[TestMethod]
 		public async Task PutGetCompactBinaryHierarchyAsync()
 		{
@@ -508,8 +505,8 @@ namespace Horde.Server.Tests.Ddc.FunctionalTests.References
 		[TestMethod]
 		public async Task ExistsChecksAsync()
 		{
-			const string objectContents = $"This is treated as a opaque blob in {nameof(ExistsChecksAsync)}";
-			byte[] data = Encoding.ASCII.GetBytes(objectContents);
+			const string ObjectContents = $"This is treated as a opaque blob in {nameof(ExistsChecksAsync)}";
+			byte[] data = Encoding.ASCII.GetBytes(ObjectContents);
 			BlobId objectHash = BlobId.FromBlob(data);
 			RefId key = RefId.FromName("newObject");
 			using HttpContent requestContent = new ByteArrayContent(data);
@@ -540,8 +537,8 @@ namespace Horde.Server.Tests.Ddc.FunctionalTests.References
 			BucketId bucket = new BucketId("bucket");
 			RefId existingObject = RefId.FromName("existingObjectMultiple");
 
-			const string objectContents = $"This is treated as a opaque blob in {nameof(ExistsChecksMultipleAsync)}";
-			byte[] data = Encoding.ASCII.GetBytes(objectContents);
+			const string ObjectContents = $"This is treated as a opaque blob in {nameof(ExistsChecksMultipleAsync)}";
+			byte[] data = Encoding.ASCII.GetBytes(ObjectContents);
 			BlobId objectHash = BlobId.FromBlob(data);
 			using HttpContent requestContent = new ByteArrayContent(data);
 			requestContent.Headers.ContentType = new MediaTypeHeaderValue(MediaTypeNames.Application.Octet);
@@ -569,19 +566,20 @@ namespace Horde.Server.Tests.Ddc.FunctionalTests.References
 		}
 
 		[Ignore]
-		[Obsolete]
 		[TestMethod]
 		public async Task PutGetObjectHierarchyAsync()
 		{
+			CancellationToken cancellationToken = CancellationToken.None;
+
 			string blobContents = $"This is a string that is referenced as a blob in {nameof(PutGetObjectHierarchyAsync)}";
 			byte[] blobData = Encoding.ASCII.GetBytes(blobContents);
 			BlobId blobHash = BlobId.FromBlob(blobData);
-			await Service.PutObjectAsync(TestNamespace, blobData, blobHash);
+			await Service.PutObjectAsync(TestNamespace, blobData, blobHash, cancellationToken);
 
 			string blobContentsChild = $"This string is also referenced as a blob but from a child object in {nameof(PutGetObjectHierarchyAsync)}";
 			byte[] dataChild = Encoding.ASCII.GetBytes(blobContentsChild);
 			BlobId blobHashChild = BlobId.FromBlob(dataChild);
-			await Service.PutObjectAsync(TestNamespace, dataChild, blobHashChild);
+			await Service.PutObjectAsync(TestNamespace, dataChild, blobHashChild, cancellationToken);
 
 			CbWriter writerChild = new CbWriter();
 			writerChild.BeginObject();
@@ -590,7 +588,7 @@ namespace Horde.Server.Tests.Ddc.FunctionalTests.References
 
 			byte[] childDataObject = writerChild.ToByteArray();
 			BlobId childDataObjectHash = BlobId.FromBlob(childDataObject);
-			await Service.PutObjectAsync(TestNamespace, childDataObject, childDataObjectHash);
+			await Service.PutObjectAsync(TestNamespace, childDataObject, childDataObjectHash, cancellationToken);
 
 			CbWriter writerParent = new CbWriter();
 			writerParent.BeginObject();
@@ -732,6 +730,8 @@ namespace Horde.Server.Tests.Ddc.FunctionalTests.References
 		[TestMethod]
 		public async Task PutPartialHierarchyAsync()
 		{
+			CancellationToken cancellationToken = CancellationToken.None;
+
 			// do not submit the content of the blobs, which should be reported in the response of the put
 			string blobContents = $"This is a string that is referenced as a blob in {nameof(PutPartialHierarchyAsync)}";
 			byte[] blobData = Encoding.ASCII.GetBytes(blobContents);
@@ -748,7 +748,7 @@ namespace Horde.Server.Tests.Ddc.FunctionalTests.References
 
 			byte[] childDataObject = writerChild.ToByteArray();
 			BlobId childDataObjectHash = BlobId.FromBlob(childDataObject);
-			await Service.PutObjectAsync(TestNamespace, childDataObject, childDataObjectHash);
+			await Service.PutObjectAsync(TestNamespace, childDataObject, childDataObjectHash, cancellationToken);
 
 			CbWriter writerParent = new CbWriter();
 			writerParent.BeginObject();
@@ -812,7 +812,6 @@ namespace Horde.Server.Tests.Ddc.FunctionalTests.References
 		}
 
 		[Ignore]
-		[Obsolete]
 		[TestMethod]
 		public async Task PutContentIdMissingBlobAsync()
 		{
@@ -886,7 +885,6 @@ namespace Horde.Server.Tests.Ddc.FunctionalTests.References
 			}
 		}
 
-
 		[TestMethod]
 		public async Task PutMissingAttachmentComplexAsync()
 		{
@@ -935,10 +933,11 @@ namespace Horde.Server.Tests.Ddc.FunctionalTests.References
 		}
 
 		[Ignore]
-		[Obsolete]
 		[TestMethod]
 		public async Task PutAndFinalizeAsync()
 		{
+			CancellationToken cancellationToken = CancellationToken.None;
+
 			BucketId bucket = new BucketId("bucket");
 			RefId key = RefId.FromName("willFinalizeObject");
 
@@ -958,7 +957,7 @@ namespace Horde.Server.Tests.Ddc.FunctionalTests.References
 
 			byte[] childDataObject = writerChild.ToByteArray();
 			BlobId childDataObjectHash = BlobId.FromBlob(childDataObject);
-			await Service.PutObjectAsync(TestNamespace, childDataObject, childDataObjectHash);
+			await Service.PutObjectAsync(TestNamespace, childDataObject, childDataObjectHash, cancellationToken);
 
 			CbWriter writerParent = new CbWriter();
 			writerParent.BeginObject();
@@ -1003,8 +1002,8 @@ namespace Horde.Server.Tests.Ddc.FunctionalTests.References
 
 			// upload missing pieces
 			{
-				await Service.PutObjectAsync(TestNamespace, blobData, blobHash);
-				await Service.PutObjectAsync(TestNamespace, dataChild, blobHashChild);
+				await Service.PutObjectAsync(TestNamespace, blobData, blobHash, cancellationToken);
+				await Service.PutObjectAsync(TestNamespace, dataChild, blobHashChild, cancellationToken);
 			}
 
 			// finalize the object as no pieces is now missing
@@ -1040,6 +1039,8 @@ namespace Horde.Server.Tests.Ddc.FunctionalTests.References
 		[TestMethod]
 		public async Task GetMissingContentIdRecordAsync()
 		{
+			CancellationToken cancellationToken = CancellationToken.None;
+			
 			string blobContents = $"This is a blob in {nameof(GetMissingContentIdRecordAsync)}";
 			byte[] blobData = Encoding.ASCII.GetBytes(blobContents);
 			BlobId uncompressedHash = BlobId.FromBlob(blobData);
@@ -1118,7 +1119,7 @@ namespace Horde.Server.Tests.Ddc.FunctionalTests.References
 
 			{
 				// delete the blob referenced by the compressed buffer
-				await Service.DeleteObjectAsync(TestNamespace, compressedHash);
+				await Service.DeleteObjectAsync(TestNamespace, compressedHash, cancellationToken);
 			}
 
 			{
@@ -1150,6 +1151,8 @@ namespace Horde.Server.Tests.Ddc.FunctionalTests.References
 		[TestMethod]
 		public async Task GetMissingCompressedBufferAttachmentAsync()
 		{
+			CancellationToken cancellationToken = CancellationToken.None;
+
 			CbObject cbObjectAttachment = CbObject.Build(writer => writer.WriteString("ValueField", "This field has a value"));
 			byte[] cbAttachmentData = cbObjectAttachment.GetView().ToArray();
 			BlobId cbAttachmentHash = BlobId.FromBlob(cbAttachmentData);
@@ -1222,7 +1225,7 @@ namespace Horde.Server.Tests.Ddc.FunctionalTests.References
 
 			{
 				// delete the blob referenced by the compressed buffer
-				await Service.DeleteObjectAsync(TestNamespace, cbAttachmentHash);
+				await Service.DeleteObjectAsync(TestNamespace, cbAttachmentHash, cancellationToken);
 			}
 
 			{
@@ -1254,6 +1257,8 @@ namespace Horde.Server.Tests.Ddc.FunctionalTests.References
 		[TestMethod]
 		public async Task GetMissingBlobRecordAsync()
 		{
+			CancellationToken cancellationToken = CancellationToken.None;
+
 			string blobContents = $"This is a blob in {nameof(GetMissingBlobRecordAsync)}";
 			byte[] blobData = Encoding.ASCII.GetBytes(blobContents);
 			BlobId blobHash = BlobId.FromBlob(blobData);
@@ -1310,7 +1315,7 @@ namespace Horde.Server.Tests.Ddc.FunctionalTests.References
 
 			{
 				// delete the blob 
-				await Service.DeleteObjectAsync(TestNamespace, blobHash);
+				await Service.DeleteObjectAsync(TestNamespace, blobHash, cancellationToken);
 			}
 
 			{
@@ -1359,8 +1364,8 @@ namespace Horde.Server.Tests.Ddc.FunctionalTests.References
 		[TestMethod]
 		public async Task DeleteObjectAsync()
 		{
-			const string objectContents = $"This is treated as a opaque blob in {nameof(DeleteObjectAsync)}";
-			byte[] data = Encoding.ASCII.GetBytes(objectContents);
+			const string ObjectContents = $"This is treated as a opaque blob in {nameof(DeleteObjectAsync)}";
+			byte[] data = Encoding.ASCII.GetBytes(ObjectContents);
 			BlobId objectHash = BlobId.FromBlob(data);
 
 			using HttpContent requestContent = new ByteArrayContent(data);
@@ -1400,14 +1405,13 @@ namespace Horde.Server.Tests.Ddc.FunctionalTests.References
 		}
 
 		[Ignore]
-		[Obsolete]
 		[TestMethod]
 		public async Task DropBucketAsync()
 		{
 			const string BucketToDelete = "delete-bucket";
 
-			const string objectContents = $"This is treated as a opaque blob in {nameof(DropBucketAsync)}";
-			byte[] data = Encoding.ASCII.GetBytes(objectContents);
+			const string ObjectContents = $"This is treated as a opaque blob in {nameof(DropBucketAsync)}";
+			byte[] data = Encoding.ASCII.GetBytes(ObjectContents);
 			BlobId objectHash = BlobId.FromBlob(data);
 
 			using HttpContent requestContent = new ByteArrayContent(data);
@@ -1725,7 +1729,6 @@ namespace Horde.Server.Tests.Ddc.FunctionalTests.References
 			getObjectOp2.WriteString("key", newReferenceObjectKey.ToString());
 			getObjectOp2.EndObject();
 
-
 			CbWriter getObjectOp3 = new CbWriter();
 			getObjectOp3.BeginObject();
 			getObjectOp3.WriteInteger("opId", 2);
diff --git a/Engine/Source/Programs/Horde/Horde.Server.Tests/Ddc/UnitTests/CompressedBufferTests.cs b/Engine/Source/Programs/Horde/Horde.Server.Tests/Ddc/UnitTests/CompressedBufferTests.cs
--- a/Engine/Source/Programs/Horde/Horde.Server.Tests/Ddc/UnitTests/CompressedBufferTests.cs
+++ b/Engine/Source/Programs/Horde/Horde.Server.Tests/Ddc/UnitTests/CompressedBufferTests.cs
@@ -2,6 +2,7 @@
 
 using System.IO;
 using System.Text;
+using System.Threading;
 using System.Threading.Tasks;
 using EpicGames.Core;
 using Horde.Server.Ddc;
@@ -33,7 +34,7 @@ namespace Horde.Server.Tests.Ddc.UnitTests
 			IoHash uncompressedHash = bufferUtils.CompressContent(ms, OoodleCompressorMethod.Mermaid, OoodleCompressionLevel.VeryFast, bytes);
 			ms.Position = 0;
 
-			using IBufferedPayload bufferedPayload = await bufferUtils.DecompressContentAsync(ms, (ulong)ms.Length);
+			using IBufferedPayload bufferedPayload = await bufferUtils.DecompressContentAsync(ms, (ulong)ms.Length, CancellationToken.None);
 			await using Stream s = bufferedPayload.GetStream();
 			byte[] roundTrippedBytes = await s.ReadAllBytesAsync();
 			CollectionAssert.AreEqual(bytes, roundTrippedBytes);
diff --git a/Engine/Source/Programs/Horde/Horde.Server/Ddc/.editorconfig b/Engine/Source/Programs/Horde/Horde.Server/Ddc/.editorconfig
--- a/Engine/Source/Programs/Horde/Horde.Server/Ddc/.editorconfig
+++ b/Engine/Source/Programs/Horde/Horde.Server/Ddc/.editorconfig
@@ -1,9 +1,3 @@
 # EditorConfig is awesome: http://EditorConfig.org
 
 [*.cs]
-
-# MA0032: Use an overload with a CancellationToken
-dotnet_diagnostic.MA0032.severity = none
-
-# VSTHRD200: Use "Async" suffix in names of methods that return an awaitable type
-dotnet_diagnostic.VSTHRD200.severity = none
diff --git a/Engine/Source/Programs/Horde/Horde.Server/Ddc/BlobId.cs b/Engine/Source/Programs/Horde/Horde.Server/Ddc/BlobId.cs
--- a/Engine/Source/Programs/Horde/Horde.Server/Ddc/BlobId.cs
+++ b/Engine/Source/Programs/Horde/Horde.Server/Ddc/BlobId.cs
@@ -7,6 +7,7 @@ using System.Globalization;
 using System.IO;
 using System.Text.Json;
 using System.Text.Json.Serialization;
+using System.Threading;
 using System.Threading.Tasks;
 using Blake3;
 using EpicGames.Core;
@@ -127,7 +128,7 @@ namespace Horde.Server.Ddc
 			return new BlobId(testObjectHash.HashData);
 		}
 
-		public static async Task<BlobId> FromStreamAsync(Stream stream)
+		public static async Task<BlobId> FromStreamAsync(Stream stream, CancellationToken cancellationToken)
 		{
 			using Hasher hasher = Hasher.New();
 
@@ -136,11 +137,11 @@ namespace Horde.Server.Ddc
 
 			try
 			{
-				int read = await stream.ReadAsync(buffer, 0, buffer.Length);
+				int read = await stream.ReadAsync(buffer, 0, buffer.Length, cancellationToken);
 				while (read > 0)
 				{
 					hasher.UpdateWithJoin(new ReadOnlySpan<byte>(buffer, 0, read));
-					read = await stream.ReadAsync(buffer, 0, buffer.Length);
+					read = await stream.ReadAsync(buffer, 0, buffer.Length, cancellationToken);
 				}
 				Hash blake3Hash = hasher.Finalize();
 
diff --git a/Engine/Source/Programs/Horde/Horde.Server/Ddc/BlobService.cs b/Engine/Source/Programs/Horde/Horde.Server/Ddc/BlobService.cs
--- a/Engine/Source/Programs/Horde/Horde.Server/Ddc/BlobService.cs
+++ b/Engine/Source/Programs/Horde/Horde.Server/Ddc/BlobService.cs
@@ -4,6 +4,7 @@ using System;
 using System.Collections.Concurrent;
 using System.Collections.Generic;
 using System.IO;
+using System.Threading;
 using System.Threading.Tasks;
 using EpicGames.Core;
 using EpicGames.Horde.Storage;
@@ -26,12 +27,12 @@ namespace Horde.Server.Ddc
 
 		public static string GetAlias(BlobId blobId) => $"ddc:{blobId}";
 
-		public async Task<ContentHash> VerifyContentMatchesHashAsync(Stream content, ContentHash identifier)
+		public async Task<ContentHash> VerifyContentMatchesHashAsync(Stream content, ContentHash identifier, CancellationToken cancellationToken)
 		{
 			ContentHash blobHash;
 			{
 				using TelemetrySpan _ = _tracer.StartActiveSpan("web.hash").SetAttribute("operation.name", "web.hash");
-				blobHash = await BlobId.FromStreamAsync(content);
+				blobHash = await BlobId.FromStreamAsync(content, cancellationToken);
 			}
 
 			if (!identifier.Equals(blobHash))
@@ -42,7 +43,7 @@ namespace Horde.Server.Ddc
 			return identifier;
 		}
 
-		public async Task<BlobId> PutObjectAsync(NamespaceId ns, IBufferedPayload payload, BlobId identifier)
+		public async Task<BlobId> PutObjectAsync(NamespaceId ns, IBufferedPayload payload, BlobId identifier, CancellationToken cancellationToken)
 		{
 			using TelemetrySpan scope = _tracer.StartActiveSpan("put_blob")
 				.SetAttribute("operation.name", "put_blob")
@@ -50,44 +51,44 @@ namespace Horde.Server.Ddc
 				.SetAttribute("Content-Length", payload.Length.ToString());
 
 			await using Stream hashStream = payload.GetStream();
-			await VerifyContentMatchesHashAsync(hashStream, identifier);
+			await VerifyContentMatchesHashAsync(hashStream, identifier, cancellationToken);
 
-			await PutObjectKnownHashAsync(ns, payload, identifier);
+			await PutObjectKnownHashAsync(ns, payload, identifier, cancellationToken);
 			return identifier;
 		}
 
-		public Task<BlobId> PutObjectAsync(NamespaceId ns, byte[] payload, BlobId identifier)
+		public Task<BlobId> PutObjectAsync(NamespaceId ns, byte[] payload, BlobId identifier, CancellationToken cancellationToken)
 		{
 			using MemoryBufferedPayload bufferedPayload = new MemoryBufferedPayload(payload);
-			return PutObjectAsync(ns, bufferedPayload, identifier);
+			return PutObjectAsync(ns, bufferedPayload, identifier, cancellationToken);
 		}
 
-		public async Task<bool> ExistsAsync(NamespaceId ns, BlobId blob, List<string>? storageLayers = null)
+		public async Task<bool> ExistsAsync(NamespaceId ns, BlobId blob, List<string>? storageLayers = null, CancellationToken cancellationToken = default)
 		{
 			using IStorageClient storageClient = _storageClientFactory.CreateClient(ns);
-			return await storageClient.FindAliasAsync(GetAlias(blob)) != null;
+			return await storageClient.FindAliasAsync(GetAlias(blob), cancellationToken) != null;
 		}
 
-		public async Task DeleteObjectAsync(NamespaceId ns, BlobId blob)
+		public async Task DeleteObjectAsync(NamespaceId ns, BlobId blob, CancellationToken cancellationToken)
 		{
 			using IStorageClient storageClient = _storageClientFactory.CreateClient(ns);
 			string aliasName = GetAlias(blob);
 
-			BlobAlias[] aliases = await storageClient.FindAliasesAsync(aliasName);
+			BlobAlias[] aliases = await storageClient.FindAliasesAsync(aliasName, cancellationToken: cancellationToken);
 			foreach (BlobAlias alias in aliases)
 			{
-				await storageClient.RemoveAliasAsync(aliasName, alias.Target);
+				await storageClient.RemoveAliasAsync(aliasName, alias.Target, cancellationToken);
 			}
 		}
 
-		public async Task<BlobId[]> FilterOutKnownBlobsAsync(NamespaceId ns, IEnumerable<BlobId> blobIds)
+		public async Task<BlobId[]> FilterOutKnownBlobsAsync(NamespaceId ns, IEnumerable<BlobId> blobIds, CancellationToken cancellationToken)
 		{
 			using IStorageClient storageClient = _storageClientFactory.CreateClient(ns);
 
 			List<BlobId> unknownBlobIds = new List<BlobId>();
 			foreach (BlobId blobId in blobIds)
 			{
-				if (await storageClient.FindAliasAsync(GetAlias(blobId)) == null)
+				if (await storageClient.FindAliasAsync(GetAlias(blobId), cancellationToken) == null)
 				{
 					unknownBlobIds.Add(blobId);
 				}
@@ -96,15 +97,15 @@ namespace Horde.Server.Ddc
 			return unknownBlobIds.ToArray();
 		}
 
-		public async Task<BlobId[]> FilterOutKnownBlobsAsync(NamespaceId ns, IAsyncEnumerable<BlobId> blobs)
+		public async Task<BlobId[]> FilterOutKnownBlobsAsync(NamespaceId ns, IAsyncEnumerable<BlobId> blobs, CancellationToken cancellationToken)
 		{
 			ConcurrentBag<BlobId> missingBlobs = new ConcurrentBag<BlobId>();
 
 			try
 			{
-				await Parallel.ForEachAsync(blobs, async (identifier, ctx) =>
+				await Parallel.ForEachAsync(blobs, cancellationToken, async (identifier, ctx) =>
 				{
-					bool exists = await ExistsAsync(ns, identifier);
+					bool exists = await ExistsAsync(ns, identifier, cancellationToken: ctx);
 
 					if (!exists)
 					{
@@ -125,32 +126,32 @@ namespace Horde.Server.Ddc
 			return missingBlobs.ToArray();
 		}
 
-		public async Task<BlobContents> GetObjectAsync(NamespaceId ns, BlobId blob, List<string>? storageLayers, bool supportsRedirectUri, bool allowOndemandReplication = true)
+		public async Task<BlobContents> GetObjectAsync(NamespaceId ns, BlobId blob, List<string>? storageLayers, bool supportsRedirectUri, bool allowOndemandReplication = true, CancellationToken cancellationToken = default)
 		{
 			using IStorageClient storageClient = _storageClientFactory.CreateClient(ns);
 
-			BlobAlias? alias = await storageClient.FindAliasAsync(GetAlias(blob));
+			BlobAlias? alias = await storageClient.FindAliasAsync(GetAlias(blob), cancellationToken);
 			if (alias == null)
 			{
 				throw new BlobNotFoundException(ns, blob);
 			}
 
-			using BlobData data = await alias.Target.ReadBlobDataAsync();
+			using BlobData data = await alias.Target.ReadBlobDataAsync(cancellationToken);
 			return new BlobContents(data.Data.ToArray());
 		}
 
-		public Task<BlobMetadata> GetObjectMetadataAsync(NamespaceId ns, BlobId blobId)
+		public Task<BlobMetadata> GetObjectMetadataAsync(NamespaceId ns, BlobId blobId, CancellationToken cancellationToken)
 		{
 			throw new NotImplementedException();
 		}
 
-		public async Task<BlobContents> GetObjectsAsync(NamespaceId ns, BlobId[] blobs)
+		public async Task<BlobContents> GetObjectsAsync(NamespaceId ns, BlobId[] blobs, CancellationToken cancellationToken)
 		{
 			using TelemetrySpan _ = _tracer.StartActiveSpan("blob.combine").SetAttribute("operation.name", "blob.combine");
 			Task<BlobContents>[] tasks = new Task<BlobContents>[blobs.Length];
 			for (int i = 0; i < blobs.Length; i++)
 			{
-				tasks[i] = GetObjectAsync(ns, blobs[i], storageLayers: null, supportsRedirectUri: false);
+				tasks[i] = GetObjectAsync(ns, blobs[i], storageLayers: null, supportsRedirectUri: false, cancellationToken: cancellationToken);
 			}
 
 			MemoryStream ms = new MemoryStream();
@@ -158,7 +159,7 @@ namespace Horde.Server.Ddc
 			{
 				BlobContents blob = await task;
 				await using Stream s = blob.Stream;
-				await s.CopyToAsync(ms);
+				await s.CopyToAsync(ms, cancellationToken);
 			}
 
 			ms.Seek(0, SeekOrigin.Begin);
@@ -166,17 +167,17 @@ namespace Horde.Server.Ddc
 			return new BlobContents(ms, ms.Length);
 		}
 
-		public Task<Uri?> GetObjectWithRedirectAsync(NamespaceId ns, BlobId blobIdentifier, List<string>? storageLayers)
+		public Task<Uri?> GetObjectWithRedirectAsync(NamespaceId ns, BlobId blobIdentifier, List<string>? storageLayers, CancellationToken cancellationToken)
 		{
 			return Task.FromResult<Uri?>(null);
 		}
 
-		public Task<Uri?> MaybePutObjectWithRedirectAsync(NamespaceId ns, BlobId identifier)
+		public Task<Uri?> MaybePutObjectWithRedirectAsync(NamespaceId ns, BlobId identifier, CancellationToken cancellationToken)
 		{
 			return Task.FromResult<Uri?>(null);
 		}
 
-		public async Task<BlobId> PutObjectKnownHashAsync(NamespaceId ns, IBufferedPayload content, BlobId identifier)
+		public async Task<BlobId> PutObjectKnownHashAsync(NamespaceId ns, IBufferedPayload content, BlobId identifier, CancellationToken cancellationToken)
 		{
 			using IStorageClient storageClient = _storageClientFactory.CreateClient(ns);
 
@@ -187,27 +188,27 @@ namespace Horde.Server.Ddc
 				Memory<byte> memory = writer.GetMemory(contentLength).Slice(0, contentLength);
 
 				using Stream stream = content.GetStream();
-				await stream.ReadFixedLengthBytesAsync(memory);
+				await stream.ReadFixedLengthBytesAsync(memory, cancellationToken);
 				writer.Advance(contentLength);
 
-				blobRef = await writer.CompleteAsync(s_rawBlobType);
-				await writer.FlushAsync();
+				blobRef = await writer.CompleteAsync(s_rawBlobType, cancellationToken);
+				await writer.FlushAsync(cancellationToken);
 			}
 
-			await storageClient.AddAliasAsync(GetAlias(identifier), blobRef, data: identifier.AsIoHash().ToByteArray());
+			await storageClient.AddAliasAsync(GetAlias(identifier), blobRef, data: identifier.AsIoHash().ToByteArray(), cancellationToken: cancellationToken);
 			return identifier;
 		}
 
-		public Task<BlobContents> ReplicateObjectAsync(NamespaceId ns, BlobId blob, bool force = false)
+		public Task<BlobContents> ReplicateObjectAsync(NamespaceId ns, BlobId blob, bool force = false, CancellationToken cancellationToken = default)
 			=> throw new NotSupportedException();
 
-		public Task<bool> ExistsInRootStoreAsync(NamespaceId ns, BlobId blob)
+		public Task<bool> ExistsInRootStoreAsync(NamespaceId ns, BlobId blob, CancellationToken cancellationToken)
 			=> Task.FromResult(false);
 
-		public Task DeleteNamespaceAsync(NamespaceId ns)
+		public Task DeleteNamespaceAsync(NamespaceId ns, CancellationToken cancellationToken)
 			=> throw new NotSupportedException();
 
-		public IAsyncEnumerable<(BlobId, DateTime)> ListObjectsAsync(NamespaceId ns)
+		public IAsyncEnumerable<(BlobId, DateTime)> ListObjectsAsync(NamespaceId ns, CancellationToken cancellationToken)
 			=> throw new NotSupportedException();
 
 		public bool ShouldFetchBlobOnDemand(NamespaceId ns)
diff --git a/Engine/Source/Programs/Horde/Horde.Server/Ddc/BlobsController.cs b/Engine/Source/Programs/Horde/Horde.Server/Ddc/BlobsController.cs
--- a/Engine/Source/Programs/Horde/Horde.Server/Ddc/BlobsController.cs
+++ b/Engine/Source/Programs/Horde/Horde.Server/Ddc/BlobsController.cs
@@ -191,7 +191,7 @@ namespace Horde.Server.Ddc
 
 			try
 			{
-				Uri? uri = await _storage.MaybePutObjectWithRedirectAsync(ns, id);
+				Uri? uri = await _storage.MaybePutObjectWithRedirectAsync(ns, id, HttpContext.RequestAborted);
 				if (uri != null)
 				{
 					return Ok(new
@@ -200,9 +200,9 @@ namespace Horde.Server.Ddc
 						RedirectUri = uri,
 					});
 				}
-				using IBufferedPayload payload = await _bufferedPayloadFactory.CreateFromRequestAsync(Request);
+				using IBufferedPayload payload = await _bufferedPayloadFactory.CreateFromRequestAsync(Request, HttpContext.RequestAborted);
 
-				BlobId identifier = await _storage.PutObjectAsync(ns, payload, id);
+				BlobId identifier = await _storage.PutObjectAsync(ns, payload, id, HttpContext.RequestAborted);
 				return Ok(new
 				{
 					Identifier = identifier.ToString()
@@ -233,12 +233,12 @@ namespace Horde.Server.Ddc
 			_diagnosticContext.Set("Content-Length", Request.ContentLength ?? -1);
 			try
 			{
-				using IBufferedPayload payload = await _bufferedPayloadFactory.CreateFromRequestAsync(Request);
+				using IBufferedPayload payload = await _bufferedPayloadFactory.CreateFromRequestAsync(Request, HttpContext.RequestAborted);
 
 				await using Stream stream = payload.GetStream();
 
-				BlobId id = await BlobId.FromStreamAsync(stream);
-				await _storage.PutObjectKnownHashAsync(ns, payload, id);
+				BlobId id = await BlobId.FromStreamAsync(stream, HttpContext.RequestAborted);
+				await _storage.PutObjectKnownHashAsync(ns, payload, id, HttpContext.RequestAborted);
 
 				return Ok(new
 				{
@@ -277,14 +277,14 @@ namespace Horde.Server.Ddc
 				return result;
 			}
 
-			await _storage.DeleteNamespaceAsync(ns);
+			await _storage.DeleteNamespaceAsync(ns, HttpContext.RequestAborted);
 
 			return NoContent();
 		}
 
 		private async Task DeleteImplAsync(NamespaceId ns, BlobId id)
 		{
-			await _storage.DeleteObjectAsync(ns, id);
+			await _storage.DeleteObjectAsync(ns, id, HttpContext.RequestAborted);
 		}
 
 		// ReSharper disable UnusedAutoPropertyAccessor.Global
@@ -402,7 +402,7 @@ namespace Horde.Server.Ddc
 							}
 
 							using MemoryBufferedPayload payload = new MemoryBufferedPayload(op.Content);
-							tasks[index] = _storage.PutObjectAsync(op.Namespace.Value, payload, op.Id).ContinueWith((t, _) => (object?)t.Result, null, TaskScheduler.Current);
+							tasks[index] = _storage.PutObjectAsync(op.Namespace.Value, payload, op.Id, HttpContext.RequestAborted).ContinueWith((t, _) => (object?)t.Result, null, TaskScheduler.Current);
 							break;
 						}
 					case BatchOp.Operation.DELETE:
diff --git a/Engine/Source/Programs/Horde/Horde.Server/Ddc/BufferedPayload.cs b/Engine/Source/Programs/Horde/Horde.Server/Ddc/BufferedPayload.cs
--- a/Engine/Source/Programs/Horde/Horde.Server/Ddc/BufferedPayload.cs
+++ b/Engine/Source/Programs/Horde/Horde.Server/Ddc/BufferedPayload.cs
@@ -2,6 +2,7 @@
 
 using System;
 using System.IO;
+using System.Threading;
 using System.Threading.Tasks;
 using Microsoft.AspNetCore.Http;
 using Microsoft.Extensions.Options;
@@ -29,12 +30,12 @@ namespace Horde.Server.Ddc
 			_buffer = source;
 		}
 
-		public static async Task<MemoryBufferedPayload> CreateAsync(Tracer tracer, Stream s)
+		public static async Task<MemoryBufferedPayload> CreateAsync(Tracer tracer, Stream s, CancellationToken cancellationToken)
 		{
 			using TelemetrySpan scope = tracer.StartActiveSpan("payload.buffer")
 				.SetAttribute("operation.name", "payload.buffer")
 				.SetAttribute("bufferType", "Memory");
-			MemoryBufferedPayload payload = new MemoryBufferedPayload(await s.ToByteArrayAsync());
+			MemoryBufferedPayload payload = new MemoryBufferedPayload(await s.ToByteArrayAsync(cancellationToken));
 
 			return payload;
 		}
@@ -123,7 +124,7 @@ namespace Horde.Server.Ddc
 			_length = _tempFile.Length;
 		}
 
-		public static async Task<FilesystemBufferedPayload> CreateAsync(Tracer tracer, Stream s, string filesystemRoot)
+		public static async Task<FilesystemBufferedPayload> CreateAsync(Tracer tracer, Stream s, string filesystemRoot, CancellationToken cancellationToken)
 		{
 			FilesystemBufferedPayload payload = new FilesystemBufferedPayload(filesystemRoot);
 
@@ -132,7 +133,7 @@ namespace Horde.Server.Ddc
 					.SetAttribute("operation.name", "payload.buffer")
 					.SetAttribute("bufferType", "Filesystem");
 				await using FileStream fs = payload._tempFile.OpenWrite();
-				await s.CopyToAsync(fs);
+				await s.CopyToAsync(fs, cancellationToken);
 			}
 
 			payload._tempFile.Refresh();
@@ -183,7 +184,7 @@ namespace Horde.Server.Ddc
 			Directory.CreateDirectory(options.CurrentValue.FilesystemTempPayloadRoot);
 		}
 
-		public Task<IBufferedPayload> CreateFromRequestAsync(HttpRequest request)
+		public Task<IBufferedPayload> CreateFromRequestAsync(HttpRequest request, CancellationToken cancellationToken)
 		{
 			long? contentLength = request.ContentLength;
 
@@ -192,23 +193,23 @@ namespace Horde.Server.Ddc
 				throw new Exception("Expected content-length on all requests");
 			}
 
-			return CreateFromStreamAsync(request.Body, contentLength.Value);
+			return CreateFromStreamAsync(request.Body, contentLength.Value, cancellationToken);
 		}
 
-		public async Task<IBufferedPayload> CreateFromStreamAsync(Stream s, long contentLength)
+		public async Task<IBufferedPayload> CreateFromStreamAsync(Stream s, long contentLength, CancellationToken cancellationToken)
 		{
 			// blob is small enough to fit into memory we just read it as is
 			if (contentLength < _options.CurrentValue.MemoryBufferSize)
 			{
-				return await MemoryBufferedPayload.CreateAsync(_tracer, s);
+				return await MemoryBufferedPayload.CreateAsync(_tracer, s, cancellationToken);
 			}
 
-			return await FilesystemBufferedPayload.CreateAsync(_tracer, s, _options.CurrentValue.FilesystemTempPayloadRoot);
+			return await FilesystemBufferedPayload.CreateAsync(_tracer, s, _options.CurrentValue.FilesystemTempPayloadRoot, cancellationToken);
 		}
 
-		public async Task<IBufferedPayload> CreateFilesystemBufferedPayloadAsync(Stream s)
+		public async Task<IBufferedPayload> CreateFilesystemBufferedPayloadAsync(Stream s, CancellationToken cancellationToken)
 		{
-			return await FilesystemBufferedPayload.CreateAsync(_tracer, s, _options.CurrentValue.FilesystemTempPayloadRoot);
+			return await FilesystemBufferedPayload.CreateAsync(_tracer, s, _options.CurrentValue.FilesystemTempPayloadRoot, cancellationToken);
 		}
 
 		public FilesystemBufferedPayloadWriter CreateFilesystemBufferedPayloadWriter()
diff --git a/Engine/Source/Programs/Horde/Horde.Server/Ddc/CompressedBlobController.cs b/Engine/Source/Programs/Horde/Horde.Server/Ddc/CompressedBlobController.cs
--- a/Engine/Source/Programs/Horde/Horde.Server/Ddc/CompressedBlobController.cs
+++ b/Engine/Source/Programs/Horde/Horde.Server/Ddc/CompressedBlobController.cs
@@ -6,6 +6,7 @@ using System.ComponentModel.DataAnnotations;
 using System.Linq;
 using System.Net;
 using System.Net.Mime;
+using System.Threading;
 using System.Threading.Tasks;
 using EpicGames.AspNet;
 using EpicGames.Horde.Storage;
@@ -226,9 +227,9 @@ namespace Horde.Server.Ddc
 
 			try
 			{
-				using IBufferedPayload payload = await _bufferedPayloadFactory.CreateFromRequestAsync(Request);
+				using IBufferedPayload payload = await _bufferedPayloadFactory.CreateFromRequestAsync(Request, HttpContext.RequestAborted);
 
-				ContentId identifier = await _storage.PutCompressedObjectAsync(ns, payload, id, HttpContext.RequestServices);
+				ContentId identifier = await _storage.PutCompressedObjectAsync(ns, payload, id, HttpContext.RequestServices, HttpContext.RequestAborted);
 
 				return Ok(new { Identifier = identifier.ToString() });
 			}
@@ -252,6 +253,8 @@ namespace Horde.Server.Ddc
 		public async Task<IActionResult> PostAsync(
 			[Required] NamespaceId ns)
 		{
+			CancellationToken cancellationToken = HttpContext.RequestAborted;
+
 			ActionResult? result = await _requestHelper.HasAccessToNamespaceAsync(User, Request, ns, new[] { JupiterAclAction.WriteObject });
 			if (result != null)
 			{
@@ -262,9 +265,9 @@ namespace Horde.Server.Ddc
 
 			try
 			{
-				using IBufferedPayload payload = await _bufferedPayloadFactory.CreateFromRequestAsync(Request);
+				using IBufferedPayload payload = await _bufferedPayloadFactory.CreateFromRequestAsync(Request, cancellationToken);
 
-				ContentId identifier = await _storage.PutCompressedObjectAsync(ns, payload, null, HttpContext.RequestServices);
+				ContentId identifier = await _storage.PutCompressedObjectAsync(ns, payload, null, HttpContext.RequestServices, cancellationToken);
 
 				return Ok(new
 				{
diff --git a/Engine/Source/Programs/Horde/Horde.Server/Ddc/CompressedBufferUtils.cs b/Engine/Source/Programs/Horde/Horde.Server/Ddc/CompressedBufferUtils.cs
--- a/Engine/Source/Programs/Horde/Horde.Server/Ddc/CompressedBufferUtils.cs
+++ b/Engine/Source/Programs/Horde/Horde.Server/Ddc/CompressedBufferUtils.cs
@@ -6,6 +6,7 @@ using System.Collections.Generic;
 using System.IO;
 using System.Linq;
 using System.Text;
+using System.Threading;
 using System.Threading.Tasks;
 using Blake3;
 using EpicGames.Compression;
@@ -184,7 +185,7 @@ namespace Horde.Server.Ddc
 			}
 		}
 
-		public async Task<IBufferedPayload> DecompressContentAsync(Stream sourceStream, ulong streamSize)
+		public async Task<IBufferedPayload> DecompressContentAsync(Stream sourceStream, ulong streamSize, CancellationToken cancellationToken)
 		{
 			using BinaryReader br = new BinaryReader(sourceStream);
 			(CompressedBufferHeader header, uint[] compressedBlockSizes) = ExtractHeader(br);
@@ -227,7 +228,7 @@ namespace Horde.Server.Ddc
 				}
 				else
 				{
-					await sourceStream.CopyToAsync(targetStream);
+					await sourceStream.CopyToAsync(targetStream, cancellationToken);
 				}
 			}
 
@@ -253,7 +254,7 @@ namespace Horde.Server.Ddc
 
 					BlobId headerIdentifier = new BlobId(slicedHash);
 					await using Stream hashStream = finalizedBufferedPayload.GetStream();
-					BlobId contentHash = await BlobId.FromStreamAsync(hashStream);
+					BlobId contentHash = await BlobId.FromStreamAsync(hashStream, cancellationToken);
 
 					if (!headerIdentifier.Equals(contentHash))
 					{
diff --git a/Engine/Source/Programs/Horde/Horde.Server/Ddc/ContentIdStore.cs b/Engine/Source/Programs/Horde/Horde.Server/Ddc/ContentIdStore.cs
--- a/Engine/Source/Programs/Horde/Horde.Server/Ddc/ContentIdStore.cs
+++ b/Engine/Source/Programs/Horde/Horde.Server/Ddc/ContentIdStore.cs
@@ -19,10 +19,8 @@ namespace Horde.Server.Ddc
 		static string GetAlias(BlobId blobId) => BlobService.GetAlias(blobId);
 		static string GetAlias(ContentId contentId) => $"cid:{contentId}";
 
-		public async Task<BlobId[]?> ResolveAsync(NamespaceId ns, ContentId contentId, bool mustBeContentId = false)
+		public async Task<BlobId[]?> ResolveAsync(NamespaceId ns, ContentId contentId, bool mustBeContentId = false, CancellationToken cancellationToken = default)
 		{
-			CancellationToken cancellationToken = CancellationToken.None;
-
 			using IStorageClient storageClient = _storageClientFactory.CreateClient(ns);
 
 			BlobAlias? blobAlias = await storageClient.FindAliasAsync(GetAlias(contentId), cancellationToken);
@@ -39,10 +37,8 @@ namespace Horde.Server.Ddc
 			return new[] { BlobId.FromIoHash(hash) };
 		}
 
-		public async Task PutAsync(NamespaceId ns, ContentId contentId, BlobId blobId, int contentWeight)
+		public async Task PutAsync(NamespaceId ns, ContentId contentId, BlobId blobId, int contentWeight, CancellationToken cancellationToken)
 		{
-			CancellationToken cancellationToken = CancellationToken.None;
-
 			using IStorageClient storageClient = _storageClientFactory.CreateClient(ns);
 
 			BlobAlias? blobAlias = await storageClient.FindAliasAsync(GetAlias(blobId), cancellationToken);
diff --git a/Engine/Source/Programs/Horde/Horde.Server/Ddc/IBlobService.cs b/Engine/Source/Programs/Horde/Horde.Server/Ddc/IBlobService.cs
--- a/Engine/Source/Programs/Horde/Horde.Server/Ddc/IBlobService.cs
+++ b/Engine/Source/Programs/Horde/Horde.Server/Ddc/IBlobService.cs
@@ -4,6 +4,7 @@ using System;
 using System.Collections.Generic;
 using System.IO;
 using System.Net.Mime;
+using System.Threading;
 using System.Threading.Tasks;
 using EpicGames.AspNet;
 using EpicGames.Horde.Storage;
@@ -16,40 +17,41 @@ namespace Horde.Server.Ddc;
 
 public interface IBlobService
 {
-	Task<ContentHash> VerifyContentMatchesHashAsync(Stream content, ContentHash identifier);
-	Task<BlobId> PutObjectKnownHashAsync(NamespaceId ns, IBufferedPayload content, BlobId identifier);
-	Task<BlobId> PutObjectAsync(NamespaceId ns, IBufferedPayload payload, BlobId identifier);
-	Task<BlobId> PutObjectAsync(NamespaceId ns, byte[] payload, BlobId identifier);
-	Task<Uri?> MaybePutObjectWithRedirectAsync(NamespaceId ns, BlobId identifier);
+	Task<ContentHash> VerifyContentMatchesHashAsync(Stream content, ContentHash identifier, CancellationToken cancellationToken);
+	Task<BlobId> PutObjectKnownHashAsync(NamespaceId ns, IBufferedPayload content, BlobId identifier, CancellationToken cancellationToken);
+	Task<BlobId> PutObjectAsync(NamespaceId ns, IBufferedPayload payload, BlobId identifier, CancellationToken cancellationToken);
+	Task<BlobId> PutObjectAsync(NamespaceId ns, byte[] payload, BlobId identifier, CancellationToken cancellationToken);
+	Task<Uri?> MaybePutObjectWithRedirectAsync(NamespaceId ns, BlobId identifier, CancellationToken cancellationToken);
 
-	Task<BlobContents> GetObjectAsync(NamespaceId ns, BlobId blob, List<string>? storageLayers = null, bool supportsRedirectUri = false, bool allowOndemandReplication = true);
+	Task<BlobContents> GetObjectAsync(NamespaceId ns, BlobId blob, List<string>? storageLayers = null, bool supportsRedirectUri = false, bool allowOndemandReplication = true, CancellationToken cancellationToken = default);
 
-	Task<Uri?> GetObjectWithRedirectAsync(NamespaceId ns, BlobId blobIdentifier, List<string>? storageLayers = null);
+	Task<Uri?> GetObjectWithRedirectAsync(NamespaceId ns, BlobId blobIdentifier, List<string>? storageLayers = null, CancellationToken cancellationToken = default);
 
-	Task<BlobMetadata> GetObjectMetadataAsync(NamespaceId ns, BlobId blobId);
+	Task<BlobMetadata> GetObjectMetadataAsync(NamespaceId ns, BlobId blobId, CancellationToken cancellationToken);
 
-	Task<BlobContents> ReplicateObjectAsync(NamespaceId ns, BlobId blob, bool force = false);
+	Task<BlobContents> ReplicateObjectAsync(NamespaceId ns, BlobId blob, bool force = false, CancellationToken cancellationToken = default);
 
-	Task<bool> ExistsAsync(NamespaceId ns, BlobId blob, List<string>? storageLayers = null);
+	Task<bool> ExistsAsync(NamespaceId ns, BlobId blob, List<string>? storageLayers = null, CancellationToken cancellationToken = default);
 
 	/// <summary>
 	/// Checks that the blob exists in the root store, the store which is last in the list and thus is intended to have every blob in it
 	/// </summary>
 	/// <param name="ns">The namespace</param>
 	/// <param name="blob">The identifier of the blob</param>
+	/// <param name="cancellationToken"></param>
 	/// <returns></returns>
-	Task<bool> ExistsInRootStoreAsync(NamespaceId ns, BlobId blob);
+	Task<bool> ExistsInRootStoreAsync(NamespaceId ns, BlobId blob, CancellationToken cancellationToken);
 
 	// Delete a object
-	Task DeleteObjectAsync(NamespaceId ns, BlobId blob);
+	Task DeleteObjectAsync(NamespaceId ns, BlobId blob, CancellationToken cancellationToken);
 
 	// delete the whole namespace
-	Task DeleteNamespaceAsync(NamespaceId ns);
+	Task DeleteNamespaceAsync(NamespaceId ns, CancellationToken cancellationToken);
 
-	IAsyncEnumerable<(BlobId, DateTime)> ListObjectsAsync(NamespaceId ns);
-	Task<BlobId[]> FilterOutKnownBlobsAsync(NamespaceId ns, IEnumerable<BlobId> blobs);
-	Task<BlobId[]> FilterOutKnownBlobsAsync(NamespaceId ns, IAsyncEnumerable<BlobId> blobs);
-	Task<BlobContents> GetObjectsAsync(NamespaceId ns, BlobId[] refRequestBlobReferences);
+	IAsyncEnumerable<(BlobId, DateTime)> ListObjectsAsync(NamespaceId ns, CancellationToken cancellationToken);
+	Task<BlobId[]> FilterOutKnownBlobsAsync(NamespaceId ns, IEnumerable<BlobId> blobs, CancellationToken cancellationToken);
+	Task<BlobId[]> FilterOutKnownBlobsAsync(NamespaceId ns, IAsyncEnumerable<BlobId> blobs, CancellationToken cancellationToken);
+	Task<BlobContents> GetObjectsAsync(NamespaceId ns, BlobId[] refRequestBlobReferences, CancellationToken cancellationToken);
 
 	bool ShouldFetchBlobOnDemand(NamespaceId ns);
 }
@@ -68,7 +70,7 @@ public class BlobMetadata
 
 public static class BlobServiceExtensions
 {
-	public static async Task<ContentId> PutCompressedObjectAsync(this IBlobService blobService, NamespaceId ns, IBufferedPayload payload, ContentId? id, IServiceProvider provider)
+	public static async Task<ContentId> PutCompressedObjectAsync(this IBlobService blobService, NamespaceId ns, IBufferedPayload payload, ContentId? id, IServiceProvider provider, CancellationToken cancellationToken)
 	{
 		IContentIdStore contentIdStore = provider.GetService<IContentIdStore>()!;
 		CompressedBufferUtils compressedBufferUtils = provider.GetService<CompressedBufferUtils>()!;
@@ -77,20 +79,20 @@ public static class BlobServiceExtensions
 		// decompress the content and generate a identifier from it to verify the identifier we got
 		await using Stream decompressStream = payload.GetStream();
 
-		using IBufferedPayload bufferedPayload = await compressedBufferUtils.DecompressContentAsync(decompressStream, (ulong)payload.Length);
+		using IBufferedPayload bufferedPayload = await compressedBufferUtils.DecompressContentAsync(decompressStream, (ulong)payload.Length, cancellationToken);
 		await using Stream decompressedStream = bufferedPayload.GetStream();
 
 		ContentId identifierDecompressedPayload;
 		if (id != null)
 		{
-			identifierDecompressedPayload = ContentId.FromContentHash(await blobService.VerifyContentMatchesHashAsync(decompressedStream, id));
+			identifierDecompressedPayload = ContentId.FromContentHash(await blobService.VerifyContentMatchesHashAsync(decompressedStream, id, cancellationToken));
 		}
 		else
 		{
 			ContentHash blobHash;
 			{
 				using TelemetrySpan _ = tracer.StartActiveSpan("web.hash").SetAttribute("operation.name", "web.hash");
-				blobHash = await BlobId.FromStreamAsync(decompressedStream);
+				blobHash = await BlobId.FromStreamAsync(decompressedStream, cancellationToken);
 			}
 
 			identifierDecompressedPayload = ContentId.FromContentHash(blobHash);
@@ -100,30 +102,30 @@ public static class BlobServiceExtensions
 		{
 			using TelemetrySpan _ = tracer.StartActiveSpan("web.hash").SetAttribute("operation.name", "web.hash");
 			await using Stream hashStream = payload.GetStream();
-			identifierCompressedPayload = await BlobId.FromStreamAsync(hashStream);
+			identifierCompressedPayload = await BlobId.FromStreamAsync(hashStream, cancellationToken);
 		}
 
 		// we still commit the compressed buffer to the object store using the hash of the compressed content
 		{
-			await blobService.PutObjectKnownHashAsync(ns, payload, identifierCompressedPayload);
+			await blobService.PutObjectKnownHashAsync(ns, payload, identifierCompressedPayload, cancellationToken);
 		}
 
 		// commit the mapping from the decompressed hash to the compressed hash, we run this in parallel with the blob store submit
 		// TODO: let users specify weight of the blob compared to previously submitted content ids
 		int contentIdWeight = (int)payload.Length;
-		Task contentIdStoreTask = contentIdStore.PutAsync(ns, identifierDecompressedPayload, identifierCompressedPayload, contentIdWeight);
+		Task contentIdStoreTask = contentIdStore.PutAsync(ns, identifierDecompressedPayload, identifierCompressedPayload, contentIdWeight, cancellationToken);
 
 		await contentIdStoreTask;
 
 		return identifierDecompressedPayload;
 	}
 
-	public static async Task<(BlobContents, string)> GetCompressedObjectAsync(this IBlobService blobService, NamespaceId ns, ContentId contentId, IServiceProvider provider, bool supportsRedirectUri = false)
+	public static async Task<(BlobContents, string)> GetCompressedObjectAsync(this IBlobService blobService, NamespaceId ns, ContentId contentId, IServiceProvider provider, bool supportsRedirectUri = false, CancellationToken cancellationToken = default)
 	{
 		IContentIdStore contentIdStore = provider.GetService<IContentIdStore>()!;
 		Tracer tracer = provider.GetService<Tracer>()!;
 
-		BlobId[]? chunks = await contentIdStore.ResolveAsync(ns, contentId, mustBeContentId: false);
+		BlobId[]? chunks = await contentIdStore.ResolveAsync(ns, contentId, mustBeContentId: false, cancellationToken);
 		if (chunks == null || chunks.Length == 0)
 		{
 			throw new ContentIdResolveException(contentId);
@@ -140,7 +142,7 @@ public static class BlobServiceExtensions
 				mimeType = MediaTypeNames.Application.Octet;
 			}
 
-			return (await blobService.GetObjectAsync(ns, blobToReturn, supportsRedirectUri: supportsRedirectUri), mimeType);
+			return (await blobService.GetObjectAsync(ns, blobToReturn, supportsRedirectUri: supportsRedirectUri, cancellationToken: cancellationToken), mimeType);
 		}
 
 		// chunked content, combine the chunks into a single stream
@@ -149,7 +151,7 @@ public static class BlobServiceExtensions
 		for (int i = 0; i < chunks.Length; i++)
 		{
 			// even if it was requested to support redirect, since we need to combine the chunks using redirects is not possible
-			tasks[i] = blobService.GetObjectAsync(ns, chunks[i], supportsRedirectUri: false);
+			tasks[i] = blobService.GetObjectAsync(ns, chunks[i], supportsRedirectUri: false, cancellationToken: cancellationToken);
 		}
 
 		MemoryStream ms = new MemoryStream();
@@ -157,7 +159,7 @@ public static class BlobServiceExtensions
 		{
 			BlobContents blob = await task;
 			await using Stream s = blob.Stream;
-			await s.CopyToAsync(ms);
+			await s.CopyToAsync(ms, cancellationToken);
 		}
 
 		ms.Seek(0, SeekOrigin.Begin);
diff --git a/Engine/Source/Programs/Horde/Horde.Server/Ddc/IContentIdService.cs b/Engine/Source/Programs/Horde/Horde.Server/Ddc/IContentIdService.cs
--- a/Engine/Source/Programs/Horde/Horde.Server/Ddc/IContentIdService.cs
+++ b/Engine/Source/Programs/Horde/Horde.Server/Ddc/IContentIdService.cs
@@ -1,6 +1,7 @@
 // Copyright Epic Games, Inc. All Rights Reserved.
 
 using System;
+using System.Threading;
 using System.Threading.Tasks;
 using EpicGames.Horde.Storage;
 
@@ -16,8 +17,9 @@ namespace Horde.Server.Ddc
 		/// <param name="ns">The namespace to operate in</param>
 		/// <param name="contentId">The identifier for the content id</param>
 		/// <param name="mustBeContentId"></param>
+		/// <param name="cancellationToken">Cancellation token for the operation</param>
 		/// <returns></returns>
-		Task<BlobId[]?> ResolveAsync(NamespaceId ns, ContentId contentId, bool mustBeContentId = false);
+		Task<BlobId[]?> ResolveAsync(NamespaceId ns, ContentId contentId, bool mustBeContentId = false, CancellationToken cancellationToken = default);
 
 		/// <summary>
 		/// Add a mapping from contentId to blobIdentifier
@@ -26,8 +28,9 @@ namespace Horde.Server.Ddc
 		/// <param name="contentId">The contentId</param>
 		/// <param name="blobIdentifier">The blob the content id maps to</param>
 		/// <param name="contentWeight">Weight of this identifier compared to previous mappings, used to determine which is more important, lower weight is considered a better fit</param>
+		/// <param name="cancellationToken">Cancellation token for the operation</param>
 		/// <returns></returns>
-		Task PutAsync(NamespaceId ns, ContentId contentId, BlobId blobIdentifier, int contentWeight);
+		Task PutAsync(NamespaceId ns, ContentId contentId, BlobId blobIdentifier, int contentWeight, CancellationToken cancellationToken = default);
 	}
 
 	public class InvalidContentIdException : Exception
diff --git a/Engine/Source/Programs/Horde/Horde.Server/Ddc/IRefService.cs b/Engine/Source/Programs/Horde/Horde.Server/Ddc/IRefService.cs
--- a/Engine/Source/Programs/Horde/Horde.Server/Ddc/IRefService.cs
+++ b/Engine/Source/Programs/Horde/Horde.Server/Ddc/IRefService.cs
@@ -2,6 +2,7 @@
 
 using System;
 using System.Collections.Generic;
+using System.Threading;
 using System.Threading.Tasks;
 using EpicGames.Horde.Storage;
 using EpicGames.Serialization;
@@ -12,18 +13,18 @@ namespace Horde.Server.Ddc
 {
 	public interface IRefService
 	{
-		Task<(RefRecord, BlobContents?)> GetAsync(NamespaceId ns, BucketId bucket, RefId key, string[] fields, bool doLastAccessTracking = true);
-		Task<(ContentId[], BlobId[])> PutAsync(NamespaceId ns, BucketId bucket, RefId key, BlobId blobHash, CbObject payload);
-		Task<(ContentId[], BlobId[])> FinalizeAsync(NamespaceId ns, BucketId bucket, RefId key, BlobId blobHash);
+		Task<(RefRecord, BlobContents?)> GetAsync(NamespaceId ns, BucketId bucket, RefId key, string[] fields, bool doLastAccessTracking = true, CancellationToken cancellationToken = default);
+		Task<(ContentId[], BlobId[])> PutAsync(NamespaceId ns, BucketId bucket, RefId key, BlobId blobHash, CbObject payload, CancellationToken cancellationToken);
+		Task<(ContentId[], BlobId[])> FinalizeAsync(NamespaceId ns, BucketId bucket, RefId key, BlobId blobHash, CancellationToken cancellationToken);
 
-		IAsyncEnumerable<NamespaceId> GetNamespacesAsync();
+		IAsyncEnumerable<NamespaceId> GetNamespacesAsync(CancellationToken cancellationToken);
 
-		Task<bool> DeleteAsync(NamespaceId ns, BucketId bucket, RefId key);
-		Task<long> DropNamespaceAsync(NamespaceId ns);
-		Task<long> DeleteBucketAsync(NamespaceId ns, BucketId bucket);
+		Task<bool> DeleteAsync(NamespaceId ns, BucketId bucket, RefId key, CancellationToken cancellationToken);
+		Task<long> DropNamespaceAsync(NamespaceId ns, CancellationToken cancellationToken);
+		Task<long> DeleteBucketAsync(NamespaceId ns, BucketId bucket, CancellationToken cancellationToken);
 
-		Task<bool> ExistsAsync(NamespaceId ns, BucketId bucket, RefId key);
-		Task<List<BlobId>> GetReferencedBlobsAsync(NamespaceId ns, BucketId bucket, RefId key, bool ignoreMissingBlobs);
+		Task<bool> ExistsAsync(NamespaceId ns, BucketId bucket, RefId key, CancellationToken cancellationToken);
+		Task<List<BlobId>> GetReferencedBlobsAsync(NamespaceId ns, BucketId bucket, RefId key, bool ignoreMissingBlobs, CancellationToken cancellationToken);
 	}
 
 	public class ObjectHashMismatchException : Exception
diff --git a/Engine/Source/Programs/Horde/Horde.Server/Ddc/IReferenceResolver.cs b/Engine/Source/Programs/Horde/Horde.Server/Ddc/IReferenceResolver.cs
--- a/Engine/Source/Programs/Horde/Horde.Server/Ddc/IReferenceResolver.cs
+++ b/Engine/Source/Programs/Horde/Horde.Server/Ddc/IReferenceResolver.cs
@@ -3,6 +3,8 @@
 using System;
 using System.Collections.Generic;
 using System.Linq;
+using System.Runtime.CompilerServices;
+using System.Threading;
 using System.Threading.Tasks;
 using EpicGames.Core;
 using EpicGames.Horde.Storage;
@@ -75,8 +77,9 @@ namespace Horde.Server.Ddc
 		/// <param name="ns">The namespace to check</param>
 		/// <param name="cb">The compact binary object to resolve references for</param>
 		/// <param name="ignoreMissingBlobs">Set to true to always returned the blobs found, ignoring anything that is missing rather then throwing</param>
+		/// <param name="cancellationToken">Cancellation token for the operation</param>
 		/// <returns></returns>
-		IAsyncEnumerable<BlobId> GetReferencedBlobsAsync(NamespaceId ns, CbObject cb, bool ignoreMissingBlobs = false);
+		IAsyncEnumerable<BlobId> GetReferencedBlobsAsync(NamespaceId ns, CbObject cb, bool ignoreMissingBlobs = false, CancellationToken cancellationToken = default);
 
 		/// <summary>
 		/// Returns which attachments exist in the cb object or any children
@@ -84,8 +87,9 @@ namespace Horde.Server.Ddc
 		/// </summary>
 		/// <param name="ns">The namespace to check</param>
 		/// <param name="cb">The compact binary object to resolve references for</param>
+		/// <param name="cancellationToken">Cancellation token for the operation</param>
 		/// <returns></returns>
-		IAsyncEnumerable<Attachment> GetAttachmentsAsync(NamespaceId ns, CbObject cb);
+		IAsyncEnumerable<Attachment> GetAttachmentsAsync(NamespaceId ns, CbObject cb, CancellationToken cancellationToken);
 	}
 
 	public class ReferenceResolver : IReferenceResolver
@@ -101,7 +105,7 @@ namespace Horde.Server.Ddc
 			_tracer = tracer;
 		}
 
-		public async IAsyncEnumerable<Attachment> GetAttachmentsAsync(NamespaceId ns, CbObject cb)
+		public async IAsyncEnumerable<Attachment> GetAttachmentsAsync(NamespaceId ns, CbObject cb, [EnumeratorCancellation] CancellationToken cancellationToken)
 		{
 			Queue<CbObject> objectsToVisit = new Queue<CbObject>();
 			objectsToVisit.Enqueue(cb);
@@ -125,13 +129,13 @@ namespace Horde.Server.Ddc
 
 						if (field.IsBinaryAttachment())
 						{
-							Task<(ContentId, BlobId[]?)> resolveContentId = ResolveContentIdAsync(ns, contentId);
+							Task<(ContentId, BlobId[]?)> resolveContentId = ResolveContentIdAsync(ns, contentId, cancellationToken);
 							pendingContentIdResolves.Add(resolveContentId);
 						}
 						else if (field.IsObjectAttachment())
 						{
 							attachments.Add(new ObjectAttachment(blobIdentifier));
-							pendingCompactBinaryAttachments.Add(ParseCompactBinaryAttachmentAsync(ns, blobIdentifier));
+							pendingCompactBinaryAttachments.Add(ParseCompactBinaryAttachmentAsync(ns, blobIdentifier, cancellationToken));
 						}
 						else
 						{
@@ -235,7 +239,7 @@ namespace Horde.Server.Ddc
 			}
 		}
 
-		public async IAsyncEnumerable<BlobId> GetReferencedBlobsAsync(NamespaceId ns, CbObject cb, bool ignoreMissingBlobs = false)
+		public async IAsyncEnumerable<BlobId> GetReferencedBlobsAsync(NamespaceId ns, CbObject cb, bool ignoreMissingBlobs = false, [EnumeratorCancellation] CancellationToken cancellationToken = default)
 		{
 			List<Task<(BlobId, bool)>> pendingBlobExistsChecks = new();
 			List<Task<(ContentIdAttachment, bool)>> pendingContentIdChecks = new();
@@ -243,21 +247,21 @@ namespace Horde.Server.Ddc
 			List<BlobId> unresolvedBlobReferences = new List<BlobId>();
 
 			// Resolve all the attachments
-			await foreach (Attachment attachment in GetAttachmentsAsync(ns, cb))
+			await foreach (Attachment attachment in GetAttachmentsAsync(ns, cb, cancellationToken))
 			{
 				if (attachment is BlobAttachment blobAttachment)
 				{
-					pendingBlobExistsChecks.Add(CheckBlobExistsAsync(ns, blobAttachment.Identifier));
+					pendingBlobExistsChecks.Add(CheckBlobExistsAsync(ns, blobAttachment.Identifier, cancellationToken));
 				}
 				else if (attachment is ContentIdAttachment contentIdAttachment)
 				{
 					// If we find a content id we resolve that into the actual blobs it references
-					pendingContentIdChecks.Add(CheckContentIdExistsAsync(ns, contentIdAttachment));
+					pendingContentIdChecks.Add(CheckContentIdExistsAsync(ns, contentIdAttachment, cancellationToken));
 				}
 				else if (attachment is ObjectAttachment objectAttachment)
 				{
 					// a object just references the same blob, traversing the object attachment is done in GetAttachments
-					pendingBlobExistsChecks.Add(CheckBlobExistsAsync(ns, objectAttachment.Identifier));
+					pendingBlobExistsChecks.Add(CheckBlobExistsAsync(ns, objectAttachment.Identifier, cancellationToken));
 				}
 				else
 				{
@@ -307,12 +311,12 @@ namespace Horde.Server.Ddc
 			}
 		}
 
-		private async Task<(ContentIdAttachment, bool)> CheckContentIdExistsAsync(NamespaceId ns, ContentIdAttachment contentIdAttachment)
+		private async Task<(ContentIdAttachment, bool)> CheckContentIdExistsAsync(NamespaceId ns, ContentIdAttachment contentIdAttachment, CancellationToken cancellationToken)
 		{
 			bool allBlobsExist = true;
 			foreach (BlobId b in contentIdAttachment.ReferencedBlobs)
 			{
-				(BlobId _, bool exists) = await CheckBlobExistsAsync(ns, b);
+				(BlobId _, bool exists) = await CheckBlobExistsAsync(ns, b, cancellationToken);
 
 				if (!exists)
 				{
@@ -324,27 +328,27 @@ namespace Horde.Server.Ddc
 			return (contentIdAttachment, allBlobsExist);
 		}
 
-		private async Task<CbObject> ParseCompactBinaryAttachmentAsync(NamespaceId ns, BlobId blobIdentifier)
+		private async Task<CbObject> ParseCompactBinaryAttachmentAsync(NamespaceId ns, BlobId blobIdentifier, CancellationToken cancellationToken)
 		{
-			BlobContents contents = await _blobStore.GetObjectAsync(ns, blobIdentifier);
-			byte[] data = await contents.Stream.ToByteArrayAsync();
+			BlobContents contents = await _blobStore.GetObjectAsync(ns, blobIdentifier, cancellationToken: cancellationToken);
+			byte[] data = await contents.Stream.ToByteArrayAsync(cancellationToken);
 			CbObject childBinaryObject = new CbObject(data);
 
 			return childBinaryObject;
 		}
 
-		private async Task<(ContentId, BlobId[]?)> ResolveContentIdAsync(NamespaceId ns, ContentId contentId)
+		private async Task<(ContentId, BlobId[]?)> ResolveContentIdAsync(NamespaceId ns, ContentId contentId, CancellationToken cancellationToken)
 		{
 			using TelemetrySpan scope = _tracer.StartActiveSpan("ReferenceResolver.ResolveContentId")
 				.SetAttribute("operation.name", "ReferenceResolver.ResolveContentId")
 				.SetAttribute("resource.name", contentId.ToString());
-			BlobId[]? resolvedBlobs = await _contentIdStore.ResolveAsync(ns, contentId);
+			BlobId[]? resolvedBlobs = await _contentIdStore.ResolveAsync(ns, contentId, cancellationToken: cancellationToken);
 			return (contentId, resolvedBlobs);
 		}
 
-		private async Task<(BlobId, bool)> CheckBlobExistsAsync(NamespaceId ns, BlobId blob)
+		private async Task<(BlobId, bool)> CheckBlobExistsAsync(NamespaceId ns, BlobId blob, CancellationToken cancellationToken)
 		{
-			return (blob, await _blobStore.ExistsAsync(ns, blob));
+			return (blob, await _blobStore.ExistsAsync(ns, blob, cancellationToken: cancellationToken));
 		}
 	}
 
diff --git a/Engine/Source/Programs/Horde/Horde.Server/Ddc/ObjectsController.cs b/Engine/Source/Programs/Horde/Horde.Server/Ddc/ObjectsController.cs
--- a/Engine/Source/Programs/Horde/Horde.Server/Ddc/ObjectsController.cs
+++ b/Engine/Source/Programs/Horde/Horde.Server/Ddc/ObjectsController.cs
@@ -158,9 +158,9 @@ namespace Horde.Server.Ddc
 			_diagnosticContext.Set("Content-Length", Request.ContentLength ?? -1);
 			try
 			{
-				using IBufferedPayload payload = await _bufferedPayloadFactory.CreateFromRequestAsync(Request);
+				using IBufferedPayload payload = await _bufferedPayloadFactory.CreateFromRequestAsync(Request, HttpContext.RequestAborted);
 
-				BlobId identifier = await _storage.PutObjectAsync(ns, payload, id);
+				BlobId identifier = await _storage.PutObjectAsync(ns, payload, id, HttpContext.RequestAborted);
 				return Ok(new PutBlobResponse(identifier));
 			}
 			catch (ClientSendSlowException e)
@@ -190,7 +190,7 @@ namespace Horde.Server.Ddc
 				return NotFound(new ValidationProblemDetails { Title = $"Object {e.Blob} not found" });
 			}
 
-			byte[] blobContents = await blob.Stream.ToByteArrayAsync();
+			byte[] blobContents = await blob.Stream.ToByteArrayAsync(HttpContext.RequestAborted);
 			if (blobContents.Length == 0)
 			{
 				_logger.LogWarning("0 byte object found for {Id} {Namespace}", id, ns);
@@ -232,7 +232,7 @@ namespace Horde.Server.Ddc
 				return result;
 			}
 
-			await _storage.DeleteObjectAsync(ns, id);
+			await _storage.DeleteObjectAsync(ns, id, HttpContext.RequestAborted);
 
 			return Ok(new DeletedResponse
 			{
@@ -250,7 +250,7 @@ namespace Horde.Server.Ddc
 				return result;
 			}
 
-			await _storage.DeleteNamespaceAsync(ns);
+			await _storage.DeleteNamespaceAsync(ns, HttpContext.RequestAborted);
 
 			return Ok();
 		}
diff --git a/Engine/Source/Programs/Horde/Horde.Server/Ddc/RefService.cs b/Engine/Source/Programs/Horde/Horde.Server/Ddc/RefService.cs
--- a/Engine/Source/Programs/Horde/Horde.Server/Ddc/RefService.cs
+++ b/Engine/Source/Programs/Horde/Horde.Server/Ddc/RefService.cs
@@ -3,6 +3,7 @@
 using System;
 using System.Collections.Generic;
 using System.Linq;
+using System.Threading;
 using System.Threading.Tasks;
 using EpicGames.Horde.Storage;
 using EpicGames.Serialization;
@@ -24,30 +25,30 @@ namespace Horde.Server.Ddc
 
 		static RefName GetRefName(BucketId bucketId, RefId refId) => $"{bucketId}/{refId}";
 
-		public async Task<bool> DeleteAsync(NamespaceId ns, BucketId bucket, RefId key)
+		public async Task<bool> DeleteAsync(NamespaceId ns, BucketId bucket, RefId key, CancellationToken cancellationToken)
 		{
 			using IStorageClient storageClient = _storageClientFactory.CreateClient(ns);
-			return await storageClient.DeleteRefAsync(GetRefName(bucket, key));
+			return await storageClient.DeleteRefAsync(GetRefName(bucket, key), cancellationToken);
 		}
 
-		public async Task<bool> ExistsAsync(NamespaceId ns, BucketId bucket, RefId key)
+		public async Task<bool> ExistsAsync(NamespaceId ns, BucketId bucket, RefId key, CancellationToken cancellationToken)
 		{
 			using IStorageClient storageClient = _storageClientFactory.CreateClient(ns);
-			return await storageClient.RefExistsAsync(GetRefName(bucket, key));
+			return await storageClient.RefExistsAsync(GetRefName(bucket, key), cancellationToken: cancellationToken);
 		}
 
-		public async Task<(ContentId[], BlobId[])> FinalizeAsync(NamespaceId ns, BucketId bucket, RefId key, BlobId blobHash)
+		public async Task<(ContentId[], BlobId[])> FinalizeAsync(NamespaceId ns, BucketId bucket, RefId key, BlobId blobHash, CancellationToken cancellationToken)
 		{
 			using IStorageClient storageClient = _storageClientFactory.CreateClient(ns);
 
-			BlobAlias? blobAlias = await storageClient.FindAliasAsync(BlobService.GetAlias(blobHash));
+			BlobAlias? blobAlias = await storageClient.FindAliasAsync(BlobService.GetAlias(blobHash), cancellationToken);
 			if (blobAlias == null)
 			{
 				throw new BlobNotFoundException(ns, blobHash);
 			}
 
 			IBlobHandle blobHandle = blobAlias.Target;
-			using BlobData blobContents = await blobHandle.ReadBlobDataAsync();
+			using BlobData blobContents = await blobHandle.ReadBlobDataAsync(cancellationToken);
 			CbObject payload = new CbObject(blobContents.Data);
 
 			BlobId[] referencedBlobs = Array.Empty<BlobId>();
@@ -59,7 +60,7 @@ namespace Horde.Server.Ddc
 			{
 				try
 				{
-					referencedBlobs = await _referenceResolver.GetReferencedBlobsAsync(ns, payload).ToArrayAsync();
+					referencedBlobs = await _referenceResolver.GetReferencedBlobsAsync(ns, payload, cancellationToken: cancellationToken).ToArrayAsync(cancellationToken);
 				}
 				catch (PartialReferenceResolveException e)
 				{
@@ -83,13 +84,13 @@ namespace Horde.Server.Ddc
 					refNode.References.Add(BlobRef.Create(blobHash.AsIoHash(), blobHandle));
 					foreach (BlobId referencedBlob in referencedBlobs)
 					{
-						BlobAlias? alias = await storageClient.FindAliasAsync(BlobService.GetAlias(referencedBlob));
+						BlobAlias? alias = await storageClient.FindAliasAsync(BlobService.GetAlias(referencedBlob), cancellationToken);
 						refNode.References.Add(BlobRef.Create(referencedBlob.AsIoHash(), alias!.Target));
 					}
-					refNodeRef = await writer.WriteBlobAsync(refNode);
+					refNodeRef = await writer.WriteBlobAsync(refNode, cancellationToken);
 				}
 
-				await storageClient.WriteRefAsync(refName, refNodeRef);
+				await storageClient.WriteRefAsync(refName, refNodeRef, cancellationToken: cancellationToken);
 			}
 
 			return (missingReferences, missingBlobs);
@@ -126,28 +127,28 @@ namespace Horde.Server.Ddc
 			return payload.Any(FieldHasAttachments);
 		}
 
-		public async Task<(RefRecord, BlobContents?)> GetAsync(NamespaceId ns, BucketId bucket, RefId key, string[] fields, bool doLastAccessTracking)
+		public async Task<(RefRecord, BlobContents?)> GetAsync(NamespaceId ns, BucketId bucket, RefId key, string[] fields, bool doLastAccessTracking, CancellationToken cancellationToken)
 		{
 			using IStorageClient storageClient = _storageClientFactory.CreateClient(ns);
 
-			DdcRefNode? node = await storageClient.TryReadRefTargetAsync<DdcRefNode>(GetRefName(bucket, key));
+			DdcRefNode? node = await storageClient.TryReadRefTargetAsync<DdcRefNode>(GetRefName(bucket, key), cancellationToken: cancellationToken);
 			if (node == null)
 			{
 				throw new RefNotFoundException(ns, bucket, key);
 			}
 
-			BlobData data = await node.References.First(x => x.Hash == node.RootHash).ReadBlobDataAsync();
+			BlobData data = await node.References.First(x => x.Hash == node.RootHash).ReadBlobDataAsync(cancellationToken);
 			BlobContents contents = new BlobContents(data.Data.ToArray());
 
 			RefRecord record = new RefRecord(ns, bucket, key, DateTime.UtcNow, null, BlobId.FromIoHash(node.RootHash), true);
 			return (record, contents);
 		}
 
-		public async Task<List<BlobId>> GetReferencedBlobsAsync(NamespaceId ns, BucketId bucket, RefId key)
+		public async Task<List<BlobId>> GetReferencedBlobsAsync(NamespaceId ns, BucketId bucket, RefId key, CancellationToken cancellationToken)
 		{
 			using IStorageClient storageClient = _storageClientFactory.CreateClient(ns);
 
-			DdcRefNode? node = await storageClient.TryReadRefTargetAsync<DdcRefNode>(GetRefName(bucket, key));
+			DdcRefNode? node = await storageClient.TryReadRefTargetAsync<DdcRefNode>(GetRefName(bucket, key), cancellationToken: cancellationToken);
 			if (node == null)
 			{
 				throw new RefNotFoundException(ns, bucket, key);
@@ -156,28 +157,28 @@ namespace Horde.Server.Ddc
 			return node.References.Select(x => BlobId.FromIoHash(x.Hash)).ToList();
 		}
 
-		public async Task<(ContentId[], BlobId[])> PutAsync(NamespaceId ns, BucketId bucket, RefId key, BlobId blobHash, CbObject payload)
+		public async Task<(ContentId[], BlobId[])> PutAsync(NamespaceId ns, BucketId bucket, RefId key, BlobId blobHash, CbObject payload, CancellationToken cancellationToken)
 		{
-			await _blobService.PutObjectAsync(ns, payload.GetView().ToArray(), blobHash);
-			return await FinalizeAsync(ns, bucket, key, blobHash);
+			await _blobService.PutObjectAsync(ns, payload.GetView().ToArray(), blobHash, cancellationToken);
+			return await FinalizeAsync(ns, bucket, key, blobHash, cancellationToken);
 		}
 
-		public IAsyncEnumerable<NamespaceId> GetNamespacesAsync()
+		public IAsyncEnumerable<NamespaceId> GetNamespacesAsync(CancellationToken cancellationToken)
 		{
 			throw new NotImplementedException();
 		}
 
-		public Task<long> DropNamespaceAsync(NamespaceId ns)
+		public Task<long> DropNamespaceAsync(NamespaceId ns, CancellationToken cancellationToken)
 		{
 			throw new NotImplementedException();
 		}
 
-		public Task<long> DeleteBucketAsync(NamespaceId ns, BucketId bucket)
+		public Task<long> DeleteBucketAsync(NamespaceId ns, BucketId bucket, CancellationToken cancellationToken)
 		{
 			throw new NotImplementedException();
 		}
 
-		public Task<List<BlobId>> GetReferencedBlobsAsync(NamespaceId ns, BucketId bucket, RefId key, bool ignoreMissingBlobs)
+		public Task<List<BlobId>> GetReferencedBlobsAsync(NamespaceId ns, BucketId bucket, RefId key, bool ignoreMissingBlobs, CancellationToken cancellationToken)
 		{
 			throw new NotImplementedException();
 		}
diff --git a/Engine/Source/Programs/Horde/Horde.Server/Ddc/RefsController.cs b/Engine/Source/Programs/Horde/Horde.Server/Ddc/RefsController.cs
--- a/Engine/Source/Programs/Horde/Horde.Server/Ddc/RefsController.cs
+++ b/Engine/Source/Programs/Horde/Horde.Server/Ddc/RefsController.cs
@@ -77,7 +77,7 @@ namespace Horde.Server.Ddc
 		[ProducesResponseType(type: typeof(ProblemDetails), 400)]
 		public async Task<IActionResult> GetNamespacesAsync()
 		{
-			NamespaceId[] namespaces = await _refService.GetNamespacesAsync().ToArrayAsync();
+			NamespaceId[] namespaces = await _refService.GetNamespacesAsync(HttpContext.RequestAborted).ToArrayAsync(HttpContext.RequestAborted);
 
 			// filter namespaces down to only the namespaces the user has access to
 			List<NamespaceId> namespacesWithAccess = new();
@@ -171,7 +171,7 @@ namespace Horde.Server.Ddc
 						}
 					case MediaTypeNames.Application.Octet:
 						{
-							byte[] blobMemory = await blob.Stream.ToByteArrayAsync();
+							byte[] blobMemory = await blob.Stream.ToByteArrayAsync(HttpContext.RequestAborted);
 							CbObject cb = new CbObject(blobMemory);
 
 							(int, CbField?) CountFields(CbObject o)
@@ -220,7 +220,7 @@ namespace Horde.Server.Ddc
 							byte[] blobMemory;
 							{
 								using TelemetrySpan scope = _tracer.StartActiveSpan("json.readblob").SetAttribute("operation.name", "json.readblob");
-								blobMemory = await blob.Stream.ToByteArrayAsync();
+								blobMemory = await blob.Stream.ToByteArrayAsync(HttpContext.RequestAborted);
 							}
 							CbObject cb = new CbObject(blobMemory);
 							string s = cb.ToJson();
@@ -232,10 +232,10 @@ namespace Horde.Server.Ddc
 					case CustomMediaTypeNames.UnrealCompactBinaryPackage:
 						{
 							using TelemetrySpan packageScope = _tracer.StartActiveSpan("cbpackage.fetch").SetAttribute("operation.name", "cbpackage.fetch");
-							byte[] blobMemory = await blob.Stream.ToByteArrayAsync();
+							byte[] blobMemory = await blob.Stream.ToByteArrayAsync(HttpContext.RequestAborted);
 							CbObject cb = new CbObject(blobMemory);
 
-							IAsyncEnumerable<Attachment> attachments = _referenceResolver.GetAttachmentsAsync(ns, cb);
+							IAsyncEnumerable<Attachment> attachments = _referenceResolver.GetAttachmentsAsync(ns, cb, HttpContext.RequestAborted);
 
 							using CbPackageBuilder writer = new CbPackageBuilder();
 							writer.AddAttachment(objectRecord.BlobIdentifier.AsIoHash(), CbPackageAttachmentFlags.IsObject, blobMemory);
@@ -251,19 +251,19 @@ namespace Horde.Server.Ddc
 									if (attachment is BlobAttachment blobAttachment)
 									{
 										BlobId referencedBlob = blobAttachment.Identifier;
-										attachmentContents = await _blobStore.GetObjectAsync(ns, referencedBlob);
+										attachmentContents = await _blobStore.GetObjectAsync(ns, referencedBlob, cancellationToken: HttpContext.RequestAborted);
 									}
 									else if (attachment is ObjectAttachment objectAttachment)
 									{
 										flags |= CbPackageAttachmentFlags.IsObject;
 										BlobId referencedBlob = objectAttachment.Identifier;
-										attachmentContents = await _blobStore.GetObjectAsync(ns, referencedBlob);
+										attachmentContents = await _blobStore.GetObjectAsync(ns, referencedBlob, cancellationToken: HttpContext.RequestAborted);
 									}
 									else if (attachment is ContentIdAttachment contentIdAttachment)
 									{
 
 										ContentId contentId = contentIdAttachment.Identifier;
-										(attachmentContents, string mime) = await _blobStore.GetCompressedObjectAsync(ns, contentId, HttpContext.RequestServices);
+										(attachmentContents, string mime) = await _blobStore.GetCompressedObjectAsync(ns, contentId, HttpContext.RequestServices, cancellationToken: HttpContext.RequestAborted);
 										if (mime == CustomMediaTypeNames.UnrealCompressedBuffer)
 										{
 											flags |= CbPackageAttachmentFlags.IsCompressed;
@@ -302,7 +302,7 @@ namespace Horde.Server.Ddc
 						}
 					case CustomMediaTypeNames.JupiterInlinedPayload:
 						{
-							byte[] blobMemory = await blob.Stream.ToByteArrayAsync();
+							byte[] blobMemory = await blob.Stream.ToByteArrayAsync(HttpContext.RequestAborted);
 							CbObject cb = new CbObject(blobMemory);
 
 							static (int, int) CountFields(CbObject o)
@@ -500,7 +500,7 @@ namespace Horde.Server.Ddc
 
 				// we have to verify the blobs are available locally, as the record of the key is replicated a head of the content
 				// TODO: Once we support inline replication this step is not needed as at least one region as this blob, just maybe not this current one
-				byte[] blobContents = await blob.Stream.ToByteArrayAsync();
+				byte[] blobContents = await blob.Stream.ToByteArrayAsync(HttpContext.RequestAborted);
 				CbObject compactBinaryObject = new CbObject(blobContents);
 				// the reference resolver will throw if any blob is missing, so no need to do anything other then process each reference
 				IAsyncEnumerable<BlobId> references = _referenceResolver.GetReferencedBlobsAsync(ns, compactBinaryObject);
@@ -508,7 +508,7 @@ namespace Horde.Server.Ddc
 
 				// we have to verify the blobs are available locally, as the record of the key is replicated a head of the content
 				// TODO: Once we support inline replication this step is not needed as at least one region as this blob, just maybe not this current one
-				BlobId[] unknownBlobs = await _blobStore.FilterOutKnownBlobsAsync(ns, new BlobId[] { record.BlobIdentifier });
+				BlobId[] unknownBlobs = await _blobStore.FilterOutKnownBlobsAsync(ns, new BlobId[] { record.BlobIdentifier }, HttpContext.RequestAborted);
 				if (unknownBlobs.Length != 0)
 				{
 					return NotFound(new ProblemDetails { Title = $"Object {bucket} {key} in namespace {ns} had at least one missing blob." });
@@ -578,7 +578,7 @@ namespace Horde.Server.Ddc
 
 					// we have to verify the blobs are available locally, as the record of the key is replicated a head of the content
 					// TODO: Once we support inline replication this step is not needed as at least one region as this blob, just maybe not this current one
-					byte[] blobContents = await blob.Stream.ToByteArrayAsync();
+					byte[] blobContents = await blob.Stream.ToByteArrayAsync(HttpContext.RequestAborted);
 					CbObject cb = new CbObject(blobContents);
 					// the reference resolver will throw if any blob is missing, so no need to do anything other then process each reference
 					IAsyncEnumerable<BlobId> references = _referenceResolver.GetReferencedBlobsAsync(ns, cb);
@@ -626,7 +626,7 @@ namespace Horde.Server.Ddc
 
 			try
 			{
-				using IBufferedPayload payload = await _bufferedPayloadFactory.CreateFromRequestAsync(Request);
+				using IBufferedPayload payload = await _bufferedPayloadFactory.CreateFromRequestAsync(Request, HttpContext.RequestAborted);
 
 				BlobId headerHash;
 				if (Request.Headers.TryGetValue(CommonHeaders.HashHeaderName, out StringValues headers))
@@ -659,7 +659,7 @@ namespace Horde.Server.Ddc
 						{
 							// TODO: define a scheme for how a json object specifies references
 
-							blobHeader = await _blobStore.PutObjectAsync(ns, payload, headerHash);
+							blobHeader = await _blobStore.PutObjectAsync(ns, payload, headerHash, HttpContext.RequestAborted);
 
 							// TODO: convert the json object into a compact binary instead
 							CbWriter writer = new CbWriter();
@@ -682,7 +682,7 @@ namespace Horde.Server.Ddc
 						}
 					case MediaTypeNames.Application.Octet:
 						{
-							blobHeader = await _blobStore.PutObjectAsync(ns, payload, headerHash);
+							blobHeader = await _blobStore.PutObjectAsync(ns, payload, headerHash, HttpContext.RequestAborted);
 
 							CbWriter writer = new CbWriter();
 							writer.BeginObject();
@@ -711,7 +711,7 @@ namespace Horde.Server.Ddc
 				return Problem(e.Message, null, (int)HttpStatusCode.RequestTimeout);
 			}
 
-			(ContentId[] missingReferences, BlobId[] missingBlobs) = await _refService.PutAsync(ns, bucket, key, blobHeader, payloadObject);
+			(ContentId[] missingReferences, BlobId[] missingBlobs) = await _refService.PutAsync(ns, bucket, key, blobHeader, payloadObject, HttpContext.RequestAborted);
 
 			{
 				using TelemetrySpan scope = _tracer.StartActiveSpan("ref.put").SetAttribute("operation.name", "ref.put");
@@ -759,11 +759,11 @@ namespace Horde.Server.Ddc
 #pragma warning disable CA2000 // Dispose objects before losing scope
 						using MemoryBufferedPayload payload = new MemoryBufferedPayload(blob);
 #pragma warning restore CA2000 // Dispose objects before losing scope
-						await _blobStore.PutCompressedObjectAsync(ns, payload, ContentId.FromIoHash(entry.AttachmentHash), HttpContext.RequestServices);
+						await _blobStore.PutCompressedObjectAsync(ns, payload, ContentId.FromIoHash(entry.AttachmentHash), HttpContext.RequestServices, HttpContext.RequestAborted);
 					}
 					else
 					{
-						await _blobStore.PutObjectAsync(ns, blob, BlobId.FromIoHash(entry.AttachmentHash));
+						await _blobStore.PutObjectAsync(ns, blob, BlobId.FromIoHash(entry.AttachmentHash), HttpContext.RequestAborted);
 					}
 				}
 			}
@@ -778,7 +778,7 @@ namespace Horde.Server.Ddc
 			CbObject rootObject = packageReader.RootObject;
 			BlobId rootObjectHash = BlobId.FromIoHash(packageReader.RootHash);
 
-			(ContentId[] missingReferences, BlobId[] missingBlobs) = await _refService.PutAsync(ns, bucket, key, rootObjectHash, rootObject);
+			(ContentId[] missingReferences, BlobId[] missingBlobs) = await _refService.PutAsync(ns, bucket, key, rootObjectHash, rootObject, HttpContext.RequestAborted);
 
 			List<ContentHash> missingHashes = new List<ContentHash>(missingReferences);
 			missingHashes.AddRange(missingBlobs);
@@ -800,7 +800,7 @@ namespace Horde.Server.Ddc
 
 			try
 			{
-				(ContentId[] missingReferences, BlobId[] missingBlobs) = await _refService.FinalizeAsync(ns, bucket, key, hash);
+				(ContentId[] missingReferences, BlobId[] missingBlobs) = await _refService.FinalizeAsync(ns, bucket, key, hash, HttpContext.RequestAborted);
 				List<ContentHash> missingHashes = new List<ContentHash>(missingReferences);
 				missingHashes.AddRange(missingBlobs);
 
@@ -874,7 +874,7 @@ namespace Horde.Server.Ddc
 						throw new Exception();
 					}
 
-					CbObject cb = new CbObject(await blob.Stream.ToByteArrayAsync());
+					CbObject cb = new CbObject(await blob.Stream.ToByteArrayAsync(HttpContext.RequestAborted));
 
 					if (op.ResolveAttachments ?? false)
 					{
@@ -911,7 +911,7 @@ namespace Horde.Server.Ddc
 
 					if (op.ResolveAttachments ?? false)
 					{
-						byte[] blobContents = await blob.Stream.ToByteArrayAsync();
+						byte[] blobContents = await blob.Stream.ToByteArrayAsync(HttpContext.RequestAborted);
 						CbObject cb = new CbObject(blobContents);
 						// the reference resolver will throw if any blob is missing, so no need to do anything other then process each reference
 						IAsyncEnumerable<BlobId> references = _referenceResolver.GetReferencedBlobsAsync(ns, cb);
@@ -956,7 +956,7 @@ namespace Horde.Server.Ddc
 						throw new HashMismatchException(headerHash, objectHash);
 					}
 
-					(ContentId[] missingReferences, BlobId[] missingBlobs) = await _refService.PutAsync(ns, op.Bucket, op.Key, objectHash, op.Payload);
+					(ContentId[] missingReferences, BlobId[] missingBlobs) = await _refService.PutAsync(ns, op.Bucket, op.Key, objectHash, op.Payload, HttpContext.RequestAborted);
 					List<ContentHash> missingHashes = new List<ContentHash>(missingReferences);
 
 					return (CbSerializer.Serialize(new PutObjectResponse(missingHashes.ToArray())), HttpStatusCode.OK);
@@ -1045,7 +1045,7 @@ namespace Horde.Server.Ddc
 
 			try
 			{
-				await _refService.DropNamespaceAsync(ns);
+				await _refService.DropNamespaceAsync(ns, HttpContext.RequestAborted);
 			}
 			catch (NamespaceNotFoundException e)
 			{
@@ -1075,7 +1075,7 @@ namespace Horde.Server.Ddc
 			long countOfDeletedRecords;
 			try
 			{
-				countOfDeletedRecords = await _refService.DeleteBucketAsync(ns, bucket);
+				countOfDeletedRecords = await _refService.DeleteBucketAsync(ns, bucket, HttpContext.RequestAborted);
 			}
 			catch (NamespaceNotFoundException e)
 			{
@@ -1107,7 +1107,7 @@ namespace Horde.Server.Ddc
 
 			try
 			{
-				bool deleted = await _refService.DeleteAsync(ns, bucket, key);
+				bool deleted = await _refService.DeleteAsync(ns, bucket, key, HttpContext.RequestAborted);
 				return Ok(new RefDeletedResponse(deleted ? 1 : 0));
 			}
 			catch (NamespaceNotFoundException e)
diff --git a/Engine/Source/Programs/Horde/Horde.Server/Ddc/StreamUtils.cs b/Engine/Source/Programs/Horde/Horde.Server/Ddc/StreamUtils.cs
--- a/Engine/Source/Programs/Horde/Horde.Server/Ddc/StreamUtils.cs
+++ b/Engine/Source/Programs/Horde/Horde.Server/Ddc/StreamUtils.cs
@@ -1,23 +1,23 @@
 ﻿// Copyright Epic Games, Inc. All Rights Reserved.
 
 using System.IO;
+using System.Threading;
 using System.Threading.Tasks;
 using EpicGames.AspNet;
 using Microsoft.AspNetCore.Http;
 
 #pragma warning disable CS1591
-#pragma warning disable MA0032 // Use an overload with a CancellationToken
 
 namespace Horde.Server.Ddc
 {
 	public static class StreamUtils
 	{
-		public static async Task<byte[]> ToByteArrayAsync(this Stream s)
+		public static async Task<byte[]> ToByteArrayAsync(this Stream s, CancellationToken cancellationToken)
 		{
 			try
 			{
 				await using MemoryStream ms = new MemoryStream();
-				await s.CopyToAsync(ms);
+				await s.CopyToAsync(ms, cancellationToken);
 				return ms.ToArray();
 			}
 			catch (BadHttpRequestException e)
