Horde: Change up automation data loading, fix no automation found handling
#rnx

diff --git a/Engine/Source/Programs/Horde/HordeDashboard/src/backend/AutomationTestData.ts b/Engine/Source/Programs/Horde/HordeDashboard/src/backend/AutomationTestData.ts
--- a/Engine/Source/Programs/Horde/HordeDashboard/src/backend/AutomationTestData.ts
+++ b/Engine/Source/Programs/Horde/HordeDashboard/src/backend/AutomationTestData.ts
@@ -45,8 +45,7 @@ export class TestDataHandler {
     constructor(search: URLSearchParams) {
         TestDataHandler.instance = this;
         makeObservable(this);
-        this.state = this.stateFromSearch(search);
-        this.load();
+        this.state = this.stateFromSearch(search);        
     }
 
     setAutomation(automation: string, initStreams: boolean = false) {
diff --git a/Engine/Source/Programs/Horde/HordeDashboard/src/components/AutomationView.tsx b/Engine/Source/Programs/Horde/HordeDashboard/src/components/AutomationView.tsx
--- a/Engine/Source/Programs/Horde/HordeDashboard/src/components/AutomationView.tsx
+++ b/Engine/Source/Programs/Horde/HordeDashboard/src/components/AutomationView.tsx
@@ -516,7 +516,7 @@ const AutomationOperationsBar: React.FC<{ handler: TestDataHandler }> = observer
    </Stack>;
 })
 
-export const AutomationView: React.FC = observer(() => {   
+export const AutomationView: React.FC = observer(() => {
 
    const navigate = useNavigate();
    const windowSize = useWindowSize();
@@ -528,8 +528,8 @@ export const AutomationView: React.FC = observer(() => {
    // subscribe
    if (!handler) {
       handler = new TestDataHandler(new URLSearchParams(searchParams));
-      if (handler.updated) { }
-      setState({ handler: handler, search: handler.search.toString() });
+      handler.load().catch((reason) => console.error(reason)).finally(() => setState({ handler: handler, search: handler!.search.toString() }))
+
       return null;
    }
 
@@ -575,7 +575,7 @@ export const AutomationView: React.FC = observer(() => {
                <div key={`windowsize_automationview_${windowSize.width}_${windowSize.height}`} style={{ width: (vw / 2 - (1440 / 2)) - 12, flexShrink: 0, backgroundColor: modeColors.background }} />
                <Stack horizontalAlign="center" grow styles={{ root: { width: "100%", padding: 12, backgroundColor: modeColors.background } }}>
                   <Stack styles={{ root: { width: "100%" } }}>
-                     {handler.loaded && !handler.automation?.length && <Stack horizontal tokens={{ childrenGap: 6 }} horizontalAlign="center">
+                     {handler.loaded && !handler.automation?.length && <Stack horizontal style={{ width: 1440, paddingTop: 30 }} tokens={{ childrenGap: 6 }} horizontalAlign="center">
                         <Text variant="mediumPlus">No automation metadata found, for more information please see</Text>
                         <a href={automationHubDocs} style={{ fontSize: "18px", "cursor": "pointer" }} onClick={(ev) => { ev.preventDefault(); ev.stopPropagation(); navigate(automationHubDocs) }}> automation hub documentation.</a>
                      </Stack>}
