Horde: Add a ticket example for globals.json
#rnx

diff --git a/Engine/Source/Programs/Horde/Horde.Server/Defaults/globals.json b/Engine/Source/Programs/Horde/Horde.Server/Defaults/globals.json
--- a/Engine/Source/Programs/Horde/Horde.Server/Defaults/globals.json
+++ b/Engine/Source/Programs/Horde/Horde.Server/Defaults/globals.json
@@ -59,7 +59,8 @@
 //			[
 //				{
 //					"userName": "test.user",
-//					"password": "password123"	
+//					"password": "password123",
+//					"ticket": "ticket if not using password"
 //				}
 //			]
 //		}
