Horde: Use the first specified credentials for server P4 commands if the serviceAccount property is not set.

diff --git a/Engine/Source/Programs/Horde/Horde.Server/Defaults/globals.json b/Engine/Source/Programs/Horde/Horde.Server/Defaults/globals.json
--- a/Engine/Source/Programs/Horde/Horde.Server/Defaults/globals.json
+++ b/Engine/Source/Programs/Horde/Horde.Server/Defaults/globals.json
@@ -28,7 +28,8 @@
 //			"name": "Default",
 //
 //			// Default account to use inside the Horde Server for querying submitted changes etc..
-//			"serviceAccount": "test.user",
+//			// If this value is not specified, the first item from the credentials below will be used.
+//			"serviceAccount": null,
 //
 //			// Whether the service account supports generating login tickets for other users, which requires a priviledged account on the server. Impersonating other users
 //			// allows Horde to submit changes on behalf of users.
@@ -40,17 +41,6 @@
 //				{
 //					// DNS name and port for the server
 //					"serverAndPort": "perforce:1666",
-//
-//					// Whether the server is running a health check service at http://server:5000/healthcheck. Horde can periodically poll this service to determine whether the
-//					// server is under heavy load and/or is being taken out of service.
-//					"healthCheck": false,
-//
-//					// Whether to manually resolve the DNS entry for serverAndPort above, and treat all records as a separate record. Putting multiple servers under a single DNS 
-//					// record is useful to allow Perforce configuration outside of the Horde configuration file.
-//					"resolveDns": false,
-//
-//					// Maximum number of agents that can be running conforms at once.
-//					"maxConformCount" : 1
 //				}
 //			],
 //
@@ -59,6 +49,8 @@
 //			[
 //				{
 //					"userName": "test.user",
+//
+//					// Specify either the password or ticket property below (not both)
 //					"password": "password123",
 //					"ticket": "ticket if not using password"
 //				}
diff --git a/Engine/Source/Programs/Horde/Horde.Server/Perforce/PerforceService.cs b/Engine/Source/Programs/Horde/Horde.Server/Perforce/PerforceService.cs
--- a/Engine/Source/Programs/Horde/Horde.Server/Perforce/PerforceService.cs
+++ b/Engine/Source/Programs/Horde/Horde.Server/Perforce/PerforceService.cs
@@ -398,7 +398,15 @@ namespace Horde.Server.Perforce
 			}
 			else
 			{
-				return new Credentials(PerforceSettings.Default.UserName, null, null, null);
+				PerforceCredentials? credentials = cluster.Credentials.FirstOrDefault();
+				if (credentials != null)
+				{
+					return new Credentials(credentials.UserName, credentials.Password, credentials.Ticket, null);
+				}
+				else
+				{
+					return new Credentials(PerforceSettings.Default.UserName, null, null, null);
+				}
 			}
 		}
 
diff --git a/Engine/Source/Programs/Horde/Horde.Server/Server/GlobalConfig.cs b/Engine/Source/Programs/Horde/Horde.Server/Server/GlobalConfig.cs
--- a/Engine/Source/Programs/Horde/Horde.Server/Server/GlobalConfig.cs
+++ b/Engine/Source/Programs/Horde/Horde.Server/Server/GlobalConfig.cs
@@ -1175,7 +1175,7 @@ namespace Horde.Server.Server
 		public string Name { get; set; } = null!;
 
 		/// <summary>
-		/// Username for Horde to log in to this server. Will use the default user if not set.
+		/// Username for Horde to log in to this server. Will use the first account specified below if not overridden.
 		/// </summary>
 		public string? ServiceAccount { get; set; }
 
