Horde: Exclude any default* files from copying to the user's config folder.

#rnx

diff --git a/Engine/Source/Programs/Horde/Horde.Server/ServerApp.cs b/Engine/Source/Programs/Horde/Horde.Server/ServerApp.cs
--- a/Engine/Source/Programs/Horde/Horde.Server/ServerApp.cs
+++ b/Engine/Source/Programs/Horde/Horde.Server/ServerApp.cs
@@ -232,7 +232,7 @@ namespace Horde.Server
 			DirectoryReference.CreateDirectory(targetDir);
 			foreach (FileReference sourceFile in DirectoryReference.EnumerateFiles(sourceDir))
 			{
-				if ((sourceFile.HasExtension(".json") || sourceFile.HasExtension(".png")) && !sourceFile.GetFileName().Equals("defaults.json", StringComparison.OrdinalIgnoreCase))
+				if ((sourceFile.HasExtension(".json") || sourceFile.HasExtension(".png")) && !sourceFile.GetFileName().StartsWith("default", StringComparison.OrdinalIgnoreCase))
 				{
 					FileReference targetFile = FileReference.Combine(targetDir, sourceFile.GetFileName());
 					if (!FileReference.Exists(targetFile))
