Horde: Remove unused 'id' field from telemetry config.

#rnx

diff --git a/Engine/Source/Programs/Horde/Horde.Server/ServerSettings.cs b/Engine/Source/Programs/Horde/Horde.Server/ServerSettings.cs
--- a/Engine/Source/Programs/Horde/Horde.Server/ServerSettings.cs
+++ b/Engine/Source/Programs/Horde/Horde.Server/ServerSettings.cs
@@ -193,11 +193,6 @@ namespace Horde.Server
 	/// </summary>
 	public class BaseTelemetryConfig
 	{
-		/// <summary>
-		/// Unique ID for this sink config (any arbitrary string)
-		/// </summary>
-		public string? Id { get; set; }
-
 		/// <summary>
 		/// Type of telemetry sink
 		/// </summary>
