Horde: Assume any exceptions while checking for config files being out of date as meaning they need to be read again.

diff --git a/Engine/Source/Programs/Horde/Horde.Server/Configuration/ConfigService.cs b/Engine/Source/Programs/Horde/Horde.Server/Configuration/ConfigService.cs
--- a/Engine/Source/Programs/Horde/Horde.Server/Configuration/ConfigService.cs
+++ b/Engine/Source/Programs/Horde/Horde.Server/Configuration/ConfigService.cs
@@ -604,37 +604,45 @@ namespace Horde.Server.Configuration
 		/// <returns>True if the snapshot is out of date</returns>
 		async Task<bool> IsOutOfDateAsync(ConfigSnapshot snapshot, CancellationToken cancellationToken)
 		{
-			// Always re-read the config file when switching server versions
-			string newServerVersion = ServerApp.Version.ToString();
-			if (!snapshot.ServerVersion.Equals(newServerVersion, StringComparison.Ordinal))
-			{
-				_logger.LogInformation("Config is out of date (server version {OldVersion} -> {NewVersion})", snapshot.ServerVersion, newServerVersion);
-				return true;
-			}
-
-			// Group the dependencies by scheme in order to allow the source to batch-query them
-			foreach (IGrouping<string, KeyValuePair<Uri, string>> group in snapshot.Dependencies.GroupBy(x => x.Key.Scheme))
+			try
 			{
-				KeyValuePair<Uri, string>[] pairs = group.ToArray();
-
-				IConfigSource? source;
-				if (!_sources.TryGetValue(group.Key, out source))
+				// Always re-read the config file when switching server versions
+				string newServerVersion = ServerApp.Version.ToString();
+				if (!snapshot.ServerVersion.Equals(newServerVersion, StringComparison.Ordinal))
 				{
-					_logger.LogInformation("Config is out of date (missing source '{Source}')", group.Key);
+					_logger.LogInformation("Config is out of date (server version {OldVersion} -> {NewVersion})", snapshot.ServerVersion, newServerVersion);
 					return true;
 				}
 
-				IConfigFile[] files = await source.GetAsync(pairs.ConvertAll(x => x.Key), cancellationToken);
-				for (int idx = 0; idx < pairs.Length; idx++)
+				// Group the dependencies by scheme in order to allow the source to batch-query them
+				foreach (IGrouping<string, KeyValuePair<Uri, string>> group in snapshot.Dependencies.GroupBy(x => x.Key.Scheme))
 				{
-					if (!files[idx].Revision.Equals(pairs[idx].Value, StringComparison.Ordinal))
+					KeyValuePair<Uri, string>[] pairs = group.ToArray();
+
+					IConfigSource? source;
+					if (!_sources.TryGetValue(group.Key, out source))
 					{
-						_logger.LogInformation("Config is out of date (file {Path}: {OldVersion} -> {NewVersion})", files[idx].Uri, pairs[idx].Value, files[idx].Revision);
+						_logger.LogInformation("Config is out of date (missing source '{Source}')", group.Key);
 						return true;
 					}
+
+					IConfigFile[] files = await source.GetAsync(pairs.ConvertAll(x => x.Key), cancellationToken);
+					for (int idx = 0; idx < pairs.Length; idx++)
+					{
+						if (!files[idx].Revision.Equals(pairs[idx].Value, StringComparison.Ordinal))
+						{
+							_logger.LogInformation("Config is out of date (file {Path}: {OldVersion} -> {NewVersion})", files[idx].Uri, pairs[idx].Value, files[idx].Revision);
+							return true;
+						}
+					}
 				}
+				return false;
+			}
+			catch (Exception ex)
+			{
+				_logger.LogInformation(ex, "Exception while checking for config files; assuming out of date");
+				return true;
 			}
-			return false;
 		}
 
 		async Task WriteSnapshotAsync(ConfigSnapshot snapshot)
