Horde: Add analytics tutorial.

diff --git a/Engine/Source/Programs/Horde/Docs/Config/Orientation.md b/Engine/Source/Programs/Horde/Docs/Config/Orientation.md
--- a/Engine/Source/Programs/Horde/Docs/Config/Orientation.md
+++ b/Engine/Source/Programs/Horde/Docs/Config/Orientation.md
@@ -22,7 +22,7 @@ Horde will detect changes to the referenced file and automatically update from i
 while trying to update configuration at runtime can be reported via Slack notifications, and the server will continue
 running with a cached version of the previous configuration until fixed.
 
-The `Globals.json` file may include other configuration files by path. For CI use cases, for example, it can be
+Horde configuration files may include other configuration files by path. For CI use cases, for example, it can be
 convenient to configure each stream within the stream itself.
 
 ## Projects and Streams
@@ -68,3 +68,7 @@ To set up Horde as a schema server in Microsoft Visual Studio, go to `Tools` > `
 
 Add the path to your Horde server as `{{ SERVER-URL }}/api/v1/schema/catalog.json`, substituting
 `{{ SERVER-URL }}` as appropriate.
+
+The extension given to the included files indicates the root element which should be expected by the schema.
+Files with a `.project.json` extension start at the project element, files with a `.stream.json` extension
+start at the stream element, and so on.
\ No newline at end of file
diff --git a/Engine/Source/Programs/Horde/Docs/Glossary.md b/Engine/Source/Programs/Horde/Docs/Glossary.md
--- a/Engine/Source/Programs/Horde/Docs/Glossary.md
+++ b/Engine/Source/Programs/Horde/Docs/Glossary.md
@@ -5,6 +5,14 @@
 * `Agent` - A service on a remote machine that connects to the Horde server and can be sent work to execute.
 * `Lease` - A unit of work that an agent is given to execute.
 
+## Analytics
+
+* `Analytics` - Umbrella term for data gathering and analysis.
+* `Metrics` - Aggregated data computed from telemetry events
+matching a set of configured criteria for a particular time interval.
+* `Telemetry` - Individual events sent by an application to the Horde server. Horde processes telemetry events as
+schema-less JSON objects.
+
 ## Storage
 
 * `Namespace` - A logical partition of the storage system that can have custom permissions, behaviors, and garbage collection policies.
diff --git a/Engine/Source/Programs/Horde/Docs/Tutorials/Analytics.md b/Engine/Source/Programs/Horde/Docs/Tutorials/Analytics.md
--- a/Engine/Source/Programs/Horde/Docs/Tutorials/Analytics.md
+++ b/Engine/Source/Programs/Horde/Docs/Tutorials/Analytics.md
@@ -4,8 +4,68 @@
 
 ## Introduction
 
+**Horde** implements a telemetry collector which can receive
+and process events sent by the Unreal Editor.
+
+Horde aggregates telemetry events into **metrics** for discrete time intervals, which can then be charted by the Horde
+dashboard to give valuable insights into the bottlenecks experienced by your team.
+
+![Analytics](../Images/Intro-Analytics.png)
+
 ## Prerequisites
 
-* Horde Server and one or more Horde Agents (see [Getting Started: Install Horde](InstallHorde.md)).
+* Horde Server installation (see [Getting Started: Install Horde](InstallHorde.md)).
+* An Unreal Engine project targetting Unreal Engine 5.4 or later.
 
 ## Steps
+
+1. Open your project in the Unreal Editor and go to the `Edit > Plugins` menu. Search for the `Studio Telemetry` plugin and
+make sure it is enabled. It should be enabled by default.
+2. Open your project's `DefaultEngine.ini` file (in the `Config` folder next to your `.uproject` file) and add
+the following lines:
+
+    ```ini
+    [StudioTelemetry.Provider.HordeAnalytics]
+    Name=HordeAnalytics
+    ProviderModule=AnalyticsET
+    UsageType=EditorAndClient
+    APIKeyET=HordeAnalytics.Dev
+    APIServerET="http://localhost:13340/" # Set your Horde Server URL here
+    APIEndpointET="api/v1/telemetry"
+    ```
+
+    Make sure to replace the value for `APIServerET` with the address of your Horde Server.
+
+3. Configure a telemetry store to aggregate metrics from your
+   telemetry events. There are some default metrics defined in a file included
+   with the Horde installation, which can be added by adding the following snippet to your
+   [globals.json](../Config/Orientation.md) file:
+
+    ```json
+    "telemetryStores": [
+        {
+            "id": "default",
+            "include": [
+                {
+                    "path": "$(HordeDir)/Defaults/default-metrics.telemetry.json"
+            ]
+        }
+    ]
+    ```
+
+4. Configure the analytics dashboard to render your metrics. A default dashboard is included in the Horde installation, and
+can be included by adding the following snippet to your [globals.json](../Config/Orientation.md) file:
+
+    ```json
+    "dashboard": {
+        "include": [
+            {
+                "path": "$(HordeDir)/Defaults/default-analytics.dashboard.json"
+            }
+        ]
+    ]
+    ```
+
+## See Also
+
+* [Configuration > Analytics](../Config/Analytics.md)
diff --git a/Engine/Source/Programs/Horde/Docs/Tutorials/BuildAutomation.md b/Engine/Source/Programs/Horde/Docs/Tutorials/BuildAutomation.md
--- a/Engine/Source/Programs/Horde/Docs/Tutorials/BuildAutomation.md
+++ b/Engine/Source/Programs/Horde/Docs/Tutorials/BuildAutomation.md
@@ -111,4 +111,7 @@ to group related job types. There are several predefined jobs on different tabs.
   environment variables that allow UnrealBuildTool to connect to the server without any additional configuration.
 * **Test Executor** - Runs a mock job with simulated errors or warnings, which is useful for testing connectivity to agents
   without syncing a Perforce workspace.
-  
\ No newline at end of file
+
+## See Also
+
+* [Configuration > Build Automation](../Config/BuildAutomation.md)
diff --git a/Engine/Source/Programs/Horde/Docs/Tutorials/RemoteCompilation.md b/Engine/Source/Programs/Horde/Docs/Tutorials/RemoteCompilation.md
--- a/Engine/Source/Programs/Horde/Docs/Tutorials/RemoteCompilation.md
+++ b/Engine/Source/Programs/Horde/Docs/Tutorials/RemoteCompilation.md
@@ -4,7 +4,7 @@
 
 ## Introduction
 
-**Horde** implements a platform for generic remote execution workloads, allowing clients to leverage idle CPU cycles on
+Horde implements a platform for generic remote execution workloads, allowing clients to leverage idle CPU cycles on
 other machines to accelerate workloads that would otherwise be executed locally. With Horde's remote execution platform,
 you can issue explicit commands to remote agents sequentially, such as "upload these files", "run this process",
 "send these files back", and so on.
@@ -18,7 +18,7 @@ remote machine behind the scenes as necessary.
 **Unreal Build Tool** can use Unreal Build Accelerator with Horde to offload compilation tasks to connected agents,
 spreading the workload over multiple machines.
 
-> Note: Unreal Build Accelerator only supports Windows in Unreal Engine 5.4. Support for Mac and Linux are planned for 
+> **Note:** Unreal Build Accelerator only supports Windows in Unreal Engine 5.4. Support for Mac and Linux are planned for 
   a future release.
 
 ## Prerequisites
@@ -86,9 +86,9 @@ spreading the workload over multiple machines.
    </UnrealBuildAccelerator>
    ```
 
-> Note: It is not recommended to run a Horde Agent on the same machine as the Horde Server for performance reasons.
+> **Note:** It is not recommended to run a Horde Agent on the same machine as the Horde Server for performance reasons.
 
-> Note: When using Horde's build automation functionality, be mindful of mixing pools of agents for UBA and 
+> **Note:** When using Horde's build automation functionality, be mindful of mixing pools of agents for UBA and 
   pools of agents for build automation. Agents used for build automation typically have higher requirements 
   and are a more scarce resource than compute helpers.
 
