Horde: Updating schema documentation

#rnx

diff --git a/Engine/Source/Programs/Horde/Docs/Config/Schema/Globals.md b/Engine/Source/Programs/Horde/Docs/Config/Schema/Globals.md
--- a/Engine/Source/Programs/Horde/Docs/Config/Schema/Globals.md
+++ b/Engine/Source/Programs/Horde/Docs/Config/Schema/Globals.md
@@ -299,7 +299,7 @@ Information about a cluster of Perforce servers.
 Name | Description
 ---- | -----------
 `name` | `string`<br>Name of the cluster
-`serviceAccount` | `string`<br>Username for Horde to log in to this server. Will use the default user if not set.
+`serviceAccount` | `string`<br>Username for Horde to log in to this server. Will use the first account specified below if not overridden.
 `canImpersonate` | `boolean`<br>Whether the service account can impersonate other users
 `supportsPartitionedWorkspaces` | `boolean`<br>Whether to use partitioned workspaces on this server
 `servers` | [`PerforceServer`](#perforceserver)`[]`<br>List of servers
diff --git a/Engine/Source/Programs/Horde/Docs/Deployment/ServerSettings.md b/Engine/Source/Programs/Horde/Docs/Deployment/ServerSettings.md
--- a/Engine/Source/Programs/Horde/Docs/Deployment/ServerSettings.md
+++ b/Engine/Source/Programs/Horde/Docs/Deployment/ServerSettings.md
@@ -319,7 +319,6 @@ Configuration for the telemetry sink
 
 Name | Description
 ---- | -----------
-`id` | `string`<br>Unique ID for this sink config (any arbitrary string)
 `type` | [`TelemetrySinkType`](#telemetrysinktype-enum)<br>Type of telemetry sink
 
 ## TelemetrySinkType (Enum)
