Horde: Lower log level of failed AWS instance start

diff --git a/Engine/Source/Programs/Horde/Horde.Server/Agents/Fleet/Providers/AwsRecyclingFleetManager.cs b/Engine/Source/Programs/Horde/Horde.Server/Agents/Fleet/Providers/AwsRecyclingFleetManager.cs
--- a/Engine/Source/Programs/Horde/Horde.Server/Agents/Fleet/Providers/AwsRecyclingFleetManager.cs
+++ b/Engine/Source/Programs/Horde/Horde.Server/Agents/Fleet/Providers/AwsRecyclingFleetManager.cs
@@ -297,7 +297,7 @@ public sealed class AwsRecyclingFleetManager : IFleetManager
 		span.SetAttribute("success", success);
 		if (!success)
 		{
-			_logger.LogError("Unable to start instances. Insufficient capacity for all instance types tried");
+			_logger.LogInformation("Unable to start instances. Insufficient capacity for all instance types tried");
 		}
 	}
 
