Horde: Un-deprecate the /api/v1/telemetry endpoint until the runtime uses a route containing a telemetry store id (ie. /api/v1/telemetry/default).

#rnx

diff --git a/Engine/Source/Programs/Horde/Horde.Server/Telemetry/TelemetryController.cs b/Engine/Source/Programs/Horde/Horde.Server/Telemetry/TelemetryController.cs
--- a/Engine/Source/Programs/Horde/Horde.Server/Telemetry/TelemetryController.cs
+++ b/Engine/Source/Programs/Horde/Horde.Server/Telemetry/TelemetryController.cs
@@ -131,7 +131,6 @@ namespace Horde.Server.Telemetry
 		/// <param name="uploadType">Type of data being uploaded</param>
 		[HttpPost]
 		[Route("/api/v1/telemetry")]
-		[Obsolete("Pass a telemetry store id in the route")]
 		public ActionResult PostEvent([FromBody] PostTelemetryEventStreamRequest request, [FromQuery][Required] string appId, [FromQuery][Required] string appVersion, [FromQuery][Required] string appEnvironment, [FromQuery][Required] TelemetryUploadType uploadType)
 		{
 			return PostEvent(TelemetryStoreId.Default, request, appId, appVersion, appEnvironment, uploadType);
