Horde: Remove query parameters from obsolete request logging.

#rnx

diff --git a/Engine/Source/Programs/Horde/Horde.Server/Startup.cs b/Engine/Source/Programs/Horde/Horde.Server/Startup.cs
--- a/Engine/Source/Programs/Horde/Horde.Server/Startup.cs
+++ b/Engine/Source/Programs/Horde/Horde.Server/Startup.cs
@@ -311,7 +311,7 @@ namespace Horde.Server
 					|| actionDescriptor?.MethodInfo?.DeclaringType?.GetCustomAttribute<ObsoleteAttribute>() != null)
 				{
 					ILogger? logger = context.HttpContext.RequestServices.GetService<ILogger<ObsoleteLoggingFilter>>();
-					logger?.LogWarning("Using obsolete endpoint for request: {Request}", context.HttpContext.Request.GetDisplayUrl());
+					logger?.LogWarning("Using obsolete endpoint: {RequestPath}", context.HttpContext.Request.Path);
 				}
 			}
 		}
