Horde: Delete existing client used for P4 replication if creation fails due to changing the workspace type.

#rnx

diff --git a/Engine/Source/Programs/Horde/Horde.Server/Replicators/PerforceReplicator.cs b/Engine/Source/Programs/Horde/Horde.Server/Replicators/PerforceReplicator.cs
--- a/Engine/Source/Programs/Horde/Horde.Server/Replicators/PerforceReplicator.cs
+++ b/Engine/Source/Programs/Horde/Horde.Server/Replicators/PerforceReplicator.cs
@@ -918,7 +918,19 @@ namespace Horde.Server.Replicators
 				newClient.Stream = streamConfig.Name;
 				newClient.Type = partitioned ? "readonly" : null;
 
-				await perforce.CreateClientAsync(newClient, cancellationToken);
+				try
+				{
+					await perforce.CreateClientAsync(newClient, cancellationToken);
+				}
+				catch (PerforceException ex) when (ex.Error != null && ex.Error.Generic == PerforceGenericCode.Usage)
+				{
+					_logger.LogInformation(ex, "Unable to create client {ClientName}; attempting to delete and retry ({Message})", newClient.Name, ex.Message);
+					await perforce.DeleteClientAsync(DeleteClientOptions.None, newClient.Name, cancellationToken);
+
+					_logger.LogInformation("Retrying create client {ClientName} after delete", newClient.Name);
+					await perforce.CreateClientAsync(newClient, cancellationToken);
+				}
+
 				_logger.LogInformation("Created client {ClientName} for {StreamName}", newClient.Name, streamConfig.Name);
 
 				PerforceSettings settings = new PerforceSettings(perforce.Settings);
