Horde: Update release notes.

#rnx

diff --git a/Engine/Source/Programs/Horde/Docs/ReleaseNotes.md b/Engine/Source/Programs/Horde/Docs/ReleaseNotes.md
--- a/Engine/Source/Programs/Horde/Docs/ReleaseNotes.md
+++ b/Engine/Source/Programs/Horde/Docs/ReleaseNotes.md
@@ -1,5 +1,18 @@
 # Release Notes
 
+## 2024-03-21
+
+* Lower log level of failed AWS instance start (32397273)
+* Assume any exceptions while checking for config files being out of date as meaning they need to be read again. (32388229)
+* Use the first specified credentials for server P4 commands if the serviceAccount property is not set. (32383940)
+* Update docs for deploying server on Mac. (32370724)
+* Log a warning whenever an obsolete endpoint is used. (32329528)
+* Fix separate RAM slots being reported separately in agent properties. (32328916)
+* Fix use of auth tokens from environment variables when the server URL is also derived from environment variables. (32305895)
+* Fix StorageService not being registered as a hosted service, preventing GC from running. (32291615)
+* Fix artifacts not being added to graphs. (32286255)
+* Add output artifacts to the graph from the exported BuildGraph definition. (32281784)
+
 ## 2024-03-15
 
 * Get the userinfo endpoint from the OIDC discovery document rather than assuming it's a fixed path. (32250458)
