Horde: Merging back Jupiter changes.

#rnx

diff --git a/Engine/Source/Programs/Horde/Horde.Server/Ddc/BlobId.cs b/Engine/Source/Programs/Horde/Horde.Server/Ddc/BlobId.cs
--- a/Engine/Source/Programs/Horde/Horde.Server/Ddc/BlobId.cs
+++ b/Engine/Source/Programs/Horde/Horde.Server/Ddc/BlobId.cs
@@ -128,7 +128,7 @@ namespace Horde.Server.Ddc
 			return new BlobId(testObjectHash.HashData);
 		}
 
-		public static async Task<BlobId> FromStreamAsync(Stream stream, CancellationToken cancellationToken)
+		public static async Task<BlobId> FromStreamAsync(Stream stream, CancellationToken cancellationToken = default)
 		{
 			using Hasher hasher = Hasher.New();
 
diff --git a/Engine/Source/Programs/Horde/Horde.Server/Ddc/CompressedBufferUtils.cs b/Engine/Source/Programs/Horde/Horde.Server/Ddc/CompressedBufferUtils.cs
--- a/Engine/Source/Programs/Horde/Horde.Server/Ddc/CompressedBufferUtils.cs
+++ b/Engine/Source/Programs/Horde/Horde.Server/Ddc/CompressedBufferUtils.cs
@@ -185,7 +185,7 @@ namespace Horde.Server.Ddc
 			}
 		}
 
-		public async Task<IBufferedPayload> DecompressContentAsync(Stream sourceStream, ulong streamSize, CancellationToken cancellationToken)
+		public async Task<IBufferedPayload> DecompressContentAsync(Stream sourceStream, ulong streamSize, CancellationToken cancellationToken = default)
 		{
 			using BinaryReader br = new BinaryReader(sourceStream);
 			(CompressedBufferHeader header, uint[] compressedBlockSizes) = ExtractHeader(br);
diff --git a/Engine/Source/Programs/Horde/Horde.Server/Ddc/IBlobService.cs b/Engine/Source/Programs/Horde/Horde.Server/Ddc/IBlobService.cs
--- a/Engine/Source/Programs/Horde/Horde.Server/Ddc/IBlobService.cs
+++ b/Engine/Source/Programs/Horde/Horde.Server/Ddc/IBlobService.cs
@@ -17,17 +17,17 @@ namespace Horde.Server.Ddc;
 
 public interface IBlobService
 {
-	Task<ContentHash> VerifyContentMatchesHashAsync(Stream content, ContentHash identifier, CancellationToken cancellationToken);
-	Task<BlobId> PutObjectKnownHashAsync(NamespaceId ns, IBufferedPayload content, BlobId identifier, CancellationToken cancellationToken);
-	Task<BlobId> PutObjectAsync(NamespaceId ns, IBufferedPayload payload, BlobId identifier, CancellationToken cancellationToken);
-	Task<BlobId> PutObjectAsync(NamespaceId ns, byte[] payload, BlobId identifier, CancellationToken cancellationToken);
-	Task<Uri?> MaybePutObjectWithRedirectAsync(NamespaceId ns, BlobId identifier, CancellationToken cancellationToken);
+	Task<ContentHash> VerifyContentMatchesHashAsync(Stream content, ContentHash identifier, CancellationToken cancellationToken = default);
+	Task<BlobId> PutObjectKnownHashAsync(NamespaceId ns, IBufferedPayload content, BlobId identifier, CancellationToken cancellationToken = default);
+	Task<BlobId> PutObjectAsync(NamespaceId ns, IBufferedPayload payload, BlobId identifier, CancellationToken cancellationToken = default);
+	Task<BlobId> PutObjectAsync(NamespaceId ns, byte[] payload, BlobId identifier, CancellationToken cancellationToken = default);
+	Task<Uri?> MaybePutObjectWithRedirectAsync(NamespaceId ns, BlobId identifier, CancellationToken cancellationToken = default);
 
 	Task<BlobContents> GetObjectAsync(NamespaceId ns, BlobId blob, List<string>? storageLayers = null, bool supportsRedirectUri = false, bool allowOndemandReplication = true, CancellationToken cancellationToken = default);
 
 	Task<Uri?> GetObjectWithRedirectAsync(NamespaceId ns, BlobId blobIdentifier, List<string>? storageLayers = null, CancellationToken cancellationToken = default);
 
-	Task<BlobMetadata> GetObjectMetadataAsync(NamespaceId ns, BlobId blobId, CancellationToken cancellationToken);
+	Task<BlobMetadata> GetObjectMetadataAsync(NamespaceId ns, BlobId blobId, CancellationToken cancellationToken = default);
 
 	Task<BlobContents> ReplicateObjectAsync(NamespaceId ns, BlobId blob, bool force = false, CancellationToken cancellationToken = default);
 
@@ -40,18 +40,18 @@ public interface IBlobService
 	/// <param name="blob">The identifier of the blob</param>
 	/// <param name="cancellationToken"></param>
 	/// <returns></returns>
-	Task<bool> ExistsInRootStoreAsync(NamespaceId ns, BlobId blob, CancellationToken cancellationToken);
+	Task<bool> ExistsInRootStoreAsync(NamespaceId ns, BlobId blob, CancellationToken cancellationToken = default);
 
 	// Delete a object
-	Task DeleteObjectAsync(NamespaceId ns, BlobId blob, CancellationToken cancellationToken);
+	Task DeleteObjectAsync(NamespaceId ns, BlobId blob, CancellationToken cancellationToken = default);
 
 	// delete the whole namespace
-	Task DeleteNamespaceAsync(NamespaceId ns, CancellationToken cancellationToken);
+	Task DeleteNamespaceAsync(NamespaceId ns, CancellationToken cancellationToken = default);
 
-	IAsyncEnumerable<(BlobId, DateTime)> ListObjectsAsync(NamespaceId ns, CancellationToken cancellationToken);
-	Task<BlobId[]> FilterOutKnownBlobsAsync(NamespaceId ns, IEnumerable<BlobId> blobs, CancellationToken cancellationToken);
-	Task<BlobId[]> FilterOutKnownBlobsAsync(NamespaceId ns, IAsyncEnumerable<BlobId> blobs, CancellationToken cancellationToken);
-	Task<BlobContents> GetObjectsAsync(NamespaceId ns, BlobId[] refRequestBlobReferences, CancellationToken cancellationToken);
+	IAsyncEnumerable<(BlobId, DateTime)> ListObjectsAsync(NamespaceId ns, CancellationToken cancellationToken = default);
+	Task<BlobId[]> FilterOutKnownBlobsAsync(NamespaceId ns, IEnumerable<BlobId> blobs, CancellationToken cancellationToken = default);
+	Task<BlobId[]> FilterOutKnownBlobsAsync(NamespaceId ns, IAsyncEnumerable<BlobId> blobs, CancellationToken cancellationToken = default);
+	Task<BlobContents> GetObjectsAsync(NamespaceId ns, BlobId[] refRequestBlobReferences, CancellationToken cancellationToken = default);
 
 	bool ShouldFetchBlobOnDemand(NamespaceId ns);
 }
diff --git a/Engine/Source/Programs/Horde/Horde.Server/Ddc/IRefService.cs b/Engine/Source/Programs/Horde/Horde.Server/Ddc/IRefService.cs
--- a/Engine/Source/Programs/Horde/Horde.Server/Ddc/IRefService.cs
+++ b/Engine/Source/Programs/Horde/Horde.Server/Ddc/IRefService.cs
@@ -14,17 +14,17 @@ namespace Horde.Server.Ddc
 	public interface IRefService
 	{
 		Task<(RefRecord, BlobContents?)> GetAsync(NamespaceId ns, BucketId bucket, RefId key, string[] fields, bool doLastAccessTracking = true, CancellationToken cancellationToken = default);
-		Task<(ContentId[], BlobId[])> PutAsync(NamespaceId ns, BucketId bucket, RefId key, BlobId blobHash, CbObject payload, CancellationToken cancellationToken);
-		Task<(ContentId[], BlobId[])> FinalizeAsync(NamespaceId ns, BucketId bucket, RefId key, BlobId blobHash, CancellationToken cancellationToken);
+		Task<(ContentId[], BlobId[])> PutAsync(NamespaceId ns, BucketId bucket, RefId key, BlobId blobHash, CbObject payload, CancellationToken cancellationToken = default);
+		Task<(ContentId[], BlobId[])> FinalizeAsync(NamespaceId ns, BucketId bucket, RefId key, BlobId blobHash, CancellationToken cancellationToken = default);
 
-		IAsyncEnumerable<NamespaceId> GetNamespacesAsync(CancellationToken cancellationToken);
+		IAsyncEnumerable<NamespaceId> GetNamespacesAsync(CancellationToken cancellationToken = default);
 
-		Task<bool> DeleteAsync(NamespaceId ns, BucketId bucket, RefId key, CancellationToken cancellationToken);
-		Task<long> DropNamespaceAsync(NamespaceId ns, CancellationToken cancellationToken);
-		Task<long> DeleteBucketAsync(NamespaceId ns, BucketId bucket, CancellationToken cancellationToken);
+		Task<bool> DeleteAsync(NamespaceId ns, BucketId bucket, RefId key, CancellationToken cancellationToken = default);
+		Task<long> DropNamespaceAsync(NamespaceId ns, CancellationToken cancellationToken = default);
+		Task<long> DeleteBucketAsync(NamespaceId ns, BucketId bucket, CancellationToken cancellationToken = default);
 
-		Task<bool> ExistsAsync(NamespaceId ns, BucketId bucket, RefId key, CancellationToken cancellationToken);
-		Task<List<BlobId>> GetReferencedBlobsAsync(NamespaceId ns, BucketId bucket, RefId key, bool ignoreMissingBlobs, CancellationToken cancellationToken);
+		Task<bool> ExistsAsync(NamespaceId ns, BucketId bucket, RefId key, CancellationToken cancellationToken = default);
+		Task<List<BlobId>> GetReferencedBlobsAsync(NamespaceId ns, BucketId bucket, RefId key, bool ignoreMissingBlobs, CancellationToken cancellationToken = default);
 	}
 
 	public class ObjectHashMismatchException : Exception
diff --git a/Engine/Source/Programs/Horde/Horde.Server/Ddc/IReferencesStore.cs b/Engine/Source/Programs/Horde/Horde.Server/Ddc/IReferencesStore.cs
--- a/Engine/Source/Programs/Horde/Horde.Server/Ddc/IReferencesStore.cs
+++ b/Engine/Source/Programs/Horde/Horde.Server/Ddc/IReferencesStore.cs
@@ -2,6 +2,7 @@
 
 using System;
 using System.Collections.Generic;
+using System.Threading;
 using System.Threading.Tasks;
 using EpicGames.Horde.Storage;
 
@@ -11,7 +12,7 @@ namespace Horde.Server.Ddc
 {
 	public interface IReferencesStore
 	{
-		Task<RefRecord> GetAsync(NamespaceId ns, BucketId bucket, RefId key, FieldFlags fieldFlags, OperationFlags opFlags);
+		Task<RefRecord> GetAsync(NamespaceId ns, BucketId bucket, RefId key, FieldFlags fieldFlags, OperationFlags opFlags, CancellationToken cancellationToken = default);
 
 		[Flags]
 		public enum FieldFlags
@@ -28,22 +29,22 @@ namespace Horde.Server.Ddc
 			BypassCache = 1
 		}
 
-		Task PutAsync(NamespaceId ns, BucketId bucket, RefId key, BlobId blobHash, byte[] blob, bool isFinalized);
-		Task FinalizeAsync(NamespaceId ns, BucketId bucket, RefId key, BlobId blobIdentifier);
+		Task PutAsync(NamespaceId ns, BucketId bucket, RefId key, BlobId blobHash, byte[] blob, bool isFinalized, CancellationToken cancellationToken = default);
+		Task FinalizeAsync(NamespaceId ns, BucketId bucket, RefId key, BlobId blobIdentifier, CancellationToken cancellationToken = default);
 
-		Task<DateTime?> GetLastAccessTimeAsync(NamespaceId ns, BucketId bucket, RefId key);
-		Task UpdateLastAccessTimeAsync(NamespaceId ns, BucketId bucket, RefId key, DateTime newLastAccessTime);
-		IAsyncEnumerable<(NamespaceId, BucketId, RefId, DateTime)> GetRecordsAsync();
+		Task<DateTime?> GetLastAccessTimeAsync(NamespaceId ns, BucketId bucket, RefId key, CancellationToken cancellationToken = default);
+		Task UpdateLastAccessTimeAsync(NamespaceId ns, BucketId bucket, RefId key, DateTime newLastAccessTime, CancellationToken cancellationToken = default);
+		IAsyncEnumerable<(NamespaceId, BucketId, RefId, DateTime)> GetRecordsAsync(CancellationToken cancellationToken = default);
 
-		IAsyncEnumerable<(NamespaceId, BucketId, RefId)> GetRecordsWithoutAccessTimeAsync();
+		IAsyncEnumerable<(NamespaceId, BucketId, RefId)> GetRecordsWithoutAccessTimeAsync(CancellationToken cancellationToken = default);
 
-		IAsyncEnumerable<(RefId, BlobId)> GetRecordsInBucketAsync(NamespaceId ns, BucketId bucket);
+		IAsyncEnumerable<(RefId, BlobId)> GetRecordsInBucketAsync(NamespaceId ns, BucketId bucket, CancellationToken cancellationToken = default);
 
-		IAsyncEnumerable<NamespaceId> GetNamespacesAsync();
-		IAsyncEnumerable<BucketId> GetBucketsAsync(NamespaceId ns);
-		Task<bool> DeleteAsync(NamespaceId ns, BucketId bucket, RefId key);
-		Task<long> DropNamespaceAsync(NamespaceId ns);
-		Task<long> DeleteBucketAsync(NamespaceId ns, BucketId bucket);
+		IAsyncEnumerable<NamespaceId> GetNamespacesAsync(CancellationToken cancellationToken = default);
+		IAsyncEnumerable<BucketId> GetBucketsAsync(NamespaceId ns, CancellationToken cancellationToken = default);
+		Task<bool> DeleteAsync(NamespaceId ns, BucketId bucket, RefId key, CancellationToken cancellationToken = default);
+		Task<long> DropNamespaceAsync(NamespaceId ns, CancellationToken cancellationToken = default);
+		Task<long> DeleteBucketAsync(NamespaceId ns, BucketId bucket, CancellationToken cancellationToken = default);
 	}
 
 	public class RefRecord
