Horde: Add remote IP address for obsolete controller warnings.

#rnx

diff --git a/Engine/Source/Programs/Horde/Horde.Server/Startup.cs b/Engine/Source/Programs/Horde/Horde.Server/Startup.cs
--- a/Engine/Source/Programs/Horde/Horde.Server/Startup.cs
+++ b/Engine/Source/Programs/Horde/Horde.Server/Startup.cs
@@ -99,7 +99,6 @@ using Microsoft.AspNetCore.DataProtection;
 using Microsoft.AspNetCore.Hosting;
 using Microsoft.AspNetCore.Hosting.Server.Features;
 using Microsoft.AspNetCore.Http;
-using Microsoft.AspNetCore.Http.Extensions;
 using Microsoft.AspNetCore.HttpOverrides;
 using Microsoft.AspNetCore.Mvc;
 using Microsoft.AspNetCore.Mvc.Controllers;
@@ -311,7 +310,7 @@ namespace Horde.Server
 					|| actionDescriptor?.MethodInfo?.DeclaringType?.GetCustomAttribute<ObsoleteAttribute>() != null)
 				{
 					ILogger? logger = context.HttpContext.RequestServices.GetService<ILogger<ObsoleteLoggingFilter>>();
-					logger?.LogWarning("Using obsolete endpoint: {RequestPath}", context.HttpContext.Request.Path);
+					logger?.LogWarning("Using obsolete endpoint: {RequestPath} (Source: {RemoteIp})", context.HttpContext.Request.Path, context.HttpContext.Connection.RemoteIpAddress);
 				}
 			}
 		}
