Horde: Updating schema documentation

#rnx

diff --git a/Engine/Source/Programs/Horde/Docs/Config/Schema/Dashboard.md b/Engine/Source/Programs/Horde/Docs/Config/Schema/Dashboard.md
new file mode 100644
--- /dev/null
+++ b/Engine/Source/Programs/Horde/Docs/Config/Schema/Dashboard.md
@@ -0,0 +1,128 @@
+[Horde](../../../README.md) > [Configuration](../../Config.md) > *.dashboard.json
+
+# *.dashboard.json
+
+Configuration for dashboard features
+
+Name | Description
+---- | -----------
+`showLandingPage` | `boolean`<br>Navigate to the landing page by default
+`showCI` | `boolean`<br>Enable CI functionality
+`showAgents` | `boolean`<br>Whether to show functionality related to agents, pools, and utilization on the dashboard.
+`showAgentRegistration` | `boolean`<br>Whether to show the agent registration page. When using registration tokens from elsewhere this is not needed.
+`showPerforceServers` | `boolean`<br>Show the Perforce server option on the server menu
+`showDeviceManager` | `boolean`<br>Show the device manager on the server menu
+`showTests` | `boolean`<br>Show automated tests on the server menu
+`agentCategories` | [`DashboardAgentCategoryConfig`](#dashboardagentcategoryconfig)`[]`<br>Configuration for different agent pages
+`poolCategories` | [`DashboardPoolCategoryConfig`](#dashboardpoolcategoryconfig)`[]`<br>Configuration for different pool pages
+`analytics` | [`TelemetryViewConfig`](#telemetryviewconfig)`[]`<br>Configuration for telemetry views
+`include` | [`ConfigInclude`](#configinclude)`[]`<br>Includes for other configuration files
+`macros` | [`ConfigMacro`](#configmacro)`[]`<br>Macros within this configuration
+
+## DashboardAgentCategoryConfig
+
+Configuration for a category of agents
+
+Name | Description
+---- | -----------
+`name` | `string`<br>Name of the category
+`condition` | `string`<br>Condition string to be evaluated for this page
+
+## DashboardPoolCategoryConfig
+
+Configuration for a category of pools
+
+Name | Description
+---- | -----------
+`name` | `string`<br>Name of the category
+`condition` | `string`<br>Condition string to be evaluated for this page
+
+## TelemetryViewConfig
+
+A telemetry view of related metrics, divided into categofies
+
+Name | Description
+---- | -----------
+`id` | `string`<br>Identifier for the view
+`name` | `string`<br>The name of the view
+`telemetryStoreId` | `string`<br>The telemetry store this view uses
+`variables` | [`TelemetryVariableConfig`](#telemetryvariableconfig)`[]`<br>The variables used to filter the view data
+`categories` | [`TelemetryCategoryConfig`](#telemetrycategoryconfig)`[]`<br>The categories contained within the view
+
+## TelemetryVariableConfig
+
+A telemetry view variable used for filtering the charting data
+
+Name | Description
+---- | -----------
+`name` | `string`<br>The name of the variable for display purposes
+`group` | `string`<br>The associated data group attached to the variable
+`defaults` | `string[]`<br>The default values to select
+
+## TelemetryCategoryConfig
+
+A chart categody, will be displayed on the dashbord under an associated pivot
+
+Name | Description
+---- | -----------
+`name` | `string`<br>The name of the category
+`charts` | [`TelemetryChartConfig`](#telemetrychartconfig)`[]`<br>The charts contained within the category
+
+## TelemetryChartConfig
+
+Telemetry chart configuraton
+
+Name | Description
+---- | -----------
+`name` | `string`<br>The name of the chart, will be displayed on the dashboard
+`display` | [`TelemetryMetricUnitType`](#telemetrymetricunittype-enum)<br>The unit to display
+`graph` | [`TelemetryMetricGraphType`](#telemetrymetricgraphtype-enum)<br>The graph type
+`metrics` | [`TelemetryChartMetricConfig`](#telemetrychartmetricconfig)`[]`<br>List of configured metrics
+`min` | `integer`<br>The min unit value for clamping chart
+`max` | `integer`<br>The max unit value for clamping chart
+
+## TelemetryMetricUnitType (Enum)
+
+The units used to present the telemetry
+
+Name | Description
+---- | -----------
+`Time` | Time duration
+`Ratio` | Ratio 0-100%
+`Value` | Artbitrary numeric value
+
+## TelemetryMetricGraphType (Enum)
+
+The type of
+
+Name | Description
+---- | -----------
+`Line` | A line graph
+`Indicator` | Key performance indicator (KPI) chart with thrasholds
+
+## TelemetryChartMetricConfig
+
+Metric attached to a telemetry chart
+
+Name | Description
+---- | -----------
+`id` | `string`<br>Associated metric id
+`threshold` | `integer`<br>The threshold for KPI values
+`alias` | `string`<br>The metric alias for display purposes
+
+## ConfigInclude
+
+Directive to merge config data from another source
+
+Name | Description
+---- | -----------
+`path` | `string`<br>Path to the config data to be included. May be relative to the including file's location.
+
+## ConfigMacro
+
+Declares a config macro
+
+Name | Description
+---- | -----------
+`name` | `string`<br>Name of the macro property
+`value` | `string`<br>Value for the macro property
diff --git a/Engine/Source/Programs/Horde/Docs/Config/Schema/Globals.md b/Engine/Source/Programs/Horde/Docs/Config/Schema/Globals.md
--- a/Engine/Source/Programs/Horde/Docs/Config/Schema/Globals.md
+++ b/Engine/Source/Programs/Horde/Docs/Config/Schema/Globals.md
@@ -9,7 +9,7 @@ Name | Description
 `version` | `integer`<br>Version number for the server. Values are indicated by the .
 `include` | [`ConfigInclude`](#configinclude)`[]`<br>Other paths to include
 `macros` | [`ConfigMacro`](#configmacro)`[]`<br>Macros within the global scope
-`dashboard` | [`DashboardConfig`](#dashboardconfig)<br>Settings for the dashboard
+`dashboard` | [`DashboardConfig`](Dashboard.md)<br>Settings for the dashboard
 `projects` | [`ProjectConfig`](Projects.md)`[]`<br>List of projects
 `pools` | [`PoolConfig`](#poolconfig)`[]`<br>List of pools
 `downtime` | [`ScheduledDowntime`](#scheduleddowntime)`[]`<br>List of scheduled downtime
@@ -25,7 +25,7 @@ Name | Description
 `agentShutdownIfDisabledGracePeriod` | `string`<br>Time to wait before shutting down an agent that has been disabled Used if no value is set on the actual pool.
 `storage` | [`StorageConfig`](#storageconfig)<br>Storage configuration
 `artifactTypes` | [`ArtifactTypeConfig`](#artifacttypeconfig)`[]`<br>Configuration for different artifact types
-`telemetryStores` | [`TelemetryStoreConfig`](#telemetrystoreconfig)`[]`<br>Metrics to aggregate on the Horde server
+`telemetryStores` | [`TelemetryStoreConfig`](Telemetry.md)`[]`<br>Metrics to aggregate on the Horde server
 `parameters` | `object`<br>General parameters for other tools. Can be queried through the api/v1/parameters endpoint.
 `acl` | [`AclConfig`](#aclconfig)<br>Access control list
 
@@ -46,116 +46,6 @@ Name | Description
 `name` | `string`<br>Name of the macro property
 `value` | `string`<br>Value for the macro property
 
-## DashboardConfig
-
-Configuration for dashboard features
-
-Name | Description
----- | -----------
-`showLandingPage` | `boolean`<br>Navigate to the landing page by default
-`showCI` | `boolean`<br>Enable CI functionality
-`showAgents` | `boolean`<br>Whether to show functionality related to agents, pools, and utilization on the dashboard.
-`showAgentRegistration` | `boolean`<br>Whether to show the agent registration page. When using registration tokens from elsewhere this is not needed.
-`showPerforceServers` | `boolean`<br>Show the Perforce server option on the server menu
-`showDeviceManager` | `boolean`<br>Show the device manager on the server menu
-`showTests` | `boolean`<br>Show automated tests on the server menu
-`agentCategories` | [`DashboardAgentCategoryConfig`](#dashboardagentcategoryconfig)`[]`<br>Configuration for different agent pages
-`poolCategories` | [`DashboardPoolCategoryConfig`](#dashboardpoolcategoryconfig)`[]`<br>Configuration for different pool pages
-`analytics` | [`TelemetryViewConfig`](#telemetryviewconfig)`[]`<br>Configuration for telemetry views
-`include` | [`ConfigInclude`](#configinclude)`[]`<br>Includes for other configuration files
-`macros` | [`ConfigMacro`](#configmacro)`[]`<br>Macros within this configuration
-
-## DashboardAgentCategoryConfig
-
-Configuration for a category of agents
-
-Name | Description
----- | -----------
-`name` | `string`<br>Name of the category
-`condition` | `string`<br>Condition string to be evaluated for this page
-
-## DashboardPoolCategoryConfig
-
-Configuration for a category of pools
-
-Name | Description
----- | -----------
-`name` | `string`<br>Name of the category
-`condition` | `string`<br>Condition string to be evaluated for this page
-
-## TelemetryViewConfig
-
-A telemetry view of related metrics, divided into categofies
-
-Name | Description
----- | -----------
-`id` | `string`<br>Identifier for the view
-`name` | `string`<br>The name of the view
-`telemetryStoreId` | `string`<br>The telemetry store this view uses
-`variables` | [`TelemetryVariableConfig`](#telemetryvariableconfig)`[]`<br>The variables used to filter the view data
-`categories` | [`TelemetryCategoryConfig`](#telemetrycategoryconfig)`[]`<br>The categories contained within the view
-
-## TelemetryVariableConfig
-
-A telemetry view variable used for filtering the charting data
-
-Name | Description
----- | -----------
-`name` | `string`<br>The name of the variable for display purposes
-`group` | `string`<br>The associated data group attached to the variable
-`defaults` | `string[]`<br>The default values to select
-
-## TelemetryCategoryConfig
-
-A chart categody, will be displayed on the dashbord under an associated pivot
-
-Name | Description
----- | -----------
-`name` | `string`<br>The name of the category
-`charts` | [`TelemetryChartConfig`](#telemetrychartconfig)`[]`<br>The charts contained within the category
-
-## TelemetryChartConfig
-
-Telemetry chart configuraton
-
-Name | Description
----- | -----------
-`name` | `string`<br>The name of the chart, will be displayed on the dashboard
-`display` | [`TelemetryMetricUnitType`](#telemetrymetricunittype-enum)<br>The unit to display
-`graph` | [`TelemetryMetricGraphType`](#telemetrymetricgraphtype-enum)<br>The graph type
-`metrics` | [`TelemetryChartMetricConfig`](#telemetrychartmetricconfig)`[]`<br>List of configured metrics
-`min` | `integer`<br>The min unit value for clamping chart
-`max` | `integer`<br>The max unit value for clamping chart
-
-## TelemetryMetricUnitType (Enum)
-
-The units used to present the telemetry
-
-Name | Description
----- | -----------
-`Time` | Time duration
-`Ratio` | Ratio 0-100%
-`Value` | Artbitrary numeric value
-
-## TelemetryMetricGraphType (Enum)
-
-The type of
-
-Name | Description
----- | -----------
-`Line` | A line graph
-`Indicator` | Key performance indicator (KPI) chart with thrasholds
-
-## TelemetryChartMetricConfig
-
-Metric attached to a telemetry chart
-
-Name | Description
----- | -----------
-`id` | `string`<br>Associated metric id
-`threshold` | `integer`<br>The threshold for KPI values
-`alias` | `string`<br>The metric alias for display purposes
-
 ## PoolConfig
 
 Mutable configuration for a pool
@@ -578,42 +468,3 @@ Name | Description
 `name` | `string`<br>Legacy 'Name' property
 `keepCount` | `integer`<br>Number of artifacts to retain
 `keepDays` | `integer`<br>Number of days to retain artifacts of this type
-
-## TelemetryStoreConfig
-
-Config for metrics
-
-Name | Description
----- | -----------
-`id` | `string`<br>Identifier for this store
-`acl` | [`AclConfig`](#aclconfig)<br>Permissions for this store
-`metrics` | [`MetricConfig`](#metricconfig)`[]`<br>Metrics to aggregate on the Horde server
-`include` | [`ConfigInclude`](#configinclude)`[]`<br>Includes for other configuration files
-`macros` | [`ConfigMacro`](#configmacro)`[]`<br>Macros within this configuration
-
-## MetricConfig
-
-Configures a metric to aggregate on the server
-
-Name | Description
----- | -----------
-`id` | `string`<br>Identifier for this metric
-`filter` | `string`<br>Filter expression to evaluate to determine which events to include. This query is evaluated against an array.
-`property` | `string`<br>Property to aggregate
-`groupBy` | `string`<br>Property to group by. Specified as a comma-separated list of JSON path expressions.
-`function` | [`AggregationFunction`](#aggregationfunction-enum)<br>How to aggregate samples for this metric
-`percentile` | `integer`<br>For the percentile function, specifies the percentile to measure
-`interval` | `string`<br>Interval for each metric. Supports times such as "2d", "1h", "1h30m", "20s".
-
-## AggregationFunction (Enum)
-
-Method for aggregating samples into a metric
-
-Name | Description
----- | -----------
-`Count` | Count the number of matching elements
-`Min` | Take the minimum value of all samples
-`Max` | Take the maximum value of all samples
-`Sum` | Sum all the reported values
-`Average` | Average all the samples
-`Percentile` | Estimates the value at a certain percentile
diff --git a/Engine/Source/Programs/Horde/Docs/Config/Schema/Telemetry.md b/Engine/Source/Programs/Horde/Docs/Config/Schema/Telemetry.md
new file mode 100644
--- /dev/null
+++ b/Engine/Source/Programs/Horde/Docs/Config/Schema/Telemetry.md
@@ -0,0 +1,98 @@
+[Horde](../../../README.md) > [Configuration](../../Config.md) > *.telemetry.json
+
+# *.telemetry.json
+
+Config for metrics
+
+Name | Description
+---- | -----------
+`id` | `string`<br>Identifier for this store
+`acl` | [`AclConfig`](#aclconfig)<br>Permissions for this store
+`metrics` | [`MetricConfig`](#metricconfig)`[]`<br>Metrics to aggregate on the Horde server
+`include` | [`ConfigInclude`](#configinclude)`[]`<br>Includes for other configuration files
+`macros` | [`ConfigMacro`](#configmacro)`[]`<br>Macros within this configuration
+
+## AclConfig
+
+Parameters to update an ACL
+
+Name | Description
+---- | -----------
+`entries` | [`AclEntryConfig`](#aclentryconfig)`[]`<br>Entries to replace the existing ACL
+`profiles` | [`AclProfileConfig`](#aclprofileconfig)`[]`<br>Defines profiles which allow grouping sets of actions into named collections
+`inherit` | `boolean`<br>Whether to inherit permissions from the parent ACL
+`exceptions` | `string[]`<br>List of exceptions to the inherited setting
+
+## AclEntryConfig
+
+Individual entry in an ACL
+
+Name | Description
+---- | -----------
+`claim` | [`AclClaimConfig`](#aclclaimconfig)<br>Name of the user or group
+`actions` | `string[]`<br>Array of actions to allow
+`profiles` | `string[]`<br>List of profiles to grant
+
+## AclClaimConfig
+
+New claim to create
+
+Name | Description
+---- | -----------
+`type` | `string`<br>The claim type
+`value` | `string`<br>The claim value
+
+## AclProfileConfig
+
+Configuration for an ACL profile. This defines a preset group of actions which can be given to a user via an ACL entry.
+
+Name | Description
+---- | -----------
+`id` | `string`<br>Identifier for this profile
+`actions` | `string[]`<br>Actions to include
+`excludeActions` | `string[]`<br>Actions to exclude from the inherited actions
+`extends` | `string[]`<br>Other profiles to extend from
+
+## MetricConfig
+
+Configures a metric to aggregate on the server
+
+Name | Description
+---- | -----------
+`id` | `string`<br>Identifier for this metric
+`filter` | `string`<br>Filter expression to evaluate to determine which events to include. This query is evaluated against an array.
+`property` | `string`<br>Property to aggregate
+`groupBy` | `string`<br>Property to group by. Specified as a comma-separated list of JSON path expressions.
+`function` | [`AggregationFunction`](#aggregationfunction-enum)<br>How to aggregate samples for this metric
+`percentile` | `integer`<br>For the percentile function, specifies the percentile to measure
+`interval` | `string`<br>Interval for each metric. Supports times such as "2d", "1h", "1h30m", "20s".
+
+## AggregationFunction (Enum)
+
+Method for aggregating samples into a metric
+
+Name | Description
+---- | -----------
+`Count` | Count the number of matching elements
+`Min` | Take the minimum value of all samples
+`Max` | Take the maximum value of all samples
+`Sum` | Sum all the reported values
+`Average` | Average all the samples
+`Percentile` | Estimates the value at a certain percentile
+
+## ConfigInclude
+
+Directive to merge config data from another source
+
+Name | Description
+---- | -----------
+`path` | `string`<br>Path to the config data to be included. May be relative to the including file's location.
+
+## ConfigMacro
+
+Declares a config macro
+
+Name | Description
+---- | -----------
+`name` | `string`<br>Name of the macro property
+`value` | `string`<br>Value for the macro property
