Horde: Add analytics and dashboard to generated schema docs
#rnx

diff --git a/Engine/Source/Programs/Horde/Docs/Config.md b/Engine/Source/Programs/Horde/Docs/Config.md
--- a/Engine/Source/Programs/Horde/Docs/Config.md
+++ b/Engine/Source/Programs/Horde/Docs/Config.md
@@ -29,3 +29,5 @@ This section targets operating and managing Horde installations and shares some
 * [Globals.json Reference](Config/Schema/Globals.md)
 * [*.project.json Reference](Config/Schema/Projects.md)
 * [*.stream.json Reference](Config/Schema/Streams.md)
+* [*.telemetry.json Reference](Config/Schema/Telemetry.md)
+* [*.dashboard.json Reference](Config/Schema/Dashboard.md)
diff --git a/Engine/Source/Programs/Horde/Horde.Server/Commands/Generate/ConfigDocsCommand.cs b/Engine/Source/Programs/Horde/Horde.Server/Commands/Generate/ConfigDocsCommand.cs
--- a/Engine/Source/Programs/Horde/Horde.Server/Commands/Generate/ConfigDocsCommand.cs
+++ b/Engine/Source/Programs/Horde/Horde.Server/Commands/Generate/ConfigDocsCommand.cs
@@ -12,9 +12,11 @@ using System.Threading.Tasks;
 using System.Xml;
 using EpicGames.Core;
 using EpicGames.Horde.Acls;
+using Horde.Server.Dashboard;
 using Horde.Server.Projects;
 using Horde.Server.Server;
 using Horde.Server.Streams;
+using Horde.Server.Telemetry.Metrics;
 using Microsoft.Extensions.Logging;
 
 namespace Horde.Server.Commands.Generate
@@ -44,6 +46,8 @@ namespace Horde.Server.Commands.Generate
 			JsonSchema globalSchema = JsonSchema.FromType(typeof(GlobalConfig), xmlDocReader);
 			JsonSchema projectSchema = JsonSchema.FromType(typeof(ProjectConfig), xmlDocReader);
 			JsonSchema streamSchema = JsonSchema.FromType(typeof(StreamConfig), xmlDocReader);
+			JsonSchema telemetryConfigSchema = JsonSchema.FromType(typeof(TelemetryStoreConfig), xmlDocReader);
+			JsonSchema dashboardConfigSchema = JsonSchema.FromType(typeof(DashboardConfig), xmlDocReader);
 
 			Dictionary<JsonSchemaType, PageInfo> typeToPageInfo = new Dictionary<JsonSchemaType, PageInfo>
 			{
@@ -51,6 +55,8 @@ namespace Horde.Server.Commands.Generate
 				[globalSchema.RootType] = new PageInfo("Globals.json", "[Horde](../../../README.md) > [Configuration](../../Config.md)", "Config/Schema/Globals.md"),
 				[projectSchema.RootType] = new PageInfo("*.project.json", "[Horde](../../../README.md) > [Configuration](../../Config.md)", "Config/Schema/Projects.md"),
 				[streamSchema.RootType] = new PageInfo("*.stream.json", "[Horde](../../../README.md) > [Configuration](../../Config.md)", "Config/Schema/Streams.md"),
+				[telemetryConfigSchema.RootType] = new PageInfo("*.telemetry.json", "[Horde](../../../README.md) > [Configuration](../../Config.md)", "Config/Schema/Telemetry.md"),
+				[dashboardConfigSchema.RootType] = new PageInfo("*.dashboard.json", "[Horde](../../../README.md) > [Configuration](../../Config.md)", "Config/Schema/Dashboard.md"),
 			};
 
 			if (Agent == null)
