Horde: Fix OIDC issues when using Horde auth with UBA.

diff --git a/Engine/Source/Programs/Horde/Horde.Server/Accounts/OAuthController.cs b/Engine/Source/Programs/Horde/Horde.Server/Accounts/OAuthController.cs
--- a/Engine/Source/Programs/Horde/Horde.Server/Accounts/OAuthController.cs
+++ b/Engine/Source/Programs/Horde/Horde.Server/Accounts/OAuthController.cs
@@ -59,7 +59,7 @@ namespace Horde.Server.Accounts
 
 			return Ok(new
 			{
-				issuer = _serverSettings.JwtIssuer,
+				issuer = new Uri(_serverSettings.ServerUrl, "api/v1/oauth2"),
 				authorization_endpoint = new Uri(_serverSettings.ServerUrl, "api/v1/oauth2/authorize"),
 				token_endpoint = new Uri(_serverSettings.ServerUrl, "api/v1/oauth2/token"),
 				userinfo_endpoint = new Uri(_serverSettings.ServerUrl, "api/v1/oauth2/userinfo"),
@@ -241,47 +241,14 @@ namespace Horde.Server.Accounts
 		{
 			// Validate the supplied token
 			IGlobals globals = await _globalsService.GetAsync(cancellationToken);
-			TokenValidationParameters validationParameters = new TokenValidationParameters()
-			{
-				ValidAudience = _serverSettings.JwtIssuer,
-				ValidIssuer = _serverSettings.JwtIssuer,
-				IssuerSigningKey = globals.RsaSigningKey
-			};
-
-			JwtSecurityTokenHandler tokenHandler = new JwtSecurityTokenHandler();
-
-			TokenValidationResult result = await tokenHandler.ValidateTokenAsync(request.AuthorizationToken, validationParameters);
-			if (!result.IsValid)
-			{
-				throw result.Exception;
-			}
-
-			JwtPayload payload = ((JwtSecurityToken)result.SecurityToken).Payload;
-
-			// Get the matching account and check the session key is still valid
-			AccountId accountId;
-			if (!TryParseAccountIdFromSubject(payload, out accountId))
-			{
-				return Unauthorized("Missing account-id subject");
-			}
-
-			IAccount? account = await _accountCollection.GetAsync(accountId, cancellationToken);
-			if (account == null)
-			{
-				return Unauthorized($"Invalid account-id ({accountId})");
-			}
 
-			string? session = GetClaimOrDefault(payload, SessionClaim);
-			if (session != account.SessionKey)
-			{
-				return Unauthorized($"Invalid session key ('{session}')");
-			}
-
-			// Check the supplied token is correct for the operation we're performing
-			string? purpose = GetClaimOrDefault(payload, PurposeClaim);
+			JwtPayload payload;
 			if (request.GrantType == "authorization_code")
 			{
-				// Validate that we have an auth code
+				payload = await ValidateTokenAsync(globals, request.AuthorizationToken);
+
+				// Check the supplied token is correct for the operation we're performing
+				string? purpose = GetClaimOrDefault(payload, PurposeClaim);
 				if (purpose != AuthorizationCodePurpose)
 				{
 					return BadRequest($"Expected authorization token, not {purpose}");
@@ -302,7 +269,10 @@ namespace Horde.Server.Accounts
 			}
 			else if (request.GrantType == "refresh_token")
 			{
+				payload = await ValidateTokenAsync(globals, request.RefreshToken);
+
 				// Validate that we have a refresh token
+				string? purpose = GetClaimOrDefault(payload, PurposeClaim);
 				if (purpose != RefreshTokenPurpose)
 				{
 					return BadRequest($"Expected refresh token, not {purpose}");
@@ -313,6 +283,25 @@ namespace Horde.Server.Accounts
 				return BadRequest($"Unsupported grant type: '{request.GrantType}'");
 			}
 
+			// Get the matching account and check the session key is still valid
+			AccountId accountId;
+			if (!TryParseAccountIdFromSubject(payload, out accountId))
+			{
+				return Unauthorized("Missing account-id subject");
+			}
+
+			IAccount? account = await _accountCollection.GetAsync(accountId, cancellationToken);
+			if (account == null)
+			{
+				return Unauthorized($"Invalid account-id ({accountId})");
+			}
+
+			string? session = GetClaimOrDefault(payload, SessionClaim);
+			if (session != account.SessionKey)
+			{
+				return Unauthorized($"Invalid session key ('{session}')");
+			}
+
 			// Create the tokens and response object
 			Response.Headers.CacheControl = "no-cache";
 
@@ -328,6 +317,26 @@ namespace Horde.Server.Accounts
 			return response;
 		}
 
+		async Task<JwtPayload> ValidateTokenAsync(IGlobals globals, string? token)
+		{
+			TokenValidationParameters validationParameters = new TokenValidationParameters()
+			{
+				ValidAudience = _serverSettings.JwtIssuer,
+				ValidIssuer = _serverSettings.JwtIssuer,
+				IssuerSigningKey = globals.RsaSigningKey
+			};
+
+			JwtSecurityTokenHandler tokenHandler = new JwtSecurityTokenHandler();
+
+			TokenValidationResult result = await tokenHandler.ValidateTokenAsync(token, validationParameters);
+			if (!result.IsValid)
+			{
+				throw result.Exception;
+			}
+
+			return ((JwtSecurityToken)result.SecurityToken).Payload;
+		}
+
 		/// <summary>
 		/// Get information about the logged in user
 		/// </summary>
diff --git a/Engine/Source/Programs/Horde/Horde.Server/Accounts/OAuthMessages.cs b/Engine/Source/Programs/Horde/Horde.Server/Accounts/OAuthMessages.cs
--- a/Engine/Source/Programs/Horde/Horde.Server/Accounts/OAuthMessages.cs
+++ b/Engine/Source/Programs/Horde/Horde.Server/Accounts/OAuthMessages.cs
@@ -83,6 +83,12 @@ namespace Horde.Server.Accounts
 		[FromForm(Name = "code")]
 		public string? AuthorizationToken { get; set; }
 
+		/// <summary>
+		/// Refresh token specified in the redirect.
+		/// </summary>
+		[FromForm(Name = "refresh_token")]
+		public string? RefreshToken { get; set; }
+
 		/// <summary>
 		/// 
 		/// </summary>
diff --git a/Engine/Source/Programs/Horde/Horde.Server/Server/ServerController.cs b/Engine/Source/Programs/Horde/Horde.Server/Server/ServerController.cs
--- a/Engine/Source/Programs/Horde/Horde.Server/Server/ServerController.cs
+++ b/Engine/Source/Programs/Horde/Horde.Server/Server/ServerController.cs
@@ -140,8 +140,16 @@ namespace Horde.Server.Server
 			GetAuthConfigResponse response = new GetAuthConfigResponse();
 			response.Method = settings.AuthMethod;
 			response.ProfileName = settings.OidcProfileName;
-			response.ServerUrl = settings.OidcAuthority;
-			response.ClientId = settings.OidcClientId;
+			if (settings.AuthMethod == AuthMethod.Horde)
+			{
+				response.ServerUrl = new Uri(_globalConfig.Value.ServerSettings.ServerUrl, "api/v1/oauth2").ToString();
+				response.ClientId = "default";
+			}
+			else
+			{
+				response.ServerUrl = settings.OidcAuthority;
+				response.ClientId = settings.OidcClientId;
+			}
 			response.LocalRedirectUrls = settings.OidcLocalRedirectUrls;
 			return response;
 		}
diff --git a/Engine/Source/Programs/Horde/Horde.Server/Startup.cs b/Engine/Source/Programs/Horde/Horde.Server/Startup.cs
--- a/Engine/Source/Programs/Horde/Horde.Server/Startup.cs
+++ b/Engine/Source/Programs/Horde/Horde.Server/Startup.cs
@@ -785,12 +785,8 @@ namespace Horde.Server
 				case AuthMethod.Horde:
 					authBuilder.AddHordeOpenId(settings, OpenIdConnectDefaults.AuthenticationScheme, OpenIdConnectDefaults.DisplayName, options =>
 					{
-						options.Authority = "/api/v1/oauth2";
+						options.Authority = new Uri(settings.ServerUrl, "api/v1/oauth2").ToString();
 						options.ClientId = "default";
-						if (settings.HttpsPort == 0)
-						{
-							options.RequireHttpsMetadata = false;
-						}
 						foreach (string scope in settings.OidcRequestedScopes)
 						{
 							options.Scope.Add(scope);
