Horde: Add device manage and pool settings to example packaged lyra build
#rnx

diff --git a/Engine/Source/Programs/Horde/Horde.Server/Defaults/ue5-release-5.4.stream.json b/Engine/Source/Programs/Horde/Horde.Server/Defaults/ue5-release-5.4.stream.json
--- a/Engine/Source/Programs/Horde/Horde.Server/Defaults/ue5-release-5.4.stream.json
+++ b/Engine/Source/Programs/Horde/Horde.Server/Defaults/ue5-release-5.4.stream.json
@@ -466,7 +466,19 @@
 							"default": false
 						}
 					]
-				}
+				},
+				{
+					"type": "Text",
+					"label": "Device Manager URL",
+					"default": "http://localhost:13340",
+					"argument": "-GauntletDeviceService="
+				},
+				{
+					"type": "Text",
+					"label": "Device Pool",
+					"default": "UE5",
+					"argument": "-GauntletDevicePool="
+				}				
 			]
 		},
 		{
