Horde: Add more code padding to markdown render
#rnx

diff --git a/Engine/Source/Programs/Horde/HordeDashboard/src/base/components/Markdown/Markdown.tsx b/Engine/Source/Programs/Horde/HordeDashboard/src/base/components/Markdown/Markdown.tsx
--- a/Engine/Source/Programs/Horde/HordeDashboard/src/base/components/Markdown/Markdown.tsx
+++ b/Engine/Source/Programs/Horde/HordeDashboard/src/base/components/Markdown/Markdown.tsx
@@ -52,6 +52,16 @@ const getStyles: IStyleFunction<IMarkdownStyleProps, IMarkdownStyles> = () => {
 
                padding: "3px 6px !important"
 
+            },
+            'a > code': {
+
+               padding: "3px 6px !important"
+
+            },
+            'td > code': {
+
+               padding: "3px 6px !important"
+
             },
             'code': {
                whiteSpace: "pre-wrap",
