Horde: Add BuildGraph target for pushing Docker image bundled with tools

Also removes outdated SetUpdateChannel call during Horde Agent install.

diff --git a/Engine/Source/Programs/Horde/BuildHorde.xml b/Engine/Source/Programs/Horde/BuildHorde.xml
--- a/Engine/Source/Programs/Horde/BuildHorde.xml
+++ b/Engine/Source/Programs/Horde/BuildHorde.xml
@@ -1,11 +1,21 @@
 <?xml version='1.0' ?>
 <BuildGraph xmlns="http://www.epicgames.com/BuildGraph" xmlns:xsi="http://www.w3.org/2001/XMLSchema-instance" xsi:schemaLocation="http://www.epicgames.com/BuildGraph ../../../../Engine/Build/Graph/Schema.xsd" >
-
+	<EnvVar Name="UE_HORDE_JOBID"/>
+	
 	<Option Name="PreflightChange" DefaultValue="" Description="Preflight changelist number"/>
 	<Option Name="Configuration" DefaultValue="release" Description="Configuration to build"/>
 	<Option Name="SignOutputs" Restrict="true|false" DefaultValue="false" Description="Whether to sign the agent included with server tools"/>
 	<Option Name="HordeDashboardConfig" DefaultValue="Production" Description="Whether to build a development or production dashboard"/>
 
+	<Property Name="DefaultImageTag" Value="Unversioned"/>
+	<Property Name="DefaultImageTag" Value="CL-$(Change)" If="'$(Change)' != ''"/>
+	<Property Name="DefaultImageTag" Value="$(DefaultImageTag)-PF-$(PreflightChange)" If="'$(PreflightChange)' != ''" />
+	<Property Name="DefaultImageTag" Value="$(DefaultImageTag)-$(UE_HORDE_JOBID)" If="'$(UE_HORDE_JOBID)' != ''" />
+	<Option Name="ImageRepoUrl" DefaultValue="" Description="URL to Docker/OCI image repository"/>
+	<Option Name="ImageName" DefaultValue="horde-server" Description="Name of Docker/OCI image"/>
+	<Option Name="ImageTag" DefaultValue="$(DefaultImageTag)" Description="Tag to apply to image"/>
+	<Option Name="ImageEnvFile" DefaultValue="" Description="Path to environment file for used for Docker invocations"/>
+
 	<Property Name="EngineDir" Value="$(RootDir)/Engine"/>
 
 	<Property Name="WixDir" Value="$(RootDir)/Engine/Source/ThirdParty/WiX/3.8"/>
@@ -118,9 +128,6 @@
 			<DotNet Arguments="publish &quot;$(RootDir)/Engine/Source/Programs/Horde/Horde.Agent/Horde.Agent.csproj&quot; --output &quot;$(ServerToolsStagingDir_HordeAgentWin64)&quot; --runtime win7-x64 --self-contained $(VersionArguments)"/>
 			<Sign Files="$(ServerToolsStagingDir_HordeAgentWin64)/HordeAgent.exe;$(ServerToolsStagingDir_HordeAgentWin64)/HordeAgent.dll;$(ServerToolsStagingDir_HordeAgentWin64)/EpicGames.*.dll" If="'$(SignOutputs)' == 'true'"/>
 
-			<!-- Set the default update channel for the win64 agent-->
-			<Spawn Exe="$(ServerToolsStagingDir_HordeAgentWin64)/HordeAgent.exe" Arguments="SetUpdateChannel -Id=horde-agent-win64"/>
-
 			<!-- Build the tray app -->
 			<DotNet Arguments="publish &quot;$(RootDir)/Engine/Source/Programs/Horde/Horde.Agent.TrayApp/Horde.Agent.TrayApp.csproj&quot; --output &quot;$(ServerToolsStagingDir_HordeAgentWin64)/TrayApp&quot; --runtime win7-x64 --self-contained $(VersionArguments)"/>
 			<Sign Files="$(ServerToolsStagingDir_HordeAgentWin64)/TrayApp/Horde.Agent.TrayApp.exe;$(ServerToolsStagingDir_HordeAgentWin64)/TrayApp/EpicGames.*.dll" If="'$(SignOutputs)' == 'true'"/>
@@ -219,7 +226,7 @@
 
 		</Node>
 
-		<Node Name="Stage Horde Server Installer" Requires="#DashboardOutputFiles;#HordeCmdFiles;#HordeAgentFiles;#HordeAgentMsiFiles;#P4VUtilsFiles;#UnrealGameSyncFiles;#UnrealGameSyncCmdFiles;#UnrealGameSyncMsiFiles;">
+		<Node Name="Stage Horde Server Installer" Requires="#DashboardOutputFiles;#HordeCmdFiles;#HordeAgentFiles;#HordeAgentMsiFiles;#P4VUtilsFiles;#UnrealGameSyncFiles;#UnrealGameSyncCmdFiles;#UnrealGameSyncMsiFiles;" Produces="#ServerWithToolsFiles">
 		
 			<!-- Stage the server -->
 			<Delete Files="$(StagingDir)/Server/..."/>
@@ -248,6 +255,7 @@
 			<!-- Copy the VCREDIST files -->
 			<Copy From="$(StagingDir)/Server/ThirdParty/Mongo/*.dll" To="$(StagingDir)/Server"/>
 
+			<Tag Files="$(StagingDir)/Server/..." With="#ServerWithToolsFiles" />
 		</Node>
 
 		<Node Name="Build Horde Server Installer" Requires="Stage Horde Server Installer">
@@ -277,13 +285,21 @@
 			<Copy Files="UnrealHordeServer.msi" From="$(InstallerOutputDir)" To="$(StagingDir)/Server-Installer"/>
 
 		</Node>
+	</Agent>
 
-		<Node Name="Build Bundled Docker Image" Requires="Stage Horde Server Installer;Build HordeServer">
+	<Agent Name="HordeInstallerLinux" Type="Linux">
+		<Property Name="StagingDir" Value="$(RootDir)/Staging"/>
+		
+		<Node Name="Build Bundled Docker Image" Requires="Stage Horde Server Installer;Build HordeServer;#ServerWithToolsFiles">
 			<!-- A Horde server image has already been built, and is referenced inside Dockerfile.bundled -->
 			<Copy From="Engine/Source/Programs/Horde/Horde.Server/Dockerfile.bundled" To="$(StagingDir)/Dockerfile.bundled" />
 			<Tag Files="$(StagingDir)/Dockerfile*;$(StagingDir)/Dashboard/...;$(StagingDir)/Server/Tools/...;$(StagingDir)/Server/*.json" With="#DockerFiles"/>
 			<Docker-Build BaseDir="$(StagingDir)" Files="#DockerFiles" Tag="horde-server:bundled" DockerFile="$(StagingDir)/Dockerfile.bundled" />
 		</Node>
+
+		<Node Name="Push Bundled Docker Image" Requires="Build Bundled Docker Image">
+			<Docker-Push Repository="$(ImageRepoUrl)" Image="horde-server:bundled" TargetImage="$(ImageName):$(ImageTag)-bundled" EnvironmentFile="$(ImageEnvFile)" />
+		</Node>
 	</Agent>
 
 </BuildGraph>
