Horde: Fix devkit pool
#rnx

diff --git a/Engine/Source/Programs/Horde/Horde.Server/Defaults/ue5-release-5.4.stream.json b/Engine/Source/Programs/Horde/Horde.Server/Defaults/ue5-release-5.4.stream.json
--- a/Engine/Source/Programs/Horde/Horde.Server/Defaults/ue5-release-5.4.stream.json
+++ b/Engine/Source/Programs/Horde/Horde.Server/Defaults/ue5-release-5.4.stream.json
@@ -179,6 +179,11 @@
 		{
 			"pool": "win-ue5",
 			"workspace": "IncrementalCompileOnly"
+		},
+		"DevkitAutomation":
+		{
+			"pool": "win-devkitautomation",
+			"workspace": "Full"			
 		}
 	},
 	"workspaceTypes":
