Horde: Fix config
#rnx

diff --git a/Engine/Source/Programs/Horde/Horde.Server/Defaults/ue5-release-5.4.stream.json b/Engine/Source/Programs/Horde/Horde.Server/Defaults/ue5-release-5.4.stream.json
--- a/Engine/Source/Programs/Horde/Horde.Server/Defaults/ue5-release-5.4.stream.json
+++ b/Engine/Source/Programs/Horde/Horde.Server/Defaults/ue5-release-5.4.stream.json
@@ -476,13 +476,13 @@
 					"type": "Text",
 					"label": "Device Manager URL",
 					"default": "http://localhost:13340",
-					"argument": "-GauntletDeviceService="
+					"argument": "-set:GauntletDeviceService="
 				},
 				{
 					"type": "Text",
 					"label": "Device Pool",
 					"default": "UE5",
-					"argument": "-GauntletDevicePool="
+					"argument": "-set:GauntletDevicePool="
 				}				
 			]
 		},
