Horde: Try not using reservation block on test installaion

diff --git a/Engine/Source/Programs/Horde/Horde.Server/Defaults/ue5-release-5.4.stream.json b/Engine/Source/Programs/Horde/Horde.Server/Defaults/ue5-release-5.4.stream.json
--- a/Engine/Source/Programs/Horde/Horde.Server/Defaults/ue5-release-5.4.stream.json
+++ b/Engine/Source/Programs/Horde/Horde.Server/Defaults/ue5-release-5.4.stream.json
@@ -249,7 +249,8 @@
 				"-Target=BuildAndTest $(Project)",
 				"-set:ProjectFile=$(ProjectPath)",
 				"-set:TargetName=$(GameTarget)",
-				"-set:EditorTargetName=$(EditorTarget)"
+				"-set:EditorTargetName=$(EditorTarget)",
+				"-set:WithDeviceReservation=false"
 			],
 			"parameters":
 			[
