Horde: Use a pseudo-random tool id for bundled tools, derived from the tool id and version number.

#rnx

diff --git a/Engine/Source/Programs/Horde/Horde.Server/Tools/ToolCollection.cs b/Engine/Source/Programs/Horde/Horde.Server/Tools/ToolCollection.cs
--- a/Engine/Source/Programs/Horde/Horde.Server/Tools/ToolCollection.cs
+++ b/Engine/Source/Programs/Horde/Horde.Server/Tools/ToolCollection.cs
@@ -1,12 +1,15 @@
 // Copyright Epic Games, Inc. All Rights Reserved.
 
 using System;
+using System.Buffers.Binary;
 using System.Collections.Generic;
 using System.IO;
 using System.Linq;
+using System.Text;
 using System.Threading;
 using System.Threading.Tasks;
 using EpicGames.Core;
+using EpicGames.Horde;
 using EpicGames.Horde.Storage;
 using EpicGames.Horde.Storage.Backends;
 using EpicGames.Horde.Storage.Bundles;
@@ -225,7 +228,9 @@ namespace Horde.Server.Tools
 			{
 				Tool tool = new Tool(bundledToolConfig);
 
-				ToolDeployment deployment = new ToolDeployment(default);
+				ToolDeploymentId deploymentId = GetDeploymentId(bundledToolConfig);
+
+				ToolDeployment deployment = new ToolDeployment(deploymentId);
 				deployment.Version = bundledToolConfig.Version;
 				deployment.State = ToolDeploymentState.Complete;
 				deployment.RefName = bundledToolConfig.RefName;
@@ -237,6 +242,19 @@ namespace Horde.Server.Tools
 			return null;
 		}
 
+		static ToolDeploymentId GetDeploymentId(BundledToolConfig bundledToolConfig)
+		{
+			// Create a pseudo-random deployment id from the tool id and version
+			IoHash hash = IoHash.Compute(Encoding.UTF8.GetBytes($"{bundledToolConfig.Id}/{bundledToolConfig.Version}"));
+
+			Span<byte> bytes = stackalloc byte[IoHash.NumBytes];
+			hash.CopyTo(bytes);
+
+			// Set a valid timestamp to allow it to appear as an ObjectId
+			BinaryPrimitives.WriteInt32BigEndian(bytes, 1446492960);
+			return new ToolDeploymentId(new BinaryId(bytes));
+		}
+
 		/// <summary>
 		/// Adds a new deployment to the given tool. The new deployment will replace the current active deployment.
 		/// </summary>
