Horde: Add tutorial for setting up auth.

diff --git a/Engine/Source/Programs/Horde/Docs/Deployment/Server.md b/Engine/Source/Programs/Horde/Docs/Deployment/Server.md
--- a/Engine/Source/Programs/Horde/Docs/Deployment/Server.md
+++ b/Engine/Source/Programs/Horde/Docs/Deployment/Server.md
@@ -182,7 +182,13 @@ Settings for port usage are defined in [Server.json](ServerSettings.md):
 To serve data over HTTPS, set the `HttpsPort` property in the [Server.json](ServerSettings.md) file. The server
 machine must have a correctly configured certificate to serve data correctly.
 
-Both HTTP/1.1 and HTTP/2.0 traffic can be served over the `HttpsPort`.
+Both HTTP/1.1 and HTTP/2.0 traffic can be served over the `HttpsPort`. Unencrypted traffic can be disabled by
+setting `HttpPort` and `Http2Port` to zero.
+
+There are occasions where the server provides links back to itself (the OIDC discovery document used when using
+Horde's internal account system, for example), and it's important that these URLs match the HTTPS certificate.
+By default, this URL is derived from the server's reported DNS name, but this can be overridden through the
+`ServerUrl` property.
 
 ### Monitoring
 
@@ -199,13 +205,16 @@ telemetry capture are [listed here](ServerSettings.md#opentelemetrysettings).
 In order to separate lighter request traffic from heavier background operations, the Horde server can be configured to
 run in different RunModes. You can configure these via the [RunMode](ServerSettings.md) setting.
 
-### Authentication
+### Authorization
 
 Horde supports [OpenID Connect (OIDC)](https://openid.net/developers/how-connect-works/) for authentication using
 an external identity provider. OIDC is a widely used auth standard, and Okta, AWS, Azure, Google, Facebook, and
 many others implement identity providers compatible with it.
 
-The following settings in [Server.json](ServerSettings.md) are required to configure an OIDC provider:
+The [Getting Started > Authorization](../Tutorials/Authorization.md) page explains how to configure Horde's internal
+account system and OIDC provider.
+
+The following settings in [Server.json](ServerSettings.md) are required to configure an external OIDC provider:
 
 * `AuthMethod`: Set this to `OpenIdConnect`.
 * `OidcAuthority`: URL of the OIDC authority. You can check the URL specified here is correct by navigating to
diff --git a/Engine/Source/Programs/Horde/Docs/Glossary.md b/Engine/Source/Programs/Horde/Docs/Glossary.md
--- a/Engine/Source/Programs/Horde/Docs/Glossary.md
+++ b/Engine/Source/Programs/Horde/Docs/Glossary.md
@@ -13,6 +13,24 @@ matching a set of configured criteria for a particular time interval.
 * `Telemetry` - Individual events sent by an application to the Horde server. Horde processes telemetry events as
 schema-less JSON objects.
 
+## Authentication
+
+* `Acl` - Abbreviation for **access control list**. An ACL in Horde defines rules that grant *entitlements* to users
+(ie. the ability to perform certain *actions*) based on *claims* that their user has.
+* `Acl Action` - The ability to perform a certain operation on an entity in Horde, such as `ViewJob` or `DownloadTool`.
+See [Config > Schema > AclActions](Config/Schema/AclActions.md) for a full list.
+* `Acl Scope` - Used to refer to a layer within Horde's hierarchical permissions system. Jobs are an acl scope within a stream,
+which are an acl scope within a project, which is an acl scope within the global permissions scope.
+Entitlements to perform actions are typically inherited from parent to child scopes, unless explicitly forbidden by
+setting the `inherit` property in an [AclConfig](Config/Schema/Globals.md#aclconfig) to false.
+* `Acl Profile` - A list of *actions* that can be given to a user, without having to list each individual action each
+time. Similar to a macro.
+* `Claims` - Key/value string pair that makes a statement about that user. Claims are part of the OAuth2 standard, and
+identity providers maintain a list of unique claims for a user. To prevent collisions between different applications,
+keys typically use a URI with an owned domain name to ensure global uniqueness - though the URI does not identify an 
+actual web resource. Horde's internally-issued claims all begin with `http://epicgames.com/ue/horde`.
+* `Entitlement` - The ability of a user to perform a certain *action*.
+
 ## Storage
 
 * `Namespace` - A logical partition of the storage system that can have custom permissions, behaviors, and garbage collection policies.
diff --git a/Engine/Source/Programs/Horde/Docs/Home.md b/Engine/Source/Programs/Horde/Docs/Home.md
--- a/Engine/Source/Programs/Horde/Docs/Home.md
+++ b/Engine/Source/Programs/Horde/Docs/Home.md
@@ -8,7 +8,8 @@
 * **[Setting up test automation](Tutorials/TestAutomation.md)**
 * **[Setting up analytics for your team](Tutorials/Analytics.md)**
 * **[Setting up mobile and console devices](Tutorials/DeviceManager.md)**
-* **[Setting up UnrealGameSync and getting Unreal Editor builds to content creators](Tutorials/UnrealGameSync.md)**
+* **[Setting up UnrealGameSync and distributing Unreal Editor to content creators](Tutorials/UnrealGameSync.md)**
+* **[Setting up authorization](Tutorials/Authorization.md)**
 
 ## Status
 
diff --git a/Engine/Source/Programs/Horde/Docs/Landing.md b/Engine/Source/Programs/Horde/Docs/Landing.md
--- a/Engine/Source/Programs/Horde/Docs/Landing.md
+++ b/Engine/Source/Programs/Horde/Docs/Landing.md
@@ -10,7 +10,8 @@
 * **[Setting up test automation](Tutorials/TestAutomation.md)**
 * **[Setting up analytics for your team](Tutorials/Analytics.md)**
 * **[Setting up mobile and console devices](Tutorials/DeviceManager.md)**
-* **[Setting up UnrealGameSync and getting Unreal Editor builds to content creators](Tutorials/UnrealGameSync.md)**
+* **[Setting up UnrealGameSync and distributing Unreal Editor to content creators](Tutorials/UnrealGameSync.md)**
+* **[Setting up authorization](Tutorials/Authorization.md)**
 
 Alternatively, check out the full documentation through the `Help` menu at the top right of the page.
 
diff --git a/Engine/Source/Programs/Horde/Docs/Tutorials/Authorization.md b/Engine/Source/Programs/Horde/Docs/Tutorials/Authorization.md
new file mode 100644
--- /dev/null
+++ b/Engine/Source/Programs/Horde/Docs/Tutorials/Authorization.md
@@ -0,0 +1,60 @@
+[Horde](../../README.md) > Getting Started: Authorization
+
+# Getting Started: Authorization
+
+## Introduction
+
+Horde ships with auth disabled by default to make it easier to demonstrate and experiment with. Most production
+deployments will probably want users to log in, and restrict the actions they can perform based on their role.
+
+To do this, Horde supports **[OAuth2](https://oauth.net/2/)** and **[OIDC](https://openid.net/developers/how-connect-works/)**,
+which is supported by most third party identity providers - including Okta, AWS, Azure, and Google.
+Configuring an external identity provider is out of scope for this documentation, though the relevant configuration
+points are touched on in the [Deployment > Server](../Deployment/Server.md#authentication) page.
+
+If you don't have an existing OIDC-compatible identity provider, Horde includes it's own - which this guide covers.
+
+## Prerequisites
+
+* Horde Server installation (see [Getting Started: Install Horde](InstallHorde.md)).
+* A valid certificate, and [HTTPS support](../Deployment/Server.md#https) enabled on your server.
+
+## Steps
+
+1. In your [server.json](../Config/Orientation.md) file set the `AuthMode` property to `Horde`, and restart the server.
+2. The first time you launch the server, you'll be prompted to enter an administrator password.
+3. After logging in, there will be an `Accounts` menu item in the `Server` menu. From here, you can manage the users
+allowed to log in to the server, and the [**claims**](../Glossary.md#authentication) that they have. Horde's account
+system uses the `http://epicgames.com/ue/horde/group` claim for groups that a user belongs to, and the dashboard will
+suggest and autocomplete any groups found in the deployment's configuration files.
+
+There are two standard groups defined in the server's `default.globals.json` file, which is included from the standard
+`globals.json` file by default: `View` and `Run`.
+
+   ```json
+    "acl": {
+        "entries": [
+            {
+                "claim": {
+                    "type": "http://epicgames.com/ue/horde/group", 
+                    "value": "View"
+                },
+                "profiles": [
+                    "default-read"
+                ]
+            },
+            {
+                "claim": {
+                    "type": "http://epicgames.com/ue/horde/group", 
+                    "value": "Run"
+                },
+                "profiles": [
+                    "default-run"
+                ]
+            }
+        ]
+    }
+   ```
+
+The `default-read` and `default-run` profiles are defined in code (`AclConfig.cs`). You can define your own profiles within
+the `profiles` element of each [AclConfig](../Config/Schema/Globals.md#aclconfig) object.
diff --git a/Engine/Source/Programs/Horde/README.md b/Engine/Source/Programs/Horde/README.md
--- a/Engine/Source/Programs/Horde/README.md
+++ b/Engine/Source/Programs/Horde/README.md
@@ -39,7 +39,8 @@ See also: [Feature Status](Docs/Features.md)
 * **[Setting up test automation](Docs/Tutorials/TestAutomation.md)**
 * **[Setting up studio analytics](Docs/Tutorials/Analytics.md)**
 * **[Setting up mobile and console devices](Docs/Tutorials/DeviceManager.md)**
-* **[Setting up UnrealGameSync and getting Unreal Editor builds to content creators](Tutorials/UnrealGameSync.md)**
+* **[Setting up UnrealGameSync and distributing Unreal Editor to content creators](Docs/Tutorials/UnrealGameSync.md)**
+* **[Setting up authorization](Docs/Tutorials/Authorization.md)**
 
 ## Reference
 
