Horde Analytics Dashbaord Fixes

diff --git a/Engine/Source/Programs/Horde/Horde.Server/Defaults/default-analytics.dashboard.json b/Engine/Source/Programs/Horde/Horde.Server/Defaults/default-analytics.dashboard.json
--- a/Engine/Source/Programs/Horde/Horde.Server/Defaults/default-analytics.dashboard.json
+++ b/Engine/Source/Programs/Horde/Horde.Server/Defaults/default-analytics.dashboard.json
@@ -85,8 +85,8 @@
 									"threshold": 360
 								},
 								{
-									"id": "core-loading-user-editorstartuptime-percentile",
-									"alias": "Editor Startup",
+									"id": "core-loading-user-editorboottime-percentile",
+									"alias": "Editor Boot",
 									"threshold": 120
 								},
 								{
@@ -113,7 +113,7 @@
 								},
 								{
 									"id": "core-loading-user-editorstartuptime",
-									"alias": "Editor Startup",
+									"alias": "Editor Boot",
 									"threshold": 120
 								},
 								{
@@ -141,6 +141,10 @@
 						{
 							"name": "User Loading Time (KPIs)",
 							"metrics": [
+								{
+									"id": "core-loading-user-editorbottime-kpi",
+									"alias": "EditorBoot"
+								},
 								{
 									"id": "core-loading-user-editorstartuptime-kpi",
 									"alias": "TotalEditorStartup"
diff --git a/Engine/Source/Programs/Horde/Horde.Server/Defaults/default-metrics.telemetry.json b/Engine/Source/Programs/Horde/Horde.Server/Defaults/default-metrics.telemetry.json
--- a/Engine/Source/Programs/Horde/Horde.Server/Defaults/default-metrics.telemetry.json
+++ b/Engine/Source/Programs/Horde/Horde.Server/Defaults/default-metrics.telemetry.json
@@ -2,7 +2,7 @@
 	"metrics": [
 		{
 			"id": "editor-performance-cache-warning",
-			"filter": "$[?(@.Payload.EventName == 'Editor.PerformanceWarning' && @.Payload.KPI_Category=='Cache' && @.Payload.KPI_DisplayType=='Percent')]",
+			"filter": "$[?(@.Payload.EventName == 'Editor.Performance.Warning' && @.Payload.KPI_Category=='Cache' && @.Payload.KPI_DisplayType=='Percent')]",
 			"groupBy": "$.Payload.Build_BranchName,$.Payload.ProjectName,$.Payload.KPI_Name",
 			"property": "$.Payload.KPI_CurrentValue",
 			"function": "Average",
@@ -10,7 +10,7 @@
 		},
 		{
 			"id": "editor-performance-editor-warning",
-			"filter": "$[?(@.Payload.EventName == 'Editor.PerformanceWarning' && @.Payload.KPI_Category=='Editor')]",
+			"filter": "$[?(@.Payload.EventName == 'Editor.Performance.Warning' && @.Payload.KPI_Category=='Editor')]",
 			"groupBy": "$.Payload.Build_BranchName,$.Payload.ProjectName,$.Payload.MapName,$.Payload.KPI_Name",
 			"property": "$.Payload.KPI_CurrentValue",
 			"function": "Average",
@@ -18,7 +18,7 @@
 		},
 		{
 			"id": "editor-performance-pie-warning",
-			"filter": "$[?(@.Payload.EventName == 'Editor.PerformanceWarning' && @.Payload.KPI_Category=='PIE')]",
+			"filter": "$[?(@.Payload.EventName == 'Editor.Performance.Warning' && @.Payload.KPI_Category=='PIE')]",
 			"groupBy": "$.Payload.Build_BranchName,$.Payload.ProjectName,$.Payload.MapName,$.Payload.KPI_Name",
 			"property": "$.Payload.KPI_CurrentValue",
 			"function": "Average",
@@ -26,7 +26,7 @@
 		},
 		{
 			"id": "editor-performance-warning-count",
-			"filter": "$[?(@.Payload.EventName == 'Editor.PerformanceWarning')]",
+			"filter": "$[?(@.Payload.EventName == 'Editor.Performance.Warning')]",
 			"groupBy": "$.Payload.Build_BranchName,$.Payload.ProjectName,$.Payload.KPI_Category,$.Payload.KPI_Name",
 			"function": "Count",
 			"interval": "00:15:00"
@@ -46,6 +46,14 @@
 			"function": "Count",
 			"interval": "00:15:00"
 		},
+		{
+			"id": "core-loading-user-editorboottime",
+			"filter": "$[?(@.Payload.EventName == 'Core.Loading' && @.Payload.LoadingName=='BootEditor' && @.Payload.Config_IsBuildMachine==false && @.Payload.Config_IsRunningCommandlet==false && @.Payload.LoadingSeconds>5 && @.Payload.LoadingSeconds<3600)]",
+			"property": "$.Payload.LoadingSeconds",
+			"groupBy": "$.Payload.Build_BranchName,$.Payload.ProjectName,$.Payload.ProjectName",
+			"function": "Average",
+			"interval": "00:15:00"
+		},
 		{
 			"id": "core-loading-user-editorstartuptime",
 			"filter": "$[?(@.Payload.EventName == 'Core.Loading' && @.Payload.LoadingName=='TotalEditorStartup' && @.Payload.Config_IsBuildMachine==false && @.Payload.Config_IsRunningCommandlet==false && @.Payload.LoadingSeconds>5 && @.Payload.LoadingSeconds<3600)]",
@@ -93,6 +101,14 @@
 			"function": "Average",
 			"interval": "00:15:00"
 		},
+		{
+			"id": "core-loading-user-editorboottime-kpi",
+			"filter": "$[?(@.Payload.EventName == 'Core.Loading' && @.Payload.LoadingName=='BootEditor' && @.Payload.Config_IsBuildMachine==false && @.Payload.Config_IsRunningCommandlet==false && @.Payload.LoadingSeconds>5 && @.Payload.LoadingSeconds<3600)]",
+			"property": "$.Payload.LoadingSeconds",
+			"groupBy": "$.Payload.Build_BranchName,$.Payload.ProjectName",
+			"function": "Average",
+			"interval": "01:00:00"
+		},
 		{
 			"id": "core-loading-user-editorstartuptime-kpi",
 			"filter": "$[?(@.Payload.EventName == 'Core.Loading' && @.Payload.LoadingName=='TotalEditorStartup' && @.Payload.Config_IsBuildMachine==false && @.Payload.Config_IsRunningCommandlet==false && @.Payload.LoadingSeconds>5 && @.Payload.LoadingSeconds<3600)]",
@@ -174,6 +190,15 @@
 			"percentile": 75,
 			"interval": "00:15:00"
 		},
+		{
+			"id": "core-loading-user-editorboottime-percentile",
+			"filter": "$[?(@.Payload.EventName == 'Core.Loading' && @.Payload.LoadingName=='BootEditor' && @.Payload.Config_IsBuildMachine==false && @.Payload.Config_IsRunningCommandlet==false && @.Payload.LoadingSeconds>5 && @.Payload.LoadingSeconds<3600 )]",
+			"groupBy": "$.Payload.Build_BranchName,$.Payload.ProjectName",
+			"property": "$.Payload.LoadingSeconds",
+			"function": "Percentile",
+			"percentile": 75,
+			"interval": "00:15:00"
+		},
 		{
 			"id": "core-loading-user-editorstartuptime-percentile",
 			"filter": "$[?(@.Payload.EventName == 'Core.Loading' && @.Payload.LoadingName=='TotalEditorStartup' && @.Payload.Config_IsBuildMachine==false && @.Payload.Config_IsRunningCommandlet==false && @.Payload.LoadingSeconds>5 && @.Payload.LoadingSeconds<3600 )]",
