Horde Analytics Dashbaord Fixes

diff --git a/Engine/Source/Programs/Horde/Horde.Server/Defaults/default-analytics.dashboard.json b/Engine/Source/Programs/Horde/Horde.Server/Defaults/default-analytics.dashboard.json
--- a/Engine/Source/Programs/Horde/Horde.Server/Defaults/default-analytics.dashboard.json
+++ b/Engine/Source/Programs/Horde/Horde.Server/Defaults/default-analytics.dashboard.json
@@ -34,7 +34,7 @@
 					"name": "LoadingName",
 					"group": "LoadingName",
 					"defaults": [
-						"TotalEditorStartup",
+						"BootEditor",
 						"LoadMap",
 						"TimeToPIE"
 					]
