Horde: Fix some broken doc links.

#rnx

diff --git a/Engine/Source/Programs/Horde/Docs/Deployment/Integrations/Slack.md b/Engine/Source/Programs/Horde/Docs/Deployment/Integrations/Slack.md
--- a/Engine/Source/Programs/Horde/Docs/Deployment/Integrations/Slack.md
+++ b/Engine/Source/Programs/Horde/Docs/Deployment/Integrations/Slack.md
@@ -64,5 +64,5 @@ fully:
 ### User mapping
 
 Horde users are mapped to Slack users by correlating the email address in the user's
-[OIDC profile](../Server.md#authentication) with their Slack user profile. Horde will use avatars configured through
+[OIDC profile](../Server.md#authorization) with their Slack user profile. Horde will use avatars configured through
 Slack in the dashboard for any successfully mapped email address.
diff --git a/Engine/Source/Programs/Horde/Docs/Tutorials/Authorization.md b/Engine/Source/Programs/Horde/Docs/Tutorials/Authorization.md
--- a/Engine/Source/Programs/Horde/Docs/Tutorials/Authorization.md
+++ b/Engine/Source/Programs/Horde/Docs/Tutorials/Authorization.md
@@ -10,7 +10,7 @@ deployments will probably want users to log in, and restrict the actions they ca
 To do this, Horde supports **[OAuth2](https://oauth.net/2/)** and **[OIDC](https://openid.net/developers/how-connect-works/)**,
 which is supported by most third party identity providers - including Okta, AWS, Azure, and Google.
 Configuring an external identity provider is out of scope for this documentation, though the relevant configuration
-points are touched on in the [Deployment > Server](../Deployment/Server.md#authentication) page.
+points are touched on in the [Deployment > Server](../Deployment/Server.md#authorization) page.
 
 If you don't have an existing OIDC-compatible identity provider, Horde includes it's own - which this guide covers.
 
@@ -24,7 +24,7 @@ If you don't have an existing OIDC-compatible identity provider, Horde includes
 1. In your [server.json](../Config/Orientation.md) file set the `AuthMode` property to `Horde`, and restart the server.
 2. The first time you launch the server, you'll be prompted to enter an administrator password.
 3. After logging in, there will be an `Accounts` menu item in the `Server` menu. From here, you can manage the users
-allowed to log in to the server, and the [**claims**](../Glossary.md#authentication) that they have. Horde's account
+allowed to log in to the server, and the [**claims**](../Glossary.md#authorization) that they have. Horde's account
 system uses the `http://epicgames.com/ue/horde/group` claim for groups that a user belongs to, and the dashboard will
 suggest and autocomplete any groups found in the deployment's configuration files.
 
