Horde: Add some logging for device reservation
#rnx

diff --git a/Engine/Source/Programs/Horde/Horde.Server/Devices/DevicesController.cs b/Engine/Source/Programs/Horde/Horde.Server/Devices/DevicesController.cs
--- a/Engine/Source/Programs/Horde/Horde.Server/Devices/DevicesController.cs
+++ b/Engine/Source/Programs/Horde/Horde.Server/Devices/DevicesController.cs
@@ -821,6 +821,11 @@ namespace Horde.Server.Devices
 
 			if (reservation == null)
 			{
+				if (!String.IsNullOrWhiteSpace(errorMessage))
+				{
+					_logger.LogWarning("Device reservation failed: {ErrorMessage}", errorMessage);
+				}
+
 				return Conflict(String.IsNullOrEmpty(errorMessage) ? "Unable to allocated devices for reservation" : errorMessage);
 			}
 
