Horde: Fix service account descriptions being clipped
#rnx

diff --git a/Engine/Source/Programs/Horde/HordeDashboard/src/components/accounts/ServiceAccountsView.tsx b/Engine/Source/Programs/Horde/HordeDashboard/src/components/accounts/ServiceAccountsView.tsx
--- a/Engine/Source/Programs/Horde/HordeDashboard/src/components/accounts/ServiceAccountsView.tsx
+++ b/Engine/Source/Programs/Horde/HordeDashboard/src/components/accounts/ServiceAccountsView.tsx
@@ -62,7 +62,7 @@ const AccountPanel: React.FC = observer(() => {
    const columns: IColumn[] = [
       { key: 'column_id', name: 'Id', minWidth: 240, maxWidth: 240, isResizable: false },
       { key: 'column_status', name: 'Status', minWidth: 80, maxWidth: 80, isResizable: false },
-      { key: 'column_description', name: 'Description', minWidth: 440, maxWidth: 440, isResizable: false },
+      { key: 'column_description', name: 'Description', minWidth: 440, isResizable: false },
       { key: 'column_edit', name: 'Edit', minWidth: 48, maxWidth: 48, isResizable: false, onRenderHeader: () => null },
    ];
 
