Horde: Remove isSuspect fuinction that is no longer used
#rnx

diff --git a/Engine/Source/Programs/Horde/HordeDashboard/src/components/IssueViewV2.tsx b/Engine/Source/Programs/Horde/HordeDashboard/src/components/IssueViewV2.tsx
--- a/Engine/Source/Programs/Horde/HordeDashboard/src/components/IssueViewV2.tsx
+++ b/Engine/Source/Programs/Horde/HordeDashboard/src/components/IssueViewV2.tsx
@@ -3,21 +3,21 @@ import { action, makeObservable, observable } from 'mobx';
 import { observer } from "mobx-react-lite";
 import moment from "moment";
 import React, { useEffect, useRef, useState } from "react";
-import { Link, useNavigate, useLocation } from "react-router-dom";
+import { Link, useLocation, useNavigate } from "react-router-dom";
 import backend from "../backend";
-import { CreateExternalIssueResponse, CreateJobRequest, EventSeverity, GetChangeSummaryResponse, GetExternalIssueProjectResponse, GetExternalIssueResponse, GetIssueResponse, GetIssueSpanResponse, GetIssueStepResponse, GetIssueStreamResponse, GetLogEventResponse, GetTemplateRefResponse, GetThinUserInfoResponse, GetUserResponse, IssueData, IssueSeverity, UpdateIssueRequest } from "../backend/Api";
+import { CreateExternalIssueResponse, CreateJobRequest, EventSeverity, GetChangeSummaryResponse, GetExternalIssueProjectResponse, GetExternalIssueResponse, GetIssueResponse, GetIssueSpanResponse, GetIssueStepResponse, GetIssueStreamResponse, GetLogEventResponse, GetTemplateRefResponse, GetUserResponse, IssueData, IssueSeverity, UpdateIssueRequest } from "../backend/Api";
 import dashboard, { StatusColor } from "../backend/Dashboard";
 import { projectStore } from "../backend/ProjectStore";
 import templateCache from '../backend/TemplateCache';
 import { Markdown } from '../base/components/Markdown';
 import { useWindowSize } from "../base/utilities/hooks";
 import { displayTimeZone, getHumanTime, getShortNiceTime } from "../base/utilities/timeUtils";
+import { getHordeStyling } from "../styles/Styles";
+import { getHordeTheme } from "../styles/theme";
 import { ErrorHandler } from "./ErrorHandler";
 import { useQuery } from "./JobDetailCommon";
 import { renderLine } from "./LogRender";
 import { UserSelect } from "./UserSelect";
-import { getHordeTheme } from "../styles/theme";
-import { getHordeStyling } from "../styles/Styles";
 
 const smallThreshhold = 1100;
 
@@ -293,21 +293,6 @@ class IssueDetails {
       return name;
    }
 
-
-   isSuspect(authorInfo?: GetThinUserInfoResponse) {
-
-      if (!authorInfo || authorInfo.name === "buildmachine" || authorInfo.name === "svc-p4-hordeproxy-p" || authorInfo.name === "robomerge") {
-         return false;
-      }
-
-      return !!this.suspects?.find(suspect => {
-
-         return suspect.id === authorInfo.id;
-
-      });
-
-   }
-
    resolved(streamId?: string, templateId?: string): boolean {
 
       if (!this.issue) {
