Horde: Add a [TryAuthorize] attribute which attempts authorization, but which does not return a challenge or forbid result on failure. Should fix tool download requests using service account tokens.

diff --git a/Engine/Source/Programs/Horde/Horde.Server/Tools/ToolsController.cs b/Engine/Source/Programs/Horde/Horde.Server/Tools/ToolsController.cs
--- a/Engine/Source/Programs/Horde/Horde.Server/Tools/ToolsController.cs
+++ b/Engine/Source/Programs/Horde/Horde.Server/Tools/ToolsController.cs
@@ -14,17 +14,22 @@ using Horde.Server.Server;
 using Horde.Server.Storage;
 using Horde.Server.Utilities;
 using HordeCommon;
+using Microsoft.AspNetCore.Authentication;
 using Microsoft.AspNetCore.Authorization;
+using Microsoft.AspNetCore.Authorization.Policy;
 using Microsoft.AspNetCore.Http;
 using Microsoft.AspNetCore.Mvc;
+using Microsoft.AspNetCore.Mvc.Authorization;
+using Microsoft.AspNetCore.Mvc.Filters;
 using Microsoft.AspNetCore.StaticFiles;
+using Microsoft.Extensions.DependencyInjection;
 using Microsoft.Extensions.Logging;
 using Microsoft.Extensions.Options;
 
 namespace Horde.Server.Tools
 {
 	/// <summary>
-	/// Controller for the /api/v1/agents endpoint
+	/// Controller for the /api/v1/tools endpoint
 	/// </summary>
 	[ApiController]
 	[Authorize]
@@ -160,6 +165,7 @@ namespace Horde.Server.Tools
 	/// Public methods available without authorization (or with very custom authorization)
 	/// </summary>
 	[ApiController]
+	[TryAuthorize]
 	[Tags("Tools")]
 	public class PublicToolsController : HordeControllerBase
 	{
@@ -181,6 +187,7 @@ namespace Horde.Server.Tools
 		/// Enumerates all the available tools.
 		/// </summary>
 		[HttpGet]
+		[TryAuthorize]
 		[Route("/api/v1/tools")]
 		public async Task<ActionResult<GetToolsSummaryResponse>> GetToolsAsync()
 		{
diff --git a/Engine/Source/Programs/Horde/Horde.Server/Utilities/TryAuthorizeAttribute.cs b/Engine/Source/Programs/Horde/Horde.Server/Utilities/TryAuthorizeAttribute.cs
new file mode 100644
--- /dev/null
+++ b/Engine/Source/Programs/Horde/Horde.Server/Utilities/TryAuthorizeAttribute.cs
@@ -0,0 +1,27 @@
+// Copyright Epic Games, Inc. All Rights Reserved.
+
+using System;
+using System.Threading.Tasks;
+using Microsoft.AspNetCore.Authorization;
+using Microsoft.AspNetCore.Authorization.Policy;
+using Microsoft.AspNetCore.Mvc.Filters;
+using Microsoft.Extensions.DependencyInjection;
+
+namespace Horde.Server.Utilities
+{
+	/// <summary>
+	/// Attempts to authorize the user, but does not challenge or forbid access if authorization fails.
+	/// </summary>
+	sealed class TryAuthorizeAttribute : Attribute, IAsyncAuthorizationFilter
+	{
+		/// <inheritdoc/>
+		public async Task OnAuthorizationAsync(AuthorizationFilterContext context)
+		{
+			IAuthorizationPolicyProvider policyProvider = context.HttpContext.RequestServices.GetRequiredService<IAuthorizationPolicyProvider>();
+			AuthorizationPolicy effectivePolicy = await policyProvider.GetDefaultPolicyAsync();
+
+			IPolicyEvaluator policyEvaluator = context.HttpContext.RequestServices.GetRequiredService<IPolicyEvaluator>();
+			await policyEvaluator.AuthenticateAsync(effectivePolicy, context.HttpContext);
+		}
+	}
+}
