Horde: Fix static analysis warnings.

#rnx

diff --git a/Engine/Source/Programs/Horde/Horde.Server/Tools/ToolsController.cs b/Engine/Source/Programs/Horde/Horde.Server/Tools/ToolsController.cs
--- a/Engine/Source/Programs/Horde/Horde.Server/Tools/ToolsController.cs
+++ b/Engine/Source/Programs/Horde/Horde.Server/Tools/ToolsController.cs
@@ -14,15 +14,10 @@ using Horde.Server.Server;
 using Horde.Server.Storage;
 using Horde.Server.Utilities;
 using HordeCommon;
-using Microsoft.AspNetCore.Authentication;
 using Microsoft.AspNetCore.Authorization;
-using Microsoft.AspNetCore.Authorization.Policy;
 using Microsoft.AspNetCore.Http;
 using Microsoft.AspNetCore.Mvc;
-using Microsoft.AspNetCore.Mvc.Authorization;
-using Microsoft.AspNetCore.Mvc.Filters;
 using Microsoft.AspNetCore.StaticFiles;
-using Microsoft.Extensions.DependencyInjection;
 using Microsoft.Extensions.Logging;
 using Microsoft.Extensions.Options;
 
diff --git a/Engine/Source/Programs/Horde/Horde.Server/Utilities/TryAuthorizeAttribute.cs b/Engine/Source/Programs/Horde/Horde.Server/Utilities/TryAuthorizeAttribute.cs
--- a/Engine/Source/Programs/Horde/Horde.Server/Utilities/TryAuthorizeAttribute.cs
+++ b/Engine/Source/Programs/Horde/Horde.Server/Utilities/TryAuthorizeAttribute.cs
@@ -12,6 +12,7 @@ namespace Horde.Server.Utilities
 	/// <summary>
 	/// Attempts to authorize the user, but does not challenge or forbid access if authorization fails.
 	/// </summary>
+	[AttributeUsage(AttributeTargets.Method | AttributeTargets.Class)]
 	sealed class TryAuthorizeAttribute : Attribute, IAsyncAuthorizationFilter
 	{
 		/// <inheritdoc/>
