Horde Analytics Dashboard Changes

diff --git a/Engine/Source/Programs/Horde/Horde.Server/Defaults/default-analytics.dashboard.json b/Engine/Source/Programs/Horde/Horde.Server/Defaults/default-analytics.dashboard.json
--- a/Engine/Source/Programs/Horde/Horde.Server/Defaults/default-analytics.dashboard.json
+++ b/Engine/Source/Programs/Horde/Horde.Server/Defaults/default-analytics.dashboard.json
@@ -106,11 +106,6 @@
 						{
 							"name": "User Load Time (Average)",
 							"metrics": [
-								{
-									"id": "core-loading-user-timetopie",
-									"alias": "Time To PIE",
-									"threshold": 360
-								},
 								{
 									"id": "core-loading-user-editorstartuptime",
 									"alias": "Editor Boot",
@@ -125,41 +120,34 @@
 									"id": "core-loading-user-piestartuptime",
 									"alias": "PIE Startup",
 									"threshold": 360
+								},
+								{
+									"id": "core-loading-user-timetopie",
+									"alias": "Time To PIE",
+									"threshold": 360
 								}
 							],
 							"graph": "indicator",
 							"max": 900
 						},
-						{
-							"name": "Performance Warnings",
-							"metrics": [
-								{
-									"id": "editor-performance-warning"
-								}
-							]
-						},
 						{
 							"name": "User Loading Time (KPIs)",
 							"metrics": [
 								{
-									"id": "core-loading-user-editorbottime-kpi",
-									"alias": "EditorBoot"
-								},
-								{
-									"id": "core-loading-user-editorstartuptime-kpi",
-									"alias": "TotalEditorStartup"
+									"id": "core-loading-user-editorboottime-kpi",
+									"alias": "Editor Boot"
 								},
 								{
 									"id": "core-loading-user-loadmaptime-kpi",
-									"alias": "LoadMap"
+									"alias": "Load Map"
 								},
 								{
 									"id": "core-loading-user-piestartuptime-kpi",
-									"alias": "PIE.TotalStartupTime"
+									"alias": "PIE StartupTime"
 								},
 								{
 									"id": "core-loading-user-timetopie-kpi",
-									"alias": "TimeToPIE"
+									"alias": "Time To PIE"
 								}
 							]
 						},
