Further tweaks to Horde Analytics Dashboard

diff --git a/Engine/Source/Programs/Horde/Horde.Server/Defaults/default-analytics.dashboard.json b/Engine/Source/Programs/Horde/Horde.Server/Defaults/default-analytics.dashboard.json
--- a/Engine/Source/Programs/Horde/Horde.Server/Defaults/default-analytics.dashboard.json
+++ b/Engine/Source/Programs/Horde/Horde.Server/Defaults/default-analytics.dashboard.json
@@ -81,22 +81,22 @@
 							"metrics": [
 								{
 									"id": "core-loading-user-timetopie-percentile",
-									"alias": "Time To PIE",
+									"alias": "TimeToPIE",
 									"threshold": 360
 								},
 								{
 									"id": "core-loading-user-editorboottime-percentile",
-									"alias": "Editor Boot",
+									"alias": "BootEditor",
 									"threshold": 120
 								},
 								{
 									"id": "core-loading-user-loadmaptime-percentile",
-									"alias": "Load Map",
+									"alias": "LoadMap",
 									"threshold": 240
 								},
 								{
 									"id": "core-loading-user-piestartuptime-percentile",
-									"alias": "PIE Startup",
+									"alias": "PIEStartup",
 									"threshold": 360
 								}
 							],
@@ -107,23 +107,23 @@
 							"name": "User Load Time (Average)",
 							"metrics": [
 								{
-									"id": "core-loading-user-editorstartuptime",
-									"alias": "Editor Boot",
+									"id": "core-loading-user-editorboottime",
+									"alias": "BootEditor",
 									"threshold": 120
 								},
 								{
 									"id": "core-loading-user-loadmaptime",
-									"alias": "Load Map",
+									"alias": "LoadMap",
 									"threshold": 240
 								},
 								{
 									"id": "core-loading-user-piestartuptime",
-									"alias": "PIE Startup",
+									"alias": "PIEStartup",
 									"threshold": 360
 								},
 								{
 									"id": "core-loading-user-timetopie",
-									"alias": "Time To PIE",
+									"alias": "TimeToPIE",
 									"threshold": 360
 								}
 							],
@@ -135,19 +135,19 @@
 							"metrics": [
 								{
 									"id": "core-loading-user-editorboottime-kpi",
-									"alias": "Editor Boot"
+									"alias": "BootEditor"
 								},
 								{
 									"id": "core-loading-user-loadmaptime-kpi",
-									"alias": "Load Map"
+									"alias": "LoadMap"
 								},
 								{
 									"id": "core-loading-user-piestartuptime-kpi",
-									"alias": "PIE StartupTime"
+									"alias": "PIEStartup"
 								},
 								{
 									"id": "core-loading-user-timetopie-kpi",
-									"alias": "Time To PIE"
+									"alias": "TimeToPIE"
 								}
 							]
 						},
