More HA Dashboard Tweaks

diff --git a/Engine/Source/Programs/Horde/Horde.Server/Defaults/default-analytics.dashboard.json b/Engine/Source/Programs/Horde/Horde.Server/Defaults/default-analytics.dashboard.json
--- a/Engine/Source/Programs/Horde/Horde.Server/Defaults/default-analytics.dashboard.json
+++ b/Engine/Source/Programs/Horde/Horde.Server/Defaults/default-analytics.dashboard.json
@@ -162,10 +162,10 @@
 							]
 						},
 						{
-							"name": "User Editor Startup Time (Live)",
+							"name": "User Editor Boot Time (Live)",
 							"metrics": [
 								{
-									"id": "core-loading-user-editorstartuptime"
+									"id": "core-loading-user-editorboottime"
 								}
 							]
 						},
