Horde: Additional logging for garbage collection.

#rnx

diff --git a/Engine/Source/Programs/Horde/Horde.Server/Storage/StorageService.cs b/Engine/Source/Programs/Horde/Horde.Server/Storage/StorageService.cs
--- a/Engine/Source/Programs/Horde/Horde.Server/Storage/StorageService.cs
+++ b/Engine/Source/Programs/Horde/Horde.Server/Storage/StorageService.cs
@@ -980,32 +980,39 @@ namespace Horde.Server.Storage
 			RedisSortedSetKey<RedisValue> checkSet = GetGcCheckSet(namespaceInfo.Id);
 			for (; ; )
 			{
-				RedisValue[] values = await _redisService.GetDatabase().SortedSetRangeByRankAsync(checkSet, 0, 0);
-				if (values.Length == 0)
+				long length = await _redisService.GetDatabase().SortedSetLengthAsync(checkSet);
+				int batchSize = (int)Math.Min(length, 1024);
+				_logger.LogInformation("Garbage collection queue for namespace {NamespaceId} ({QueueName}) has {Length} entries; taking {Count}", namespaceInfo.Id, checkSet.Inner, length, batchSize);
+
+				if (length == 0)
 				{
 					break;
 				}
 
-				ObjectId blobInfoId = new ObjectId(((byte[]?)values[0])!);
-				if (blobInfoId < lastImportBlobInfoId && !await IsBlobReferencedAsync(blobInfoId, cancellationToken))
+				RedisValue[] values = await _redisService.GetDatabase().SortedSetRangeByRankAsync(checkSet, 0, batchSize);
+				foreach (RedisValue value in values)
 				{
-					BlobInfo? info = await _blobCollection.FindOneAndDeleteAsync(x => x.Id == blobInfoId, cancellationToken: cancellationToken);
-					if (info != null)
+					ObjectId blobInfoId = new ObjectId(((byte[]?)value)!);
+					if (blobInfoId < lastImportBlobInfoId && !await IsBlobReferencedAsync(blobInfoId, cancellationToken))
 					{
-						if (info.Imports != null)
+						BlobInfo? info = await _blobCollection.FindOneAndDeleteAsync(x => x.Id == blobInfoId, cancellationToken: cancellationToken);
+						if (info != null)
 						{
-							SortedSetEntry<RedisValue>[] entries = info.Imports.Select(x => new SortedSetEntry<RedisValue>(x.ToByteArray(), score)).ToArray();
-							_ = _redisService.GetDatabase().SortedSetAddAsync(checkSet, entries, flags: CommandFlags.FireAndForget);
-							score = Math.BitIncrement(score);
-						}
+							if (info.Imports != null)
+							{
+								SortedSetEntry<RedisValue>[] entries = info.Imports.Select(x => new SortedSetEntry<RedisValue>(x.ToByteArray(), score)).ToArray();
+								_ = _redisService.GetDatabase().SortedSetAddAsync(checkSet, entries, flags: CommandFlags.FireAndForget);
+								score = Math.BitIncrement(score);
+							}
 
-						ObjectKey objectKey = GetObjectKey(new BlobLocator(info.Path));
-						_logger.LogDebug("Deleting object: {Key}", objectKey);
-						await namespaceInfo.Store.DeleteAsync(objectKey, cancellationToken);
-						numItemsRemoved++;
+							ObjectKey objectKey = GetObjectKey(new BlobLocator(info.Path));
+							_logger.LogDebug("Deleting {NamespaceId} blob {BlobId}, key: {Key} ({ImportCount} imports)", namespaceInfo.Id, blobInfoId, objectKey, info.Imports?.Count ?? 0);
+							await namespaceInfo.Store.DeleteAsync(objectKey, cancellationToken);
+							numItemsRemoved++;
+						}
 					}
+					_ = _redisService.GetDatabase().SortedSetRemoveAsync(checkSet, values[0], CommandFlags.FireAndForget);
 				}
-				_ = _redisService.GetDatabase().SortedSetRemoveAsync(checkSet, values[0], CommandFlags.FireAndForget);
 			}
 
 			await _gcState.UpdateAsync(state => state.FindOrAddNamespace(namespaceInfo.Id).LastTime = utcNow, cancellationToken);
