Horde: Allow using Horde auth without HTTPS for tests.

diff --git a/Engine/Source/Programs/Horde/Horde.Server/Startup.cs b/Engine/Source/Programs/Horde/Horde.Server/Startup.cs
--- a/Engine/Source/Programs/Horde/Horde.Server/Startup.cs
+++ b/Engine/Source/Programs/Horde/Horde.Server/Startup.cs
@@ -787,6 +787,10 @@ namespace Horde.Server
 					{
 						options.Authority = new Uri(settings.ServerUrl, "api/v1/oauth2").ToString();
 						options.ClientId = "default";
+						if (settings.HttpsPort == 0)
+						{
+							options.RequireHttpsMetadata = false;
+						}
 						foreach (string scope in settings.OidcRequestedScopes)
 						{
 							options.Scope.Add(scope);
