Horde: Fix broken links in docs.

#rnx

diff --git a/Engine/Source/Programs/Horde/Docs/Config/Analytics.md b/Engine/Source/Programs/Horde/Docs/Config/Analytics.md
--- a/Engine/Source/Programs/Horde/Docs/Config/Analytics.md
+++ b/Engine/Source/Programs/Horde/Docs/Config/Analytics.md
@@ -30,9 +30,9 @@ You can configure telemetry sinks through the `Telemetry` property in the server
 
 To provide efficient aggregation of analytics data over large time periods, Horde aggregates telemetry events into
 running metrics for each time interval. This aggregation is performed according to rules specified in the
-`Telemetry.Metrics` section of the globals.json file (see [MetricConfig](Schema/Globals.md#metricconfig)).
+`Telemetry.Metrics` section of the globals.json file (see [MetricConfig](Schema/Telemetry.md#metricconfig)).
 
 ## Charting
 
 The Horde dashboard can render charts showing aggregated metrics collected on the server. These views are configured
-using the `Dashboard.Analytics` section of the globals.json file (see [TelemetryViewConfig](Schema/Globals.md#telemetryviewconfig)).
\ No newline at end of file
+using the `Dashboard.Analytics` section of the globals.json file (see [TelemetryViewConfig](Schema/Telemetry.md#telemetryviewconfig)).
diff --git a/Engine/Source/Programs/Horde/Docs/Glossary.md b/Engine/Source/Programs/Horde/Docs/Glossary.md
--- a/Engine/Source/Programs/Horde/Docs/Glossary.md
+++ b/Engine/Source/Programs/Horde/Docs/Glossary.md
@@ -13,7 +13,7 @@ matching a set of configured criteria for a particular time interval.
 * `Telemetry` - Individual events sent by an application to the Horde server. Horde processes telemetry events as
 schema-less JSON objects.
 
-## Authentication
+## Authorization
 
 * `Acl` - Abbreviation for **access control list**. An ACL in Horde defines rules that grant *entitlements* to users
 (ie. the ability to perform certain *actions*) based on *claims* that their user has.
