Horde: Add trace scopes for GC.

diff --git a/Engine/Source/Programs/Horde/Horde.Server/Storage/StorageService.cs b/Engine/Source/Programs/Horde/Horde.Server/Storage/StorageService.cs
--- a/Engine/Source/Programs/Horde/Horde.Server/Storage/StorageService.cs
+++ b/Engine/Source/Programs/Horde/Horde.Server/Storage/StorageService.cs
@@ -27,6 +27,7 @@ using Microsoft.Extensions.Options;
 using MongoDB.Bson;
 using MongoDB.Bson.Serialization.Attributes;
 using MongoDB.Driver;
+using OpenTelemetry.Trace;
 using StackExchange.Redis;
 
 namespace Horde.Server.Storage
@@ -358,6 +359,7 @@ namespace Horde.Server.Storage
 		readonly IMemoryCache _memoryCache;
 		readonly IObjectStoreFactory _objectStoreFactory;
 		readonly IOptionsMonitor<GlobalConfig> _globalConfig;
+		readonly Tracer _tracer;
 		readonly ILogger _logger;
 
 		readonly IMongoCollection<BlobInfo> _blobCollection;
@@ -377,7 +379,7 @@ namespace Horde.Server.Storage
 		/// <summary>
 		/// Constructor
 		/// </summary>
-		public StorageService(MongoService mongoService, RedisService redisService, IClock clock, BundleCache bundleCache, IMemoryCache memoryCache, IObjectStoreFactory objectStoreFactory, IOptionsMonitor<GlobalConfig> globalConfig, ILogger<StorageService> logger)
+		public StorageService(MongoService mongoService, RedisService redisService, IClock clock, BundleCache bundleCache, IMemoryCache memoryCache, IObjectStoreFactory objectStoreFactory, IOptionsMonitor<GlobalConfig> globalConfig, Tracer tracer, ILogger<StorageService> logger)
 		{
 			_redisService = redisService;
 			_clock = clock;
@@ -385,6 +387,7 @@ namespace Horde.Server.Storage
 			_memoryCache = memoryCache;
 			_objectStoreFactory = objectStoreFactory;
 			_globalConfig = globalConfig;
+			_tracer = tracer;
 			_logger = logger;
 
 			List<MongoIndex<BlobInfo>> blobIndexes = new List<MongoIndex<BlobInfo>>();
@@ -552,16 +555,26 @@ namespace Horde.Server.Storage
 		/// <inheritdoc/>
 		async Task<bool> IsBlobReferencedAsync(ObjectId blobInfoId, CancellationToken cancellationToken = default)
 		{
-			FilterDefinition<BlobInfo> blobFilter = Builders<BlobInfo>.Filter.AnyEq(x => x.Imports, blobInfoId);
-			if (await _blobCollection.Find(blobFilter).AnyAsync(cancellationToken))
+			using (TelemetrySpan span = _tracer.StartActiveSpan($"{nameof(StorageService)}.{nameof(IsBlobReferencedAsync)}.Blobs"))
 			{
-				return true;
+				span.SetAttribute("BlobId", blobInfoId.ToString());
+
+				FilterDefinition<BlobInfo> blobFilter = Builders<BlobInfo>.Filter.AnyEq(x => x.Imports, blobInfoId);
+				if (await _blobCollection.Find(blobFilter).AnyAsync(cancellationToken))
+				{
+					return true;
+				}
 			}
 
-			FilterDefinition<RefInfo> refFilter = Builders<RefInfo>.Filter.Eq(x => x.TargetBlobId, blobInfoId);
-			if (await _refCollection.Find(refFilter).AnyAsync(cancellationToken))
+			using (TelemetrySpan span = _tracer.StartActiveSpan($"{nameof(StorageService)}.{nameof(IsBlobReferencedAsync)}.Refs"))
 			{
-				return true;
+				span.SetAttribute("BlobId", blobInfoId.ToString());
+
+				FilterDefinition<RefInfo> refFilter = Builders<RefInfo>.Filter.Eq(x => x.TargetBlobId, blobInfoId);
+				if (await _refCollection.Find(refFilter).AnyAsync(cancellationToken))
+				{
+					return true;
+				}
 			}
 
 			return false;
@@ -993,6 +1006,10 @@ namespace Horde.Server.Storage
 				foreach (RedisValue value in values)
 				{
 					ObjectId blobInfoId = new ObjectId(((byte[]?)value)!);
+
+					using TelemetrySpan span = _tracer.StartActiveSpan($"{nameof(StorageService)}.{nameof(TickGcForNamespaceAsync)}");
+					span.SetAttribute("BlobId", blobInfoId.ToString());
+
 					if (blobInfoId < lastImportBlobInfoId && !await IsBlobReferencedAsync(blobInfoId, cancellationToken))
 					{
 						BlobInfo? info = await _blobCollection.FindOneAndDeleteAsync(x => x.Id == blobInfoId, cancellationToken: cancellationToken);
@@ -1011,7 +1028,7 @@ namespace Horde.Server.Storage
 							numItemsRemoved++;
 						}
 					}
-					_ = _redisService.GetDatabase().SortedSetRemoveAsync(checkSet, values[0], CommandFlags.FireAndForget);
+					_ = _redisService.GetDatabase().SortedSetRemoveAsync(checkSet, value, CommandFlags.FireAndForget);
 				}
 			}
 
