Horde: Fix slow queries in storage service garbage collection.

MongoDB implements IFindFluent.AnyAsync() by projecting the result down to the _id field, and running the query with a limit of 1. This seems to break the DocumentDB query planner, causing it to do a full document scan (over a massive collection) rather than an index scan for no reason.

diff --git a/Engine/Source/Programs/Horde/Horde.Server/Storage/StorageService.cs b/Engine/Source/Programs/Horde/Horde.Server/Storage/StorageService.cs
--- a/Engine/Source/Programs/Horde/Horde.Server/Storage/StorageService.cs
+++ b/Engine/Source/Programs/Horde/Horde.Server/Storage/StorageService.cs
@@ -555,26 +555,16 @@ namespace Horde.Server.Storage
 		/// <inheritdoc/>
 		async Task<bool> IsBlobReferencedAsync(ObjectId blobInfoId, CancellationToken cancellationToken = default)
 		{
-			using (TelemetrySpan span = _tracer.StartActiveSpan($"{nameof(StorageService)}.{nameof(IsBlobReferencedAsync)}.Blobs"))
+			FilterDefinition<BlobInfo> blobFilter = Builders<BlobInfo>.Filter.AnyEq(x => x.Imports, blobInfoId);
+			if (await _blobCollection.Find(blobFilter).Limit(1).CountDocumentsAsync(cancellationToken) > 0)
 			{
-				span.SetAttribute("BlobId", blobInfoId.ToString());
-
-				FilterDefinition<BlobInfo> blobFilter = Builders<BlobInfo>.Filter.AnyEq(x => x.Imports, blobInfoId);
-				if (await _blobCollection.Find(blobFilter).AnyAsync(cancellationToken))
-				{
-					return true;
-				}
+				return true;
 			}
 
-			using (TelemetrySpan span = _tracer.StartActiveSpan($"{nameof(StorageService)}.{nameof(IsBlobReferencedAsync)}.Refs"))
+			FilterDefinition<RefInfo> refFilter = Builders<RefInfo>.Filter.Eq(x => x.TargetBlobId, blobInfoId);
+			if (await _refCollection.Find(refFilter).Limit(1).CountDocumentsAsync(cancellationToken) > 0)
 			{
-				span.SetAttribute("BlobId", blobInfoId.ToString());
-
-				FilterDefinition<RefInfo> refFilter = Builders<RefInfo>.Filter.Eq(x => x.TargetBlobId, blobInfoId);
-				if (await _refCollection.Find(refFilter).AnyAsync(cancellationToken))
-				{
-					return true;
-				}
+				return true;
 			}
 
 			return false;
