Horde: Allow configuring different telemetry stores for each project and stream.

diff --git a/Engine/Source/Programs/Horde/Horde.Server/Issues/IssueCollection.cs b/Engine/Source/Programs/Horde/Horde.Server/Issues/IssueCollection.cs
--- a/Engine/Source/Programs/Horde/Horde.Server/Issues/IssueCollection.cs
+++ b/Engine/Source/Programs/Horde/Horde.Server/Issues/IssueCollection.cs
@@ -21,11 +21,13 @@ using Horde.Server.Auditing;
 using Horde.Server.Jobs;
 using Horde.Server.Jobs.Graphs;
 using Horde.Server.Server;
+using Horde.Server.Streams;
 using Horde.Server.Telemetry;
 using Horde.Server.Users;
 using Horde.Server.Utilities;
 using Microsoft.Extensions.Caching.Memory;
 using Microsoft.Extensions.Logging;
+using Microsoft.Extensions.Options;
 using MongoDB.Bson;
 using MongoDB.Bson.Serialization;
 using MongoDB.Bson.Serialization.Attributes;
@@ -595,6 +597,7 @@ namespace Horde.Server.Issues
 		readonly IMongoCollection<IssueSuspect> _issueSuspects;
 		readonly IAuditLog<int> _auditLog;
 		readonly ITelemetrySink _telemetrySink;
+		readonly IOptionsMonitor<GlobalConfig> _globalConfig;
 		readonly Tracer _tracer;
 		readonly ILogger _logger;
 
@@ -616,11 +619,12 @@ namespace Horde.Server.Issues
 			});
 		}
 
-		public IssueCollection(MongoService mongoService, RedisService redisService, IUserCollection userCollection, IAuditLogFactory<int> auditLogFactory, ITelemetrySink telemetrySink, Tracer tracer, ILogger<IssueCollection> logger)
+		public IssueCollection(MongoService mongoService, RedisService redisService, IUserCollection userCollection, IAuditLogFactory<int> auditLogFactory, ITelemetrySink telemetrySink, IOptionsMonitor<GlobalConfig> globalConfig, Tracer tracer, ILogger<IssueCollection> logger)
 		{
 			_redisService = redisService;
 			_userCollection = userCollection;
 			_telemetrySink = telemetrySink;
+			_globalConfig = globalConfig;
 			_tracer = tracer;
 			_logger = logger;
 
@@ -726,9 +730,18 @@ namespace Horde.Server.Issues
 
 		void SendTelemetry(IIssue issue)
 		{
-			if (_telemetrySink.Enabled)
+			List<TelemetryStoreId> telemetryStoreIds = new List<TelemetryStoreId>();
+			foreach (IIssueStream stream in issue.Streams)
 			{
-				_telemetrySink.SendEvent(TelemetryStoreId.Default, TelemetryRecordMeta.CurrentHordeInstance, new
+				if (_globalConfig.CurrentValue.TryGetStream(stream.StreamId, out StreamConfig? streamConfig) && !streamConfig.TelemetryStoreId.IsEmpty)
+				{
+					telemetryStoreIds.Add(streamConfig.TelemetryStoreId);
+				}
+			}
+
+			foreach (TelemetryStoreId telemetryStoreId in telemetryStoreIds)
+			{
+				_telemetrySink.SendEvent(telemetryStoreId, TelemetryRecordMeta.CurrentHordeInstance, new
 				{
 					EventName = "State.Issue",
 					Id = issue.Id,
@@ -1497,9 +1510,9 @@ namespace Horde.Server.Issues
 
 		void SendTelemetry(IIssueSpan issueSpan)
 		{
-			if (_telemetrySink.Enabled)
+			if (_globalConfig.CurrentValue.TryGetStream(issueSpan.StreamId, out StreamConfig? streamConfig) && !streamConfig.TelemetryStoreId.IsEmpty)
 			{
-				_telemetrySink.SendEvent(TelemetryStoreId.Default, TelemetryRecordMeta.CurrentHordeInstance, new
+				_telemetrySink.SendEvent(streamConfig.TelemetryStoreId, TelemetryRecordMeta.CurrentHordeInstance, new
 				{
 					EventName = "State.IssueSpan",
 					Id = issueSpan.Id,
diff --git a/Engine/Source/Programs/Horde/Horde.Server/Jobs/JobCollection.cs b/Engine/Source/Programs/Horde/Horde.Server/Jobs/JobCollection.cs
--- a/Engine/Source/Programs/Horde/Horde.Server/Jobs/JobCollection.cs
+++ b/Engine/Source/Programs/Horde/Horde.Server/Jobs/JobCollection.cs
@@ -29,6 +29,7 @@ using Horde.Server.Utilities;
 using HordeCommon;
 using HordeCommon.Rpc.Tasks;
 using Microsoft.Extensions.Logging;
+using Microsoft.Extensions.Options;
 using MongoDB.Bson;
 using MongoDB.Bson.IO;
 using MongoDB.Bson.Serialization;
@@ -389,20 +390,17 @@ namespace Horde.Server.Jobs
 		readonly ITelemetrySink _telemetrySink;
 		readonly IClock _clock;
 		readonly Tracer _tracer;
+		readonly IOptionsMonitor<GlobalConfig> _globalConfig;
 		readonly ILogger<JobCollection> _logger;
 
 		/// <summary>
 		/// Constructor
 		/// </summary>
-		/// <param name="mongoService">The database service singleton</param>
-		/// <param name="clock"></param>
-		/// <param name="telemetrySink">Telemetry sink for data</param>
-		/// <param name="tracer">Tracer</param>
-		/// <param name="logger">The logger instance</param>
-		public JobCollection(MongoService mongoService, IClock clock, ITelemetrySink telemetrySink, Tracer tracer, ILogger<JobCollection> logger)
+		public JobCollection(MongoService mongoService, IClock clock, ITelemetrySink telemetrySink, IOptionsMonitor<GlobalConfig> globalConfig, Tracer tracer, ILogger<JobCollection> logger)
 		{
 			_clock = clock;
 			_telemetrySink = telemetrySink;
+			_globalConfig = globalConfig;
 			_tracer = tracer;
 			_logger = logger;
 
@@ -450,9 +448,9 @@ namespace Horde.Server.Jobs
 
 			await _jobs.InsertOneAsync(newJob, null, cancellationToken);
 
-			if (_telemetrySink.Enabled)
+			if(_globalConfig.CurrentValue.TryGetStream(streamId, out StreamConfig? streamConfig) && !streamConfig.TelemetryStoreId.IsEmpty)
 			{
-				_telemetrySink.SendEvent(TelemetryStoreId.Default, TelemetryRecordMeta.CurrentHordeInstance, new
+				_telemetrySink.SendEvent(streamConfig.TelemetryStoreId, TelemetryRecordMeta.CurrentHordeInstance, new
 				{
 					EventName = "State.Job",
 					Id = newJob.Id,
diff --git a/Engine/Source/Programs/Horde/Horde.Server/Jobs/JobStepRefCollection.cs b/Engine/Source/Programs/Horde/Horde.Server/Jobs/JobStepRefCollection.cs
--- a/Engine/Source/Programs/Horde/Horde.Server/Jobs/JobStepRefCollection.cs
+++ b/Engine/Source/Programs/Horde/Horde.Server/Jobs/JobStepRefCollection.cs
@@ -13,9 +13,11 @@ using EpicGames.Horde.Logs;
 using EpicGames.Horde.Streams;
 using EpicGames.Horde.Telemetry;
 using Horde.Server.Server;
+using Horde.Server.Streams;
 using Horde.Server.Telemetry;
 using Horde.Server.Utilities;
 using HordeCommon;
+using Microsoft.Extensions.Options;
 using MongoDB.Bson.Serialization.Attributes;
 using MongoDB.Driver;
 
@@ -105,19 +107,19 @@ namespace Horde.Server.Jobs
 
 		readonly IMongoCollection<JobStepRef> _jobStepRefs;
 		readonly ITelemetrySink _telemetrySink;
+		readonly IOptionsMonitor<GlobalConfig> _globalConfig;
 
 		/// <summary>
 		/// Constructor
 		/// </summary>
-		/// <param name="mongoService">The database service instance</param>		
-		/// <param name="telemetrySink">Telemetry sink</param>
-		public JobStepRefCollection(MongoService mongoService, ITelemetrySink telemetrySink)
+		public JobStepRefCollection(MongoService mongoService, ITelemetrySink telemetrySink, IOptionsMonitor<GlobalConfig> globalConfig)
 		{
 			List<MongoIndex<JobStepRef>> indexes = new List<MongoIndex<JobStepRef>>();
 			indexes.Add(keys => keys.Ascending(x => x.StreamId).Ascending(x => x.TemplateId).Ascending(x => x.Name).Descending(x => x.Change));
 
 			_jobStepRefs = mongoService.GetCollection<JobStepRef>("JobStepRefs", indexes);
 			_telemetrySink = telemetrySink;
+			_globalConfig = globalConfig;
 		}
 
 		/// <inheritdoc/>
@@ -126,9 +128,9 @@ namespace Horde.Server.Jobs
 			JobStepRef newJobStepRef = new JobStepRef(id, jobName, stepName, streamId, templateId, change, logId, poolId, agentId, state, outcome, updateIssues, lastSuccess, lastWarning, waitTime, initTime, jobStartTimeUtc, startTimeUtc, finishTimeUtc);
 			await _jobStepRefs.ReplaceOneAsync(Builders<JobStepRef>.Filter.Eq(x => x.Id, newJobStepRef.Id), newJobStepRef, new ReplaceOptions { IsUpsert = true });
 
-			if (_telemetrySink.Enabled)
+			if (_globalConfig.CurrentValue.TryGetStream(streamId, out StreamConfig? streamConfig) && !streamConfig.TelemetryStoreId.IsEmpty)
 			{
-				_telemetrySink.SendEvent(TelemetryStoreId.Default, TelemetryRecordMeta.CurrentHordeInstance, new
+				_telemetrySink.SendEvent(streamConfig.TelemetryStoreId, TelemetryRecordMeta.CurrentHordeInstance, new
 				{
 					EventName = "State.JobStepRef",
 					Id = id.ToString(),
diff --git a/Engine/Source/Programs/Horde/Horde.Server/Projects/ProjectConfig.cs b/Engine/Source/Programs/Horde/Horde.Server/Projects/ProjectConfig.cs
--- a/Engine/Source/Programs/Horde/Horde.Server/Projects/ProjectConfig.cs
+++ b/Engine/Source/Programs/Horde/Horde.Server/Projects/ProjectConfig.cs
@@ -9,6 +9,7 @@ using System.Text.Json.Serialization;
 using EpicGames.Core;
 using EpicGames.Horde.Acls;
 using EpicGames.Horde.Projects;
+using EpicGames.Horde.Telemetry;
 using Horde.Server.Acls;
 using Horde.Server.Agents.Pools;
 using Horde.Server.Configuration;
@@ -85,6 +86,11 @@ namespace Horde.Server.Projects
 		/// </summary>
 		public JobOptions JobOptions { get; set; } = new JobOptions();
 
+		/// <summary>
+		/// Telemetry store for Horde data for this project
+		/// </summary>
+		public TelemetryStoreId TelemetryStoreId { get; set; }
+
 		/// <summary>
 		/// List of streams
 		/// </summary>
diff --git a/Engine/Source/Programs/Horde/Horde.Server/Streams/StreamConfig.cs b/Engine/Source/Programs/Horde/Horde.Server/Streams/StreamConfig.cs
--- a/Engine/Source/Programs/Horde/Horde.Server/Streams/StreamConfig.cs
+++ b/Engine/Source/Programs/Horde/Horde.Server/Streams/StreamConfig.cs
@@ -19,6 +19,7 @@ using EpicGames.Horde.Issues;
 using EpicGames.Horde.Jobs.Templates;
 using EpicGames.Horde.Replicators;
 using EpicGames.Horde.Streams;
+using EpicGames.Horde.Telemetry;
 using Horde.Server.Acls;
 using Horde.Server.Configuration;
 using Horde.Server.Issues;
@@ -155,6 +156,11 @@ namespace Horde.Server.Streams
 		/// </summary>
 		public JobOptions JobOptions { get; set; } = new JobOptions();
 
+		/// <summary>
+		/// Telemetry store for Horde data for this stream
+		/// </summary>
+		public TelemetryStoreId TelemetryStoreId { get; set; }
+
 		/// <summary>
 		/// View for the AutoSDK paths to sync. If null, the whole thing will be synced.
 		/// </summary>
@@ -253,6 +259,11 @@ namespace Horde.Server.Streams
 
 			JobOptions.MergeDefaults(projectConfig.JobOptions);
 
+			if (TelemetryStoreId.IsEmpty)
+			{
+				TelemetryStoreId = projectConfig.TelemetryStoreId;
+			}
+
 			foreach (TemplateRefConfig template in Templates)
 			{
 				template.PostLoad(this);
