Horde: Minor documentation fixes.

diff --git a/Engine/Source/Programs/Horde/Docs/Config/Analytics.md b/Engine/Source/Programs/Horde/Docs/Config/Analytics.md
--- a/Engine/Source/Programs/Horde/Docs/Config/Analytics.md
+++ b/Engine/Source/Programs/Horde/Docs/Config/Analytics.md
@@ -35,4 +35,4 @@ running metrics for each time interval. This aggregation is performed according
 ## Charting
 
 The Horde dashboard can render charts showing aggregated metrics collected on the server. These views are configured
-using the `Dashboard.Analytics` section of the globals.json file (see [TelemetryViewConfig](Schema/Telemetry.md#telemetryviewconfig)).
+using the `Dashboard.Analytics` section of the globals.json file (see [TelemetryViewConfig](Schema/Dashboard.md#telemetryviewconfig)).
diff --git a/Engine/Source/Programs/Horde/Docs/Tutorials/Analytics.md b/Engine/Source/Programs/Horde/Docs/Tutorials/Analytics.md
--- a/Engine/Source/Programs/Horde/Docs/Tutorials/Analytics.md
+++ b/Engine/Source/Programs/Horde/Docs/Tutorials/Analytics.md
@@ -48,6 +48,7 @@ the following lines:
             "include": [
                 {
                     "path": "$(HordeDir)/Defaults/default-metrics.telemetry.json"
+                }
             ]
         }
     ]
@@ -63,7 +64,7 @@ can be included by adding the following snippet to your [globals.json](../Config
                 "path": "$(HordeDir)/Defaults/default-analytics.dashboard.json"
             }
         ]
-    ]
+    }
     ```
 
 ## See Also
