Horde: Support setting arbitrary artifact metadata. This cannot be queried, but will be returned in artifact responses.

diff --git a/Engine/Source/Programs/Horde/Horde.Agent/Execution/JobExecutor.cs b/Engine/Source/Programs/Horde/Horde.Agent/Execution/JobExecutor.cs
--- a/Engine/Source/Programs/Horde/Horde.Agent/Execution/JobExecutor.cs
+++ b/Engine/Source/Programs/Horde/Horde.Agent/Execution/JobExecutor.cs
@@ -132,6 +132,7 @@ namespace Horde.Agent.Execution
 			public string? Description { get; set; }
 			public string? BasePath { get; set; }
 			public List<string> Keys { get; set; } = new List<string>();
+			public List<string> Metadata { get; set; } = new List<string>();
 			public string OutputName { get; set; } = String.Empty;
 		}
 
@@ -759,6 +760,7 @@ namespace Horde.Agent.Execution
 				createArtifact.Description = exportedArtifact.Description ?? String.Empty;
 				createArtifact.BasePath = exportedArtifact.BasePath ?? String.Empty;
 				createArtifact.Keys.AddRange(exportedArtifact.Keys);
+				createArtifact.Metadata.AddRange(exportedArtifact.Metadata);
 				createArtifact.OutputName = exportedArtifact.OutputName;
 
 				updateGraph.Artifacts.Add(createArtifact);
@@ -1162,6 +1164,7 @@ namespace Horde.Agent.Execution
 					artifactRequest.Type = graphArtifact.Type;
 					artifactRequest.Description = graphArtifact.Description;
 					artifactRequest.Keys.AddRange(graphArtifact.Keys);
+					artifactRequest.Metadata.AddRange(graphArtifact.Metadata);
 
 					CreateJobArtifactResponseV2 artifact = await jobRpc.Client.CreateArtifactV2Async(artifactRequest, cancellationToken: cancellationToken);
 					logger.LogInformation("Created artifact {ArtifactId} '{ArtifactName}' ({ArtifactType}) with ref {RefName} ({RefUrl})", artifact.Id, artifactRequest.Name, ArtifactType.StepOutput, artifact.RefName, $"{Session.ServerUrl.ToString().TrimEnd('/')}/api/v1/storage/{artifact.NamespaceId}/refs/{artifact.RefName}");
diff --git a/Engine/Source/Programs/Horde/Horde.Server.Tests/Artifacts/ArtifactTests.cs b/Engine/Source/Programs/Horde/Horde.Server.Tests/Artifacts/ArtifactTests.cs
--- a/Engine/Source/Programs/Horde/Horde.Server.Tests/Artifacts/ArtifactTests.cs
+++ b/Engine/Source/Programs/Horde/Horde.Server.Tests/Artifacts/ArtifactTests.cs
@@ -25,7 +25,7 @@ namespace Horde.Server.Tests.Artifacts
 			StreamId streamId = new StreamId("foo");
 
 			IArtifactCollection artifactCollection = ServiceProvider.GetRequiredService<IArtifactCollection>();
-			IArtifact artifact = await artifactCollection.AddAsync(new ArtifactName("default"), ArtifactType.StepOutput, null, streamId, 1, new string[] { "test1", "test2" }, AclScopeName.Root);
+			IArtifact artifact = await artifactCollection.AddAsync(new ArtifactName("default"), ArtifactType.StepOutput, null, streamId, 1, new string[] { "test1", "test2" }, Array.Empty<string>(), AclScopeName.Root);
 
 			{
 				List<IArtifact> artifacts = await artifactCollection.FindAsync(streamId, keys: new[] { "test1" }).ToListAsync();
@@ -59,7 +59,7 @@ namespace Horde.Server.Tests.Artifacts
 			StreamId streamId = new StreamId("foo");
 
 			IArtifactCollection artifactCollection = ServiceProvider.GetRequiredService<IArtifactCollection>();
-			IArtifact artifact = await artifactCollection.AddAsync(new ArtifactName("default"), type, null, streamId, 1, new string[] { "test1", "test2" }, AclScopeName.Root);
+			IArtifact artifact = await artifactCollection.AddAsync(new ArtifactName("default"), type, null, streamId, 1, new string[] { "test1", "test2" }, Array.Empty<string>(), AclScopeName.Root);
 
 			{
 				List<IArtifact> artifacts = await artifactCollection.FindAsync(streamId, keys: new[] { "test1" }).ToListAsync();
@@ -90,7 +90,7 @@ namespace Horde.Server.Tests.Artifacts
 			StreamId streamId = new StreamId("foo");
 
 			IArtifactCollection artifactCollection = ServiceProvider.GetRequiredService<IArtifactCollection>();
-			IArtifact artifact = await artifactCollection.AddAsync(new ArtifactName("default"), type, null, streamId, 1, new string[] { "test1", "test2" }, AclScopeName.Root);
+			IArtifact artifact = await artifactCollection.AddAsync(new ArtifactName("default"), type, null, streamId, 1, new string[] { "test1", "test2" }, Array.Empty<string>(), AclScopeName.Root);
 
 			{
 				List<IArtifact> artifacts = await artifactCollection.FindAsync(streamId, keys: new[] { "test1" }).ToListAsync();
@@ -133,11 +133,11 @@ namespace Horde.Server.Tests.Artifacts
 			IArtifactCollection artifactCollection = ServiceProvider.GetRequiredService<IArtifactCollection>();
 			for (int idx = 0; idx < 10; idx++)
 			{
-				await artifactCollection.AddAsync(new ArtifactName($"default-{idx}"), type, null, fooStreamId, 1, new string[] { "test1", "test2" }, AclScopeName.Root);
+				await artifactCollection.AddAsync(new ArtifactName($"default-{idx}"), type, null, fooStreamId, 1, new string[] { "test1", "test2" }, Array.Empty<string>(), AclScopeName.Root);
 			}
 			for (int idx = 0; idx < 10; idx++)
 			{
-				await artifactCollection.AddAsync(new ArtifactName($"default-{idx}"), type, null, barStreamId, 1, new string[] { "test1", "test2" }, AclScopeName.Root);
+				await artifactCollection.AddAsync(new ArtifactName($"default-{idx}"), type, null, barStreamId, 1, new string[] { "test1", "test2" }, Array.Empty<string>(), AclScopeName.Root);
 			}
 
 			await Clock.AdvanceAsync(TimeSpan.FromDays(2.0));
diff --git a/Engine/Source/Programs/Horde/Horde.Server.Tests/Jobs/JobCollectionTests.cs b/Engine/Source/Programs/Horde/Horde.Server.Tests/Jobs/JobCollectionTests.cs
--- a/Engine/Source/Programs/Horde/Horde.Server.Tests/Jobs/JobCollectionTests.cs
+++ b/Engine/Source/Programs/Horde/Horde.Server.Tests/Jobs/JobCollectionTests.cs
@@ -822,7 +822,7 @@ namespace Horde.Server.Tests.Jobs
 			IGraph baseGraph = await GraphCollection.AddAsync(templateMock.Object, null);
 
 			List<NewGraphArtifact> newArtifacts = new List<NewGraphArtifact>();
-			newArtifacts.Add(new NewGraphArtifact(new ArtifactName("foo"), new ArtifactType("type"), "hello world", "Engine/Source", new List<string>(), "fileset"));
+			newArtifacts.Add(new NewGraphArtifact(new ArtifactName("foo"), new ArtifactType("type"), "hello world", "Engine/Source", new List<string>(), new List<string>(), "fileset"));
 
 			IGraph graph = await GraphCollection.AppendAsync(baseGraph, newArtifactRequests: newArtifacts);
 			Assert.AreEqual(1, graph.Artifacts.Count);
diff --git a/Engine/Source/Programs/Horde/Horde.Server/Artifacts/ArtifactCollection.cs b/Engine/Source/Programs/Horde/Horde.Server/Artifacts/ArtifactCollection.cs
--- a/Engine/Source/Programs/Horde/Horde.Server/Artifacts/ArtifactCollection.cs
+++ b/Engine/Source/Programs/Horde/Horde.Server/Artifacts/ArtifactCollection.cs
@@ -50,6 +50,11 @@ namespace Horde.Server.Artifacts
 
 			IReadOnlyList<string> IArtifact.Keys => Keys;
 
+			[BsonElement("met")]
+			public List<string> Metadata { get; set; } = new List<string>();
+
+			IReadOnlyList<string> IArtifact.Metadata => Metadata;
+
 			[BsonElement("ns")]
 			public NamespaceId NamespaceId { get; set; }
 
@@ -72,7 +77,7 @@ namespace Horde.Server.Artifacts
 			{
 			}
 
-			public Artifact(ArtifactId id, ArtifactName name, ArtifactType type, string? description, StreamId streamId, int change, IEnumerable<string> keys, NamespaceId namespaceId, RefName refName, DateTime createdAtUtc, AclScopeName scopeName)
+			public Artifact(ArtifactId id, ArtifactName name, ArtifactType type, string? description, StreamId streamId, int change, IEnumerable<string> keys, IEnumerable<string> metadata, NamespaceId namespaceId, RefName refName, DateTime createdAtUtc, AclScopeName scopeName)
 			{
 				Id = id;
 				Name = name;
@@ -81,6 +86,7 @@ namespace Horde.Server.Artifacts
 				StreamId = streamId;
 				Change = change;
 				Keys.AddRange(keys);
+				Metadata.AddRange(metadata);
 				NamespaceId = namespaceId;
 				RefName = refName;
 				CreatedAtUtc = createdAtUtc;
@@ -111,14 +117,14 @@ namespace Horde.Server.Artifacts
 		public static string GetArtifactPath(StreamId streamId, ArtifactName name, ArtifactType type) => $"{streamId}/{name}/{type}";
 
 		/// <inheritdoc/>
-		public async Task<IArtifact> AddAsync(ArtifactName name, ArtifactType type, string? description, StreamId streamId, int change, IEnumerable<string> keys, AclScopeName scopeName, CancellationToken cancellationToken)
+		public async Task<IArtifact> AddAsync(ArtifactName name, ArtifactType type, string? description, StreamId streamId, int change, IEnumerable<string> keys, IEnumerable<string> metadata, AclScopeName scopeName, CancellationToken cancellationToken)
 		{
 			ArtifactId id = new ArtifactId(BinaryIdUtils.CreateNew());
 
 			NamespaceId namespaceId = Namespace.Artifacts;
 			RefName refName = new RefName($"{GetArtifactPath(streamId, name, type)}/{change}/{id}");
 
-			Artifact artifact = new Artifact(id, name, type, description, streamId, change, keys, namespaceId, refName, _clock.UtcNow, scopeName);
+			Artifact artifact = new Artifact(id, name, type, description, streamId, change, keys, metadata, namespaceId, refName, _clock.UtcNow, scopeName);
 			await _artifacts.InsertOneAsync(artifact, null, cancellationToken);
 			return artifact;
 		}
@@ -131,7 +137,7 @@ namespace Horde.Server.Artifacts
 		}
 
 		/// <inheritdoc/>
-		public async IAsyncEnumerable<IArtifact> FindAsync(StreamId? streamId = null, int? minChange = null, int? maxChange = null, ArtifactName? name = null, ArtifactType? type = null, IEnumerable<string>? keys = null, [EnumeratorCancellation] CancellationToken cancellationToken = default)
+		public async IAsyncEnumerable<IArtifact> FindAsync(StreamId? streamId = null, int? minChange = null, int? maxChange = null, ArtifactName? name = null, ArtifactType? type = null, IEnumerable<string>? keys = null, int maxResults = 100, [EnumeratorCancellation] CancellationToken cancellationToken = default)
 		{
 			FilterDefinition<Artifact> filter = FilterDefinition<Artifact>.Empty;
 			if (streamId != null)
@@ -159,7 +165,7 @@ namespace Horde.Server.Artifacts
 				filter &= Builders<Artifact>.Filter.All(x => x.Keys, keys);
 			}
 
-			using (IAsyncCursor<Artifact> cursor = await _artifacts.Find(filter).SortByDescending(x => x.Change).ThenByDescending(x => x.Id).ToCursorAsync(cancellationToken))
+			using (IAsyncCursor<Artifact> cursor = await _artifacts.Find(filter).SortByDescending(x => x.Change).ThenByDescending(x => x.Id).Limit(maxResults).ToCursorAsync(cancellationToken))
 			{
 				while (await cursor.MoveNextAsync(cancellationToken))
 				{
diff --git a/Engine/Source/Programs/Horde/Horde.Server/Artifacts/ArtifactsController.cs b/Engine/Source/Programs/Horde/Horde.Server/Artifacts/ArtifactsController.cs
--- a/Engine/Source/Programs/Horde/Horde.Server/Artifacts/ArtifactsController.cs
+++ b/Engine/Source/Programs/Horde/Horde.Server/Artifacts/ArtifactsController.cs
@@ -76,7 +76,7 @@ namespace Horde.Server.Artifacts
 			{
 				if (streamConfig.Authorize(ArtifactAclAction.WriteArtifact, User))
 				{
-					return await CreateArtifactInternalAsync(request.Name, request.Type, request.Description, request.StreamId.Value, request.Change.Value, request.Keys, AclScopeName.Root, cancellationToken);
+					return await CreateArtifactInternalAsync(request.Name, request.Type, request.Description, request.StreamId.Value, request.Change.Value, request.Keys, request.Metadata, AclScopeName.Root, cancellationToken);
 				}
 			}
 
@@ -128,15 +128,15 @@ namespace Horde.Server.Artifacts
 					scopeName = templateRefConfig.Acl.ScopeName;
 				}
 
-				return await CreateArtifactInternalAsync(request.Name, request.Type, request.Description, streamId, request.Change ?? job.Change, keys, scopeName, cancellationToken);
+				return await CreateArtifactInternalAsync(request.Name, request.Type, request.Description, streamId, request.Change ?? job.Change, keys, request.Metadata, scopeName, cancellationToken);
 			}
 
 			return Forbid(ArtifactAclAction.WriteArtifact);
 		}
 
-		async Task<ActionResult<CreateArtifactResponse>> CreateArtifactInternalAsync(ArtifactName name, ArtifactType type, string? description, StreamId streamId, int change, List<string> keys, AclScopeName scopeName, CancellationToken cancellationToken)
+		async Task<ActionResult<CreateArtifactResponse>> CreateArtifactInternalAsync(ArtifactName name, ArtifactType type, string? description, StreamId streamId, int change, List<string> keys, List<string> metadata, AclScopeName scopeName, CancellationToken cancellationToken)
 		{
-			IArtifact artifact = await _artifactCollection.AddAsync(name, type, description, streamId, change, keys, scopeName, cancellationToken);
+			IArtifact artifact = await _artifactCollection.AddAsync(name, type, description, streamId, change, keys, metadata, scopeName, cancellationToken);
 			RefName? prevRefName = await GetPrevRefNameForArtifactAsync(artifact, cancellationToken);
 
 			List<AclClaimConfig> claims = new List<AclClaimConfig>();
@@ -187,7 +187,7 @@ namespace Horde.Server.Artifacts
 				return Forbid(ArtifactAclAction.ReadArtifact, artifact.AclScope);
 			}
 
-			return PropertyFilter.Apply(new GetArtifactResponse(artifact.Id, artifact.Name, artifact.Type, artifact.Description, artifact.StreamId, artifact.Change, artifact.Keys), filter);
+			return PropertyFilter.Apply(new GetArtifactResponse(artifact.Id, artifact.Name, artifact.Type, artifact.Description, artifact.StreamId, artifact.Change, artifact.Keys, artifact.Metadata), filter);
 		}
 
 		/// <summary>
@@ -538,24 +538,25 @@ namespace Horde.Server.Artifacts
 		/// <param name="name">Artifact name</param>
 		/// <param name="type">Type of the artifact</param>
 		/// <param name="keys">Keys to find</param>
+		/// <param name="maxResults">Maximum number of results to return</param>
 		/// <param name="filter">Filter for returned values</param>
 		/// <returns>Information about all the artifacts</returns>
 		[HttpGet]
 		[Route("/api/v2/artifacts")]
 		[ProducesResponseType(typeof(FindArtifactsResponse), 200)]
-		public async Task<ActionResult<object>> FindArtifactsAsync([FromQuery] StreamId? streamId = null, [FromQuery] int? minChange = null, [FromQuery] int? maxChange = null, [FromQuery(Name = "name")] ArtifactName? name = null, [FromQuery(Name = "type")] ArtifactType? type = null, [FromQuery(Name = "key")] IEnumerable<string>? keys = null, [FromQuery] PropertyFilter? filter = null)
+		public async Task<ActionResult<object>> FindArtifactsAsync([FromQuery] StreamId? streamId = null, [FromQuery] int? minChange = null, [FromQuery] int? maxChange = null, [FromQuery(Name = "name")] ArtifactName? name = null, [FromQuery(Name = "type")] ArtifactType? type = null, [FromQuery(Name = "key")] IEnumerable<string>? keys = null, [FromQuery] int maxResults = 100, [FromQuery] PropertyFilter? filter = null)
 		{
-			if (streamId == null && (minChange != null || maxChange != null || name != null || type != null))
+			if (streamId == null && keys == null)
 			{
-				return BadRequest("Missing StreamId parameter");
+				return BadRequest("Missing streamId or key parameter");
 			}
 
 			FindArtifactsResponse response = new FindArtifactsResponse();
-			await foreach (IArtifact artifact in _artifactCollection.FindAsync(streamId, minChange, maxChange, name, type, keys, HttpContext.RequestAborted))
+			await foreach (IArtifact artifact in _artifactCollection.FindAsync(streamId, minChange, maxChange, name, type, keys, maxResults, HttpContext.RequestAborted))
 			{
 				if (_globalConfig.Authorize(artifact.AclScope, ArtifactAclAction.ReadArtifact, User))
 				{
-					response.Artifacts.Add(new GetArtifactResponse(artifact.Id, artifact.Name, artifact.Type, artifact.Description, artifact.StreamId, artifact.Change, artifact.Keys));
+					response.Artifacts.Add(new GetArtifactResponse(artifact.Id, artifact.Name, artifact.Type, artifact.Description, artifact.StreamId, artifact.Change, artifact.Keys, artifact.Metadata));
 				}
 			}
 
diff --git a/Engine/Source/Programs/Horde/Horde.Server/Artifacts/IArtifact.cs b/Engine/Source/Programs/Horde/Horde.Server/Artifacts/IArtifact.cs
--- a/Engine/Source/Programs/Horde/Horde.Server/Artifacts/IArtifact.cs
+++ b/Engine/Source/Programs/Horde/Horde.Server/Artifacts/IArtifact.cs
@@ -49,6 +49,11 @@ namespace Horde.Server.Artifacts
 		/// </summary>
 		public IReadOnlyList<string> Keys { get; }
 
+		/// <summary>
+		/// Metadata for the artifact
+		/// </summary>
+		public IReadOnlyList<string> Metadata { get; }
+
 		/// <summary>
 		/// Storage namespace containing the data
 		/// </summary>
diff --git a/Engine/Source/Programs/Horde/Horde.Server/Artifacts/IArtifactCollection.cs b/Engine/Source/Programs/Horde/Horde.Server/Artifacts/IArtifactCollection.cs
--- a/Engine/Source/Programs/Horde/Horde.Server/Artifacts/IArtifactCollection.cs
+++ b/Engine/Source/Programs/Horde/Horde.Server/Artifacts/IArtifactCollection.cs
@@ -24,10 +24,11 @@ namespace Horde.Server.Artifacts
 		/// <param name="streamId">Stream that the artifact was built from</param>
 		/// <param name="change">Change number that the artifact was built from</param>
 		/// <param name="keys">Keys for the artifact</param>
+		/// <param name="metadata">Metadata for the artifact</param>
 		/// <param name="scopeName">Inherited scope used for permissions</param>
 		/// <param name="cancellationToken">Cancellation token for the operation</param>
 		/// <returns>The new log file document</returns>
-		Task<IArtifact> AddAsync(ArtifactName name, ArtifactType type, string? description, StreamId streamId, int change, IEnumerable<string> keys, AclScopeName scopeName, CancellationToken cancellationToken = default);
+		Task<IArtifact> AddAsync(ArtifactName name, ArtifactType type, string? description, StreamId streamId, int change, IEnumerable<string> keys, IEnumerable<string> metadata, AclScopeName scopeName, CancellationToken cancellationToken = default);
 
 		/// <summary>
 		/// Deletes artifacts
@@ -55,9 +56,10 @@ namespace Horde.Server.Artifacts
 		/// <param name="name">Name of the artifact to search for</param>
 		/// <param name="type">The artifact type</param>
 		/// <param name="keys">Set of keys, all of which must all be present on any returned artifacts</param>
+		/// <param name="maxResults">Maximum number of results to return</param>
 		/// <param name="cancellationToken">Cancellation token for the operation</param>
 		/// <returns>Sequence of artifacts. Ordered by descending CL order, then by descending order in which they were created.</returns>
-		IAsyncEnumerable<IArtifact> FindAsync(StreamId? streamId = null, int? minChange = null, int? maxChange = null, ArtifactName? name = null, ArtifactType? type = null, IEnumerable<string>? keys = null, CancellationToken cancellationToken = default);
+		IAsyncEnumerable<IArtifact> FindAsync(StreamId? streamId = null, int? minChange = null, int? maxChange = null, ArtifactName? name = null, ArtifactType? type = null, IEnumerable<string>? keys = null, int maxResults = 100, CancellationToken cancellationToken = default);
 
 		/// <summary>
 		/// Gets an artifact by ID
diff --git a/Engine/Source/Programs/Horde/Horde.Server/Jobs/Graphs/GraphCollection.cs b/Engine/Source/Programs/Horde/Horde.Server/Jobs/Graphs/GraphCollection.cs
--- a/Engine/Source/Programs/Horde/Horde.Server/Jobs/Graphs/GraphCollection.cs
+++ b/Engine/Source/Programs/Horde/Horde.Server/Jobs/Graphs/GraphCollection.cs
@@ -197,25 +197,29 @@ namespace Horde.Server.Jobs.Graphs
 			public string Description { get; set; }
 			public string BasePath { get; set; }
 			public List<string> Keys { get; set; }
+			public List<string> Metadata { get; set; }
 			public string OutputName { get; set; }
 
 			IReadOnlyList<string> IGraphArtifact.Keys => Keys;
+			IReadOnlyList<string> IGraphArtifact.Metadata => Metadata;
 
 			private GraphArtifact()
 			{
 				Description = String.Empty;
 				BasePath = String.Empty;
 				Keys = new List<string>();
+				Metadata = new List<string>();
 				OutputName = String.Empty;
 			}
 
-			public GraphArtifact(ArtifactName name, ArtifactType type, string description, string basePath, IReadOnlyList<string> keys, string outputName)
+			public GraphArtifact(ArtifactName name, ArtifactType type, string description, string basePath, IReadOnlyList<string> keys, IReadOnlyList<string> metadata, string outputName)
 			{
 				Name = name;
 				Type = type;
 				Description = description;
 				BasePath = basePath;
 				Keys = keys.ToList();
+				Metadata = metadata.ToList();
 				OutputName = outputName;
 			}
 		}
@@ -335,7 +339,7 @@ namespace Horde.Server.Jobs.Graphs
 				{
 					foreach (NewGraphArtifact newArtifactRequest in newArtifactRequests)
 					{
-						newArtifacts.Add(new GraphArtifact(newArtifactRequest.Name, newArtifactRequest.Type, newArtifactRequest.Description, newArtifactRequest.BasePath, newArtifactRequest.Keys, newArtifactRequest.OutputName));
+						newArtifacts.Add(new GraphArtifact(newArtifactRequest.Name, newArtifactRequest.Type, newArtifactRequest.Description, newArtifactRequest.BasePath, newArtifactRequest.Keys, newArtifactRequest.Metadata, newArtifactRequest.OutputName));
 					}
 				}
 
diff --git a/Engine/Source/Programs/Horde/Horde.Server/Jobs/Graphs/IGraph.cs b/Engine/Source/Programs/Horde/Horde.Server/Jobs/Graphs/IGraph.cs
--- a/Engine/Source/Programs/Horde/Horde.Server/Jobs/Graphs/IGraph.cs
+++ b/Engine/Source/Programs/Horde/Horde.Server/Jobs/Graphs/IGraph.cs
@@ -298,6 +298,11 @@ namespace Horde.Server.Jobs.Graphs
 		/// </summary>
 		public IReadOnlyList<string> Keys { get; }
 
+		/// <summary>
+		/// Metadata for the artifact
+		/// </summary>
+		public IReadOnlyList<string> Metadata { get; }
+
 		/// <summary>
 		/// Tag for the artifact files
 		/// </summary>
@@ -682,5 +687,5 @@ namespace Horde.Server.Jobs.Graphs
 	/// <summary>
 	/// Information about an artifact
 	/// </summary>
-	public record class NewGraphArtifact(ArtifactName Name, ArtifactType Type, string Description, string BasePath, List<string> Keys, string OutputName);
+	public record class NewGraphArtifact(ArtifactName Name, ArtifactType Type, string Description, string BasePath, List<string> Keys, List<string> Metadata, string OutputName);
 }
diff --git a/Engine/Source/Programs/Horde/Horde.Server/Jobs/JobRpcService.cs b/Engine/Source/Programs/Horde/Horde.Server/Jobs/JobRpcService.cs
--- a/Engine/Source/Programs/Horde/Horde.Server/Jobs/JobRpcService.cs
+++ b/Engine/Source/Programs/Horde/Horde.Server/Jobs/JobRpcService.cs
@@ -89,12 +89,14 @@ namespace Horde.Server.Jobs
 			keys.Add($"job:{job.Id}");
 			keys.Add($"job:{job.Id}/step:{step.Id}");
 
+			List<string> metadata = new List<string>();
+
 			if (!_globalConfig.TryGetTemplate(job.StreamId, job.TemplateId, out TemplateRefConfig? templateConfig))
 			{
 				throw new StructuredRpcException(StatusCode.NotFound, "Couldn't find template {TemplateId} in stream {StreamId}", job.TemplateId, job.StreamId);
 			}
 
-			IArtifact artifact = await _artifactCollection.AddAsync(new ArtifactName("default"), type, null, job.StreamId, job.Change, keys, templateConfig.Acl.ScopeName, context.CancellationToken);
+			IArtifact artifact = await _artifactCollection.AddAsync(new ArtifactName("default"), type, null, job.StreamId, job.Change, keys, metadata, templateConfig.Acl.ScopeName, context.CancellationToken);
 
 			List<AclClaimConfig> claims = new List<AclClaimConfig>();
 			claims.Add(new AclClaimConfig(HordeClaimTypes.WriteNamespace, artifact.NamespaceId.ToString()));
@@ -127,7 +129,7 @@ namespace Horde.Server.Jobs
 				description = request.Name;
 			}
 
-			IArtifact artifact = await _artifactCollection.AddAsync(name, type, description, job.StreamId, job.Change, keys, templateConfig.Acl.ScopeName, context.CancellationToken);
+			IArtifact artifact = await _artifactCollection.AddAsync(name, type, description, job.StreamId, job.Change, keys, request.Metadata, templateConfig.Acl.ScopeName, context.CancellationToken);
 
 			List<AclClaimConfig> claims = new List<AclClaimConfig>();
 			claims.Add(new AclClaimConfig(HordeClaimTypes.WriteNamespace, $"{artifact.NamespaceId}:{artifact.RefName}"));
@@ -863,7 +865,7 @@ namespace Horde.Server.Jobs
 						description = artifact.Name;
 					}
 
-					newArtifacts.Add(new NewGraphArtifact(name, type, description, artifact.BasePath, artifact.Keys.ToList(), artifact.OutputName));
+					newArtifacts.Add(new NewGraphArtifact(name, type, description, artifact.BasePath, artifact.Keys.ToList(), artifact.Metadata.ToList(), artifact.OutputName));
 				}
 
 				// Create the new graph
diff --git a/Engine/Source/Programs/Shared/EpicGames.Horde/Artifacts/ArtifactMessages.cs b/Engine/Source/Programs/Shared/EpicGames.Horde/Artifacts/ArtifactMessages.cs
--- a/Engine/Source/Programs/Shared/EpicGames.Horde/Artifacts/ArtifactMessages.cs
+++ b/Engine/Source/Programs/Shared/EpicGames.Horde/Artifacts/ArtifactMessages.cs
@@ -18,7 +18,8 @@ namespace EpicGames.Horde.Artifacts
 	/// <param name="StreamId">Stream to create the artifact for</param>
 	/// <param name="Change">Change number for the artifact</param>
 	/// <param name="Keys">Keys used to identify the artifact</param>
-	public record CreateArtifactRequest(ArtifactName Name, ArtifactType Type, string? Description, StreamId? StreamId, int? Change, List<string> Keys);
+	/// <param name="Metadata">Metadata for the artifact</param>
+	public record CreateArtifactRequest(ArtifactName Name, ArtifactType Type, string? Description, StreamId? StreamId, int? Change, List<string> Keys, List<string> Metadata);
 
 	/// <summary>
 	/// Information about a created artifact
@@ -70,10 +71,15 @@ namespace EpicGames.Horde.Artifacts
 		/// </summary>
 		public IReadOnlyList<string> Keys { get; }
 
+		/// <summary>
+		/// List of metadata properties stored with the artifact, in the form 'Key=Value'
+		/// </summary>
+		public IReadOnlyList<string> Metadata { get; }
+
 		/// <summary>
 		/// Constructor
 		/// </summary>
-		public GetArtifactResponse(ArtifactId id, ArtifactName name, ArtifactType type, string? description, StreamId streamId, int change, IReadOnlyList<string> keys)
+		public GetArtifactResponse(ArtifactId id, ArtifactName name, ArtifactType type, string? description, StreamId streamId, int change, IReadOnlyList<string> keys, IReadOnlyList<string> metadata)
 		{
 			Id = id;
 			Name = name;
@@ -82,6 +88,7 @@ namespace EpicGames.Horde.Artifacts
 			StreamId = streamId;
 			Change = change;
 			Keys = keys;
+			Metadata = metadata;
 		}
 	}
 
diff --git a/Engine/Source/Programs/Shared/EpicGames.Horde/HordeHttpClient.cs b/Engine/Source/Programs/Shared/EpicGames.Horde/HordeHttpClient.cs
--- a/Engine/Source/Programs/Shared/EpicGames.Horde/HordeHttpClient.cs
+++ b/Engine/Source/Programs/Shared/EpicGames.Horde/HordeHttpClient.cs
@@ -121,10 +121,11 @@ namespace EpicGames.Horde
 		/// <param name="streamId">Stream to create the artifact for</param>
 		/// <param name="change">Change number for the artifact</param>
 		/// <param name="keys">Keys used to identify the artifact</param>
+		/// <param name="metadata">Metadata for the artifact</param>
 		/// <param name="cancellationToken">Cancellation token for the operation</param>
-		public Task<CreateArtifactResponse> CreateArtifactAsync(ArtifactName name, ArtifactType type, string? description, StreamId? streamId = null, int? change = null, List<string>? keys = null, CancellationToken cancellationToken = default)
+		public Task<CreateArtifactResponse> CreateArtifactAsync(ArtifactName name, ArtifactType type, string? description, StreamId? streamId = null, int? change = null, List<string>? keys = null, List<string>? metadata = null, CancellationToken cancellationToken = default)
 		{
-			return PostAsync<CreateArtifactResponse, CreateArtifactRequest>(_httpClient, $"api/v2/artifacts", new CreateArtifactRequest(name, type, description, streamId, change, keys ?? new List<string>()), cancellationToken);
+			return PostAsync<CreateArtifactResponse, CreateArtifactRequest>(_httpClient, $"api/v2/artifacts", new CreateArtifactRequest(name, type, description, streamId, change, keys ?? new List<string>(), metadata ?? new List<string>()), cancellationToken);
 		}
 
 		/// <summary>
@@ -183,9 +184,10 @@ namespace EpicGames.Horde
 		/// <param name="streamId">Stream to look for the artifact in</param>
 		/// <param name="minChange">The minimum change number for the artifacts</param>
 		/// <param name="maxChange">The minimum change number for the artifacts</param>
+		/// <param name="keys">Keys for artifacts to return</param>
 		/// <param name="cancellationToken">Cancellation token for the operation</param>
 		/// <returns>Information about all the artifacts</returns>
-		public async Task<List<GetArtifactResponse>> FindArtifactsByTypeAsync(ArtifactType type, StreamId? streamId = null, int? minChange = null, int? maxChange = null, CancellationToken cancellationToken = default)
+		public async Task<List<GetArtifactResponse>> FindArtifactsByTypeAsync(ArtifactType type, StreamId? streamId = null, int? minChange = null, int? maxChange = null, IEnumerable<string>? keys = null, CancellationToken cancellationToken = default)
 		{
 			QueryStringBuilder queryParams = new QueryStringBuilder();
 			queryParams.Add("type", type.ToString());
@@ -205,6 +207,14 @@ namespace EpicGames.Horde
 				queryParams.Add("maxChange", maxChange.ToString()!);
 			}
 
+			if (keys != null)
+			{
+				foreach (string key in keys)
+				{
+					queryParams.Add("key", key);
+				}
+			}
+
 			FindArtifactsResponse response = await GetAsync<FindArtifactsResponse>(_httpClient, $"api/v2/artifacts?{queryParams}", cancellationToken);
 			return response.Artifacts;
 		}
