Horde: Updating schema documentation

#rnx

diff --git a/Engine/Source/Programs/Horde/Docs/Config/Schema/Projects.md b/Engine/Source/Programs/Horde/Docs/Config/Schema/Projects.md
--- a/Engine/Source/Programs/Horde/Docs/Config/Schema/Projects.md
+++ b/Engine/Source/Programs/Horde/Docs/Config/Schema/Projects.md
@@ -16,6 +16,7 @@ Name | Description
 `pools` | [`PoolConfig`](#poolconfig)`[]`<br>List of pools for this project
 `categories` | [`ProjectCategoryConfig`](#projectcategoryconfig)`[]`<br>Categories to include in this project
 `jobOptions` | [`JobOptions`](#joboptions)<br>Default settings for executing jobs
+`telemetryStoreId` | `string`<br>Telemetry store for Horde data for this project
 `streams` | [`StreamConfig`](Streams.md)`[]`<br>List of streams
 `acl` | [`AclConfig`](#aclconfig)<br>Acl entries
 
diff --git a/Engine/Source/Programs/Horde/Docs/Config/Schema/Streams.md b/Engine/Source/Programs/Horde/Docs/Config/Schema/Streams.md
--- a/Engine/Source/Programs/Horde/Docs/Config/Schema/Streams.md
+++ b/Engine/Source/Programs/Horde/Docs/Config/Schema/Streams.md
@@ -18,6 +18,7 @@ Name | Description
 `notificationChannelFilter` | `string`<br>Notification channel filter for this template. Can be Success|Failure|Warnings
 `triageChannel` | `string`<br>Channel to post issue triage notifications
 `jobOptions` | [`JobOptions`](#joboptions)<br>Default settings for executing jobs
+`telemetryStoreId` | `string`<br>Telemetry store for Horde data for this stream
 `autoSdkView` | `string[]`<br>View for the AutoSDK paths to sync. If null, the whole thing will be synced.
 `defaultPreflightTemplate` | `string`<br>Legacy name for the default preflight template
 `defaultPreflight` | [`DefaultPreflightConfig`](#defaultpreflightconfig)<br>Default template for running preflights
