Horde: Fix metadata not being sent back to agents for creating artifacts.

#rnx

diff --git a/Engine/Source/Programs/Horde/Horde.Server/Jobs/JobRpcService.cs b/Engine/Source/Programs/Horde/Horde.Server/Jobs/JobRpcService.cs
--- a/Engine/Source/Programs/Horde/Horde.Server/Jobs/JobRpcService.cs
+++ b/Engine/Source/Programs/Horde/Horde.Server/Jobs/JobRpcService.cs
@@ -651,6 +651,7 @@ namespace Horde.Server.Jobs
 					{
 						CreateGraphArtifactRequest stepArtifact = new CreateGraphArtifactRequest { Name = artifact.Name.ToString(), Type = artifact.Type.ToString(), Description = artifact.Description, BasePath = artifact.BasePath, OutputName = artifact.OutputName };
 						stepArtifact.Keys.AddRange(artifact.Keys);
+						stepArtifact.Metadata.AddRange(artifact.Metadata);
 						response.Artifacts.Add(stepArtifact);
 					}
 				}
