Horde: Device service -> Device manager in docs
#submitool bypass
#rnx

diff --git a/Engine/Source/Programs/Horde/Docs/Config/Devices.md b/Engine/Source/Programs/Horde/Docs/Config/Devices.md
--- a/Engine/Source/Programs/Horde/Docs/Config/Devices.md
+++ b/Engine/Source/Programs/Horde/Docs/Config/Devices.md
@@ -1,8 +1,8 @@
 [Horde](../../README.md) > [Configuration](../Config.md) > Devices
 
-# Device Service
+# Device Manager
 
-The Horde Device Service manages mobile and console development kit resources. The service is used extensively at Epic and is quite mature.  It features:
+The Horde device manager administers mobile and console development kit resources. The service is used extensively at Epic and is quite mature.  It features:
 
 * Configurable device platforms and pools
 * Shared remote device resources with user checkouts
@@ -21,7 +21,7 @@ It is also possible to set up automation jobs targeting the checked-out devices
 
 ## Automation Pools
 
-The Device Service supports automation job device reservations, which can be constrained by pool, platform, and model.
+The device manager supports automation job device reservations, which can be constrained by pool, platform, and model.
 
 We recommend that you use the [Gauntlet integration](#gauntlet-integration) with the reservation system.
 However, you can also drive custom solutions using the same REST API.
@@ -32,7 +32,7 @@ Device platform hardware is partitioned into pools for use by automated tests an
 
 #### Example
 
-The following config fragment declares a device service configuration that:
+The following config fragment declares a device manager configuration that:
 
 * Adds an `Android` device platform specifying several models
 * Adds two device pools, an `Automation` pool which automated tests can use, and a `Shared` pool from which users can check out remote device hardware
@@ -82,20 +82,21 @@ This includes support for:
 
 ## Gauntlet Integration
 
-[Gauntlet](https://docs.unrealengine.com/en-US/gauntlet-automation-framework-overview-in-unreal-engine) tests can reserve hardware from the Horde Device Service. This integration includes features such as reporting problems with devices and recovery with new devices. It also supports reservation blocks, which can be used to reuse a device with an installed build across a series of automation tests.  
+[Gauntlet](https://docs.unrealengine.com/en-US/gauntlet-automation-framework-overview-in-unreal-engine) tests can reserve hardware from the Horde device manager. This integration includes features such as reporting problems with devices and recovery with new devices. It also supports reservation blocks, which can be used to reuse a device with an installed build across a series of automation tests.  
 
 #### Example
 
 The following [**BuildGraph**](https://docs.unrealengine.com/en-US/buildgraph-for-unreal-engine/) fragment declares:
 
-* `HordeDeviceService` and `HordeDevicePool` properties that specify your Horde server and which pool to use.
+* `HordeDeviceManager` and `HordeDevicePool` properties that specify your Horde server and which pool to use.
 * Adds a `BootTest Android` node, which will reserve an Android Pixel 8 for the test
 
 ---
-	<Property Name="HordeDeviceService" Value="https://horde.yourdomain.com" />
+
+	<Property Name="HordeDeviceManager" Value="https://horde.yourdomain.com" />
 	<Property Name="HordeDevicePool" Value="UE5" />
 		
 	<Node Name="BootTest Android">
-		<Command Name="RunUnreal" Arguments="-test=UE.BootTest -platform=Android " -deviceurl=&quot;$(HordeDeviceService)&quot; -devicepool=&quot;$(HordeDevicePool)&quot; -PerfModel=Pixel8/>
+		<Command Name="RunUnreal" Arguments="-test=UE.BootTest -platform=Android " -deviceurl=&quot;$(HordeDeviceManager)&quot; -devicepool=&quot;$(HordeDevicePool)&quot; -PerfModel=Pixel8/>
 	</Node>
 ---
diff --git a/Engine/Source/Programs/Horde/HordeDashboard/package.json b/Engine/Source/Programs/Horde/HordeDashboard/package.json
--- a/Engine/Source/Programs/Horde/HordeDashboard/package.json
+++ b/Engine/Source/Programs/Horde/HordeDashboard/package.json
@@ -1,6 +1,6 @@
 {
   "name": "horde-dashboard",
-  "version": "0.2.0",
+  "version": "0.2.1",
   "private": true,
   "main": "index.js",
   "proxy": "https://localhost:5001",
