Horde: Fix telemetry view analytics links
#rnx

diff --git a/Engine/Source/Programs/Horde/HordeDashboard/src/components/telemetry/TelemetryView.tsx b/Engine/Source/Programs/Horde/HordeDashboard/src/components/telemetry/TelemetryView.tsx
--- a/Engine/Source/Programs/Horde/HordeDashboard/src/components/telemetry/TelemetryView.tsx
+++ b/Engine/Source/Programs/Horde/HordeDashboard/src/components/telemetry/TelemetryView.tsx
@@ -14,8 +14,6 @@ import { Breadcrumbs } from "../Breadcrumbs";
 import { TopNav } from "../TopNav";
 import { TelemetryViewData, clearTelemetryViewMetrics, getTelemetryViewData, graphColors } from "./TelemetryData";
 import { TelemetryLineRenderer } from "./TelemetryLineGraph";
-import { getSiteConfig } from "../../backend/Config";
-import { UnderConstruction } from "../UnderConstruction";
 
 const timeSelections: TimeSelection[] = [
    {
@@ -1350,8 +1348,10 @@ export const SearchUpdate: React.FC = observer(() => {
    const csearch = "query=" + btoa(handler.search.toString());
 
    useEffect(() => {
-
-      setSearchParams(csearch, { replace: true });
+      if (handler.search.toString().length) {
+         setSearchParams(csearch, { replace: true });
+      }
+      
    }, [csearch, setSearchParams])
 
    // subscribe
@@ -1360,14 +1360,10 @@ export const SearchUpdate: React.FC = observer(() => {
    return null;
 });
 
-export const TelemetryView: React.FC = () => {   
-
-   const siteConfig = getSiteConfig();
+export const TelemetryView: React.FC = () => {
 
    useEffect(() => {
-      if (siteConfig.environment === "dev") {
-         handler.initialize();
-      }
+      handler.initialize();
       return () => {
          handler.clear();
       };
@@ -1376,10 +1372,6 @@ export const TelemetryView: React.FC = () => {
    const windowSize = useWindowSize();
    const navigate = useNavigate();
 
-   if (siteConfig.environment !== "dev") {
-      return <UnderConstruction/>
-   }
-
    const vw = Math.max(document.documentElement.clientWidth, window.innerWidth || 0);
 
    const { hordeClasses, modeColors } = getHordeStyling();
@@ -1396,7 +1388,7 @@ export const TelemetryView: React.FC = () => {
       <Breadcrumbs items={[{ text: 'Analytics' }]} />
       <Stack horizontal styles={{ root: { backgroundColor: modeColors.background } }}>
          <Stack styles={{ root: { width: "100%" } }}>
-            {!dashboard.telemetryViews.length && <Stack horizontal tokens={{ childrenGap: 6 }} horizontalAlign="center" style={{paddingTop: 30}}>
+            {!dashboard.telemetryViews.length && <Stack horizontal tokens={{ childrenGap: 6 }} horizontalAlign="center" style={{ paddingTop: 30 }}>
                <Text variant="mediumPlus">No analytic views found, for more information please see</Text>
                <a href={telemetrybDocs} style={{ fontSize: "18px", "cursor": "pointer" }} onClick={(ev) => { ev.preventDefault(); ev.stopPropagation(); navigate(telemetrybDocs) }}> Horde analytics documentation.</a>
             </Stack>}
