Horde: Make job details more tolerant of strings appended to job links by accident, such as periods clipped out of sentences
#rnx

diff --git a/Engine/Source/Programs/Horde/HordeDashboard/src/components/jobDetailsV2/JobDetailViewV2.tsx b/Engine/Source/Programs/Horde/HordeDashboard/src/components/jobDetailsV2/JobDetailViewV2.tsx
--- a/Engine/Source/Programs/Horde/HordeDashboard/src/components/jobDetailsV2/JobDetailViewV2.tsx
+++ b/Engine/Source/Programs/Horde/HordeDashboard/src/components/jobDetailsV2/JobDetailViewV2.tsx
@@ -1,13 +1,14 @@
 import { mergeStyleSets, Pivot, PivotItem, Spinner, SpinnerSize, Stack } from "@fluentui/react";
 import { observer } from "mobx-react-lite";
 import { useEffect, useRef, useState } from "react";
-import { useLocation, useNavigate, useParams } from "react-router-dom";
+import { useLocation, useNavigate, useParams, useSearchParams } from "react-router-dom";
 import { GetJobsTabResponse } from "../../backend/Api";
 import { useWindowSize } from "../../base/utilities/hooks";
 import { getHordeStyling } from "../../styles/Styles";
 import { BreadcrumbItem, Breadcrumbs } from "../Breadcrumbs";
 import { useQuery } from "../JobDetailCommon";
 import { TopNav } from "../TopNav";
+import { JobArtifactsPanel } from "./JobDetailArtifacts";
 import { BisectionPanel } from "./JobDetailBisection";
 import { HealthPanel } from "./JobDetailHealthV2";
 import { PreflightPanel } from "./JobDetailPreflight";
@@ -17,7 +18,6 @@ import { StepsPanelV2 } from "./JobDetailViewSteps";
 import { SummaryPanel } from "./JobDetailViewSummary";
 import { JobOperations } from "./JobOperationsBar";
 import { StepDetailView } from "./StepDetailView";
-import { JobArtifactsPanel } from "./JobDetailArtifacts";
 
 class BreadcrumbDataView extends JobDataView {
 
@@ -54,7 +54,7 @@ const JobBreadCrumbs: React.FC<{ jobDetails: JobDetailsV2 }> = observer(({ jobDe
          return <Breadcrumbs items={[{ text: "Loading Job" }]} title={"Loading Job"} spinner={true} />
       }
       else {
-         console.error(`Unable to load job ${jobDetails.jobId}: ${jobDetails.jobError}` )
+         console.error(`Unable to load job ${jobDetails.jobId}: ${jobDetails.jobError}`)
       }
 
       return null;
@@ -307,9 +307,9 @@ const DetailsViewOverview: React.FC<{ jobDetails: JobDetailsV2 }> = ({ jobDetail
       <JobArtifactsPanel jobDetails={details} />
       <StepsPanelV2 jobDetails={details} />
       <HealthPanel jobDetails={details} />
-      <TimelinePanel jobDetails={details} />      
+      <TimelinePanel jobDetails={details} />
       <BisectionPanel jobDetails={details} />
-      {!jobDetails.viewsReady && <Stack style={{paddingTop: 32}}>
+      {!jobDetails.viewsReady && <Stack style={{ paddingTop: 32 }}>
          <Spinner size={SpinnerSize.large} />
       </Stack>}
 
@@ -341,7 +341,7 @@ const ScrollRestore: React.FC<{ jobDetails: JobDetailsV2, scrollRef: React.RefOb
 const DetailsView: React.FC<{ jobDetails: JobDetailsV2 }> = ({ jobDetails }) => {
 
    const windowSize = useWindowSize();
-   const scrollRef = useRef<HTMLDivElement>(null);   
+   const scrollRef = useRef<HTMLDivElement>(null);
    const { modeColors } = getHordeStyling();
 
    const vw = Math.max(document.documentElement.clientWidth, window.innerWidth || 0);
@@ -373,7 +373,7 @@ const DetailsView: React.FC<{ jobDetails: JobDetailsV2 }> = ({ jobDetails }) =>
                <Stack horizontal>
                   <Stack /*className={classNames.pointerSuppress} */ style={{ position: "relative", width: "100%", height: 'calc(100vh - 228px)' }}>
                      <div id="hordeContentArea" ref={scrollRef} style={{ overflowX: "auto", overflowY: "visible" }}>
-                        <Stack horizontal style={{paddingBottom: 48}}>
+                        <Stack horizontal style={{ paddingBottom: 48 }}>
                            <Stack key={`${key}_2`} style={{ paddingLeft: centerAlign }} />
                            <Stack style={{ width: rootWidth }}>
                               <DetailsViewOverview jobDetails={details} />
@@ -391,7 +391,7 @@ const DetailsView: React.FC<{ jobDetails: JobDetailsV2 }> = ({ jobDetails }) =>
 
 export const JobDetailViewV2: React.FC = () => {
 
-   const { jobId } = useParams<{ jobId: string }>();
+   let { jobId } = useParams<{ jobId: string }>();
    const [state, setState] = useState<{ jobDetails?: JobDetailsV2 }>({});
 
    useEffect(() => {
@@ -407,6 +407,39 @@ export const JobDetailViewV2: React.FC = () => {
       return null;
    }
 
+   let stepId = new URLSearchParams(window.location.search).get("step") ?? "";
+
+   let needNavigate = false;
+   if (stepId && stepId.length !== 4) {
+      needNavigate = true;
+      stepId = stepId.slice(0, 4);
+      if (stepId.length !== 4) {
+         stepId = "";
+      }
+   }
+
+   if (jobId.length !== 24) {
+      needNavigate = true;
+      jobId = jobId.slice(0, 24);
+   }
+
+   if (needNavigate) {
+      if (jobId.length !== 24) {
+         // navigate isn't working here
+         window.location.assign("/");
+      } else {
+         let url = `/job/${jobId}`;
+
+         if (stepId) {
+            url += `?step=${stepId}`;
+         }
+
+         // navigate isn't working here
+         window.location.assign(url);
+      }
+      return null;
+   }
+
    if (!state.jobDetails || state.jobDetails.jobId !== jobId) {
       if (state.jobDetails) {
          state.jobDetails.clear();
