Horde: Tweak dark mode scrollbar track color for Chad
#rnx

diff --git a/Engine/Source/Programs/Horde/HordeDashboard/src/styles/darkTheme.ts b/Engine/Source/Programs/Horde/HordeDashboard/src/styles/darkTheme.ts
--- a/Engine/Source/Programs/Horde/HordeDashboard/src/styles/darkTheme.ts
+++ b/Engine/Source/Programs/Horde/HordeDashboard/src/styles/darkTheme.ts
@@ -52,8 +52,8 @@ const hordeDarkTheme: HordeThemeExtensions = {
     // divides sections
     dividerColor: "#25282A",
     // Scrollbar theme colors
-    scrollbarThumbColor: "#575757",
-    scrollbarTrackColor: baseNeutral
+    scrollbarThumbColor: "#606060",
+    scrollbarTrackColor: "#3F3F3F"
 
 }
 
