Horde: Fix overflow on new build dialog
#rnx

diff --git a/Engine/Source/Programs/Horde/HordeDashboard/src/components/NewBuild.tsx b/Engine/Source/Programs/Horde/HordeDashboard/src/components/NewBuild.tsx
--- a/Engine/Source/Programs/Horde/HordeDashboard/src/components/NewBuild.tsx
+++ b/Engine/Source/Programs/Horde/HordeDashboard/src/components/NewBuild.tsx
@@ -1779,6 +1779,7 @@ export const NewBuild: React.FC<{ streamId: string; show: boolean; onClose: (new
                      {mode === "Basic" && <Stack style={{
                         height: estimatedHeight,
                         position: 'relative',
+                        width: 767,
                         maxHeight: 'calc(100vh - 360px)'
                      }}><ScrollablePane scrollbarVisibility={ScrollbarVisibility.auto}>
                            <Stack tokens={{ childrenGap: parameterGap }}>
@@ -1804,7 +1805,7 @@ export const NewBuild: React.FC<{ streamId: string; show: boolean; onClose: (new
                         </ScrollablePane>
                      </Stack>}
 
-                     {mode === "Advanced" && <Stack style={{ paddingBottom: 12 }} tokens={{ childrenGap: 12 }}>
+                     {mode === "Advanced" && <Stack style={{ paddingBottom: 12, width: 767, }} tokens={{ childrenGap: 12 }}>
                         {!!stream.configRevision && <Stack>
                            <TextField label="Template Path" readOnly={true} value={stream.configPath ?? ""} />
                         </Stack>
