Horde: Fix dashboard warning
#rnx

diff --git a/Engine/Source/Programs/Horde/HordeDashboard/src/components/jobDetailsV2/JobDetailViewV2.tsx b/Engine/Source/Programs/Horde/HordeDashboard/src/components/jobDetailsV2/JobDetailViewV2.tsx
--- a/Engine/Source/Programs/Horde/HordeDashboard/src/components/jobDetailsV2/JobDetailViewV2.tsx
+++ b/Engine/Source/Programs/Horde/HordeDashboard/src/components/jobDetailsV2/JobDetailViewV2.tsx
@@ -1,7 +1,7 @@
 import { mergeStyleSets, Pivot, PivotItem, Spinner, SpinnerSize, Stack } from "@fluentui/react";
 import { observer } from "mobx-react-lite";
 import { useEffect, useRef, useState } from "react";
-import { useLocation, useNavigate, useParams, useSearchParams } from "react-router-dom";
+import { useLocation, useNavigate, useParams } from "react-router-dom";
 import { GetJobsTabResponse } from "../../backend/Api";
 import { useWindowSize } from "../../base/utilities/hooks";
 import { getHordeStyling } from "../../styles/Styles";
