Horde: Fix automation test report hover color in dark mode
#rnx

diff --git a/Engine/Source/Programs/Horde/HordeDashboard/src/components/AutomationReport.tsx b/Engine/Source/Programs/Horde/HordeDashboard/src/components/AutomationReport.tsx
--- a/Engine/Source/Programs/Horde/HordeDashboard/src/components/AutomationReport.tsx
+++ b/Engine/Source/Programs/Horde/HordeDashboard/src/components/AutomationReport.tsx
@@ -52,7 +52,7 @@ const getStyles = () => {
          getFocusStyle(theme, { inset: -1 }),
          {
             selectors: {
-               '&:hover': { background: theme.palette.neutralLight }
+               '&:hover': { background: dashboard.darktheme ? undefined : theme.palette.neutralLight }
             }
          }
       ],
