Horde: Populate query string line index when using error/warning buttongs for easier link sharing
#rnx

diff --git a/Engine/Source/Programs/Horde/HordeDashboard/src/components/LogView.tsx b/Engine/Source/Programs/Horde/HordeDashboard/src/components/LogView.tsx
--- a/Engine/Source/Programs/Horde/HordeDashboard/src/components/LogView.tsx
+++ b/Engine/Source/Programs/Horde/HordeDashboard/src/components/LogView.tsx
@@ -1010,12 +1010,18 @@ export const LogList: React.FC<{ logId: string }> = observer(({ logId }) => {
       if (!handler || handler.currentError === undefined) return;
       handler.currentWarning = undefined;
       handler.stopTrailing();
-      let lineIdx = errors[handler.currentError].lineIndex - 10;
+      let lineIdx = errors[handler.currentError].lineIndex;
+      const search = new URLSearchParams(window.location.search);
+      search.set("lineindex", (lineIdx).toString());
+      const url = `${window.location.pathname}?` + search.toString();               
+      lineIdx -= 10;
       if (lineIdx < 0) {
          lineIdx = 0;
       }
       handler.externalUpdate();
       listRef?.scrollToIndex(lineIdx, () => handler.lineHeight, ScrollToMode.top);
+
+      navigate(url, { replace: true })
    }
 
    function prevError() {
@@ -1051,12 +1057,19 @@ export const LogList: React.FC<{ logId: string }> = observer(({ logId }) => {
       handler.currentError = undefined;
 
       handler.stopTrailing();
-      let lineIdx = warnings[handler.currentWarning].lineIndex - 10;
+      let lineIdx = warnings[handler.currentWarning].lineIndex ;
+      const search = new URLSearchParams(window.location.search);
+      search.set("lineindex", (lineIdx).toString());
+      const url = `${window.location.pathname}?` + search.toString();               
+
+      lineIdx -= 10
       if (lineIdx < 0) {
          lineIdx = 0;
       }
       handler.externalUpdate();
       listRef?.scrollToIndex(lineIdx, () => handler.lineHeight, ScrollToMode.top);
+
+      navigate(url, { replace: true })
    }
 
    function prevWarning() {
