Horde: Fix automation test tutorial docs
#rnx

diff --git a/Engine/Source/Programs/Horde/Docs/Tutorials/TestAutomation.md b/Engine/Source/Programs/Horde/Docs/Tutorials/TestAutomation.md
--- a/Engine/Source/Programs/Horde/Docs/Tutorials/TestAutomation.md
+++ b/Engine/Source/Programs/Horde/Docs/Tutorials/TestAutomation.md
@@ -27,8 +27,8 @@ The Horde Automation Hub surfaces individual and suite [Gauntlet](https://docs.u
 
 2. The job will build, cook, and run a Lyra Boot Test on the Android device, reserving it from the device manager during the process.
 
-![Automation Labels](../Images/Tutorial-TestAutomation-Labels.png)
+    ![Automation Labels](../Images/Tutorial-TestAutomation-Labels.png)
 
-1. Once completed the test results will be availble in the [Automation Hub](../Config/AutomationHub.md) which features finely grained filters and views that can cross compare platforms and streams.
+3. Once completed the test results will be availble in the [Automation Hub](../Config/AutomationHub.md) which features finely grained filters and views that can cross compare platforms and streams.
 
     ![Test Results](../Images/Tutorial-TestAutomation-TestResult.png)
\ No newline at end of file
