Horde: Normalize keys for artifacts to allow case-insensitive searching.

diff --git a/Engine/Source/Programs/Horde/Horde.Server/Artifacts/ArtifactCollection.cs b/Engine/Source/Programs/Horde/Horde.Server/Artifacts/ArtifactCollection.cs
--- a/Engine/Source/Programs/Horde/Horde.Server/Artifacts/ArtifactCollection.cs
+++ b/Engine/Source/Programs/Horde/Horde.Server/Artifacts/ArtifactCollection.cs
@@ -111,6 +111,11 @@ namespace Horde.Server.Artifacts
 			_clock = clock;
 		}
 
+#pragma warning disable CA1308 // Expect ansi-only keys here
+		static string NormalizeKey(string key)
+			=> key.ToLowerInvariant();
+#pragma warning restore CA1308
+
 		/// <summary>
 		/// Gets the base path for a set of artifacts
 		/// </summary>
@@ -124,7 +129,7 @@ namespace Horde.Server.Artifacts
 			NamespaceId namespaceId = Namespace.Artifacts;
 			RefName refName = new RefName($"{GetArtifactPath(streamId, name, type)}/{change}/{id}");
 
-			Artifact artifact = new Artifact(id, name, type, description, streamId, change, keys, metadata, namespaceId, refName, _clock.UtcNow, scopeName);
+			Artifact artifact = new Artifact(id, name, type, description, streamId, change, keys.Select(x => NormalizeKey(x)), metadata, namespaceId, refName, _clock.UtcNow, scopeName);
 			await _artifacts.InsertOneAsync(artifact, null, cancellationToken);
 			return artifact;
 		}
@@ -162,7 +167,7 @@ namespace Horde.Server.Artifacts
 			}
 			if (keys != null && keys.Any())
 			{
-				filter &= Builders<Artifact>.Filter.All(x => x.Keys, keys);
+				filter &= Builders<Artifact>.Filter.All(x => x.Keys, keys.Select(x => NormalizeKey(x)));
 			}
 
 			using (IAsyncCursor<Artifact> cursor = await _artifacts.Find(filter).SortByDescending(x => x.Change).ThenByDescending(x => x.Id).Limit(maxResults).ToCursorAsync(cancellationToken))
