Horde: Update release notes.

#rnx

diff --git a/Engine/Source/Programs/Horde/Docs/ReleaseNotes.md b/Engine/Source/Programs/Horde/Docs/ReleaseNotes.md
--- a/Engine/Source/Programs/Horde/Docs/ReleaseNotes.md
+++ b/Engine/Source/Programs/Horde/Docs/ReleaseNotes.md
@@ -1,5 +1,15 @@
 # Release Notes
 
+## 2024-03-26
+
+* Normalize keys for artifacts to allow case-insensitive searching. (32495123)
+* Support setting arbitrary artifact metadata. This cannot be queried, but will be returned in artifact responses. (32464180)
+* Run a background task to monitor for hangs when executing Perforce commands. (32461961)
+* Allow configuring different telemetry stores for each project and stream. (32460547)
+* Fix slow queries in storage service garbage collection. (32460451)
+* Add a [TryAuthorize] attribute which attempts authorization, but which does not return a challenge or forbid result on failure. Should fix tool download requests using service account tokens. (32447487)
+* Fix OIDC issues when using Horde auth with UBA. (32406994)
+
 ## 2024-03-21
 
 * Lower log level of failed AWS instance start (32397273)
