Horde: Add logging to trace when blobs are scanned for imports, to help debug some logs being expired while still referenced.

diff --git a/Engine/Source/Programs/Horde/Horde.Server/Storage/StorageService.cs b/Engine/Source/Programs/Horde/Horde.Server/Storage/StorageService.cs
--- a/Engine/Source/Programs/Horde/Horde.Server/Storage/StorageService.cs
+++ b/Engine/Source/Programs/Horde/Horde.Server/Storage/StorageService.cs
@@ -610,7 +610,7 @@ namespace Horde.Server.Storage
 								}
 								catch (Exception ex)
 								{
-									_logger.LogInformation(ex, "Unable to read references for blob {Locator}: {Message}", blobInfo.Locator, ex.Message);
+									_logger.LogInformation(ex, "Unable to read references for {NamespaceId} blob {BlobId} (key: {ObjectKey}): {Message}", blobInfo.NamespaceId, blobInfo.Id, GetObjectKey(blobInfo.Locator), ex.Message);
 								}
 							}
 						}
@@ -647,6 +647,7 @@ namespace Horde.Server.Storage
 			await _blobCollection.UpdateOneAsync(x => x.Id == blobInfo.Id, Builders<BlobInfo>.Update.Set(x => x.Imports, importInfoIds), null, cancellationToken);
 
 			AddGcCheckRecord(blobInfo.NamespaceId, blobInfo.Id);
+			_logger.LogDebug("Added {Count} imports for {NamespaceId} blob {BlobId}", importInfoIds.Count, blobInfo.NamespaceId, blobInfo.Id);
 		}
 
 		#endregion
@@ -1013,7 +1014,7 @@ namespace Horde.Server.Storage
 							}
 
 							ObjectKey objectKey = GetObjectKey(new BlobLocator(info.Path));
-							_logger.LogDebug("Deleting {NamespaceId} blob {BlobId}, key: {Key} ({ImportCount} imports)", namespaceInfo.Id, blobInfoId, objectKey, info.Imports?.Count ?? 0);
+							_logger.LogDebug("Deleting {NamespaceId} blob {BlobId}, key: {ObjectKey} ({ImportCount} imports)", namespaceInfo.Id, blobInfoId, objectKey, info.Imports?.Count ?? 0);
 							await namespaceInfo.Store.DeleteAsync(objectKey, cancellationToken);
 							numItemsRemoved++;
 						}
