Horde: Temporarily disable GC due to instances of incorrect reference scanning.

#rnx

diff --git a/Engine/Source/Programs/Horde/Horde.Server/Storage/StorageService.cs b/Engine/Source/Programs/Horde/Horde.Server/Storage/StorageService.cs
--- a/Engine/Source/Programs/Horde/Horde.Server/Storage/StorageService.cs
+++ b/Engine/Source/Programs/Horde/Horde.Server/Storage/StorageService.cs
@@ -444,7 +444,8 @@ namespace Horde.Server.Storage
 		{
 			await _blobTicker.StartAsync();
 			await _refTicker.StartAsync();
-			await _gcTicker.StartAsync();
+			// Disable GC for the moment; seeing some log blobs being removed incorrectly.
+//			await _gcTicker.StartAsync();
 		}
 
 		/// <inheritdoc/>
