Horde: Remove trailing comma from config file.

#rnx

diff --git a/Engine/Source/Programs/Horde/Horde.Server/Defaults/globals.json b/Engine/Source/Programs/Horde/Horde.Server/Defaults/globals.json
--- a/Engine/Source/Programs/Horde/Horde.Server/Defaults/globals.json
+++ b/Engine/Source/Programs/Horde/Horde.Server/Defaults/globals.json
@@ -40,7 +40,7 @@
 //			[
 //				{
 //					// DNS name and port for the server
-//					"serverAndPort": "perforce:1666",
+//					"serverAndPort": "perforce:1666"
 //				}
 //			],
 //
