Horde: Re-enable garbage collection.

#rnx

diff --git a/Engine/Source/Programs/Horde/Horde.Server/Storage/StorageService.cs b/Engine/Source/Programs/Horde/Horde.Server/Storage/StorageService.cs
--- a/Engine/Source/Programs/Horde/Horde.Server/Storage/StorageService.cs
+++ b/Engine/Source/Programs/Horde/Horde.Server/Storage/StorageService.cs
@@ -444,8 +444,7 @@ namespace Horde.Server.Storage
 		{
 			await _blobTicker.StartAsync();
 			await _refTicker.StartAsync();
-			// Disable GC for the moment; seeing some log blobs being removed incorrectly.
-//			await _gcTicker.StartAsync();
+			await _gcTicker.StartAsync();
 		}
 
 		/// <inheritdoc/>
