Horde: Use the //UE5/Release-5.4-HordeSync for the Horde demo stream. Also add a few comments explaining the purpose of various sections in the config file.

diff --git a/Engine/Source/Programs/Horde/Horde.Server/Defaults/ue5-release-5.4.stream.json b/Engine/Source/Programs/Horde/Horde.Server/Defaults/ue5-release-5.4.stream.json
--- a/Engine/Source/Programs/Horde/Horde.Server/Defaults/ue5-release-5.4.stream.json
+++ b/Engine/Source/Programs/Horde/Horde.Server/Defaults/ue5-release-5.4.stream.json
@@ -1,5 +1,8 @@
 {
-	"name": "//UE5/Release-5.4",
+	// Update this to the name of your Perforce stream
+	"name": "//UE5/Release-5.4-HordeSync",
+
+	// Macros for strings in the rest of this file (as $(Project) etc...)
 	"macros":
 	[
 		{
@@ -19,9 +22,13 @@
 			"value": "LyraEditor"
 		}
 	],
+
+	// Default job to run when triggering a preflight from P4VUtils (can be installed through UGS)
 	"defaultPreflight":
 	{
 		"templateId": "default-presubmit-tests",
+
+		// Set the base changelist number to the last incremental build that succeeded
 		"change":
 		{
 			"templateId": "incremental-build",
@@ -33,6 +40,8 @@
 			]
 		}
 	},
+
+	// Tabs to display in the dashboard for different jobs. Each job needs to be added to a tab to be shown in the New Build dialog.
 	"tabs":
 	[
 		{
@@ -155,6 +164,8 @@
 			]
 		}
 	],
+
+	// Definitions for different agent types. Agent types are referenced by BuildGraph scripts, and map onto pools and workspace definitions configured here.
 	"agentTypes":
 	{
 		"Win64":
@@ -186,13 +197,12 @@
 			"workspace": "Full"			
 		}
 	},
+
+	// Different workspace types, meaning Perforce client views that should be synced by particular agent types.
 	"workspaceTypes":
 	{
 		"Full":
 		{
-			// Use a virtual stream for complex filtering
-			"stream": "//UE5/Release-5.4-HordeSync",
-
 			// Set a unique identifier to prevent overlap with other streams; will cause workspace to be synced to a different folder
 			"identifier": "Demo",
 
@@ -201,9 +211,6 @@
 		},
 		"IncrementalCompileOnly":
 		{
-			// Use a virtual stream for complex filtering
-			"stream": "//UE5/Release-5.4-HordeSync",
-
 			// Set a unique identifier to prevent overlap with other streams; will cause workspace to be synced to a different folder
 			"identifier": "Demo-Inc",
 
@@ -221,6 +228,8 @@
 			"method": "name=managedWorkspace&preferNativeClient=true"
 		}
 	},
+
+	// Templates for all the jobs that can be run, as well as any parameters to show in the dashboard's New Build dialog
 	"templates":
 	[
 		{
