Horde: Add a default pool to include interactive agents, and map it to the TestWin64 agent type.

diff --git a/Engine/Source/Programs/Horde/Horde.Server/Defaults/default.global.json b/Engine/Source/Programs/Horde/Horde.Server/Defaults/default.global.json
--- a/Engine/Source/Programs/Horde/Horde.Server/Defaults/default.global.json
+++ b/Engine/Source/Programs/Horde/Horde.Server/Defaults/default.global.json
@@ -70,6 +70,11 @@
 			"condition": "Platform == 'Win64'",
 			"color": "Blue"
 		},
+		{
+			"name": "Win-UE5-GPU",
+			"condition": "Platform == 'Win64' && IsInteractive == true",
+			"color": "Blue"
+		},
 		{
 			"name": "Mac-UE5",
 			"condition": "Platform == 'Mac'",
diff --git a/Engine/Source/Programs/Horde/Horde.Server/Defaults/ue5-release-5.4.stream.json b/Engine/Source/Programs/Horde/Horde.Server/Defaults/ue5-release-5.4.stream.json
--- a/Engine/Source/Programs/Horde/Horde.Server/Defaults/ue5-release-5.4.stream.json
+++ b/Engine/Source/Programs/Horde/Horde.Server/Defaults/ue5-release-5.4.stream.json
@@ -171,10 +171,7 @@
 		"Win64":
 		{
 			"pool": "win-ue5",
-			"workspace": "Full",
-			"environment":
-			{
-			}
+			"workspace": "Full"
 		},
 		"Win64_Licensee": // Used for installed engine build
 		{
@@ -183,7 +180,7 @@
 		},
 		"TestWin64":
 		{
-			"pool": "win-ue5",
+			"pool": "win-ue5-gpu",
 			"workspace": "Full"
 		},
 		"IncrementalWin64":
