Horde: Fix typo in "interactive" property.

diff --git a/Engine/Source/Programs/Horde/Horde.Server/Defaults/default.global.json b/Engine/Source/Programs/Horde/Horde.Server/Defaults/default.global.json
--- a/Engine/Source/Programs/Horde/Horde.Server/Defaults/default.global.json
+++ b/Engine/Source/Programs/Horde/Horde.Server/Defaults/default.global.json
@@ -72,7 +72,7 @@
 		},
 		{
 			"name": "Win-UE5-GPU",
-			"condition": "Platform == 'Win64' && IsInteractive == true",
+			"condition": "Platform == 'Win64' && Interactive == true",
 			"color": "Blue"
 		},
 		{
