Added UBT Metrics to Hrode Analytics

diff --git a/Engine/Source/Programs/Horde/Horde.Server/Defaults/default-analytics.dashboard.json b/Engine/Source/Programs/Horde/Horde.Server/Defaults/default-analytics.dashboard.json
--- a/Engine/Source/Programs/Horde/Horde.Server/Defaults/default-analytics.dashboard.json
+++ b/Engine/Source/Programs/Horde/Horde.Server/Defaults/default-analytics.dashboard.json
@@ -532,18 +532,24 @@
 					"name": "Compile",
 					"charts": [
 						{
-							"name": "Full Compile",
+							"name": "UBA Compile Duration",
 							"metrics": [
 								{
-									"id": "horde-monolithic-compile-jobs"
+									"id": "ubt-uba-duration",
 								}
 							]
 						},
 						{
-							"name": "Incremental Compile",
+							"name": "UBA Compile Distribution",
+							"display": "Ratio",
 							"metrics": [
 								{
-									"id": "horde-incremental-compile-jobs"
+									"id": "ubt-uba-remote-usage",
+									"alias": "Remote"
+								},
+								{
+									"id": "ubt-uba-local-usage",
+									"alias": "Local"
 								}
 							]
 						}
diff --git a/Engine/Source/Programs/Horde/Horde.Server/Defaults/default-metrics.telemetry.json b/Engine/Source/Programs/Horde/Horde.Server/Defaults/default-metrics.telemetry.json
--- a/Engine/Source/Programs/Horde/Horde.Server/Defaults/default-metrics.telemetry.json
+++ b/Engine/Source/Programs/Horde/Horde.Server/Defaults/default-metrics.telemetry.json
@@ -481,6 +481,30 @@
 			"property": "$.Payload.Duration",
 			"function": "Average",
 			"interval": "01:00:00"
+		},
+		{
+			"id": "ubt-uba-duration",
+			"filter": "$[?(@.Payload.EventName == 'Core.UBT.Executor' && @.Payload.Executor=='Unreal Build Accelerator local' )]",
+			"groupBy": "$.Payload.Target,$.Payload.Mode,$.Payload.Configuration,$.Payload.Platform",
+			"property": "$.Payload.Duration",
+			"function": "Average",
+			"interval": "00:01:00"
+		},
+		{
+			"id": "ubt-uba-remote-usage",
+			"filter": "$[?(@.Payload.EventName == 'Core.UBT.Executor' && @.Payload.Executor=='Unreal Build Accelerator local' )]",
+			"groupBy": "$.Payload.Target,$.Payload.Mode,$.Payload.Configuration,$.Payload.Platform",
+			"property": "$.Payload.RemoteUsage",
+			"function": "Average",
+			"interval": "00:01:00"
+		},
+		{
+			"id": "ubt-uba-local-usage",
+			"filter": "$[?(@.Payload.EventName == 'Core.UBT.Executor' && @.Payload.Executor=='Unreal Build Accelerator local' )]",
+			"groupBy": "$.Payload.Target,$.Payload.Mode,$.Payload.Configuration,$.Payload.Platform",
+			"property": "$.Payload.LocalUsage",
+			"function": "Average",
+			"interval": "00:01:00"
 		}
 	]
 }
\ No newline at end of file
