Adding back Horde Compile Jobs To Horde Analytics Dashboard

diff --git a/Engine/Source/Programs/Horde/Horde.Server/Defaults/default-analytics.dashboard.json b/Engine/Source/Programs/Horde/Horde.Server/Defaults/default-analytics.dashboard.json
--- a/Engine/Source/Programs/Horde/Horde.Server/Defaults/default-analytics.dashboard.json
+++ b/Engine/Source/Programs/Horde/Horde.Server/Defaults/default-analytics.dashboard.json
@@ -531,11 +531,27 @@
 				{
 					"name": "Compile",
 					"charts": [
+						{
+							"name": "Full Compile",
+							"metrics": [
+								{
+									"id": "horde-monolithic-compile-jobs"
+								}
+							]
+						},
+						{
+							"name": "Incremental Compile",
+							"metrics": [
+								{
+									"id": "horde-incremental-compile-jobs"
+								}
+							]
+						},
 						{
 							"name": "UBA Compile Duration",
 							"metrics": [
 								{
-									"id": "ubt-uba-duration",
+									"id": "ubt-uba-duration"
 								}
 							]
 						},
