Added UBA Category to Horde Analytics
Added UBA Remote and Local Actions to Horde Analytics

diff --git a/Engine/Source/Programs/Horde/Horde.Server/Defaults/default-analytics.dashboard.json b/Engine/Source/Programs/Horde/Horde.Server/Defaults/default-analytics.dashboard.json
--- a/Engine/Source/Programs/Horde/Horde.Server/Defaults/default-analytics.dashboard.json
+++ b/Engine/Source/Programs/Horde/Horde.Server/Defaults/default-analytics.dashboard.json
@@ -70,6 +70,22 @@
 				{
 					"name": "StepName",
 					"group": "StepName"
+				},
+				{
+					"name": "Target",
+					"group": "Target"
+				},
+				{
+					"name": "Mode",
+					"group": "Mode"
+				},
+				{
+					"name": "Platform",
+					"group": "Platform"
+				},
+				{
+					"name": "Configuration",
+					"group": "Configuration"
 				}
 			],
 			"categories": [
@@ -546,7 +562,12 @@
 									"id": "horde-incremental-compile-jobs"
 								}
 							]
-						},
+						}	
+					]
+				},
+				{
+					"name": "UBA",
+					"charts": [
 						{
 							"name": "UBA Compile Duration",
 							"metrics": [
@@ -568,6 +589,20 @@
 									"alias": "Local"
 								}
 							]
+						},
+						{
+							"name": "UBA Compile Actions",
+							"display": "Ratio",
+							"metrics": [
+								{
+									"id": "ubt-uba-remote-actions",
+									"alias": "Remote"
+								},
+								{
+									"id": "ubt-uba-local-actions",
+									"alias": "Local"
+								}
+							]
 						}
 					]
 				}
diff --git a/Engine/Source/Programs/Horde/Horde.Server/Defaults/default-metrics.telemetry.json b/Engine/Source/Programs/Horde/Horde.Server/Defaults/default-metrics.telemetry.json
--- a/Engine/Source/Programs/Horde/Horde.Server/Defaults/default-metrics.telemetry.json
+++ b/Engine/Source/Programs/Horde/Horde.Server/Defaults/default-metrics.telemetry.json
@@ -505,6 +505,22 @@
 			"property": "$.Payload.LocalUsage",
 			"function": "Average",
 			"interval": "00:01:00"
+		},
+		{
+			"id": "ubt-uba-remote-actions",
+			"filter": "$[?(@.Payload.EventName == 'Core.UBT.Executor' && @.Payload.Executor=='Unreal Build Accelerator local' )]",
+			"groupBy": "$.Payload.Target,$.Payload.Mode,$.Payload.Configuration,$.Payload.Platform",
+			"property": "$.Payload.RemoteActions",
+			"function": "Sum",
+			"interval": "00:01:00"
+		},
+		{
+			"id": "ubt-uba-local-actions",
+			"filter": "$[?(@.Payload.EventName == 'Core.UBT.Executor' && @.Payload.Executor=='Unreal Build Accelerator local' )]",
+			"groupBy": "$.Payload.Target,$.Payload.Mode,$.Payload.Configuration,$.Payload.Platform",
+			"property": "$.Payload.LocalActions",
+			"function": "Sum",
+			"interval": "00:01:00"
 		}
 	]
 }
\ No newline at end of file
