Fix For UBA Metrics

diff --git a/Engine/Source/Programs/Horde/Horde.Server/Defaults/default-metrics.telemetry.json b/Engine/Source/Programs/Horde/Horde.Server/Defaults/default-metrics.telemetry.json
--- a/Engine/Source/Programs/Horde/Horde.Server/Defaults/default-metrics.telemetry.json
+++ b/Engine/Source/Programs/Horde/Horde.Server/Defaults/default-metrics.telemetry.json
@@ -484,7 +484,7 @@
 		},
 		{
 			"id": "ubt-uba-duration",
-			"filter": "$[?(@.Payload.EventName == 'Core.UBT.Executor' && @.Payload.Executor=='Unreal Build Accelerator local' )]",
+			"filter": "$[?(@.Payload.EventName == 'Core.UBT.Executor' && search(@.Payload.Executor, \"Unreal Build Accelerator\") )]",
 			"groupBy": "$.Payload.Target,$.Payload.Mode,$.Payload.Configuration,$.Payload.Platform",
 			"property": "$.Payload.Duration",
 			"function": "Average",
@@ -492,7 +492,7 @@
 		},
 		{
 			"id": "ubt-uba-remote-usage",
-			"filter": "$[?(@.Payload.EventName == 'Core.UBT.Executor' && @.Payload.Executor=='Unreal Build Accelerator local' )]",
+			"filter": "$[?(@.Payload.EventName == 'Core.UBT.Executor' && search(@.Payload.Executor, \"Unreal Build Accelerator\") )]",
 			"groupBy": "$.Payload.Target,$.Payload.Mode,$.Payload.Configuration,$.Payload.Platform",
 			"property": "$.Payload.RemoteUsage",
 			"function": "Average",
@@ -500,7 +500,7 @@
 		},
 		{
 			"id": "ubt-uba-local-usage",
-			"filter": "$[?(@.Payload.EventName == 'Core.UBT.Executor' && @.Payload.Executor=='Unreal Build Accelerator local' )]",
+			"filter": "$[?(@.Payload.EventName == 'Core.UBT.Executor' && search(@.Payload.Executor, \"Unreal Build Accelerator\") )]",
 			"groupBy": "$.Payload.Target,$.Payload.Mode,$.Payload.Configuration,$.Payload.Platform",
 			"property": "$.Payload.LocalUsage",
 			"function": "Average",
@@ -508,7 +508,7 @@
 		},
 		{
 			"id": "ubt-uba-remote-actions",
-			"filter": "$[?(@.Payload.EventName == 'Core.UBT.Executor' && @.Payload.Executor=='Unreal Build Accelerator local' )]",
+			"filter": "$[?(@.Payload.EventName == 'Core.UBT.Executor' && search(@.Payload.Executor, \"Unreal Build Accelerator\") )]",
 			"groupBy": "$.Payload.Target,$.Payload.Mode,$.Payload.Configuration,$.Payload.Platform",
 			"property": "$.Payload.RemoteActions",
 			"function": "Sum",
@@ -516,7 +516,7 @@
 		},
 		{
 			"id": "ubt-uba-local-actions",
-			"filter": "$[?(@.Payload.EventName == 'Core.UBT.Executor' && @.Payload.Executor=='Unreal Build Accelerator local' )]",
+			"filter": "$[?(@.Payload.EventName == 'Core.UBT.Executor' && search(@.Payload.Executor, \"Unreal Build Accelerator\") )]",
 			"groupBy": "$.Payload.Target,$.Payload.Mode,$.Payload.Configuration,$.Payload.Platform",
 			"property": "$.Payload.LocalActions",
 			"function": "Sum",
