Use exact count rather than ratio fro UBA Actions in Horde Analytics Dashboard

diff --git a/Engine/Source/Programs/Horde/Horde.Server/Defaults/default-analytics.dashboard.json b/Engine/Source/Programs/Horde/Horde.Server/Defaults/default-analytics.dashboard.json
--- a/Engine/Source/Programs/Horde/Horde.Server/Defaults/default-analytics.dashboard.json
+++ b/Engine/Source/Programs/Horde/Horde.Server/Defaults/default-analytics.dashboard.json
@@ -592,7 +592,6 @@
 						},
 						{
 							"name": "UBA Compile Actions",
-							"display": "Ratio",
 							"metrics": [
 								{
 									"id": "ubt-uba-remote-actions",
