Use value display type for UBA Actions in Horde Analytics

diff --git a/Engine/Source/Programs/Horde/Horde.Server/Defaults/default-analytics.dashboard.json b/Engine/Source/Programs/Horde/Horde.Server/Defaults/default-analytics.dashboard.json
--- a/Engine/Source/Programs/Horde/Horde.Server/Defaults/default-analytics.dashboard.json
+++ b/Engine/Source/Programs/Horde/Horde.Server/Defaults/default-analytics.dashboard.json
@@ -592,6 +592,7 @@
 						},
 						{
 							"name": "UBA Compile Actions",
+							"display": "Value",
 							"metrics": [
 								{
 									"id": "ubt-uba-remote-actions",
