Added Total Actions to UBA telemetry

diff --git a/Engine/Source/Programs/Horde/Horde.Server/Defaults/default-analytics.dashboard.json b/Engine/Source/Programs/Horde/Horde.Server/Defaults/default-analytics.dashboard.json
--- a/Engine/Source/Programs/Horde/Horde.Server/Defaults/default-analytics.dashboard.json
+++ b/Engine/Source/Programs/Horde/Horde.Server/Defaults/default-analytics.dashboard.json
@@ -594,6 +594,10 @@
 							"name": "UBA Compile Actions",
 							"display": "Value",
 							"metrics": [
+								{
+									"id": "ubt-uba-total-actions",
+									"alias": "Total"
+								},
 								{
 									"id": "ubt-uba-remote-actions",
 									"alias": "Remote"
diff --git a/Engine/Source/Programs/Horde/Horde.Server/Defaults/default-metrics.telemetry.json b/Engine/Source/Programs/Horde/Horde.Server/Defaults/default-metrics.telemetry.json
--- a/Engine/Source/Programs/Horde/Horde.Server/Defaults/default-metrics.telemetry.json
+++ b/Engine/Source/Programs/Horde/Horde.Server/Defaults/default-metrics.telemetry.json
@@ -521,6 +521,14 @@
 			"property": "$.Payload.LocalActions",
 			"function": "Sum",
 			"interval": "00:01:00"
+		},
+		{
+			"id": "ubt-uba-total-actions",
+			"filter": "$[?(@.Payload.EventName == 'Core.UBT.Executor' && search(@.Payload.Executor, \"Unreal Build Accelerator\") )]",
+			"groupBy": "$.Payload.Target,$.Payload.Mode,$.Payload.Configuration,$.Payload.Platform",
+			"property": "$.Payload.TotalActions",
+			"function": "Sum",
+			"interval": "00:01:00"
 		}
 	]
 }
\ No newline at end of file
