Changed UBA metrics frequency to 1h

diff --git a/Engine/Source/Programs/Horde/Horde.Server/Defaults/default-metrics.telemetry.json b/Engine/Source/Programs/Horde/Horde.Server/Defaults/default-metrics.telemetry.json
--- a/Engine/Source/Programs/Horde/Horde.Server/Defaults/default-metrics.telemetry.json
+++ b/Engine/Source/Programs/Horde/Horde.Server/Defaults/default-metrics.telemetry.json
@@ -488,7 +488,7 @@
 			"groupBy": "$.Payload.Target,$.Payload.Mode,$.Payload.Configuration,$.Payload.Platform",
 			"property": "$.Payload.Duration",
 			"function": "Average",
-			"interval": "00:01:00"
+			"interval": "01:00:00"
 		},
 		{
 			"id": "ubt-uba-remote-usage",
@@ -496,7 +496,7 @@
 			"groupBy": "$.Payload.Target,$.Payload.Mode,$.Payload.Configuration,$.Payload.Platform",
 			"property": "$.Payload.RemoteUsage",
 			"function": "Average",
-			"interval": "00:01:00"
+			"interval": "01:00:00"
 		},
 		{
 			"id": "ubt-uba-local-usage",
@@ -504,7 +504,7 @@
 			"groupBy": "$.Payload.Target,$.Payload.Mode,$.Payload.Configuration,$.Payload.Platform",
 			"property": "$.Payload.LocalUsage",
 			"function": "Average",
-			"interval": "00:01:00"
+			"interval": "01:00:00"
 		},
 		{
 			"id": "ubt-uba-remote-actions",
@@ -512,7 +512,7 @@
 			"groupBy": "$.Payload.Target,$.Payload.Mode,$.Payload.Configuration,$.Payload.Platform",
 			"property": "$.Payload.RemoteActions",
 			"function": "Sum",
-			"interval": "00:01:00"
+			"interval": "01:00:00"
 		},
 		{
 			"id": "ubt-uba-local-actions",
@@ -520,7 +520,7 @@
 			"groupBy": "$.Payload.Target,$.Payload.Mode,$.Payload.Configuration,$.Payload.Platform",
 			"property": "$.Payload.LocalActions",
 			"function": "Sum",
-			"interval": "00:01:00"
+			"interval": "01:00:00"
 		},
 		{
 			"id": "ubt-uba-total-actions",
@@ -528,7 +528,7 @@
 			"groupBy": "$.Payload.Target,$.Payload.Mode,$.Payload.Configuration,$.Payload.Platform",
 			"property": "$.Payload.TotalActions",
 			"function": "Sum",
-			"interval": "00:01:00"
+			"interval": "01:00:00"
 		}
 	]
 }
\ No newline at end of file
