Changed UBA Action Counts to Averages in Horde Analytics

diff --git a/Engine/Source/Programs/Horde/Horde.Server/Defaults/default-analytics.dashboard.json b/Engine/Source/Programs/Horde/Horde.Server/Defaults/default-analytics.dashboard.json
--- a/Engine/Source/Programs/Horde/Horde.Server/Defaults/default-analytics.dashboard.json
+++ b/Engine/Source/Programs/Horde/Horde.Server/Defaults/default-analytics.dashboard.json
@@ -591,7 +591,25 @@
 							]
 						},
 						{
-							"name": "UBA Compile Actions",
+							"name": "UBA Compile Actions (Average)",
+							"display": "Value",
+							"metrics": [
+								{
+									"id": "ubt-uba-total-actions-avg",
+									"alias": "Total"
+								},
+								{
+									"id": "ubt-uba-remote-actions-avg",
+									"alias": "Remote"
+								},
+								{
+									"id": "ubt-uba-local-actions-avg",
+									"alias": "Local"
+								}
+							]
+						},
+						{
+							"name": "UBA Compile Actions (Total)",
 							"display": "Value",
 							"metrics": [
 								{
diff --git a/Engine/Source/Programs/Horde/Horde.Server/Defaults/default-metrics.telemetry.json b/Engine/Source/Programs/Horde/Horde.Server/Defaults/default-metrics.telemetry.json
--- a/Engine/Source/Programs/Horde/Horde.Server/Defaults/default-metrics.telemetry.json
+++ b/Engine/Source/Programs/Horde/Horde.Server/Defaults/default-metrics.telemetry.json
@@ -511,7 +511,7 @@
 			"filter": "$[?(@.Payload.EventName == 'Core.UBT.Executor' && search(@.Payload.Executor, \"Unreal Build Accelerator\") )]",
 			"groupBy": "$.Payload.Target,$.Payload.Mode,$.Payload.Configuration,$.Payload.Platform",
 			"property": "$.Payload.RemoteActions",
-			"function": "Sum",
+			"function": "Count",
 			"interval": "01:00:00"
 		},
 		{
@@ -519,7 +519,7 @@
 			"filter": "$[?(@.Payload.EventName == 'Core.UBT.Executor' && search(@.Payload.Executor, \"Unreal Build Accelerator\") )]",
 			"groupBy": "$.Payload.Target,$.Payload.Mode,$.Payload.Configuration,$.Payload.Platform",
 			"property": "$.Payload.LocalActions",
-			"function": "Sum",
+			"function": "Count",
 			"interval": "01:00:00"
 		},
 		{
@@ -527,7 +527,31 @@
 			"filter": "$[?(@.Payload.EventName == 'Core.UBT.Executor' && search(@.Payload.Executor, \"Unreal Build Accelerator\") )]",
 			"groupBy": "$.Payload.Target,$.Payload.Mode,$.Payload.Configuration,$.Payload.Platform",
 			"property": "$.Payload.TotalActions",
-			"function": "Sum",
+			"function": "Count",
+			"interval": "01:00:00"
+		},
+		{
+			"id": "ubt-uba-remote-actions-avg",
+			"filter": "$[?(@.Payload.EventName == 'Core.UBT.Executor' && search(@.Payload.Executor, \"Unreal Build Accelerator\") )]",
+			"groupBy": "$.Payload.Target,$.Payload.Mode,$.Payload.Configuration,$.Payload.Platform",
+			"property": "$.Payload.RemoteActions",
+			"function": "Average",
+			"interval": "01:00:00"
+		},
+		{
+			"id": "ubt-uba-local-actions-avg",
+			"filter": "$[?(@.Payload.EventName == 'Core.UBT.Executor' && search(@.Payload.Executor, \"Unreal Build Accelerator\") )]",
+			"groupBy": "$.Payload.Target,$.Payload.Mode,$.Payload.Configuration,$.Payload.Platform",
+			"property": "$.Payload.LocalActions",
+			"function": "Average",
+			"interval": "01:00:00"
+		},
+		{
+			"id": "ubt-uba-total-actions-avg",
+			"filter": "$[?(@.Payload.EventName == 'Core.UBT.Executor' && search(@.Payload.Executor, \"Unreal Build Accelerator\") )]",
+			"groupBy": "$.Payload.Target,$.Payload.Mode,$.Payload.Configuration,$.Payload.Platform",
+			"property": "$.Payload.TotalActions",
+			"function": "Average",
 			"interval": "01:00:00"
 		}
 	]
