Horde: Updating schema documentation

#rnx

diff --git a/Engine/Source/Programs/Horde/Docs/Config/Schema/Globals.md b/Engine/Source/Programs/Horde/Docs/Config/Schema/Globals.md
--- a/Engine/Source/Programs/Horde/Docs/Config/Schema/Globals.md
+++ b/Engine/Source/Programs/Horde/Docs/Config/Schema/Globals.md
@@ -385,7 +385,7 @@ Name | Description
 `id` | `string`<br>Unique identifier for the tool
 `name` | `string`<br>Name of the tool
 `description` | `string`<br>Description for the tool
-`parentId` | `string`<br>Tool id to nest this tool under
+`category` | `string`<br>Category for the tool. Will cause the tool to be shown in a different tab in the dashboard.
 `public` | `boolean`<br>Whether this tool should be exposed for download on a public endpoint without authentication
 `showInUgs` | `boolean`<br>Whether to show this tool for download in the UGS tools menu
 `showInDashboard` | `boolean`<br>Whether to show this tool for download in the dashboard
diff --git a/Engine/Source/Programs/Horde/Docs/Deployment/ServerSettings.md b/Engine/Source/Programs/Horde/Docs/Deployment/ServerSettings.md
--- a/Engine/Source/Programs/Horde/Docs/Deployment/ServerSettings.md
+++ b/Engine/Source/Programs/Horde/Docs/Deployment/ServerSettings.md
@@ -344,7 +344,7 @@ Name | Description
 `id` | `string`<br>Unique identifier for the tool
 `name` | `string`<br>Name of the tool
 `description` | `string`<br>Description for the tool
-`parentId` | `string`<br>Tool id to nest this tool under
+`category` | `string`<br>Category for the tool. Will cause the tool to be shown in a different tab in the dashboard.
 `public` | `boolean`<br>Whether this tool should be exposed for download on a public endpoint without authentication
 `showInUgs` | `boolean`<br>Whether to show this tool for download in the UGS tools menu
 `showInDashboard` | `boolean`<br>Whether to show this tool for download in the dashboard
