Horde: Remove /userinfo call when authenticating via JWT

If an access token is passed, it's not guaranteed it has permission to access /userinfo from OIDC. ID tokens during normal web-based login does on the other hand.
The limited user information provided by the access token is however enough to find or create the user in Horde. Should the user later login via the web, information such as email and name will get populated.

diff --git a/Engine/Source/Programs/Horde/Horde.Server/Authentication/ExternalJwtAuthHandler.cs b/Engine/Source/Programs/Horde/Horde.Server/Authentication/ExternalJwtAuthHandler.cs
--- a/Engine/Source/Programs/Horde/Horde.Server/Authentication/ExternalJwtAuthHandler.cs
+++ b/Engine/Source/Programs/Horde/Horde.Server/Authentication/ExternalJwtAuthHandler.cs
@@ -4,15 +4,8 @@ using System;
 using System.Collections.Concurrent;
 using System.IdentityModel.Tokens.Jwt;
 using System.Linq;
-using System.Net.Http;
-using System.Net.Http.Headers;
-using System.Net.Http.Json;
 using System.Security.Claims;
-using System.Text.Json;
-using System.Text.Json.Nodes;
-using System.Threading;
 using System.Threading.Tasks;
-using EpicGames.Core;
 using Horde.Server.Users;
 using Horde.Server.Utilities;
 using Microsoft.AspNetCore.Authentication;
@@ -23,8 +16,6 @@ using Microsoft.IdentityModel.Tokens;
 
 namespace Horde.Server.Authentication
 {
-	using JsonObject = System.Text.Json.Nodes.JsonObject;
-
 	/// <summary>
 	/// JWT bearer handler capable of processing tokens generated by external OIDC providers (ie. access tokens issued by the OIDC server rather than the Horde server)
 	/// </summary>
@@ -38,10 +29,8 @@ namespace Horde.Server.Authentication
 		// Caches the lookup of "sub" claim in access token to actual user used internally by Horde.
 		// Required to avoid database lookup on every request when authenticating.
 		private readonly ConcurrentDictionary<string, IUser> _subToUser = new();
-		private readonly ConcurrentDictionary<string, JsonElement> _subToUserInfo = new();
 
 		readonly ServerSettings _settings;
-		Uri? _userInfoUrl;
 
 		/// <summary>
 		/// Constructor
@@ -113,85 +102,9 @@ namespace Horde.Server.Authentication
 			logger.LogInformation("Using HordeJwtBearerHandler");
 			if (!_subToUser.TryGetValue(accessToken.Subject, out IUser? user))
 			{
-				// No cached mapping of 'sub' claim to actual login ID used to fetch the user.
-				// Call /userinfo to get the ID token and get information needed to populate (and maybe create a new user)
-				CancellationToken cancellationToken = context.HttpContext.RequestAborted;
-				using HttpClient httpClient = serviceProvider.GetRequiredService<IHttpClientFactory>().CreateClient();
-
-				if (_userInfoUrl == null)
-				{
-					Uri discoveryUrl = new Uri(_settings.OidcAuthority.TrimEnd('/') + "/.well-known/openid-configuration");
-
-					using HttpResponseMessage discoveryResponse = await httpClient.GetAsync(discoveryUrl, cancellationToken);
-					if (!discoveryResponse.IsSuccessStatusCode)
-					{
-						ReportError(logger, context, "Bad status code from OpenID authority when fetching OIDC discovery document: " + discoveryResponse.StatusCode);
-						return;
-					}
-
-					JsonObject? wellKnownUrls = await discoveryResponse.Content.ReadFromJsonAsync<JsonObject>(cancellationToken);
-					if (wellKnownUrls == null || !wellKnownUrls.TryGetPropertyValue("userinfo_endpoint", out JsonNode? userInfoNode))
-					{
-						ReportError(logger, context, "Missing userinfo endpoint from OIDC discovery document");
-						return;
-					}
-
-					string? userInfoNodeValue = (string?)userInfoNode;
-					if (userInfoNodeValue == null)
-					{
-						ReportError(logger, context, "Missing userinfo value from OIDC discovery document");
-						return;
-					}
-
-					_userInfoUrl = new Uri(userInfoNodeValue);
-				}
-
-				using HttpRequestMessage request = new(HttpMethod.Get, _userInfoUrl);
-				request.Headers.Authorization = new AuthenticationHeaderValue("Bearer", accessToken!.RawData);
-
-				HttpResponseMessage response = await httpClient.SendAsync(request, cancellationToken);
-				if (!response.IsSuccessStatusCode)
-				{
-					ReportError(logger, context, "Bad status code from OpenID authority when fetching user info: " + response.StatusCode);
-					return;
-				}
-
-				string userInfoResponse = await response.Content.ReadAsStringAsync(cancellationToken);
-
-				JsonElement userInfo;
-				MediaTypeHeaderValue? contentType = response.Content.Headers.ContentType;
-				if (contentType?.MediaType?.Equals("application/json", StringComparison.OrdinalIgnoreCase) ?? false)
-				{
-					userInfo = JsonDocument.Parse(userInfoResponse).RootElement;
-				}
-				else if (contentType?.MediaType?.Equals("application/jwt", StringComparison.OrdinalIgnoreCase) ?? false)
-				{
-					JwtSecurityToken userInfoEndpointJwt = new(userInfoResponse);
-					userInfo = JsonDocument.Parse(userInfoEndpointJwt.Payload.SerializeToJson()).RootElement;
-				}
-				else
-				{
-					ReportError(logger, context, "Unknown response type: " + contentType?.MediaType);
-					return;
-				}
-
-				if (!userInfo.TryGetStringProperty("preferred_username", out string? login))
-				{
-					ReportError(logger, context, "Unable to read property 'preferred_username' from /userinfo");
-					return;
-				}
-
-				if (!userInfo.TryGetStringProperty("name", out string? name))
-				{
-					ReportError(logger, context, "Unable to read property 'name' from /userinfo");
-					return;
-				}
-
-				if (!userInfo.TryGetStringProperty("email", out string? email))
-				{
-					ReportError(logger, context, "Unable to read property 'email' from /userinfo");
-					return;
-				}
+				string login = accessToken.Subject;
+				string? name = accessToken.Claims.FirstOrDefault(c => c.Type == "name")?.Value;
+				string? email = accessToken.Claims.FirstOrDefault(c => c.Type == "email")?.Value;
 
 				IUserCollection userCollection = serviceProvider.GetRequiredService<IUserCollection>();
 				user = await userCollection.FindOrAddUserByLoginAsync(login, name, email);
@@ -199,18 +112,11 @@ namespace Horde.Server.Authentication
 				await userCollection.UpdateClaimsAsync(user.Id, accessToken.Claims.Select(x => new UserClaim(x.Type, x.Value)), context.HttpContext.RequestAborted);
 
 				_subToUser[accessToken.Subject] = user;
-				_subToUserInfo[accessToken.Subject] = userInfo;
-			}
-
-			if (!_subToUserInfo.TryGetValue(accessToken.Subject, out JsonElement cachedUserInfo))
-			{
-				ReportError(logger, context, "Cached user info not found");
-				return;
 			}
-
+			
 			identity.AddClaim(new Claim(HordeClaimTypes.Version, HordeClaimTypes.CurrentVersion));
 			identity.AddClaim(new Claim(HordeClaimTypes.UserId, user.Id.ToString()));
-			OidcAuthHandler.AddUserInfoClaims(_settings, cachedUserInfo, identity);
+			OidcAuthHandler.MapAdminClaim(_settings, identity);
 		}
 
 		private static void ReportError(ILogger logger, ResultContext<JwtBearerOptions> context, string message)
diff --git a/Engine/Source/Programs/Horde/Horde.Server/Authentication/OidcAuthHandler.cs b/Engine/Source/Programs/Horde/Horde.Server/Authentication/OidcAuthHandler.cs
--- a/Engine/Source/Programs/Horde/Horde.Server/Authentication/OidcAuthHandler.cs
+++ b/Engine/Source/Programs/Horde/Horde.Server/Authentication/OidcAuthHandler.cs
@@ -75,6 +75,11 @@ namespace Horde.Server.Authentication
 				}
 			}
 
+			MapAdminClaim(settings, identity);
+		}
+
+		public static void MapAdminClaim(ServerSettings settings, ClaimsIdentity identity)
+		{
 			if (!String.IsNullOrEmpty(settings.AdminClaimType) && !String.IsNullOrEmpty(settings.AdminClaimValue))
 			{
 				if (identity.HasClaim(settings.AdminClaimType, settings.AdminClaimValue))
