Horde: Copy group claims to use Horde's role claim type

diff --git a/Engine/Source/Programs/Horde/Horde.Server/Authentication/ExternalJwtAuthHandler.cs b/Engine/Source/Programs/Horde/Horde.Server/Authentication/ExternalJwtAuthHandler.cs
--- a/Engine/Source/Programs/Horde/Horde.Server/Authentication/ExternalJwtAuthHandler.cs
+++ b/Engine/Source/Programs/Horde/Horde.Server/Authentication/ExternalJwtAuthHandler.cs
@@ -110,12 +110,16 @@ namespace Horde.Server.Authentication
 				user = await userCollection.FindOrAddUserByLoginAsync(login, name, email);
 
 				await userCollection.UpdateClaimsAsync(user.Id, accessToken.Claims.Select(x => new UserClaim(x.Type, x.Value)), context.HttpContext.RequestAborted);
-
 				_subToUser[accessToken.Subject] = user;
 			}
 			
 			identity.AddClaim(new Claim(HordeClaimTypes.Version, HordeClaimTypes.CurrentVersion));
 			identity.AddClaim(new Claim(HordeClaimTypes.UserId, user.Id.ToString()));
+
+			foreach (Claim claim in accessToken.Claims.Where(x => x.Type == "groups"))
+			{
+				identity.AddClaim(new Claim(ClaimTypes.Role, claim.Value));
+			}
 			OidcAuthHandler.MapAdminClaim(_settings, identity);
 		}
 
