Added Failure Rate To UBA Dashboard

diff --git a/Engine/Source/Programs/Horde/Horde.Server/Defaults/default-analytics.dashboard.json b/Engine/Source/Programs/Horde/Horde.Server/Defaults/default-analytics.dashboard.json
--- a/Engine/Source/Programs/Horde/Horde.Server/Defaults/default-analytics.dashboard.json
+++ b/Engine/Source/Programs/Horde/Horde.Server/Defaults/default-analytics.dashboard.json
@@ -569,7 +569,7 @@
 					"name": "UBA",
 					"charts": [
 						{
-							"name": "UBA Compile Duration",
+							"name": "Compile Duration",
 							"metrics": [
 								{
 									"id": "ubt-uba-duration"
@@ -577,7 +577,7 @@
 							]
 						},
 						{
-							"name": "UBA Compile Distribution",
+							"name": "Compile Distribution",
 							"display": "Ratio",
 							"metrics": [
 								{
@@ -591,7 +591,7 @@
 							]
 						},
 						{
-							"name": "UBA Compile Actions (Average)",
+							"name": "Compile Actions (Average)",
 							"display": "Value",
 							"metrics": [
 								{
@@ -609,7 +609,7 @@
 							]
 						},
 						{
-							"name": "UBA Compile Actions (Total)",
+							"name": "Compile Actions (Total)",
 							"display": "Value",
 							"metrics": [
 								{
@@ -625,6 +625,16 @@
 									"alias": "Local"
 								}
 							]
+						},
+						{
+							"name": "Failure Rate",
+							"display": "Ratio",
+							"metrics": [
+								{
+									"id": "ubt-uba-failure-rate"
+
+								}
+							]
 						}
 					]
 				}
diff --git a/Engine/Source/Programs/Horde/Horde.Server/Defaults/default-metrics.telemetry.json b/Engine/Source/Programs/Horde/Horde.Server/Defaults/default-metrics.telemetry.json
--- a/Engine/Source/Programs/Horde/Horde.Server/Defaults/default-metrics.telemetry.json
+++ b/Engine/Source/Programs/Horde/Horde.Server/Defaults/default-metrics.telemetry.json
@@ -488,7 +488,7 @@
 			"groupBy": "$.Payload.Target,$.Payload.Mode,$.Payload.Configuration,$.Payload.Platform",
 			"property": "$.Payload.Duration",
 			"function": "Average",
-			"interval": "01:00:00"
+			"interval": "00:30:00"
 		},
 		{
 			"id": "ubt-uba-remote-usage",
@@ -496,7 +496,7 @@
 			"groupBy": "$.Payload.Target,$.Payload.Mode,$.Payload.Configuration,$.Payload.Platform",
 			"property": "$.Payload.RemoteUsage",
 			"function": "Average",
-			"interval": "01:00:00"
+			"interval": "00:30:00"
 		},
 		{
 			"id": "ubt-uba-local-usage",
@@ -504,7 +504,7 @@
 			"groupBy": "$.Payload.Target,$.Payload.Mode,$.Payload.Configuration,$.Payload.Platform",
 			"property": "$.Payload.LocalUsage",
 			"function": "Average",
-			"interval": "01:00:00"
+			"interval": "00:30:00"
 		},
 		{
 			"id": "ubt-uba-remote-actions",
@@ -512,7 +512,7 @@
 			"groupBy": "$.Payload.Target,$.Payload.Mode,$.Payload.Configuration,$.Payload.Platform",
 			"property": "$.Payload.RemoteActions",
 			"function": "Count",
-			"interval": "01:00:00"
+			"interval": "00:30:00"
 		},
 		{
 			"id": "ubt-uba-local-actions",
@@ -520,7 +520,7 @@
 			"groupBy": "$.Payload.Target,$.Payload.Mode,$.Payload.Configuration,$.Payload.Platform",
 			"property": "$.Payload.LocalActions",
 			"function": "Count",
-			"interval": "01:00:00"
+			"interval": "00:30:00"
 		},
 		{
 			"id": "ubt-uba-total-actions",
@@ -528,7 +528,7 @@
 			"groupBy": "$.Payload.Target,$.Payload.Mode,$.Payload.Configuration,$.Payload.Platform",
 			"property": "$.Payload.TotalActions",
 			"function": "Count",
-			"interval": "01:00:00"
+			"interval": "00:30:00"
 		},
 		{
 			"id": "ubt-uba-remote-actions-avg",
@@ -536,7 +536,7 @@
 			"groupBy": "$.Payload.Target,$.Payload.Mode,$.Payload.Configuration,$.Payload.Platform",
 			"property": "$.Payload.RemoteActions",
 			"function": "Average",
-			"interval": "01:00:00"
+			"interval": "00:30:00"
 		},
 		{
 			"id": "ubt-uba-local-actions-avg",
@@ -544,7 +544,7 @@
 			"groupBy": "$.Payload.Target,$.Payload.Mode,$.Payload.Configuration,$.Payload.Platform",
 			"property": "$.Payload.LocalActions",
 			"function": "Average",
-			"interval": "01:00:00"
+			"interval": "00:30:00"
 		},
 		{
 			"id": "ubt-uba-total-actions-avg",
@@ -552,7 +552,15 @@
 			"groupBy": "$.Payload.Target,$.Payload.Mode,$.Payload.Configuration,$.Payload.Platform",
 			"property": "$.Payload.TotalActions",
 			"function": "Average",
-			"interval": "01:00:00"
+			"interval": "00:30:00"
+		},
+		{
+			"id": "ubt-uba-failure-rate",
+			"filter": "$[?(@.Payload.EventName == 'Core.UBT.Executor' && search(@.Payload.Executor, \"Unreal Build Accelerator\") )]",
+			"groupBy": "$.Payload.Target,$.Payload.Mode,$.Payload.Configuration,$.Payload.Platform",
+			"property": "$.Payload.FailureRate",
+			"function": "Average",
+			"interval": "00:30:00"
 		}
 	]
 }
\ No newline at end of file
