Changed Remote, Local and Total Actions to Sum rather than Count

diff --git a/Engine/Source/Programs/Horde/Horde.Server/Defaults/default-metrics.telemetry.json b/Engine/Source/Programs/Horde/Horde.Server/Defaults/default-metrics.telemetry.json
--- a/Engine/Source/Programs/Horde/Horde.Server/Defaults/default-metrics.telemetry.json
+++ b/Engine/Source/Programs/Horde/Horde.Server/Defaults/default-metrics.telemetry.json
@@ -511,7 +511,7 @@
 			"filter": "$[?(@.Payload.EventName == 'Core.UBT.Executor' && search(@.Payload.Executor, \"Unreal Build Accelerator\") )]",
 			"groupBy": "$.Payload.Target,$.Payload.Mode,$.Payload.Configuration,$.Payload.Platform",
 			"property": "$.Payload.RemoteActions",
-			"function": "Count",
+			"function": "Sum",
 			"interval": "00:30:00"
 		},
 		{
@@ -519,7 +519,7 @@
 			"filter": "$[?(@.Payload.EventName == 'Core.UBT.Executor' && search(@.Payload.Executor, \"Unreal Build Accelerator\") )]",
 			"groupBy": "$.Payload.Target,$.Payload.Mode,$.Payload.Configuration,$.Payload.Platform",
 			"property": "$.Payload.LocalActions",
-			"function": "Count",
+			"function": "Sum",
 			"interval": "00:30:00"
 		},
 		{
@@ -527,7 +527,7 @@
 			"filter": "$[?(@.Payload.EventName == 'Core.UBT.Executor' && search(@.Payload.Executor, \"Unreal Build Accelerator\") )]",
 			"groupBy": "$.Payload.Target,$.Payload.Mode,$.Payload.Configuration,$.Payload.Platform",
 			"property": "$.Payload.TotalActions",
-			"function": "Count",
+			"function": "Sum",
 			"interval": "00:30:00"
 		},
 		{
