Horde: Remove references to specific claims other than the user id in the claims principal.

diff --git a/Engine/Source/Programs/Horde/Horde.Server/Agents/AgentsController.cs b/Engine/Source/Programs/Horde/Horde.Server/Agents/AgentsController.cs
--- a/Engine/Source/Programs/Horde/Horde.Server/Agents/AgentsController.cs
+++ b/Engine/Source/Programs/Horde/Horde.Server/Agents/AgentsController.cs
@@ -11,10 +11,12 @@ using EpicGames.Horde.Agents.Leases;
 using EpicGames.Horde.Agents.Pools;
 using EpicGames.Horde.Agents.Sessions;
 using EpicGames.Horde.Common;
+using EpicGames.Horde.Users;
 using Horde.Server.Agents.Leases;
 using Horde.Server.Agents.Sessions;
 using Horde.Server.Auditing;
 using Horde.Server.Server;
+using Horde.Server.Users;
 using Horde.Server.Utilities;
 using Microsoft.AspNetCore.Authorization;
 using Microsoft.AspNetCore.Mvc;
@@ -32,15 +34,17 @@ namespace Horde.Server.Agents
 	public class AgentsController : HordeControllerBase
 	{
 		readonly AgentService _agentService;
+		readonly IUserCollection _userCollection;
 		readonly IOptionsSnapshot<GlobalConfig> _globalConfig;
 		readonly ILogger<AgentsController> _logger;
 
 		/// <summary>
 		/// Constructor
 		/// </summary>
-		public AgentsController(AgentService agentService, IOptionsSnapshot<GlobalConfig> globalConfig, ILogger<AgentsController> logger)
+		public AgentsController(AgentService agentService, IUserCollection userCollection, IOptionsSnapshot<GlobalConfig> globalConfig, ILogger<AgentsController> logger)
 		{
 			_agentService = agentService;
+			_userCollection = userCollection;
 			_globalConfig = globalConfig;
 			_logger = logger;
 		}
@@ -197,7 +201,8 @@ namespace Horde.Server.Agents
 
 			List<PoolId>? updatePools = update.Pools?.ConvertAll(x => new PoolId(x));
 
-			string userName = User.GetUserName() ?? "Unknown";
+			IUser? user = await _userCollection.GetUserAsync(User, HttpContext.RequestAborted);
+			string userName = user?.Name ?? "Unknown";
 
 			for (; ; )
 			{
diff --git a/Engine/Source/Programs/Horde/Horde.Server/Issues/IssuesController.cs b/Engine/Source/Programs/Horde/Horde.Server/Issues/IssuesController.cs
--- a/Engine/Source/Programs/Horde/Horde.Server/Issues/IssuesController.cs
+++ b/Engine/Source/Programs/Horde/Horde.Server/Issues/IssuesController.cs
@@ -916,10 +916,9 @@ namespace Horde.Server.Issues
 			}
 
 			IUser? user = await _userCollection.GetUserAsync(User, cancellationToken);
-
 			if (user == null)
 			{
-				return BadRequest($"Missing user for {User.GetUserName()}");
+				return BadRequest($"Missing user for {User.GetUserId()}");
 			}
 
 			(string? key, string? url) = await _externalIssueService.CreateIssueAsync(user, User.GetExternalIssueUser(), issueRequest.IssueId, issueRequest.Summary, issueRequest.ProjectId, issueRequest.ComponentId, issueRequest.IssueTypeId, issueRequest.Description, issueRequest.HordeIssueLink, cancellationToken);
diff --git a/Engine/Source/Programs/Horde/Horde.Server/Notifications/SubscriptionsController.cs b/Engine/Source/Programs/Horde/Horde.Server/Notifications/SubscriptionsController.cs
--- a/Engine/Source/Programs/Horde/Horde.Server/Notifications/SubscriptionsController.cs
+++ b/Engine/Source/Programs/Horde/Horde.Server/Notifications/SubscriptionsController.cs
@@ -155,20 +155,11 @@ namespace Horde.Server.Notifications
 				objectId = newObjectId;
 				return true;
 			}
-
-			string? currentUserName = User.GetUserName();
-			if (currentUserName != null && String.Equals(userName, currentUserName, StringComparison.OrdinalIgnoreCase))
+			else
 			{
-				UserId? currentUserId = User.GetUserId();
-				if (currentUserId != null)
-				{
-					objectId = currentUserId.Value;
-					return true;
-				}
+				objectId = default;
+				return false;
 			}
-
-			objectId = default;
-			return false;
 		}
 	}
 }
diff --git a/Engine/Source/Programs/Horde/Horde.Server/Utilities/HordeClaimTypes.cs b/Engine/Source/Programs/Horde/Horde.Server/Utilities/HordeClaimTypes.cs
--- a/Engine/Source/Programs/Horde/Horde.Server/Utilities/HordeClaimTypes.cs
+++ b/Engine/Source/Programs/Horde/Horde.Server/Utilities/HordeClaimTypes.cs
@@ -140,41 +140,6 @@ namespace Horde.Server.Utilities
 			}
 		}
 
-		/// <summary>
-		/// Gets the Horde user name from a principal
-		/// </summary>
-		/// <param name="principal"></param>
-		/// <returns></returns>
-		public static string? GetUserName(this ClaimsPrincipal principal)
-		{
-			return principal.FindFirstValue(HordeClaimTypes.User) ?? principal.FindFirstValue(ClaimTypes.Name);
-		}
-
-		/// <summary>
-		/// Get the email for the given user
-		/// </summary>
-		/// <param name="user">The user to query the email for</param>
-		/// <returns>The user's email address or null if not found</returns>
-		public static string? GetEmail(this ClaimsPrincipal user)
-		{
-			return user.Claims.FirstOrDefault(x => x.Type == ClaimTypes.Email)?.Value;
-		}
-
-		/// <summary>
-		/// Gets the perforce username for the given principal
-		/// </summary>
-		/// <param name="user">The principal to get the Perforce user for</param>
-		/// <returns>Perforce user name</returns>
-		public static string? GetPerforceUser(this ClaimsPrincipal user)
-		{
-			Claim? claim = user.FindFirst(HordeClaimTypes.PerforceUser);
-			if (claim == null)
-			{
-				return null;
-			}
-			return claim.Value;
-		}
-
 		/// <summary>
 		/// Gets the external issue username for the given principal
 		/// </summary>
