Horde: Add category pivot to tools view, styling tweaks
#rnx

diff --git a/Engine/Source/Programs/Horde/HordeDashboard/package.json b/Engine/Source/Programs/Horde/HordeDashboard/package.json
--- a/Engine/Source/Programs/Horde/HordeDashboard/package.json
+++ b/Engine/Source/Programs/Horde/HordeDashboard/package.json
@@ -1,6 +1,6 @@
 {
   "name": "horde-dashboard",
-  "version": "0.2.1",
+  "version": "0.2.2",
   "private": true,
   "main": "index.js",
   "proxy": "https://localhost:5001",
diff --git a/Engine/Source/Programs/Horde/HordeDashboard/src/backend/Api.ts b/Engine/Source/Programs/Horde/HordeDashboard/src/backend/Api.ts
--- a/Engine/Source/Programs/Horde/HordeDashboard/src/backend/Api.ts
+++ b/Engine/Source/Programs/Horde/HordeDashboard/src/backend/Api.ts
@@ -5059,6 +5059,9 @@ export type GetToolSummaryResponse = {
 	/** Name of tool */
 	name: string;
 
+	/** Category of the tool */
+	category?: string;
+
 	/** Description of tool */
 	description: string;
 
diff --git a/Engine/Source/Programs/Horde/HordeDashboard/src/components/PoolsView.tsx b/Engine/Source/Programs/Horde/HordeDashboard/src/components/PoolsView.tsx
--- a/Engine/Source/Programs/Horde/HordeDashboard/src/components/PoolsView.tsx
+++ b/Engine/Source/Programs/Horde/HordeDashboard/src/components/PoolsView.tsx
@@ -484,7 +484,7 @@ export const PoolPivot: React.FC = () => {
 
    pivotItems.unshift(<PivotItem headerText="All" itemKey="all" key={"all"} style={{ color: modeColors.text }} />);
 
-   return <Stack grow>
+   return <Stack grow style={{paddingBottom: 12}}>
       <Pivot className={hordeClasses.pivot}
          overflowBehavior='menu'
          selectedKey={handler.category?.name ?? "all"}
@@ -556,7 +556,7 @@ export const PoolsView: React.FC = observer(() => {
       <Stack styles={{ root: { width: "100%", backgroundColor: modeColors.background } }}>
          <Stack style={{ width: "100%", backgroundColor: modeColors.background }}>
             <Stack style={{ position: "relative", width: "100%", height: 'calc(100vh - 148px)' }}>
-               {<Stack style={{ paddingTop: "30px", paddingBottom: "4px" }}>
+               {<Stack style={{ paddingTop: "12px", paddingBottom: "4px" }}>
                   <Stack style={{ width: 1440, marginLeft: centerAlign }}>
                      <Stack horizontal style={{ width: "100%" }}>
                         {!poolId && <PoolPivot />}
diff --git a/Engine/Source/Programs/Horde/HordeDashboard/src/components/ToolView.tsx b/Engine/Source/Programs/Horde/HordeDashboard/src/components/ToolView.tsx
--- a/Engine/Source/Programs/Horde/HordeDashboard/src/components/ToolView.tsx
+++ b/Engine/Source/Programs/Horde/HordeDashboard/src/components/ToolView.tsx
@@ -1,8 +1,8 @@
 // Copyright Epic Games, Inc. All Rights Reserved.
 
-import { DetailsList, DetailsListLayoutMode, IColumn, PrimaryButton, SelectionMode, Stack, Text } from "@fluentui/react";
+import { DetailsList, DetailsListLayoutMode, IColumn, Pivot, PivotItem, PrimaryButton, SelectionMode, Stack, Text } from "@fluentui/react";
 import { observer } from "mobx-react-lite";
-import React, { useEffect } from "react";
+import React, { useEffect, useState } from "react";
 import backend from "../backend";
 import { GetToolSummaryResponse } from "../backend/Api";
 import { PollBase } from "../backend/PollBase";
@@ -11,6 +11,8 @@ import { Breadcrumbs } from "./Breadcrumbs";
 import { TopNav } from "./TopNav";
 import { getHordeStyling } from "../styles/Styles";
 
+const defaultCategory = "General";
+
 class ToolHandler extends PollBase {
 
    constructor(pollTime = 30000) {
@@ -20,6 +22,7 @@ class ToolHandler extends PollBase {
    }
 
    clear() {
+      this.categories = new Map();
       this.loaded = false;
       super.stop();
    }
@@ -30,6 +33,16 @@ class ToolHandler extends PollBase {
 
          this.tools = await backend.getTools();
          this.tools = this.tools.filter(t => t.showInDashboard);
+         this.categories = new Map();
+
+         this.tools.forEach(t => {
+            const cat = t.category ?? defaultCategory;
+            if (!this.categories.has(cat)) {
+               this.categories.set(cat, []);
+            }
+            this.categories.get(cat)!.push(t);
+         })
+
          this.loaded = true;
          this.setUpdated();
 
@@ -38,13 +51,16 @@ class ToolHandler extends PollBase {
       }
 
    }
+
+   categories: Map<string, GetToolSummaryResponse[]> = new Map();
+
    loaded = false;
    tools: GetToolSummaryResponse[] = [];
 }
 
 const handler = new ToolHandler();
 
-const ToolPanel: React.FC = observer(() => {
+const ToolPanel: React.FC<{ selectedKey: string }> = observer(({ selectedKey }) => {
 
    useEffect(() => {
 
@@ -70,6 +86,8 @@ const ToolPanel: React.FC = observer(() => {
 
    let tools = [...handler.tools];
 
+   tools = tools.filter(t => (t.category ?? defaultCategory) === selectedKey)
+
    tools = tools.sort((a, b) => a.name.localeCompare(b.name));
 
    const renderItem = (item: any, index?: number, column?: IColumn) => {
@@ -88,7 +106,7 @@ const ToolPanel: React.FC = observer(() => {
          if (!item.version) {
             return null;
          }
-         return <Stack horizontalAlign="center" verticalAlign="center" verticalFill={true}>
+         return <Stack verticalAlign="center" verticalFill={true}>
             <Text style={{ color: modeColors.text }}>{item.version}</Text>
          </Stack>
       }
@@ -121,6 +139,7 @@ const ToolPanel: React.FC = observer(() => {
          <Stack styles={{ root: { paddingLeft: 12, paddingRight: 12, paddingBottom: 12, width: "100%" } }} >
             <DetailsList
                isHeaderVisible={true}
+               styles={{ headerWrapper: {paddingTop: 0} }}
                items={tools}
                columns={columns}
                selectionMode={SelectionMode.none}
@@ -133,34 +152,75 @@ const ToolPanel: React.FC = observer(() => {
    </Stack>
 });
 
+export const ToolViewInner: React.FC = observer(() => {
 
-export const ToolView: React.FC = () => {
+   const [selectedKey, setSelectedKey] = useState<string>(defaultCategory);
 
    const windowSize = useWindowSize();
+
+   // subscribe
+   if (handler.updated) { };
+
+   const { hordeClasses, modeColors } = getHordeStyling();
    const vw = Math.max(document.documentElement.clientWidth, window.innerWidth || 0);
    const centerAlign = vw / 2 - 720;
+   const key = `windowsize_view_${windowSize.width}_${windowSize.height}`;
 
-   const { hordeClasses, modeColors } = getHordeStyling();
+   const categories = Array.from(handler.categories.keys()).sort((a, b) => a.localeCompare(b));
 
-   const key = `windowsize_view_${windowSize.width}_${windowSize.height}`;
+   const pivotItems = categories.map(cat => {
+      if (cat === defaultCategory) {
+         return undefined;
+      }
+      return <PivotItem headerText={cat} itemKey={cat} key={cat} style={{ color: modeColors.text }} />;
+   }).filter(p => !!p);
+
+   pivotItems.unshift(<PivotItem headerText={defaultCategory} itemKey={defaultCategory} key={defaultCategory} style={{ color: modeColors.text }} />);
+
+   return <Stack styles={{ root: { width: "100%", backgroundColor: modeColors.background } }}>
+      <Stack style={{ width: "100%", backgroundColor: modeColors.background }}>
+         <Stack style={{ position: "relative", width: "100%", height: 'calc(100vh - 148px)' }}>
+            <div style={{ overflowX: "auto", overflowY: "visible" }}>
+               <Stack horizontal style={{ paddingTop: 12, paddingBottom: 48 }}>
+                  <Stack key={`${key}`} style={{ paddingLeft: centerAlign }} />
+                  <Stack style={{ width: 1440 }}>
+                     <Stack style={{paddingBottom: 12}}>
+                        <Pivot className={hordeClasses.pivot}
+                           overflowBehavior='menu'
+                           selectedKey={selectedKey}
+                           linkSize="normal"
+                           linkFormat="links"
+                           onLinkClick={(item) => {
+                              if (item?.props.itemKey) {
+
+                                 setSelectedKey(item.props.itemKey);
+                              }
+                           }}>
+                           {pivotItems}
+                        </Pivot>
 
-   return <Stack className={hordeClasses.horde}>
-      <TopNav />
-      <Breadcrumbs items={[{ text: 'Tools' }]} />
-      <Stack styles={{ root: { width: "100%", backgroundColor: modeColors.background } }}>
-         <Stack style={{ width: "100%", backgroundColor: modeColors.background }}>
-            <Stack style={{ position: "relative", width: "100%", height: 'calc(100vh - 148px)' }}>
-               <div style={{ overflowX: "auto", overflowY: "visible" }}>
-                  <Stack horizontal style={{ paddingTop: 30, paddingBottom: 48 }}>
-                     <Stack key={`${key}`} style={{ paddingLeft: centerAlign }} />
-                     <Stack style={{ width: 1440 }}>
-                        <ToolPanel />
                      </Stack>
+
+                     <ToolPanel selectedKey={selectedKey} />
                   </Stack>
-               </div>
-            </Stack>
+               </Stack>
+            </div>
          </Stack>
       </Stack>
    </Stack>
+
+})
+
+export const ToolView: React.FC = () => {
+
+
+   const { hordeClasses } = getHordeStyling();
+
+
+   return <Stack className={hordeClasses.horde}>
+      <TopNav />
+      <Breadcrumbs items={[{ text: 'Tools' }]} />
+      <ToolViewInner />
+   </Stack>
 };
 
