Horde: Tweak legend selection
#rnx

diff --git a/Engine/Source/Programs/Horde/HordeDashboard/src/components/telemetry/TelemetryView.tsx b/Engine/Source/Programs/Horde/HordeDashboard/src/components/telemetry/TelemetryView.tsx
--- a/Engine/Source/Programs/Horde/HordeDashboard/src/components/telemetry/TelemetryView.tsx
+++ b/Engine/Source/Programs/Horde/HordeDashboard/src/components/telemetry/TelemetryView.tsx
@@ -471,21 +471,9 @@ class MetricsHandler {
 
    }
 
-   setFilterKey(key: string, filtered: boolean) {
+   setFilterKeys(keys: Set<string>) {
 
-      if (filtered && this.filteredKeys.has(key)) {
-         return;
-      }
-
-      if (!filtered && !this.filteredKeys.has(key)) {
-         return;
-      }
-
-      if (filtered) {
-         this.filteredKeys.add(key);
-      } else {
-         this.filteredKeys.delete(key);
-      }
+      this.filteredKeys = keys;
 
       this.setUpdated();
    }
@@ -939,10 +927,51 @@ const Legend: React.FC<{ chart: GetTelemetryChartResponse }> = observer(({ chart
          directionalHint={DirectionalHint.leftCenter}
          calloutProps={calloutProps}
          styles={hostStyles}>
-         <Stack key={`key_legend_${v.key}`} horizontal verticalAlign="center" tokens={{ childrenGap: 8 }} onClick={() => {
-            if (legend.length > 1) {
-               handler.setFilterKey(v.key, !filtered);
+         <Stack key={`key_legend_${v.key}`} style={{ userSelect: "none" }} horizontal verticalAlign="center" tokens={{ childrenGap: 8 }} onClick={(ev) => {
+
+            let keys = legend.map(l => l.key);
+
+            const mod = ev.shiftKey || ev.ctrlKey;
+            if (!mod) {
+
+               if (handler.filteredKeys.size && !handler.filteredKeys.has(v.key)) {
+                  handler.setFilterKeys(new Set());
+               } else {
+                  const nset = new Set<string>();
+                  keys.forEach(k => {
+                     if (k !== v.key) {
+                        nset.add(k);
+                     }
+                  })
+                  handler.setFilterKeys(nset);
+               }
+            } else {
+
+               let nset = new Set(handler.filteredKeys);
+               if (nset.has(v.key)) {
+                  nset.delete(v.key);
+               } else {
+                  if (!nset.size) {
+                     nset = new Set<string>();
+
+                     keys.forEach(k => {
+                        if (k !== v.key) {
+                           nset.add(k);
+                        }
+                     })
+                  } else {
+                     nset.add(v.key);
+                  }
+               }
+
+               if (nset.size === legend.length) {
+                  nset = new Set();
+               }
+
+               handler.setFilterKeys(nset);
+
             }
+
          }}>
             <Stack>
                <FontIcon style={{ color: filtered ? "#999999" : graphColors[index % graphColors.length], paddingTop: 2 }} iconName="Square" />
@@ -1351,7 +1380,7 @@ export const SearchUpdate: React.FC = observer(() => {
       if (handler.search.toString().length) {
          setSearchParams(csearch, { replace: true });
       }
-      
+
    }, [csearch, setSearchParams])
 
    // subscribe
