Horde: Improve error messages when files are missing from the object store.

diff --git a/Engine/Source/Programs/Horde/Horde.Server/Storage/ObjectStores/AwsObjectStore.cs b/Engine/Source/Programs/Horde/Horde.Server/Storage/ObjectStores/AwsObjectStore.cs
--- a/Engine/Source/Programs/Horde/Horde.Server/Storage/ObjectStores/AwsObjectStore.cs
+++ b/Engine/Source/Programs/Horde/Horde.Server/Storage/ObjectStores/AwsObjectStore.cs
@@ -4,6 +4,7 @@ using System;
 using System.Collections.Generic;
 using System.Globalization;
 using System.IO;
+using System.Net;
 using System.Runtime.CompilerServices;
 using System.Threading;
 using System.Threading.Tasks;
@@ -24,14 +25,14 @@ namespace Horde.Server.Storage.ObjectStores
 	/// <summary>
 	/// Exception wrapper for S3 requests
 	/// </summary>
-	public sealed class AwsException : Exception
+	public sealed class AwsException : StorageException
 	{
 		/// <summary>
 		/// Constructor
 		/// </summary>
 		/// <param name="message">Message for the exception</param>
 		/// <param name="innerException">Inner exception data</param>
-		public AwsException(string? message, Exception? innerException)
+		public AwsException(string message, Exception? innerException)
 			: base(message, innerException)
 		{
 		}
@@ -325,7 +326,15 @@ namespace Horde.Server.Storage.ObjectStores
 				}
 
 				_logger.LogWarning(ex, "Unable to read {Path} from S3", fullPath);
-				throw new StorageException($"Unable to read {fullPath} from {_options.AwsBucketName}", ex);
+
+				if (ex is AmazonS3Exception s3ex && s3ex.StatusCode == HttpStatusCode.NotFound)
+				{
+					throw new ObjectNotFoundException(key, $"Object {key} not found in bucket {_options.AwsBucketName}", ex);
+				}
+				else
+				{
+					throw new StorageException($"Unable to read {fullPath} from {_options.AwsBucketName}: {ex.Message}", ex);
+				}
 			}
 		}
 
diff --git a/Engine/Source/Programs/Horde/Horde.Server/Storage/StorageService.cs b/Engine/Source/Programs/Horde/Horde.Server/Storage/StorageService.cs
--- a/Engine/Source/Programs/Horde/Horde.Server/Storage/StorageService.cs
+++ b/Engine/Source/Programs/Horde/Horde.Server/Storage/StorageService.cs
@@ -608,9 +608,13 @@ namespace Horde.Server.Storage
 								{
 									await TickBlobAsync(storageClient, blobInfo, cancellationToken);
 								}
+								catch (ObjectNotFoundException ex)
+								{
+									_logger.LogInformation(ex, "Unable to read references for {NamespaceId} blob {BlobId}: {Message}", blobInfo.NamespaceId, blobInfo.Id, GetObjectKey(blobInfo.Locator), ex.Message);
+								}
 								catch (Exception ex)
 								{
-									_logger.LogInformation(ex, "Unable to read references for {NamespaceId} blob {BlobId} (key: {ObjectKey}): {Message}", blobInfo.NamespaceId, blobInfo.Id, GetObjectKey(blobInfo.Locator), ex.Message);
+									_logger.LogWarning(ex, "Unable to read references for {NamespaceId} blob {BlobId} (key: {ObjectKey}): {Message}", blobInfo.NamespaceId, blobInfo.Id, GetObjectKey(blobInfo.Locator), ex.Message);
 								}
 							}
 						}
diff --git a/Engine/Source/Programs/Shared/EpicGames.Horde/Storage/Bundles/V2/PacketHandle.cs b/Engine/Source/Programs/Shared/EpicGames.Horde/Storage/Bundles/V2/PacketHandle.cs
--- a/Engine/Source/Programs/Shared/EpicGames.Horde/Storage/Bundles/V2/PacketHandle.cs
+++ b/Engine/Source/Programs/Shared/EpicGames.Horde/Storage/Bundles/V2/PacketHandle.cs
@@ -169,6 +169,11 @@ namespace EpicGames.Horde.Storage.Bundles.V2
 			{
 				throw;
 			}
+			catch (ObjectNotFoundException ex)
+			{
+				Utf8String locator = GetLocator();
+				throw new ObjectNotFoundException(ex.Key, $"Unable to read {locator}: {ex.Message}", ex);
+			}
 			catch (Exception ex)
 			{
 				Utf8String locator = GetLocator();
diff --git a/Engine/Source/Programs/Shared/EpicGames.Horde/Storage/StorageExceptions.cs b/Engine/Source/Programs/Shared/EpicGames.Horde/Storage/StorageExceptions.cs
--- a/Engine/Source/Programs/Shared/EpicGames.Horde/Storage/StorageExceptions.cs
+++ b/Engine/Source/Programs/Shared/EpicGames.Horde/Storage/StorageExceptions.cs
@@ -29,7 +29,7 @@ namespace EpicGames.Horde.Storage
 	/// <summary>
 	/// Exception thrown when an object does not exist
 	/// </summary>
-	public sealed class ObjectNotFoundException : Exception
+	public sealed class ObjectNotFoundException : StorageException
 	{
 		/// <summary>
 		/// Path to the object
@@ -39,8 +39,16 @@ namespace EpicGames.Horde.Storage
 		/// <summary>
 		/// Constructor
 		/// </summary>
-		public ObjectNotFoundException(ObjectKey key)
-			: base($"Object '{key}' was not found")
+		public ObjectNotFoundException(ObjectKey key, Exception? innerException = null)
+			: this(key, $"Object '{key}' was not found", innerException)
+		{
+		}
+
+		/// <summary>
+		/// Constructor
+		/// </summary>
+		public ObjectNotFoundException(ObjectKey key, string message, Exception? innerException = null)
+			: base(message, innerException)
 		{
 			Key = key;
 		}
