Horde: Fix padding on detail list at top of panel for tool and pools views
#rnx

diff --git a/Engine/Source/Programs/Horde/HordeDashboard/src/components/PoolsView.tsx b/Engine/Source/Programs/Horde/HordeDashboard/src/components/PoolsView.tsx
--- a/Engine/Source/Programs/Horde/HordeDashboard/src/components/PoolsView.tsx
+++ b/Engine/Source/Programs/Horde/HordeDashboard/src/components/PoolsView.tsx
@@ -320,7 +320,7 @@ const PoolList: React.FC = observer(() => {
    },];
 
    const onRenderDetailsHeader: IDetailsListProps['onRenderDetailsHeader'] = (props) => {
-      const customStyles: Partial<IDetailsHeaderStyles> = {};
+      const customStyles: Partial<IDetailsHeaderStyles> = {root: {paddingTop: 0}};
       if (props) {
          return (
             <Sticky stickyPosition={StickyPositionType.Header} isScrollSynced={true}>
diff --git a/Engine/Source/Programs/Horde/HordeDashboard/src/components/ToolView.tsx b/Engine/Source/Programs/Horde/HordeDashboard/src/components/ToolView.tsx
--- a/Engine/Source/Programs/Horde/HordeDashboard/src/components/ToolView.tsx
+++ b/Engine/Source/Programs/Horde/HordeDashboard/src/components/ToolView.tsx
@@ -136,10 +136,10 @@ const ToolPanel: React.FC<{ selectedKey: string }> = observer(({ selectedKey })
       </Stack>}
 
       {!!tools.length && <Stack className={hordeClasses.raised} >
-         <Stack styles={{ root: { paddingLeft: 12, paddingRight: 12, paddingBottom: 12, width: "100%" } }} >
+         <Stack styles={{ root: { paddingLeft: 12, paddingRight: 12, paddingBottom: 12, width: "100%", selectors: { ".ms-DetailsHeader": { "paddingTop": "0px" } } } }} >
             <DetailsList
                isHeaderVisible={true}
-               styles={{ headerWrapper: {paddingTop: 0} }}
+               styles={{ headerWrapper: { paddingTop: 0 } }}
                items={tools}
                columns={columns}
                selectionMode={SelectionMode.none}
@@ -184,7 +184,7 @@ export const ToolViewInner: React.FC = observer(() => {
                <Stack horizontal style={{ paddingTop: 12, paddingBottom: 48 }}>
                   <Stack key={`${key}`} style={{ paddingLeft: centerAlign }} />
                   <Stack style={{ width: 1440 }}>
-                     <Stack style={{paddingBottom: 12}}>
+                     <Stack style={{ paddingBottom: 12 }}>
                         <Pivot className={hordeClasses.pivot}
                            overflowBehavior='menu'
                            selectedKey={selectedKey}
