Horde: Fix replicator permissions checks.

diff --git a/Engine/Source/Programs/Horde/Horde.Server/Replicators/ReplicatorsController.cs b/Engine/Source/Programs/Horde/Horde.Server/Replicators/ReplicatorsController.cs
--- a/Engine/Source/Programs/Horde/Horde.Server/Replicators/ReplicatorsController.cs
+++ b/Engine/Source/Programs/Horde/Horde.Server/Replicators/ReplicatorsController.cs
@@ -88,7 +88,7 @@ namespace Horde.Server.Replicators
 			{
 				return NotFound();
 			}
-			if (streamConfig.Authorize(ReplicatorAclAction.ViewReplicator, User))
+			if (!streamConfig.Authorize(ReplicatorAclAction.ViewReplicator, User))
 			{
 				return Forbid(ReplicatorAclAction.ViewReplicator);
 			}
@@ -104,7 +104,7 @@ namespace Horde.Server.Replicators
 		}
 
 		/// <summary>
-		/// Query a replicator state
+		/// Update a replicator state
 		/// </summary>
 		/// <param name="replicatorId">Repliactor to query</param>
 		/// <param name="request">Update request</param>
@@ -118,7 +118,7 @@ namespace Horde.Server.Replicators
 			{
 				return NotFound();
 			}
-			if (streamConfig.Authorize(ReplicatorAclAction.UpdateReplicator, User))
+			if (!streamConfig.Authorize(ReplicatorAclAction.UpdateReplicator, User))
 			{
 				return Forbid(ReplicatorAclAction.UpdateReplicator);
 			}
