Horde: Increase verbosity of storage log output.

diff --git a/Engine/Source/Programs/Horde/Horde.Agent/appsettings.json b/Engine/Source/Programs/Horde/Horde.Agent/appsettings.json
--- a/Engine/Source/Programs/Horde/Horde.Agent/appsettings.json
+++ b/Engine/Source/Programs/Horde/Horde.Agent/appsettings.json
@@ -15,7 +15,8 @@
 				"Microsoft.AspNetCore.Hosting.Diagnostics": "Warning",
 				"Microsoft.AspNetCore.Mvc.Infrastructure.ControllerActionInvoker": "Warning",
 				"Serilog.AspNetCore.RequestLoggingMiddleware": "Warning",
-				"Grpc.Net.Client.Internal.GrpcCall": "Warning"
+				"Grpc.Net.Client.Internal.GrpcCall": "Warning",
+				"EpicGames.Horde.Storage": "Debug"
 			}
 		}
 	}
diff --git a/Engine/Source/Programs/Horde/Horde.Server/Storage/ObjectStores/AwsObjectStore.cs b/Engine/Source/Programs/Horde/Horde.Server/Storage/ObjectStores/AwsObjectStore.cs
--- a/Engine/Source/Programs/Horde/Horde.Server/Storage/ObjectStores/AwsObjectStore.cs
+++ b/Engine/Source/Programs/Horde/Horde.Server/Storage/ObjectStores/AwsObjectStore.cs
@@ -373,6 +373,7 @@ namespace Horde.Server.Storage.ObjectStores
 				newGetRequest.ResponseHeaderOverrides.CacheControl = "private, max-age=2592000, immutable"; // 30 days
 
 				string url = _client.GetPreSignedURL(newGetRequest);
+				_logger.LogDebug("Creating presigned URL for {Verb} to {Path}", verb, fullPath);
 				return new Uri(url);
 			}
 			catch (Exception ex)
@@ -399,7 +400,7 @@ namespace Horde.Server.Storage.ObjectStores
 				{
 					using IDisposable semaLock = await _semaphore.WaitDisposableAsync(cancellationToken);
 					await WriteInternalAsync(fullPath, inputStream, cancellationToken);
-					_logger.LogDebug("Written data to {Path}", key);
+					_logger.LogDebug("Written data to {Path}", fullPath);
 					break;
 				}
 				catch (Exception ex) when (ex is not OperationCanceledException)
diff --git a/Engine/Source/Programs/Horde/Horde.Server/appsettings.json b/Engine/Source/Programs/Horde/Horde.Server/appsettings.json
--- a/Engine/Source/Programs/Horde/Horde.Server/appsettings.json
+++ b/Engine/Source/Programs/Horde/Horde.Server/appsettings.json
@@ -18,7 +18,8 @@
 				"System.Net.Http.HttpClient": "Warning",
 				"Grpc": "Warning",
 				"Horde.Server.Issues.IssueService": "Debug",
-				"Horde.Server.Perforce.PerforceService": "Verbose"
+				"Horde.Server.Perforce.PerforceService": "Verbose",
+				"EpicGames.Horde.Storage": "Debug"
 			}
 		}
 	}
