Horde: Allow downloading artifacts as .uartifact files.

diff --git a/Engine/Source/Programs/Horde/Horde.Server/Agents/AgentsController.cs b/Engine/Source/Programs/Horde/Horde.Server/Agents/AgentsController.cs
--- a/Engine/Source/Programs/Horde/Horde.Server/Agents/AgentsController.cs
+++ b/Engine/Source/Programs/Horde/Horde.Server/Agents/AgentsController.cs
@@ -11,7 +11,6 @@ using EpicGames.Horde.Agents.Leases;
 using EpicGames.Horde.Agents.Pools;
 using EpicGames.Horde.Agents.Sessions;
 using EpicGames.Horde.Common;
-using EpicGames.Horde.Users;
 using Horde.Server.Agents.Leases;
 using Horde.Server.Agents.Sessions;
 using Horde.Server.Auditing;
diff --git a/Engine/Source/Programs/Horde/Horde.Server/Artifacts/ArtifactsController.cs b/Engine/Source/Programs/Horde/Horde.Server/Artifacts/ArtifactsController.cs
--- a/Engine/Source/Programs/Horde/Horde.Server/Artifacts/ArtifactsController.cs
+++ b/Engine/Source/Programs/Horde/Horde.Server/Artifacts/ArtifactsController.cs
@@ -174,7 +174,7 @@ namespace Horde.Server.Artifacts
 		/// <returns>Information about all the artifacts</returns>
 		[HttpGet]
 		[Route("/api/v2/artifacts/{id}")]
-		[ProducesResponseType(typeof(FindArtifactsResponse), 200)]
+		[ProducesResponseType(typeof(GetArtifactResponse), 200)]
 		public async Task<ActionResult<object>> GetArtifactAsync(ArtifactId id, [FromQuery] PropertyFilter? filter = null)
 		{
 			IArtifact? artifact = await _artifactCollection.GetAsync(id, HttpContext.RequestAborted);
@@ -467,34 +467,36 @@ namespace Horde.Server.Artifacts
 		}
 
 		/// <summary>
-		/// Downloads an individual file from an artifact
+		/// Downloads the artifact data
 		/// </summary>
 		/// <param name="id">Identifier of the artifact to retrieve</param>
-		/// <param name="filter">Paths to include in the zip file. The post version of this request allows for more parameters than can fit in a request string.</param>
+		/// <param name="format">Format for the download type</param>
+		/// <param name="filter">Paths to include. The post version of this request allows for more parameters than can fit in a request string.</param>
 		/// <param name="cancellationToken">Cancellation token for the operation</param>
 		/// <returns>Information about all the artifacts</returns>
 		[HttpGet]
-		[Route("/api/v2/artifacts/{id}/zip")]
-		public async Task<ActionResult<object>> GetZipAsync(ArtifactId id, [FromQuery(Name = "filter")] string[]? filter, CancellationToken cancellationToken = default)
+		[Route("/api/v2/artifacts/{id}/download")]
+		public async Task<ActionResult<object>> DownloadAsync(ArtifactId id, [FromQuery] DownloadArtifactFormat? format, [FromQuery(Name = "filter")] string[]? filter, CancellationToken cancellationToken = default)
 		{
-			return await GetZipInternalAsync(id, filter, cancellationToken);
+			return await DownloadInternalAsync(id, format, filter, cancellationToken);
 		}
 
 		/// <summary>
 		/// Downloads an individual file from an artifact
 		/// </summary>
 		/// <param name="id">Identifier of the artifact to retrieve</param>
+		/// <param name="format">Format for the download type</param>
 		/// <param name="request">Filter for the zip file</param>
 		/// <param name="cancellationToken">Cancellation token for the operation</param>
 		/// <returns>Information about all the artifacts</returns>
 		[HttpPost]
-		[Route("/api/v2/artifacts/{id}/zip")]
-		public async Task<ActionResult<object>> CreateZipFromFilterAsync(ArtifactId id, CreateZipRequest request, CancellationToken cancellationToken = default)
+		[Route("/api/v2/artifacts/{id}/download")]
+		public async Task<ActionResult<object>> DownloadWithFilterAsync(ArtifactId id, [FromQuery] DownloadArtifactFormat? format, CreateZipRequest request, CancellationToken cancellationToken = default)
 		{
-			return await GetZipInternalAsync(id, request.Filter, cancellationToken);
+			return await DownloadInternalAsync(id, format, request.Filter, cancellationToken);
 		}
 
-		async Task<ActionResult> GetZipInternalAsync(ArtifactId id, IEnumerable<string>? fileFilter, CancellationToken cancellationToken)
+		async Task<ActionResult> DownloadInternalAsync(ArtifactId id, DownloadArtifactFormat? format, IReadOnlyList<string>? fileFilter, CancellationToken cancellationToken)
 		{
 			IArtifact? artifact = await _artifactCollection.GetAsync(id, cancellationToken);
 			if (artifact == null)
@@ -506,6 +508,19 @@ namespace Horde.Server.Artifacts
 				return Forbid(ArtifactAclAction.ReadArtifact, artifact.AclScope);
 			}
 
+			switch (format ?? DownloadArtifactFormat.Zip)
+			{
+				case DownloadArtifactFormat.Zip:
+					return await GetZipInternalAsync(artifact, fileFilter, cancellationToken);
+				case DownloadArtifactFormat.Ugs:
+					return GetDescriptorInternal(artifact, fileFilter);
+				default:
+					return BadRequest("Unhandled download format");
+			}
+		}
+
+		async Task<ActionResult> GetZipInternalAsync(IArtifact artifact, IEnumerable<string>? fileFilter, CancellationToken cancellationToken)
+		{
 			FileFilter? filter = null;
 			if (fileFilter != null && fileFilter.Any())
 			{
@@ -529,6 +544,44 @@ namespace Horde.Server.Artifacts
 #pragma warning restore CA2000
 		}
 
+		ActionResult GetDescriptorInternal(IArtifact artifact, IReadOnlyList<string>? fileFilter)
+		{
+			Uri baseUri = new Uri(_globalConfig.ServerSettings.ServerUrl, $"api/v2/artfiacts/{artifact.Id}");
+
+			ArtifactDescriptor descriptor = new ArtifactDescriptor(baseUri, new RefName("default"), fileFilter);
+
+			byte[] data = descriptor.Serialize();
+			return new FileStreamResult(new MemoryStream(data), "application/x-horde-artifact") { FileDownloadName = $"{artifact.Name}.uartifact" };
+		}
+
+		/// <summary>
+		/// Downloads an individual file from an artifact
+		/// </summary>
+		/// <param name="id">Identifier of the artifact to retrieve</param>
+		/// <param name="filter">Paths to include in the zip file. The post version of this request allows for more parameters than can fit in a request string.</param>
+		/// <param name="cancellationToken">Cancellation token for the operation</param>
+		/// <returns>Information about all the artifacts</returns>
+		[HttpGet]
+		[Route("/api/v2/artifacts/{id}/zip")]
+		public async Task<ActionResult> GetZipAsync(ArtifactId id, [FromQuery(Name = "filter")] string[]? filter, CancellationToken cancellationToken = default)
+		{
+			return await DownloadInternalAsync(id, DownloadArtifactFormat.Zip, filter, cancellationToken);
+		}
+
+		/// <summary>
+		/// Downloads an individual file from an artifact
+		/// </summary>
+		/// <param name="id">Identifier of the artifact to retrieve</param>
+		/// <param name="request">Filter for the zip file</param>
+		/// <param name="cancellationToken">Cancellation token for the operation</param>
+		/// <returns>Information about all the artifacts</returns>
+		[HttpPost]
+		[Route("/api/v2/artifacts/{id}/zip")]
+		public async Task<ActionResult<object>> CreateZipFromFilterAsync(ArtifactId id, CreateZipRequest request, CancellationToken cancellationToken = default)
+		{
+			return await DownloadInternalAsync(id, DownloadArtifactFormat.Zip, request.Filter, cancellationToken);
+		}
+
 		/// <summary>
 		/// Finds artifacts matching certain criteria
 		/// </summary>
diff --git a/Engine/Source/Programs/Horde/Horde.Server/Notifications/SubscriptionsController.cs b/Engine/Source/Programs/Horde/Horde.Server/Notifications/SubscriptionsController.cs
--- a/Engine/Source/Programs/Horde/Horde.Server/Notifications/SubscriptionsController.cs
+++ b/Engine/Source/Programs/Horde/Horde.Server/Notifications/SubscriptionsController.cs
@@ -1,6 +1,5 @@
 // Copyright Epic Games, Inc. All Rights Reserved.
 
-using System;
 using System.Collections.Generic;
 using System.Threading.Tasks;
 using EpicGames.Horde.Users;
@@ -147,7 +146,7 @@ namespace Horde.Server.Notifications
 		/// <param name="userName"></param>
 		/// <param name="objectId"></param>
 		/// <returns></returns>
-		bool TryParseUserId(string userName, out UserId objectId)
+		static bool TryParseUserId(string userName, out UserId objectId)
 		{
 			UserId newObjectId;
 			if (UserId.TryParse(userName, out newObjectId))
diff --git a/Engine/Source/Programs/Horde/Horde.Server/Server/GlobalConfig.cs b/Engine/Source/Programs/Horde/Horde.Server/Server/GlobalConfig.cs
--- a/Engine/Source/Programs/Horde/Horde.Server/Server/GlobalConfig.cs
+++ b/Engine/Source/Programs/Horde/Horde.Server/Server/GlobalConfig.cs
@@ -5,7 +5,6 @@ using System.Collections.Generic;
 using System.ComponentModel.DataAnnotations;
 using System.Diagnostics;
 using System.Diagnostics.CodeAnalysis;
-using System.IO;
 using System.Linq;
 using System.Net;
 using System.Security.Claims;
diff --git a/Engine/Source/Programs/Horde/Horde.Server/Storage/StorageService.cs b/Engine/Source/Programs/Horde/Horde.Server/Storage/StorageService.cs
--- a/Engine/Source/Programs/Horde/Horde.Server/Storage/StorageService.cs
+++ b/Engine/Source/Programs/Horde/Horde.Server/Storage/StorageService.cs
@@ -610,7 +610,7 @@ namespace Horde.Server.Storage
 								}
 								catch (ObjectNotFoundException ex)
 								{
-									_logger.LogInformation(ex, "Unable to read references for {NamespaceId} blob {BlobId}: {Message}", blobInfo.NamespaceId, blobInfo.Id, GetObjectKey(blobInfo.Locator), ex.Message);
+									_logger.LogInformation(ex, "Unable to read references for {NamespaceId} blob {BlobId}: {Message}", blobInfo.NamespaceId, blobInfo.Id, ex.Message);
 								}
 								catch (Exception ex)
 								{
diff --git a/Engine/Source/Programs/Horde/Horde.Server/Utilities/HordeClaimTypes.cs b/Engine/Source/Programs/Horde/Horde.Server/Utilities/HordeClaimTypes.cs
--- a/Engine/Source/Programs/Horde/Horde.Server/Utilities/HordeClaimTypes.cs
+++ b/Engine/Source/Programs/Horde/Horde.Server/Utilities/HordeClaimTypes.cs
@@ -1,6 +1,5 @@
 // Copyright Epic Games, Inc. All Rights Reserved.
 
-using System.Linq;
 using System.Security.Claims;
 using EpicGames.Horde.Accounts;
 using EpicGames.Horde.Users;
diff --git a/Engine/Source/Programs/Shared/EpicGames.Horde/Artifacts/ArtifactDescriptor.cs b/Engine/Source/Programs/Shared/EpicGames.Horde/Artifacts/ArtifactDescriptor.cs
new file mode 100644
--- /dev/null
+++ b/Engine/Source/Programs/Shared/EpicGames.Horde/Artifacts/ArtifactDescriptor.cs
@@ -0,0 +1,90 @@
+// Copyright Epic Games, Inc. All Rights Reserved.
+
+using System;
+using System.Collections.Generic;
+using System.Text.Json;
+using System.Threading;
+using System.Threading.Tasks;
+using EpicGames.Core;
+
+namespace EpicGames.Horde.Storage
+{
+	/// <summary>
+	/// Class which describes an artifact on Horde which can be serialized to JSON.
+	/// </summary>
+	public class ArtifactDescriptor
+	{
+		/// <summary>
+		/// Base URL for downloading from
+		/// </summary>
+		public Uri BaseUrl { get; set; }
+
+		/// <summary>
+		/// Name of the ref to download
+		/// </summary>
+		public RefName RefName { get; set; }
+
+		/// <summary>
+		/// Filter for the files selected for download
+		/// </summary>
+		public List<string>? Filter { get; }
+
+		/// <summary>
+		/// Constructor
+		/// </summary>
+		public ArtifactDescriptor(Uri baseUrl, RefName refName, IReadOnlyCollection<string>? filter)
+		{
+			BaseUrl = baseUrl;
+			RefName = refName;
+			if (filter != null && filter.Count > 0)
+			{
+				Filter = new List<string>(filter);
+			}
+		}
+
+		/// <summary>
+		/// Deserialize a descriptor from utf8 bytes
+		/// </summary>
+		/// <param name="data">Data to deserialize from</param>
+		/// <returns>New descriptor instance</returns>
+		public static ArtifactDescriptor Deserialize(ReadOnlySpan<byte> data)
+		{
+			return JsonSerializer.Deserialize<ArtifactDescriptor>(data, HordeHttpClient.JsonSerializerOptions)
+				?? throw new InvalidOperationException("Cannot deserialize descriptor");
+		}
+
+		/// <summary>
+		/// Writes this descriptor to storage
+		/// </summary>
+		/// <param name="file">File to read from</param>
+		/// <param name="cancellationToken">Cancellation token for the operation</param>
+		public static async Task<ArtifactDescriptor> ReadAsync(FileReference file, CancellationToken cancellationToken)
+		{
+			byte[] data = await FileReference.ReadAllBytesAsync(file, cancellationToken);
+			return Deserialize(data);
+		}
+
+		/// <summary>
+		/// Serializes a descriptor
+		/// </summary>
+		/// <returns>Data for the serialized descriptor</returns>
+		public byte[] Serialize()
+		{
+			JsonSerializerOptions serializerOptions = new JsonSerializerOptions(HordeHttpClient.JsonSerializerOptions);
+			serializerOptions.WriteIndented = true;
+
+			return JsonSerializer.SerializeToUtf8Bytes<ArtifactDescriptor>(this, serializerOptions);
+		}
+
+		/// <summary>
+		/// Writes this descriptor to storage
+		/// </summary>
+		/// <param name="file">File to write to</param>
+		/// <param name="cancellationToken">Cancellation token for the operation</param>
+		public async Task WriteAsync(FileReference file, CancellationToken cancellationToken)
+		{
+			byte[] data = Serialize();
+			await FileReference.WriteAllBytesAsync(file, data, cancellationToken);
+		}
+	}
+}
diff --git a/Engine/Source/Programs/Shared/EpicGames.Horde/Artifacts/ArtifactMessages.cs b/Engine/Source/Programs/Shared/EpicGames.Horde/Artifacts/ArtifactMessages.cs
--- a/Engine/Source/Programs/Shared/EpicGames.Horde/Artifacts/ArtifactMessages.cs
+++ b/Engine/Source/Programs/Shared/EpicGames.Horde/Artifacts/ArtifactMessages.cs
@@ -31,6 +31,22 @@ namespace EpicGames.Horde.Artifacts
 	/// <param name="Token">Token which can be used to upload blobs for the artifact, and read blobs from the previous artifact</param>
 	public record CreateArtifactResponse(ArtifactId ArtifactId, NamespaceId NamespaceId, RefName RefName, RefName? PrevRefName, string Token);
 
+	/// <summary>
+	/// Type of data to download for an artifact
+	/// </summary>
+	public enum DownloadArtifactFormat
+	{
+		/// <summary>
+		/// Download as a zip file
+		/// </summary>
+		Zip,
+
+		/// <summary>
+		/// Download as a UGS link
+		/// </summary>
+		Ugs
+	}
+
 	/// <summary>
 	/// Describes an artifact
 	/// </summary>
