Horde: Fix incorrect URL to artifacts endpoint.

diff --git a/Engine/Source/Programs/Horde/Horde.Server/Artifacts/ArtifactsController.cs b/Engine/Source/Programs/Horde/Horde.Server/Artifacts/ArtifactsController.cs
--- a/Engine/Source/Programs/Horde/Horde.Server/Artifacts/ArtifactsController.cs
+++ b/Engine/Source/Programs/Horde/Horde.Server/Artifacts/ArtifactsController.cs
@@ -546,7 +546,7 @@ namespace Horde.Server.Artifacts
 
 		ActionResult GetDescriptorInternal(IArtifact artifact, IReadOnlyList<string>? fileFilter)
 		{
-			Uri baseUri = new Uri(_globalConfig.ServerSettings.ServerUrl, $"api/v2/artfiacts/{artifact.Id}");
+			Uri baseUri = new Uri(_globalConfig.ServerSettings.ServerUrl, $"api/v2/artifacts/{artifact.Id}");
 
 			ArtifactDescriptor descriptor = new ArtifactDescriptor(baseUri, new RefName("default"), fileFilter);
 
