Added UBT Completed Event Horde Analytics

diff --git a/Engine/Source/Programs/Horde/Horde.Server/Defaults/default-analytics.dashboard.json b/Engine/Source/Programs/Horde/Horde.Server/Defaults/default-analytics.dashboard.json
--- a/Engine/Source/Programs/Horde/Horde.Server/Defaults/default-analytics.dashboard.json
+++ b/Engine/Source/Programs/Horde/Horde.Server/Defaults/default-analytics.dashboard.json
@@ -566,10 +566,10 @@
 					]
 				},
 				{
-					"name": "UBA",
+					"name": "UBT",
 					"charts": [
 						{
-							"name": "Compile Duration",
+							"name": "UBA Compile WallTime (Average)",
 							"metrics": [
 								{
 									"id": "ubt-uba-duration"
@@ -577,7 +577,7 @@
 							]
 						},
 						{
-							"name": "Compile Distribution",
+							"name": "UBA Compile Distribution (Average)",
 							"display": "Ratio",
 							"metrics": [
 								{
@@ -591,7 +591,7 @@
 							]
 						},
 						{
-							"name": "Compile Actions (Average)",
+							"name": "UBA Compile Actions (Average)",
 							"display": "Value",
 							"metrics": [
 								{
@@ -609,7 +609,7 @@
 							]
 						},
 						{
-							"name": "Compile Actions (Total)",
+							"name": "UBA Compile Actions (Total)",
 							"display": "Value",
 							"metrics": [
 								{
@@ -627,12 +627,21 @@
 							]
 						},
 						{
-							"name": "Failure Rate",
+							"name": "UBA Failure Rate (Average)",
 							"display": "Ratio",
 							"metrics": [
 								{
 									"id": "ubt-uba-failure-rate"
 
+								}
+							]
+						},
+						{
+							"name": "UBT Wall Time (Average)",
+							"metrics": [
+								{
+									"id": "ubt-completed-duration"
+
 								}
 							]
 						}
diff --git a/Engine/Source/Programs/Horde/Horde.Server/Defaults/default-metrics.telemetry.json b/Engine/Source/Programs/Horde/Horde.Server/Defaults/default-metrics.telemetry.json
--- a/Engine/Source/Programs/Horde/Horde.Server/Defaults/default-metrics.telemetry.json
+++ b/Engine/Source/Programs/Horde/Horde.Server/Defaults/default-metrics.telemetry.json
@@ -561,6 +561,14 @@
 			"property": "$.Payload.FailureRate",
 			"function": "Average",
 			"interval": "00:30:00"
+		},
+		{
+			"id": "ubt-completed-duration",
+			"filter": "$[?(@.Payload.EventName == 'Core.UBT.Completed' )]",
+			"groupBy": "$.Payload.Target,$.Payload.Mode,$.Payload.Configuration,$.Payload.Platform",
+			"property": "$.Payload.Duration",
+			"function": "Average",
+			"interval": "00:30:00"
 		}
 	]
 }
\ No newline at end of file
