Horde: Add parameter for image repository auth file

diff --git a/Engine/Source/Programs/Horde/BuildHorde.xml b/Engine/Source/Programs/Horde/BuildHorde.xml
--- a/Engine/Source/Programs/Horde/BuildHorde.xml
+++ b/Engine/Source/Programs/Horde/BuildHorde.xml
@@ -15,6 +15,7 @@
 	<Option Name="ImageName" DefaultValue="horde-server" Description="Name of Docker/OCI image"/>
 	<Option Name="ImageTag" DefaultValue="$(DefaultImageTag)" Description="Tag to apply to image"/>
 	<Option Name="ImageEnvFile" DefaultValue="" Description="Path to environment file for used for Docker invocations"/>
+	<Option Name="ImageRepoAuthFile" DefaultValue="" Description="Path to file with repository auth credentials"/>
 
 	<Property Name="EngineDir" Value="$(RootDir)/Engine"/>
 
@@ -298,7 +299,7 @@
 		</Node>
 
 		<Node Name="Push Bundled Docker Image" Requires="Build Bundled Docker Image">
-			<Docker-Push Repository="$(ImageRepoUrl)" Image="horde-server:bundled" TargetImage="$(ImageName):$(ImageTag)-bundled" EnvironmentFile="$(ImageEnvFile)" />
+			<Docker-Push Repository="$(ImageRepoUrl)" Image="horde-server:bundled" TargetImage="$(ImageName):$(ImageTag)-bundled" EnvironmentFile="$(ImageEnvFile)" RepositoryAuthFile="$(ImageRepoAuthFile)" />
 		</Node>
 	</Agent>
 
