Split UBT and UBA into separate categories in Horde Analytics

diff --git a/Engine/Source/Programs/Horde/Horde.Server/Defaults/default-analytics.dashboard.json b/Engine/Source/Programs/Horde/Horde.Server/Defaults/default-analytics.dashboard.json
--- a/Engine/Source/Programs/Horde/Horde.Server/Defaults/default-analytics.dashboard.json
+++ b/Engine/Source/Programs/Horde/Horde.Server/Defaults/default-analytics.dashboard.json
@@ -566,10 +566,10 @@
 					]
 				},
 				{
-					"name": "UBT",
+					"name": "UBA",
 					"charts": [
 						{
-							"name": "UBA Compile WallTime (Average)",
+							"name": "Compile WallTime (Average)",
 							"metrics": [
 								{
 									"id": "ubt-uba-duration"
@@ -577,7 +577,7 @@
 							]
 						},
 						{
-							"name": "UBA Compile Distribution (Average)",
+							"name": "Compile Distribution (Average)",
 							"display": "Ratio",
 							"metrics": [
 								{
@@ -591,7 +591,7 @@
 							]
 						},
 						{
-							"name": "UBA Compile Actions (Average)",
+							"name": "Compile Actions (Average)",
 							"display": "Value",
 							"metrics": [
 								{
@@ -609,7 +609,7 @@
 							]
 						},
 						{
-							"name": "UBA Compile Actions (Total)",
+							"name": "Compile Actions (Total)",
 							"display": "Value",
 							"metrics": [
 								{
@@ -627,7 +627,7 @@
 							]
 						},
 						{
-							"name": "UBA Failure Rate (Average)",
+							"name": "Failure Rate (Average)",
 							"display": "Ratio",
 							"metrics": [
 								{
@@ -635,9 +635,14 @@
 
 								}
 							]
-						},
+						}
+					]
+				},
+				{
+					"name": "UBT",
+					"charts": [
 						{
-							"name": "UBT Wall Time (Average)",
+							"name": "Wall Time (Average)",
 							"metrics": [
 								{
 									"id": "ubt-completed-duration"
