Horde: Fix some documentation formatting.

diff --git a/Engine/Source/Programs/Horde/Horde.Server/Commands/Generate/ConfigDocsCommand.cs b/Engine/Source/Programs/Horde/Horde.Server/Commands/Generate/ConfigDocsCommand.cs
--- a/Engine/Source/Programs/Horde/Horde.Server/Commands/Generate/ConfigDocsCommand.cs
+++ b/Engine/Source/Programs/Horde/Horde.Server/Commands/Generate/ConfigDocsCommand.cs
@@ -283,7 +283,7 @@ namespace Horde.Server.Commands.Generate
 
 		static string GetMarkdownDescription(string? description)
 		{
-			return (description ?? String.Empty).Replace("\n", "<br>", StringComparison.Ordinal);
+			return (description ?? String.Empty).Replace("|", @"\|").Replace("\n", "<br>", StringComparison.Ordinal);
 		}
 
 		static string GetHeadingName(JsonSchemaType type)
diff --git a/Engine/Source/Programs/Horde/Horde.Server/Streams/StreamConfig.cs b/Engine/Source/Programs/Horde/Horde.Server/Streams/StreamConfig.cs
--- a/Engine/Source/Programs/Horde/Horde.Server/Streams/StreamConfig.cs
+++ b/Engine/Source/Programs/Horde/Horde.Server/Streams/StreamConfig.cs
@@ -142,7 +142,7 @@ namespace Horde.Server.Streams
 		public string? NotificationChannel { get; set; }
 
 		/// <summary>
-		/// Notification channel filter for this template. Can be Success|Failure|Warnings
+		/// Notification channel filter for this template. Can be Success, Failure, or Warnings.
 		/// </summary>
 		public string? NotificationChannelFilter { get; set; }
 
@@ -824,7 +824,7 @@ namespace Horde.Server.Streams
 		public string? NotificationChannel { get; set; }
 
 		/// <summary>
-		/// Notification channel filter for this template. Can be Success|Failure|Warnings
+		/// Notification channel filter for this template. Can be a combination of "Success", "Failure" and "Warnings" separated by pipe characters.
 		/// </summary>
 		public string? NotificationChannelFilter { get; set; }
 
diff --git a/Engine/Source/Programs/Horde/Horde.Server/Streams/StreamMessages.cs b/Engine/Source/Programs/Horde/Horde.Server/Streams/StreamMessages.cs
--- a/Engine/Source/Programs/Horde/Horde.Server/Streams/StreamMessages.cs
+++ b/Engine/Source/Programs/Horde/Horde.Server/Streams/StreamMessages.cs
@@ -232,7 +232,7 @@ namespace Horde.Server.Streams
 		public string? NotificationChannel { get; set; }
 
 		/// <summary>
-		/// Notification channel filter for this template. Can be Success|Failure|Warnings
+		/// Notification channel filter for this template. Can be a combination of "Success", "Failure" and "Warnings" separated by pipe characters.
 		/// </summary>
 		public string? NotificationChannelFilter { get; set; }
 
@@ -326,7 +326,7 @@ namespace Horde.Server.Streams
 		public string? NotificationChannel { get; set; }
 
 		/// <summary>
-		/// Notification channel filter for this template. Can be Success|Failure|Warnings
+		/// Notification channel filter for this template. Can be a combination of "Success", "Failure" and "Warnings" separated by pipe characters.
 		/// </summary>
 		public string? NotificationChannelFilter { get; set; }
 
