Added UBT/UBA telemetry filtering by [REDACTED_PERSON] to Horde Analytics Dashboard

diff --git a/Engine/Source/Programs/Horde/Horde.Server/Defaults/default-analytics.dashboard.json b/Engine/Source/Programs/Horde/Horde.Server/Defaults/default-analytics.dashboard.json
--- a/Engine/Source/Programs/Horde/Horde.Server/Defaults/default-analytics.dashboard.json
+++ b/Engine/Source/Programs/Horde/Horde.Server/Defaults/default-analytics.dashboard.json
@@ -77,7 +77,10 @@
 				},
 				{
 					"name": "Mode",
-					"group": "Mode"
+					"group": "Mode",
+					"defaults": [
+						"Rebuild"
+					]
 				},
 				{
 					"name": "Platform",
diff --git a/Engine/Source/Programs/Horde/Horde.Server/Defaults/default-metrics.telemetry.json b/Engine/Source/Programs/Horde/Horde.Server/Defaults/default-metrics.telemetry.json
--- a/Engine/Source/Programs/Horde/Horde.Server/Defaults/default-metrics.telemetry.json
+++ b/Engine/Source/Programs/Horde/Horde.Server/Defaults/default-metrics.telemetry.json
@@ -478,14 +478,14 @@
 			"id": "horde-incremental-compile-jobs",
 			"filter": "$[?(@.Payload.EventName == 'State.JobStepRef' && @.Payload.JobName=='Incremental CI' && search(@.Payload.StepName, \"Incremental Compile\") && @.Payload.Duration>0 && @.Payload.State=='Completed')]",
 			"groupBy": "$.Payload.StreamId,$.Payload.StepName",
-			"property": "$.Payload.Duration",
+			"property": "Build_BranchName$.Payload.Duration",
 			"function": "Average",
 			"interval": "01:00:00"
 		},
 		{
 			"id": "ubt-uba-duration",
 			"filter": "$[?(@.Payload.EventName == 'Core.UBT.Executor' && search(@.Payload.Executor, \"Unreal Build Accelerator\") )]",
-			"groupBy": "$.Payload.Target,$.Payload.Mode,$.Payload.Configuration,$.Payload.Platform",
+			"groupBy": "$.Payload.Build_BranchName,$.Payload.Target,$.Payload.Mode,$.Payload.Configuration,$.Payload.Platform",
 			"property": "$.Payload.Duration",
 			"function": "Average",
 			"interval": "00:30:00"
@@ -493,7 +493,7 @@
 		{
 			"id": "ubt-uba-remote-usage",
 			"filter": "$[?(@.Payload.EventName == 'Core.UBT.Executor' && search(@.Payload.Executor, \"Unreal Build Accelerator\") )]",
-			"groupBy": "$.Payload.Target,$.Payload.Mode,$.Payload.Configuration,$.Payload.Platform",
+			"groupBy": "$.Payload.Build_BranchName,$.Payload.Target,$.Payload.Mode,$.Payload.Configuration,$.Payload.Platform",
 			"property": "$.Payload.RemoteUsage",
 			"function": "Average",
 			"interval": "00:30:00"
@@ -509,7 +509,7 @@
 		{
 			"id": "ubt-uba-remote-actions",
 			"filter": "$[?(@.Payload.EventName == 'Core.UBT.Executor' && search(@.Payload.Executor, \"Unreal Build Accelerator\") )]",
-			"groupBy": "$.Payload.Target,$.Payload.Mode,$.Payload.Configuration,$.Payload.Platform",
+			"groupBy": "$.Payload.Build_BranchName,$.Payload.Target,$.Payload.Mode,$.Payload.Configuration,$.Payload.Platform",
 			"property": "$.Payload.RemoteActions",
 			"function": "Sum",
 			"interval": "00:30:00"
@@ -517,7 +517,7 @@
 		{
 			"id": "ubt-uba-local-actions",
 			"filter": "$[?(@.Payload.EventName == 'Core.UBT.Executor' && search(@.Payload.Executor, \"Unreal Build Accelerator\") )]",
-			"groupBy": "$.Payload.Target,$.Payload.Mode,$.Payload.Configuration,$.Payload.Platform",
+			"groupBy": "$.Payload.Build_BranchName,$.Payload.Target,$.Payload.Mode,$.Payload.Configuration,$.Payload.Platform",
 			"property": "$.Payload.LocalActions",
 			"function": "Sum",
 			"interval": "00:30:00"
@@ -525,7 +525,7 @@
 		{
 			"id": "ubt-uba-total-actions",
 			"filter": "$[?(@.Payload.EventName == 'Core.UBT.Executor' && search(@.Payload.Executor, \"Unreal Build Accelerator\") )]",
-			"groupBy": "$.Payload.Target,$.Payload.Mode,$.Payload.Configuration,$.Payload.Platform",
+			"groupBy": "$.Payload.Build_BranchName,$.Payload.Target,$.Payload.Mode,$.Payload.Configuration,$.Payload.Platform",
 			"property": "$.Payload.TotalActions",
 			"function": "Sum",
 			"interval": "00:30:00"
@@ -533,7 +533,7 @@
 		{
 			"id": "ubt-uba-remote-actions-avg",
 			"filter": "$[?(@.Payload.EventName == 'Core.UBT.Executor' && search(@.Payload.Executor, \"Unreal Build Accelerator\") )]",
-			"groupBy": "$.Payload.Target,$.Payload.Mode,$.Payload.Configuration,$.Payload.Platform",
+			"groupBy": "$.Payload.Build_BranchName,$.Payload.Target,$.Payload.Mode,$.Payload.Configuration,$.Payload.Platform",
 			"property": "$.Payload.RemoteActions",
 			"function": "Average",
 			"interval": "00:30:00"
@@ -541,7 +541,7 @@
 		{
 			"id": "ubt-uba-local-actions-avg",
 			"filter": "$[?(@.Payload.EventName == 'Core.UBT.Executor' && search(@.Payload.Executor, \"Unreal Build Accelerator\") )]",
-			"groupBy": "$.Payload.Target,$.Payload.Mode,$.Payload.Configuration,$.Payload.Platform",
+			"groupBy": "$.Payload.Build_BranchName,$.Payload.Target,$.Payload.Mode,$.Payload.Configuration,$.Payload.Platform",
 			"property": "$.Payload.LocalActions",
 			"function": "Average",
 			"interval": "00:30:00"
@@ -549,7 +549,7 @@
 		{
 			"id": "ubt-uba-total-actions-avg",
 			"filter": "$[?(@.Payload.EventName == 'Core.UBT.Executor' && search(@.Payload.Executor, \"Unreal Build Accelerator\") )]",
-			"groupBy": "$.Payload.Target,$.Payload.Mode,$.Payload.Configuration,$.Payload.Platform",
+			"groupBy": "$.Payload.Build_BranchName,$.Payload.Target,$.Payload.Mode,$.Payload.Configuration,$.Payload.Platform",
 			"property": "$.Payload.TotalActions",
 			"function": "Average",
 			"interval": "00:30:00"
@@ -557,7 +557,7 @@
 		{
 			"id": "ubt-uba-failure-rate",
 			"filter": "$[?(@.Payload.EventName == 'Core.UBT.Executor' && search(@.Payload.Executor, \"Unreal Build Accelerator\") )]",
-			"groupBy": "$.Payload.Target,$.Payload.Mode,$.Payload.Configuration,$.Payload.Platform",
+			"groupBy": "$.Payload.Build_BranchName,$.Payload.Target,$.Payload.Mode,$.Payload.Configuration,$.Payload.Platform",
 			"property": "$.Payload.FailureRate",
 			"function": "Average",
 			"interval": "00:30:00"
@@ -565,7 +565,7 @@
 		{
 			"id": "ubt-completed-duration",
 			"filter": "$[?(@.Payload.EventName == 'Core.UBT.Completed' )]",
-			"groupBy": "$.Payload.Target,$.Payload.Mode,$.Payload.Configuration,$.Payload.Platform",
+			"groupBy": "$.Payload.Build_BranchName,$.Payload.Target,$.Payload.Mode,$.Payload.Configuration,$.Payload.Platform",
 			"property": "$.Payload.Duration",
 			"function": "Average",
 			"interval": "00:30:00"
