Horde: Static analysis fix.

#rnx

diff --git a/Engine/Source/Programs/Horde/Horde.Server/Commands/Generate/ConfigDocsCommand.cs b/Engine/Source/Programs/Horde/Horde.Server/Commands/Generate/ConfigDocsCommand.cs
--- a/Engine/Source/Programs/Horde/Horde.Server/Commands/Generate/ConfigDocsCommand.cs
+++ b/Engine/Source/Programs/Horde/Horde.Server/Commands/Generate/ConfigDocsCommand.cs
@@ -283,7 +283,7 @@ namespace Horde.Server.Commands.Generate
 
 		static string GetMarkdownDescription(string? description)
 		{
-			return (description ?? String.Empty).Replace("|", @"\|").Replace("\n", "<br>", StringComparison.Ordinal);
+			return (description ?? String.Empty).Replace("|", @"\|", StringComparison.Ordinal).Replace("\n", "<br>", StringComparison.Ordinal);
 		}
 
 		static string GetHeadingName(JsonSchemaType type)
