Added sensible defaults for UBt/UBA Horde Analytics serach filters

diff --git a/Engine/Source/Programs/Horde/Horde.Server/Defaults/default-analytics.dashboard.json b/Engine/Source/Programs/Horde/Horde.Server/Defaults/default-analytics.dashboard.json
--- a/Engine/Source/Programs/Horde/Horde.Server/Defaults/default-analytics.dashboard.json
+++ b/Engine/Source/Programs/Horde/Horde.Server/Defaults/default-analytics.dashboard.json
@@ -84,11 +84,17 @@
 				},
 				{
 					"name": "Platform",
-					"group": "Platform"
+					"group": "Platform",
+					"defaults": [
+						"Win64"
+					]
 				},
 				{
 					"name": "Configuration",
-					"group": "Configuration"
+					"group": "Configuration",
+					"defaults": [
+						"Development"
+					]
 				}
 			],
 			"categories": [
