Fix for un parsable json in telemetry configuration

diff --git a/Engine/Source/Programs/Horde/Horde.Server/Defaults/default-metrics.telemetry.json b/Engine/Source/Programs/Horde/Horde.Server/Defaults/default-metrics.telemetry.json
--- a/Engine/Source/Programs/Horde/Horde.Server/Defaults/default-metrics.telemetry.json
+++ b/Engine/Source/Programs/Horde/Horde.Server/Defaults/default-metrics.telemetry.json
@@ -478,7 +478,7 @@
 			"id": "horde-incremental-compile-jobs",
 			"filter": "$[?(@.Payload.EventName == 'State.JobStepRef' && @.Payload.JobName=='Incremental CI' && search(@.Payload.StepName, \"Incremental Compile\") && @.Payload.Duration>0 && @.Payload.State=='Completed')]",
 			"groupBy": "$.Payload.StreamId,$.Payload.StepName",
-			"property": "Build_BranchName$.Payload.Duration",
+			"property": "$.Payload.Build_BranchName$.Payload.Duration",
 			"function": "Average",
 			"interval": "01:00:00"
 		},
@@ -501,7 +501,7 @@
 		{
 			"id": "ubt-uba-local-usage",
 			"filter": "$[?(@.Payload.EventName == 'Core.UBT.Executor' && search(@.Payload.Executor, \"Unreal Build Accelerator\") )]",
-			"groupBy": "$.Payload.Target,$.Payload.Mode,$.Payload.Configuration,$.Payload.Platform",
+			"groupBy": "$.Payload.Build_BranchName,$.Payload.Target,$.Payload.Mode,$.Payload.Configuration,$.Payload.Platform",
 			"property": "$.Payload.LocalUsage",
 			"function": "Average",
 			"interval": "00:30:00"
