Horde: Updating schema documentation

#rnx

diff --git a/Engine/Source/Programs/Horde/Docs/Config/Schema/Streams.md b/Engine/Source/Programs/Horde/Docs/Config/Schema/Streams.md
--- a/Engine/Source/Programs/Horde/Docs/Config/Schema/Streams.md
+++ b/Engine/Source/Programs/Horde/Docs/Config/Schema/Streams.md
@@ -15,7 +15,7 @@ Name | Description
 `order` | `integer`<br>Order for this stream
 `initialAgentType` | `string`<br>Default initial agent type for templates
 `notificationChannel` | `string`<br>Notification channel for all jobs in this stream
-`notificationChannelFilter` | `string`<br>Notification channel filter for this template. Can be Success|Failure|Warnings
+`notificationChannelFilter` | `string`<br>Notification channel filter for this template. Can be Success, Failure, or Warnings.
 `triageChannel` | `string`<br>Channel to post issue triage notifications
 `jobOptions` | [`JobOptions`](#joboptions)<br>Default settings for executing jobs
 `telemetryStoreId` | `string`<br>Telemetry store for Horde data for this stream
@@ -216,7 +216,7 @@ Name | Description
 `showUgsBadges` | `boolean`<br>Whether to show badges in UGS for these jobs
 `showUgsAlerts` | `boolean`<br>Whether to show alerts in UGS for these jobs
 `notificationChannel` | `string`<br>Notification channel for this template. Overrides the stream channel if set.
-`notificationChannelFilter` | `string`<br>Notification channel filter for this template. Can be Success|Failure|Warnings
+`notificationChannelFilter` | `string`<br>Notification channel filter for this template. Can be a combination of "Success", "Failure" and "Warnings" separated by pipe characters.
 `triageChannel` | `string`<br>Triage channel for this template. Overrides the stream channel if set.
 `workflowId` | `string`<br>Workflow to user for this stream
 `annotations` | `string` `->` `string`<br>Default annotations to apply to nodes in this template
