UBA Dashboard Tweak. Removed Local values as this simply the inverse of the remote values

diff --git a/Engine/Source/Programs/Horde/Horde.Server/Defaults/default-analytics.dashboard.json b/Engine/Source/Programs/Horde/Horde.Server/Defaults/default-analytics.dashboard.json
--- a/Engine/Source/Programs/Horde/Horde.Server/Defaults/default-analytics.dashboard.json
+++ b/Engine/Source/Programs/Horde/Horde.Server/Defaults/default-analytics.dashboard.json
@@ -571,7 +571,7 @@
 									"id": "horde-incremental-compile-jobs"
 								}
 							]
-						}	
+						}
 					]
 				},
 				{
@@ -592,10 +592,6 @@
 								{
 									"id": "ubt-uba-remote-usage",
 									"alias": "Remote"
-								},
-								{
-									"id": "ubt-uba-local-usage",
-									"alias": "Local"
 								}
 							]
 						},
@@ -610,10 +606,6 @@
 								{
 									"id": "ubt-uba-remote-actions-avg",
 									"alias": "Remote"
-								},
-								{
-									"id": "ubt-uba-local-actions-avg",
-									"alias": "Local"
 								}
 							]
 						},
@@ -628,10 +620,6 @@
 								{
 									"id": "ubt-uba-remote-actions",
 									"alias": "Remote"
-								},
-								{
-									"id": "ubt-uba-local-actions",
-									"alias": "Local"
 								}
 							]
 						},
