More Changes To UBA Charts in Horde Analytics

diff --git a/Engine/Source/Programs/Horde/Horde.Server/Defaults/default-analytics.dashboard.json b/Engine/Source/Programs/Horde/Horde.Server/Defaults/default-analytics.dashboard.json
--- a/Engine/Source/Programs/Horde/Horde.Server/Defaults/default-analytics.dashboard.json
+++ b/Engine/Source/Programs/Horde/Horde.Server/Defaults/default-analytics.dashboard.json
@@ -577,6 +577,30 @@
 				{
 					"name": "UBA",
 					"charts": [
+						{
+							"name": "Compile Distibution All (Average)",
+							"display": "Ratio",
+							"metrics": [
+								{
+									"id": "ubt-uba-remote-usage-all",
+									"alias": "Remote"
+								}
+							]
+						},
+						{
+							"name": "Compile Actions All (Sum)",
+							"display": "Value",
+							"metrics": [
+								{
+									"id": "ubt-uba-total-actions-all",
+									"alias": "Total"
+								},
+								{
+									"id": "ubt-uba-remote-actions-all",
+									"alias": "Remote"
+								}
+							]
+						},
 						{
 							"name": "Compile WallTime (Average)",
 							"metrics": [
@@ -610,7 +634,7 @@
 							]
 						},
 						{
-							"name": "Compile Actions (Total)",
+							"name": "Compile Actions (Sum)",
 							"display": "Value",
 							"metrics": [
 								{
@@ -629,7 +653,6 @@
 							"metrics": [
 								{
 									"id": "ubt-uba-failure-rate"
-
 								}
 							]
 						}
diff --git a/Engine/Source/Programs/Horde/Horde.Server/Defaults/default-metrics.telemetry.json b/Engine/Source/Programs/Horde/Horde.Server/Defaults/default-metrics.telemetry.json
--- a/Engine/Source/Programs/Horde/Horde.Server/Defaults/default-metrics.telemetry.json
+++ b/Engine/Source/Programs/Horde/Horde.Server/Defaults/default-metrics.telemetry.json
@@ -562,6 +562,38 @@
 			"function": "Average",
 			"interval": "00:30:00"
 		},
+		{
+			"id": "ubt-uba-remote-usage-all",
+			"filter": "$[?(@.Payload.EventName == 'Core.UBT.Executor' && search(@.Payload.Executor, \"Unreal Build Accelerator\") )]",
+			"groupBy": "$.Payload.Build_BranchName",
+			"property": "$.Payload.RemoteUsage",
+			"function": "Average",
+			"interval": "00:30:00"
+		},
+		{
+			"id": "ubt-uba-remote-actions-all",
+			"filter": "$[?(@.Payload.EventName == 'Core.UBT.Executor' && search(@.Payload.Executor, \"Unreal Build Accelerator\") )]",
+			"groupBy": "$.Payload.Build_BranchName",
+			"property": "$.Payload.RemoteActions",
+			"function": "Sum",
+			"interval": "00:30:00"
+		},
+		{
+			"id": "ubt-uba-total-actions-all",
+			"filter": "$[?(@.Payload.EventName == 'Core.UBT.Executor' && search(@.Payload.Executor, \"Unreal Build Accelerator\") )]",
+			"groupBy": "$.Payload.Build_BranchName",
+			"property": "$.Payload.TotalActions",
+			"function": "Sum",
+			"interval": "00:30:00"
+		},
+		{
+			"id": "ubt-uba-remote-actions-all",
+			"filter": "$[?(@.Payload.EventName == 'Core.UBT.Executor' && search(@.Payload.Executor, \"Unreal Build Accelerator\") )]",
+			"groupBy": "$.Payload.Build_BranchName",
+			"property": "$.Payload.RemoteActions",
+			"function": "Sum",
+			"interval": "00:30:00"
+		},
 		{
 			"id": "ubt-completed-duration",
 			"filter": "$[?(@.Payload.EventName == 'Core.UBT.Completed' )]",
