Fix for telemetry config in Horde Analytics

diff --git a/Engine/Source/Programs/Horde/Horde.Server/Defaults/default-metrics.telemetry.json b/Engine/Source/Programs/Horde/Horde.Server/Defaults/default-metrics.telemetry.json
--- a/Engine/Source/Programs/Horde/Horde.Server/Defaults/default-metrics.telemetry.json
+++ b/Engine/Source/Programs/Horde/Horde.Server/Defaults/default-metrics.telemetry.json
@@ -585,15 +585,7 @@
 			"property": "$.Payload.TotalActions",
 			"function": "Sum",
 			"interval": "00:30:00"
-		},
-		{
-			"id": "ubt-uba-remote-actions-all",
-			"filter": "$[?(@.Payload.EventName == 'Core.UBT.Executor' && search(@.Payload.Executor, \"Unreal Build Accelerator\") )]",
-			"groupBy": "$.Payload.Build_BranchName",
-			"property": "$.Payload.RemoteActions",
-			"function": "Sum",
-			"interval": "00:30:00"
-		},
+		}
 		{
 			"id": "ubt-completed-duration",
 			"filter": "$[?(@.Payload.EventName == 'Core.UBT.Completed' )]",
