Added UBA Failure Rate All to Horde Analytics

diff --git a/Engine/Source/Programs/Horde/Horde.Server/Defaults/default-analytics.dashboard.json b/Engine/Source/Programs/Horde/Horde.Server/Defaults/default-analytics.dashboard.json
--- a/Engine/Source/Programs/Horde/Horde.Server/Defaults/default-analytics.dashboard.json
+++ b/Engine/Source/Programs/Horde/Horde.Server/Defaults/default-analytics.dashboard.json
@@ -601,6 +601,15 @@
 								}
 							]
 						},
+						{
+							"name": "Failure Rate All (Average)",
+							"display": "Ratio",
+							"metrics": [
+								{
+									"id": "ubt-uba-failure-rate-all"
+								}
+							]
+						},
 						{
 							"name": "Compile WallTime (Average)",
 							"metrics": [
diff --git a/Engine/Source/Programs/Horde/Horde.Server/Defaults/default-metrics.telemetry.json b/Engine/Source/Programs/Horde/Horde.Server/Defaults/default-metrics.telemetry.json
--- a/Engine/Source/Programs/Horde/Horde.Server/Defaults/default-metrics.telemetry.json
+++ b/Engine/Source/Programs/Horde/Horde.Server/Defaults/default-metrics.telemetry.json
@@ -586,6 +586,14 @@
 			"function": "Sum",
 			"interval": "00:30:00"
 		},
+		{
+			"id": "ubt-uba-failure-rate-all",
+			"filter": "$[?(@.Payload.EventName == 'Core.UBT.Executor' && search(@.Payload.Executor, \"Unreal Build Accelerator\") )]",
+			"groupBy": "$.Payload.Build_BranchName",
+			"property": "$.Payload.FailureRate",
+			"function": "Average",
+			"interval": "00:30:00"
+		},
 		{
 			"id": "ubt-completed-duration",
 			"filter": "$[?(@.Payload.EventName == 'Core.UBT.Completed' )]",
