Removed Remote Alias from Distribution Charts

diff --git a/Engine/Source/Programs/Horde/Horde.Server/Defaults/default-analytics.dashboard.json b/Engine/Source/Programs/Horde/Horde.Server/Defaults/default-analytics.dashboard.json
--- a/Engine/Source/Programs/Horde/Horde.Server/Defaults/default-analytics.dashboard.json
+++ b/Engine/Source/Programs/Horde/Horde.Server/Defaults/default-analytics.dashboard.json
@@ -582,8 +582,7 @@
 							"display": "Ratio",
 							"metrics": [
 								{
-									"id": "ubt-uba-remote-usage-all",
-									"alias": "Remote"
+									"id": "ubt-uba-remote-usage-all"
 								}
 							]
 						},
@@ -623,8 +622,7 @@
 							"display": "Ratio",
 							"metrics": [
 								{
-									"id": "ubt-uba-remote-usage",
-									"alias": "Remote"
+									"id": "ubt-uba-remote-usage"
 								}
 							]
 						},
