UBA Dasboard Tweak

diff --git a/Engine/Source/Programs/Horde/Horde.Server/Defaults/default-analytics.dashboard.json b/Engine/Source/Programs/Horde/Horde.Server/Defaults/default-analytics.dashboard.json
--- a/Engine/Source/Programs/Horde/Horde.Server/Defaults/default-analytics.dashboard.json
+++ b/Engine/Source/Programs/Horde/Horde.Server/Defaults/default-analytics.dashboard.json
@@ -578,7 +578,7 @@
 					"name": "UBA",
 					"charts": [
 						{
-							"name": "Compile Distibution All (Average)",
+							"name": "Remote Distibution All (Average)",
 							"display": "Ratio",
 							"metrics": [
 								{
@@ -587,7 +587,7 @@
 							]
 						},
 						{
-							"name": "Compile Actions All (Sum)",
+							"name": "Actions All (Sum)",
 							"display": "Value",
 							"metrics": [
 								{
@@ -610,7 +610,7 @@
 							]
 						},
 						{
-							"name": "Compile WallTime (Average)",
+							"name": "WallTime (Average)",
 							"metrics": [
 								{
 									"id": "ubt-uba-duration"
@@ -618,7 +618,7 @@
 							]
 						},
 						{
-							"name": "Compile Distribution (Average)",
+							"name": "Remote Distribution (Average)",
 							"display": "Ratio",
 							"metrics": [
 								{
@@ -627,7 +627,7 @@
 							]
 						},
 						{
-							"name": "Compile Actions (Average)",
+							"name": "Actions (Average)",
 							"display": "Value",
 							"metrics": [
 								{
@@ -641,7 +641,7 @@
 							]
 						},
 						{
-							"name": "Compile Actions (Sum)",
+							"name": "Actions (Sum)",
 							"display": "Value",
 							"metrics": [
 								{
