Horde: Add docs for Docker Compose

#skipci

diff --git a/Engine/Source/Programs/Horde/Docs/Deployment/Server.md b/Engine/Source/Programs/Horde/Docs/Deployment/Server.md
--- a/Engine/Source/Programs/Horde/Docs/Deployment/Server.md
+++ b/Engine/Source/Programs/Horde/Docs/Deployment/Server.md
@@ -37,10 +37,27 @@ environment variable (see below).
 Running multiple Horde servers behind a load balancer does not require explicit configuration as long as each
 server points to the same MongoDB and Redis instance.
 
+Using Docker containers on Linux is [Epic's preferred way of running Horde](../Deployment.md).
+
+### Docker Compose (Linux)
+
+[Docker Compose](https://docs.docker.com/compose/) simplifies the setup of a Docker-based installation
+by providing a preconfigured group of Docker containers, which includes instances of MongoDB and Redis.
+
+Similar to the MSI installer, this method is suitable for testing Horde or deploying it in small-scale environments.
+To access the prebuilt images, refer to the Docker section above.
+The necessary Docker Compose configuration can be found within the `Engine/Source/Programs/Horde.Server/docker-compose.yml` file.
+
+From the same directory, start the containers with:
+```bash
+docker compose up
+````
+For additional guidance, see the comments within the YAML file.
+
 ### Homebrew (Mac)
 
 We don't provide any prebuilt binaries for running the server on Mac, though it's relatively straightforward to
-install all the prerequistes using [Homebrew](https://brew.sh/).
+install all the prerequisites using [Homebrew](https://brew.sh/).
 
 1. Install the .NET 8 SDK
 
