Horde: Add a global setting to enable or disable garbage collection.

diff --git a/Engine/Source/Programs/Horde/Horde.Server/Storage/StorageConfig.cs b/Engine/Source/Programs/Horde/Horde.Server/Storage/StorageConfig.cs
--- a/Engine/Source/Programs/Horde/Horde.Server/Storage/StorageConfig.cs
+++ b/Engine/Source/Programs/Horde/Horde.Server/Storage/StorageConfig.cs
@@ -50,6 +50,11 @@ namespace Horde.Server.Storage
 	/// </summary>
 	public class StorageConfig
 	{
+		/// <summary>
+		/// Whether to enable garbage collection
+		/// </summary>
+		public bool EnableGC { get; set; } = true;
+
 		/// <summary>
 		/// List of storage backends
 		/// </summary>
diff --git a/Engine/Source/Programs/Horde/Horde.Server/Storage/StorageService.cs b/Engine/Source/Programs/Horde/Horde.Server/Storage/StorageService.cs
--- a/Engine/Source/Programs/Horde/Horde.Server/Storage/StorageService.cs
+++ b/Engine/Source/Programs/Horde/Horde.Server/Storage/StorageService.cs
@@ -912,7 +912,12 @@ namespace Horde.Server.Storage
 				GlobalConfig globalConfig = _globalConfig.CurrentValue;
 
 				State state = CreateState(globalConfig);
+
 				StorageConfig storageConfig = state.Config.Storage;
+				if (!storageConfig.EnableGC)
+				{
+					break;
+				}
 
 				// Synchronize the list of configured namespaces with the GC state object
 				GcState gcState = await _gcState.GetAsync(cancellationToken);
@@ -986,7 +991,7 @@ namespace Horde.Server.Storage
 			double score = GetGcTimestamp(utcNow);
 
 			RedisSortedSetKey<RedisValue> checkSet = GetGcCheckSet(namespaceInfo.Id);
-			for (; ; )
+			while (_globalConfig.CurrentValue.Storage.EnableGC)
 			{
 				long length = await _redisService.GetDatabase().SortedSetLengthAsync(checkSet);
 				int batchSize = (int)Math.Min(length, 1024);
@@ -994,6 +999,8 @@ namespace Horde.Server.Storage
 
 				if (length == 0)
 				{
+					await _gcState.UpdateAsync(state => state.FindOrAddNamespace(namespaceInfo.Id).LastTime = utcNow, cancellationToken);
+					_logger.LogInformation("Finished garbage collection for namespace {NamespaceId} in {TimeSecs}s ({NumItems} removed)", namespaceInfo.Id, timer.Elapsed.TotalSeconds, numItemsRemoved);
 					break;
 				}
 
@@ -1026,9 +1033,6 @@ namespace Horde.Server.Storage
 					_ = _redisService.GetDatabase().SortedSetRemoveAsync(checkSet, value, CommandFlags.FireAndForget);
 				}
 			}
-
-			await _gcState.UpdateAsync(state => state.FindOrAddNamespace(namespaceInfo.Id).LastTime = utcNow, cancellationToken);
-			_logger.LogInformation("Finished garbage collection for namespace {NamespaceId} in {TimeSecs}s ({NumItems} removed)", namespaceInfo.Id, timer.Elapsed.TotalSeconds, numItemsRemoved);
 		}
 
 		static void SyncNamespaceList(GcState state, List<NamespaceConfig> namespaces)
