Horde: Update release notes.

diff --git a/Engine/Source/Programs/Horde/Docs/ReleaseNotes.md b/Engine/Source/Programs/Horde/Docs/ReleaseNotes.md
--- a/Engine/Source/Programs/Horde/Docs/ReleaseNotes.md
+++ b/Engine/Source/Programs/Horde/Docs/ReleaseNotes.md
@@ -1,7 +1,11 @@
 # Release Notes
 
-## 2024-03-26
+## 2024-03-28
 
+* Remove /userinfo call when authenticating via JWT. If an access token is passed, it's not guaranteed it has permission to access /userinfo from OIDC. ID tokens during normal web-based login does on the other hand. (32539012)
+* Fix parsing of true/false literals in condition expressions. (32518640)
+* Add a default pool to include interactive agents, and map it to the TestWin64 agent type. (32514202)
+* Add logging to trace when blobs are scanned for imports, to help debug some logs being expired while still referenced. (32506875)
 * Normalize keys for artifacts to allow case-insensitive searching. (32495123)
 * Support setting arbitrary artifact metadata. This cannot be queried, but will be returned in artifact responses. (32464180)
 * Run a background task to monitor for hangs when executing Perforce commands. (32461961)
