Horde: Add docs about agent GPU access

#skipci

diff --git a/Engine/Source/Programs/Horde/Docs/Deployment/Agent.md b/Engine/Source/Programs/Horde/Docs/Deployment/Agent.md
--- a/Engine/Source/Programs/Horde/Docs/Deployment/Agent.md
+++ b/Engine/Source/Programs/Horde/Docs/Deployment/Agent.md
@@ -90,7 +90,7 @@ command:
 
     dotnet HordeAgent.dll service install [-UserName=..] [-Password=..]
 
-Where -UserName and -Password specify credentials for the account to run the service under.
+Where `-UserName` and `-Password` specify credentials for the account to run the service under.
 
 The service may be uninstalled using the following command:
 
@@ -206,3 +206,14 @@ Setting the `%TEMP%` and `%TMP%` environment variables to this drive on Windows
 
 The agent can be configured to mount certain network shares at startup before taking on any work using the `Shares`
 property in the agent's [`appsettings.json`](AgentSettings.md) file.
+
+### GPU Access
+
+Certain jobs running in Horde, particularly tests, require access to the GPU. However, running the
+agent as a Windows service prevents accessing the GPU. To work around this issue, the agent must run
+as a normal Windows user in an interactive session. This can be achieved in different ways,
+such as running the agent from a PowerShell script on the desktop or using an external tool such as [NSSM](https://nssm.cc/).
+It's important to ensure the agent process is restarted in case of an unexpected exit.
+
+Since most agents won't require GPU access, it's recommended to create a separate pool
+specifically for interactive agents that need this support.
\ No newline at end of file
