Horde: Add split zip/ugs downloads
#rnx

diff --git a/Engine/Source/Programs/Horde/HordeDashboard/src/components/jobDetailsV2/JobDetailArtifacts.tsx b/Engine/Source/Programs/Horde/HordeDashboard/src/components/jobDetailsV2/JobDetailArtifacts.tsx
--- a/Engine/Source/Programs/Horde/HordeDashboard/src/components/jobDetailsV2/JobDetailArtifacts.tsx
+++ b/Engine/Source/Programs/Horde/HordeDashboard/src/components/jobDetailsV2/JobDetailArtifacts.tsx
@@ -1,15 +1,14 @@
 // Copyright Epic Games, Inc. All Rights Reserved.
 
-import { DetailsList, DetailsListLayoutMode, FontIcon, IColumn, SelectionMode, Stack, Text, mergeStyleSets } from '@fluentui/react';
+import { DetailsList, DetailsListLayoutMode, DirectionalHint, FontIcon, IColumn, IContextualMenuProps, PrimaryButton, SelectionMode, Stack, Text } from '@fluentui/react';
 import { observer } from 'mobx-react-lite';
-import React, { useEffect, useState } from 'react';
+import React, { useEffect } from 'react';
 import { GetJobArtifactResponse, JobStepState } from '../../backend/Api';
 import dashboard, { StatusColor } from '../../backend/Dashboard';
 import { ISideRailLink } from '../../base/components/SideRail';
 import { getStepETA, getStepFinishTime } from '../../base/utilities/timeUtils';
 import { getHordeStyling } from '../../styles/Styles';
 import { getStepStatusColor } from '../../styles/colors';
-import { JobArtifactsModal } from '../artifacts/ArtifactsModal';
 import { JobDataView, JobDetailsV2 } from "./JobDetailsViewCommon";
 
 const sideRail: ISideRailLink = { text: "Artifacts", url: "rail_artifacts" };
@@ -26,7 +25,7 @@ class JobArtifactsDataView extends JobDataView {
       super.clear();
    }
 
-   set() { 
+   set() {
 
       if (!this.initial) {
          return;
@@ -52,12 +51,12 @@ class JobArtifactsDataView extends JobDataView {
    detailsUpdated() {
 
       const hasArtifacts = !!this.details?.jobData?.artifacts?.length;
-      if (this.hasArtifacts !== hasArtifacts) {         
+      if (this.hasArtifacts !== hasArtifacts) {
          this.hasArtifacts = hasArtifacts;
          this.initialize(hasArtifacts ? [sideRail] : undefined);
          this.details?.setRootUpdated();
       }
-         
+
    }
 
    initial = true;
@@ -69,37 +68,8 @@ class JobArtifactsDataView extends JobDataView {
 
 JobDetailsV2.registerDataView("JobArtifactsDataView", (details: JobDetailsV2) => new JobArtifactsDataView(details));
 
-let _styles: any;
-
-const getStyles = () => {
-
-   const styles = _styles ?? mergeStyleSets({
-      list: {
-         selectors: {
-            'a': {
-               height: "unset !important",
-            },
-            ".ms-DetailsRow #artifactview": {
-               opacity: 0
-            },
-            ".ms-DetailsRow:hover #artifactview": {
-               opacity: 1
-            },
-         }
-      }
-   });
-
-   _styles = styles;
-
-   return styles;
-
-}
-
-
 export const JobArtifactsPanel: React.FC<{ jobDetails: JobDetailsV2 }> = observer(({ jobDetails }) => {
 
-   const [selected, setSelected] = useState<GetJobArtifactResponse | undefined>(undefined);   
-   
    const artifactView = jobDetails.getDataView<JobArtifactsDataView>("JobArtifactsDataView");
 
    jobDetails.subscribe();
@@ -124,13 +94,12 @@ export const JobArtifactsPanel: React.FC<{ jobDetails: JobDetailsV2 }> = observe
       return null;
    }
 
-   const styles = getStyles();
    const { hordeClasses, modeColors } = getHordeStyling();
 
    const columns: IColumn[] = [
       { key: 'column_desc', name: 'Description', minWidth: 580, isResizable: false, isMultiline: true },
       { key: 'column_time', name: 'Time', minWidth: 120, isResizable: false, isMultiline: true },
-      { key: 'column_cloud', name: 'Cloud', minWidth: 64, isResizable: false, isMultiline: true },
+      { key: 'column_download', name: 'Download', minWidth: 148, isResizable: false, isMultiline: true },
    ];
 
    let artifacts = [...jobData.artifacts!];
@@ -151,13 +120,31 @@ export const JobArtifactsPanel: React.FC<{ jobDetails: JobDetailsV2 }> = observe
       let step = jobDetails.stepById(item.stepId);
 
       let stepFinished = !!step?.finishTime;
-      const cursor = stepFinished ? "pointer" : undefined;
 
-      if (column.key === 'column_cloud') {
-         return <Stack horizontalAlign="end" verticalAlign="center" verticalFill={true} style={{ cursor: cursor, paddingRight: 8 }} onClick={() => { if (stepFinished) setSelected(item) }}>
-            <FontIcon id="artifactview" style={{ fontSize: "14px", color: stepFinished ? "#106EBE" : "#106EBE77" }} iconName="CloudDownload" />            
+      const downloadProps: IContextualMenuProps = {
+         items: [
+            {
+               key: 'download_ugs',
+               text: 'Download with UGS',
+               onClick: () => {
+                  window.location.assign(`/api/v2/artifacts/${item.id}/download?format=ugs`);
+               }
+            }
+         ],
+         directionalHint: DirectionalHint.bottomLeftEdge
+      };
+
+      if (column.key === 'column_download') {
+         return <Stack id="downloadview" horizontalAlign="end" verticalAlign="center" verticalFill={true} style={{ paddingRight: 8 }}>
+            <PrimaryButton split text="Download" menuProps={downloadProps}
+               disabled={!stepFinished}
+               style={{ fontFamily: "Horde Open Sans SemiBold" }}
+               onClick={() => {
+                  window.location.assign(`/api/v2/artifacts/${item.id}/download?format=zip`);
+               }}
+            />
          </Stack>
-      }      
+      }
 
       let eta = {
          display: "",
@@ -181,7 +168,7 @@ export const JobArtifactsPanel: React.FC<{ jobDetails: JobDetailsV2 }> = observe
             eta.server = finished.server;
          }
 
-         return <Stack horizontalAlign={"end"} verticalAlign="center" verticalFill={true} style={{ cursor: cursor }} onClick={() => { if (stepFinished) setSelected(item) }}>
+         return <Stack horizontalAlign={"end"} verticalAlign="center" verticalFill={true}>
             <Stack horizontal tokens={{ childrenGap: 2 }}>
                {!!eta.display && !step.finishTime && <Text style={{ fontSize: "11px", paddingTop: 2 }}>~</Text>}
                <Text style={{ fontSize: "13px" }}>
@@ -199,9 +186,9 @@ export const JobArtifactsPanel: React.FC<{ jobDetails: JobDetailsV2 }> = observe
             const colors = dashboard.getStatusColors();
             color = colors.get(StatusColor.Running)!;
          }
-         
-         return <Stack horizontal verticalAlign="center" verticalFill={true} style={{cursor: cursor}} onClick={() => { if (stepFinished) setSelected(item) }}>
-            {!!step && <FontIcon iconName="Square" style={{color: color, paddingTop: 2, fontSize: 13, paddingRight: 8}} />}
+
+         return <Stack horizontal verticalAlign="center" verticalFill={true} >
+            {!!step && <FontIcon iconName="Square" style={{ color: color, paddingTop: 2, fontSize: 13, paddingRight: 8 }} />}
             <Text style={{ color: modeColors.text, fontFamily: "Horde Open Sans SemiBold" }}>{item.description ?? item.name}</Text>
          </Stack>
       }
@@ -211,7 +198,6 @@ export const JobArtifactsPanel: React.FC<{ jobDetails: JobDetailsV2 }> = observe
 
 
    return (<Stack id={sideRail.url} styles={{ root: { paddingTop: 18, paddingRight: 12 } }}>
-      {!!selected && <JobArtifactsModal jobId={jobData.id} stepId={selected.stepId} contextType={selected.type} onClose={() => setSelected(undefined)} />}
       <Stack className={hordeClasses.raised} >
          <Stack tokens={{ childrenGap: 12 }} grow>
             <Stack horizontal>
@@ -224,7 +210,6 @@ export const JobArtifactsPanel: React.FC<{ jobDetails: JobDetailsV2 }> = observe
                   <DetailsList
                      styles={{ root: { overflowX: "hidden" } }}
                      isHeaderVisible={false}
-                     className={styles.list}
                      items={artifacts}
                      columns={columns}
                      selectionMode={SelectionMode.none}
diff --git a/Engine/Source/Programs/Horde/HordeDashboard/src/styles/lightTheme.ts b/Engine/Source/Programs/Horde/HordeDashboard/src/styles/lightTheme.ts
--- a/Engine/Source/Programs/Horde/HordeDashboard/src/styles/lightTheme.ts
+++ b/Engine/Source/Programs/Horde/HordeDashboard/src/styles/lightTheme.ts
@@ -41,6 +41,17 @@ export const lightTheme = createTheme({
                 }
             }
         },
+        "PrimaryButton": {
+            styles: {
+                splitButtonDivider: {
+                    backgroundColor: bodyText
+                },
+
+                splitButtonMenuButton: {
+                    backgroundColor: "#0078D4"
+                }
+            }
+        },
         "Stack": {
             styles: {
                 root: {
