Horde: Fix dashboard build
#rnx

diff --git a/Engine/Source/Programs/Horde/HordeDashboard/src/backend/Api.ts b/Engine/Source/Programs/Horde/HordeDashboard/src/backend/Api.ts
--- a/Engine/Source/Programs/Horde/HordeDashboard/src/backend/Api.ts
+++ b/Engine/Source/Programs/Horde/HordeDashboard/src/backend/Api.ts
@@ -1201,7 +1201,7 @@ export type GetReportResponse = {
 
 export type GetJobArtifactResponse = {
 	/// Identifier for this artifact
-	artifactId: string;
+	id: string;
 	/// Name of the artifact
 	name: string;
 	/// Artifact type
