Horde: Fix incorrect description for agent executor setting

#skipci

diff --git a/Engine/Source/Programs/Horde/Horde.Agent/AgentSettings.cs b/Engine/Source/Programs/Horde/Horde.Agent/AgentSettings.cs
--- a/Engine/Source/Programs/Horde/Horde.Agent/AgentSettings.cs
+++ b/Engine/Source/Programs/Horde/Horde.Agent/AgentSettings.cs
@@ -201,7 +201,7 @@ namespace Horde.Agent
 		public bool Ephemeral { get; set; } = false;
 
 		/// <summary>
-		/// The executor to use for jobs. Defaults to the Perforce executor.
+		/// The executor to use for jobs
 		/// </summary>
 		public string Executor { get; set; } = Execution.WorkspaceExecutor.Name;
 
