Horde: Comment out the placeholder ticket value in the default config file.

diff --git a/Engine/Source/Programs/Horde/Horde.Server/Defaults/globals.json b/Engine/Source/Programs/Horde/Horde.Server/Defaults/globals.json
--- a/Engine/Source/Programs/Horde/Horde.Server/Defaults/globals.json
+++ b/Engine/Source/Programs/Horde/Horde.Server/Defaults/globals.json
@@ -50,9 +50,9 @@
 //				{
 //					"userName": "test.user",
 //
-//					// Specify either the password or ticket property below (not both)
+//					// Specify either the password or ticket property below (but not both)
 //					"password": "password123",
-//					"ticket": "ticket if not using password"
+//					// "ticket": "ticket if not using password"
 //				}
 //			]
 //		}
