Horde: Change registration => enrollment
#rnx

diff --git a/Engine/Source/Programs/Horde/HordeDashboard/src/backend/Backend.ts b/Engine/Source/Programs/Horde/HordeDashboard/src/backend/Backend.ts
--- a/Engine/Source/Programs/Horde/HordeDashboard/src/backend/Backend.ts
+++ b/Engine/Source/Programs/Horde/HordeDashboard/src/backend/Backend.ts
@@ -129,7 +129,7 @@ export class Backend {
 
     getAgentRegistrationRequests(): Promise<GetPendingAgentsResponse> {
         return new Promise<GetPendingAgentsResponse>((resolve, reject) => {
-            this.backend.get(`/api/v1/registration`).then((response) => {
+            this.backend.get(`/api/v1/enrollment`).then((response) => {
                 const agent = response.data as GetPendingAgentsResponse;
                 resolve(agent);
             }).catch(reason => { reject(reason); });
@@ -139,7 +139,7 @@ export class Backend {
     // create a new account
     registerAgents(request: ApproveAgentsRequest): Promise<void> {
         return new Promise<void>((resolve, reject) => {
-            this.backend.post(`/api/v1/registration`, request).then(() => {
+            this.backend.post(`/api/v1/enrollment`, request).then(() => {
                 resolve();
             }).catch(reason => {
                 reject(reason);
