Horde: Doc updates.

diff --git a/Engine/Source/Programs/Horde/Docs/Deployment/Integrations/Slack.md b/Engine/Source/Programs/Horde/Docs/Deployment/Integrations/Slack.md
--- a/Engine/Source/Programs/Horde/Docs/Deployment/Integrations/Slack.md
+++ b/Engine/Source/Programs/Horde/Docs/Deployment/Integrations/Slack.md
@@ -64,5 +64,5 @@ fully:
 ### User mapping
 
 Horde users are mapped to Slack users by correlating the email address in the user's
-[OIDC profile](../Server.md#authorization) with their Slack user profile. Horde will use avatars configured through
+[OIDC profile](../Server.md#authentication) with their Slack user profile. Horde will use avatars configured through
 Slack in the dashboard for any successfully mapped email address.
diff --git a/Engine/Source/Programs/Horde/Docs/Deployment/Server.md b/Engine/Source/Programs/Horde/Docs/Deployment/Server.md
--- a/Engine/Source/Programs/Horde/Docs/Deployment/Server.md
+++ b/Engine/Source/Programs/Horde/Docs/Deployment/Server.md
@@ -222,13 +222,13 @@ telemetry capture are [listed here](ServerSettings.md#opentelemetrysettings).
 In order to separate lighter request traffic from heavier background operations, the Horde server can be configured to
 run in different RunModes. You can configure these via the [RunMode](ServerSettings.md) setting.
 
-### Authorization
+### Authentication
 
 Horde supports [OpenID Connect (OIDC)](https://openid.net/developers/how-connect-works/) for authentication using
 an external identity provider. OIDC is a widely used auth standard, and Okta, AWS, Azure, Google, Facebook, and
 many others implement identity providers compatible with it.
 
-The [Getting Started > Authorization](../Tutorials/Authorization.md) page explains how to configure Horde's internal
+The [Getting Started > Authentication](../Tutorials/Authentication.md) page explains how to configure Horde's internal
 account system and OIDC provider.
 
 The following settings in [Server.json](ServerSettings.md) are required to configure an external OIDC provider:
diff --git a/Engine/Source/Programs/Horde/Docs/Home.md b/Engine/Source/Programs/Horde/Docs/Home.md
--- a/Engine/Source/Programs/Horde/Docs/Home.md
+++ b/Engine/Source/Programs/Horde/Docs/Home.md
@@ -2,22 +2,14 @@
 
 ## Getting Started
 
-* **[Installing Horde](Tutorials/InstallHorde.md)**
-* **[Setting up remote C++ compilation with Unreal Build Accelerator](Tutorials/RemoteCompilation.md)**
-* **[Setting up build automation](Tutorials/BuildAutomation.md)**
-* **[Setting up test automation](Tutorials/TestAutomation.md)**
-* **[Setting up analytics for your team](Tutorials/Analytics.md)**
-* **[Setting up mobile and console devices](Tutorials/DeviceManager.md)**
-* **[Setting up UnrealGameSync and distributing Unreal Editor to content creators](Tutorials/UnrealGameSync.md)**
-* **[Setting up authorization](Tutorials/Authorization.md)**
-
-## Status
-
-Horde is under heavy development, and large parts of it are still in flux. While we use aspects of it
-(particularly the CI system) heavily at Epic, we consider it experimental for Unreal Engine licensees
-and offer limited support for it.
-
-See also: [Feature Status](Features.md)
+* **[Install the Horde agent](Tutorials/InstallHorde.md)**
+* **[Enable remote C++ compilation with Unreal Build Accelerator](Tutorials/RemoteCompilation.md)**
+* **[Set up build automation](Tutorials/BuildAutomation.md)**
+* **[Enable test automation with Gauntlet](Tutorials/TestAutomation.md)**
+* **[Get telemetry and analytics for your team](Tutorials/Analytics.md)**
+* **[Work with mobile and console devices](Tutorials/DeviceManager.md)**
+* **[Install UnrealGameSync and distribute Unreal Editor to your team](Tutorials/UnrealGameSync.md)**
+* **[Enable authentication](Tutorials/Authentication.md)**
 
 ## Reference
 
@@ -33,6 +25,14 @@ Horde's reference documentation is divided into sections by target audience:
   * Describes how to build and modify Horde, and its architecture.
   * **Target audience:** Developers wishing to extend Horde.
 
+## Status
+
+Horde is under heavy development, and large parts of it are still in flux. While we use aspects of it
+(particularly the CI system) heavily at Epic, we consider it experimental for Unreal Engine licensees
+and offer limited support for it.
+
+See also: [Feature Status](Features.md)
+
 ## Further Reading
 
 * [Glossary](Glossary.md)
diff --git a/Engine/Source/Programs/Horde/Docs/Landing.md b/Engine/Source/Programs/Horde/Docs/Landing.md
--- a/Engine/Source/Programs/Horde/Docs/Landing.md
+++ b/Engine/Source/Programs/Horde/Docs/Landing.md
@@ -4,14 +4,14 @@
 
 **Horde** has been installed successfully. To continue setting up your deployment, check out the following guides:
 
-* **[Installing the Horde Agent software](Tutorials/InstallHorde.md#horde-agents)**
-* **[Setting up remote C++ compilation with Unreal Build Accelerator](Tutorials/RemoteCompilation.md)**
-* **[Setting up build automation](Tutorials/BuildAutomation.md)**
-* **[Setting up test automation](Tutorials/TestAutomation.md)**
-* **[Setting up analytics for your team](Tutorials/Analytics.md)**
-* **[Setting up mobile and console devices](Tutorials/DeviceManager.md)**
-* **[Setting up UnrealGameSync and distributing Unreal Editor to content creators](Tutorials/UnrealGameSync.md)**
-* **[Setting up authorization](Tutorials/Authorization.md)**
+* **[Install the Horde agent](Tutorials/InstallHorde.md)**
+* **[Enable remote C++ compilation with Unreal Build Accelerator](Tutorials/RemoteCompilation.md)**
+* **[Set up build automation](Tutorials/BuildAutomation.md)**
+* **[Enable test automation with Gauntlet](Tutorials/TestAutomation.md)**
+* **[Get telemetry and analytics for your team](Tutorials/Analytics.md)**
+* **[Work with mobile and console devices](Tutorials/DeviceManager.md)**
+* **[Install UnrealGameSync and distribute Unreal Editor to your team](Tutorials/UnrealGameSync.md)**
+* **[Enable authentication](Tutorials/Authentication.md)**
 
 Alternatively, check out the full documentation through the `Help` menu at the top right of the page.
 
diff --git a/Engine/Source/Programs/Horde/Docs/Tutorials/Authentication.md b/Engine/Source/Programs/Horde/Docs/Tutorials/Authentication.md
new file mode 100644
--- /dev/null
+++ b/Engine/Source/Programs/Horde/Docs/Tutorials/Authentication.md
@@ -0,0 +1,60 @@
+[Horde](../../README.md) > Getting Started: Authentication
+
+# Getting Started: Authentication
+
+## Introduction
+
+Horde ships with auth disabled by default to make it easier to demonstrate and experiment with. Most production
+deployments will probably want users to log in, and restrict the actions they can perform based on their role.
+
+To do this, Horde supports **[OAuth2](https://oauth.net/2/)** and **[OIDC](https://openid.net/developers/how-connect-works/)**,
+which is supported by most third party identity providers - including Okta, AWS, Azure, and Google.
+Configuring an external identity provider is out of scope for this documentation, though the relevant configuration
+points are touched on in the [Deployment > Server](../Deployment/Server.md#authentication) page.
+
+If you don't have an existing OIDC-compatible identity provider, Horde includes it's own - which this guide covers.
+
+## Prerequisites
+
+* Horde Server installation (see [Getting Started: Install Horde](InstallHorde.md)).
+* A valid certificate, and [HTTPS support](../Deployment/Server.md#https) enabled on your server.
+
+## Steps
+
+1. In your [server.json](../Config/Orientation.md) file set the `AuthMode` property to `Horde`, and restart the server.
+2. The first time you launch the server, you'll be prompted to enter an administrator password.
+3. After logging in, there will be an `Accounts` menu item in the `Server` menu. From here, you can manage the users
+allowed to log in to the server, and the [**claims**](../Glossary.md#authorization) that they have. Horde's account
+system uses the `http://epicgames.com/ue/horde/group` claim for groups that a user belongs to, and the dashboard will
+suggest and autocomplete any groups found in the deployment's configuration files.
+
+There are two standard groups defined in the server's `default.globals.json` file, which is included from the standard
+`globals.json` file by default: `View` and `Run`.
+
+   ```json
+    "acl": {
+        "entries": [
+            {
+                "claim": {
+                    "type": "http://epicgames.com/ue/horde/group", 
+                    "value": "View"
+                },
+                "profiles": [
+                    "default-read"
+                ]
+            },
+            {
+                "claim": {
+                    "type": "http://epicgames.com/ue/horde/group", 
+                    "value": "Run"
+                },
+                "profiles": [
+                    "default-run"
+                ]
+            }
+        ]
+    }
+   ```
+
+The `default-read` and `default-run` profiles are defined in code (`AclConfig.cs`). You can define your own profiles within
+the `profiles` element of each [AclConfig](../Config/Schema/Globals.md#aclconfig) object.
diff --git a/Engine/Source/Programs/Horde/Docs/Tutorials/Authorization.md b/Engine/Source/Programs/Horde/Docs/Tutorials/Authorization.md
deleted file mode 100644
diff --git a/Engine/Source/Programs/Horde/README.md b/Engine/Source/Programs/Horde/README.md
--- a/Engine/Source/Programs/Horde/README.md
+++ b/Engine/Source/Programs/Horde/README.md
@@ -33,14 +33,14 @@ See also: [Feature Status](Docs/Features.md)
 
 ## Getting Started
 
-* **[Installing Horde](Docs/Tutorials/InstallHorde.md)**
-* **[Setting up remote C++ compilation with Unreal Build Accelerator](Docs/Tutorials/RemoteCompilation.md)**
-* **[Setting up build automation](Docs/Tutorials/BuildAutomation.md)**
-* **[Setting up test automation](Docs/Tutorials/TestAutomation.md)**
-* **[Setting up studio analytics](Docs/Tutorials/Analytics.md)**
-* **[Setting up mobile and console devices](Docs/Tutorials/DeviceManager.md)**
-* **[Setting up UnrealGameSync and distributing Unreal Editor to content creators](Docs/Tutorials/UnrealGameSync.md)**
-* **[Setting up authorization](Docs/Tutorials/Authorization.md)**
+* **[Install the Horde agent](Docs/Tutorials/InstallHorde.md)**
+* **[Enable remote C++ compilation with Unreal Build Accelerator](Docs/Tutorials/RemoteCompilation.md)**
+* **[Set up build automation](Docs/Tutorials/BuildAutomation.md)**
+* **[Enable test automation with Gauntlet](Docs/Tutorials/TestAutomation.md)**
+* **[Get telemetry and analytics for your team](Docs/Tutorials/Analytics.md)**
+* **[Work with mobile and console devices](Docs/Tutorials/DeviceManager.md)**
+* **[Install UnrealGameSync and distribute Unreal Editor to your team](Docs/Tutorials/UnrealGameSync.md)**
+* **[Enable authentication](Docs/Tutorials/Authentication.md)**
 
 ## Reference
 
