Horde: Updating schema documentation

#rnx

diff --git a/Engine/Source/Programs/Horde/Docs/Config/Schema/Globals.md b/Engine/Source/Programs/Horde/Docs/Config/Schema/Globals.md
--- a/Engine/Source/Programs/Horde/Docs/Config/Schema/Globals.md
+++ b/Engine/Source/Programs/Horde/Docs/Config/Schema/Globals.md
@@ -398,6 +398,7 @@ Configuration for storage
 
 Name | Description
 ---- | -----------
+`enableGC` | `boolean`<br>Whether to enable garbage collection
 `backends` | [`BackendConfig`](#backendconfig)`[]`<br>List of storage backends
 `namespaces` | [`NamespaceConfig`](#namespaceconfig)`[]`<br>List of namespaces for storage
 
diff --git a/Engine/Source/Programs/Horde/Docs/Deployment/AgentSettings.md b/Engine/Source/Programs/Horde/Docs/Deployment/AgentSettings.md
--- a/Engine/Source/Programs/Horde/Docs/Deployment/AgentSettings.md
+++ b/Engine/Source/Programs/Horde/Docs/Deployment/AgentSettings.md
@@ -11,7 +11,7 @@ Name | Description
 `name` | `string`<br>Name of agent to report as when connecting to server. By default, the computer's hostname will be used.
 `installed` | `boolean`<br>Whether the server is running in 'installed' mode. In this mode, on Windows, the default data directory will use the common application data folder (C:\ProgramData\Epic\Horde), and configuration data will be read from here and the registry. This setting is overridden to false for local builds from appsettings.Local.json.
 `ephemeral` | `boolean`<br>Whether agent should register as being ephemeral. Doing so will not persist any long-lived data on the server and once disconnected it's assumed to have been deleted permanently. Ideal for short-lived agents, such as spot instances on AWS EC2.
-`executor` | `string`<br>The executor to use for jobs. Defaults to the Perforce executor.
+`executor` | `string`<br>The executor to use for jobs
 `localExecutor` | [`LocalExecutorSettings`](#localexecutorsettings)<br>Settings for the local executor
 `perforceExecutor` | [`PerforceExecutorSettings`](#perforceexecutorsettings)<br>Settings for the perforce executor
 `workingDir` | [`DirectoryReference`](#directoryreference)<br>Working directory
