Horde: Reorder agent view filters and right align search box
#rnx

diff --git a/Engine/Source/Programs/Horde/HordeDashboard/src/components/AgentView.tsx b/Engine/Source/Programs/Horde/HordeDashboard/src/components/AgentView.tsx
--- a/Engine/Source/Programs/Horde/HordeDashboard/src/components/AgentView.tsx
+++ b/Engine/Source/Programs/Horde/HordeDashboard/src/components/AgentView.tsx
@@ -481,7 +481,7 @@ class LocalState {
             isSorted: colState.key === "pools" ? undefined : colState.isSorted,
             isSortedDescending: colState.key === "pools" ? undefined : colState.isSortedDescending,
             onColumnClick: this._onColumnClick.bind(this),
-            onRender: colState.key === "editAgent" ? () => {               
+            onRender: colState.key === "editAgent" ? () => {
                return <Stack style={{ cursor: "pointer" }} verticalFill verticalAlign='center' horizontalAlign='start' onClick={(ev) => {
                   localState.setRightClickDiv(ev?.clientX, ev?.clientY);
                   localState.setAgentContextMenuOpen(true);
@@ -1412,7 +1412,7 @@ export const AgentMenuBar: React.FC<{ agentView?: boolean }> = observer(({ agent
    if (localState.agentsSelectedCount !== 0) {
       selectedButton = <PrimaryButton
          disabled={localState.agentsSelectedCount === 0 ? true : false}
-         styles={{ root: { marginRight: 18, fontFamily: 'Horde Open Sans SemiBold !important' } }}
+         styles={{ root: { fontFamily: 'Horde Open Sans SemiBold !important' } }}
          menuProps={agentSelectedProps}
          onClick={() => { }}>
          {`${localState.agentsSelectedCount} Agent${localState.agentsSelectedCount > 1 ? "s" : ""} Selected`}
@@ -1428,65 +1428,60 @@ export const AgentMenuBar: React.FC<{ agentView?: boolean }> = observer(({ agent
    });
 
    return (
-      <Stack horizontal horizontalAlign="space-between" grow={!!agentView} style={{ paddingTop: "6px" }}>
-         <Stack.Item styles={{ root: { paddingLeft: '20px' } }}>
-            <Stack horizontal tokens={{ childrenGap: 12 }}>
-               <Stack>
-                  <SearchBox
-                     showIcon={true}
-                     disableAnimation={true}
-                     placeholder="Search Agents"
-                     value={localState.agentFilter}
-                     styles={{ root: { width: 240 } }}
-                     onChange={(event?: React.ChangeEvent<HTMLInputElement> | undefined, newValue?: string | undefined) => { localState.setAgentFilter(newValue ?? ""); }}
-                     onClear={() => { localState.setAgentFilter(""); }}
-                  />
-               </Stack>
-               <Stack>
-                  <Checkbox styles={{ root: { paddingTop: 6 } }} label={"Exact Match"} checked={localState.filterExactMatch} onChange={(ev, checked) => localState.setExactMatch(checked!)} />
-               </Stack>
-               <Stack style={{ paddingLeft: 18 }}>
-                  <Dropdown
-                     placeholder="Filter Status"
-                     style={{ width: 200 }}
-                     selectedKeys={Array.from(localState.agentStatusFilter).map(status => `dropdown_agent_status_${status}_key`)}
-                     multiSelect
-                     options={agentStatusItems}
-                     onChange={(event, option, index) => {
-
-                        if (option) {
-
-                           // new set instance to update observable action update
-                           const newFilter = new Set(localState.agentStatusFilter);
-
-                           if (option.selected) {
-                              newFilter.add((option as any).status);
-                           } else {
-                              newFilter.delete((option as any).status);
-                           }
+      <Stack horizontal verticalAlign='center' tokens={{ childrenGap: 24 }}>
+         {!!agentView && <Stack horizontal tokens={{ childrenGap: 12 }} verticalAlign='center'>
+            {!!dashboard.user?.dashboardFeatures?.showPoolEditor && <CommandButton
+               onClick={() => { editPoolsModalState.setOpen(); }}
+               iconProps={{ iconName: 'Edit' }}
+               styles={{ root: { fontFamily: 'Horde Open Sans SemiBold !important' } }}>
+               {`Pools`}
+            </CommandButton>}
+            {selectedButton}
+            <PoolEditorModal></PoolEditorModal>
+            <PoolSelectionModal></PoolSelectionModal>
+         </Stack>
+         }
+         <Stack>
+            <Dropdown
+               placeholder="Filter Status"
+               style={{ width: 200 }}
+               selectedKeys={Array.from(localState.agentStatusFilter).map(status => `dropdown_agent_status_${status}_key`)}
+               multiSelect
+               options={agentStatusItems}
+               onChange={(event, option, index) => {
+
+                  if (option) {
+
+                     // new set instance to update observable action update
+                     const newFilter = new Set(localState.agentStatusFilter);
+
+                     if (option.selected) {
+                        newFilter.add((option as any).status);
+                     } else {
+                        newFilter.delete((option as any).status);
+                     }
 
-                           localState.setAgentStaus(newFilter);
+                     localState.setAgentStaus(newFilter);
 
-                        }
-                     }}
-                  />
-               </Stack>
-            </Stack>
-         </Stack.Item>
-         {!!agentView && <Stack><Stack grow />
-            <Stack horizontal tokens={{ childrenGap: 12 }} grow>
-               {!!dashboard.user?.dashboardFeatures?.showPoolEditor && <CommandButton
-                  onClick={() => { editPoolsModalState.setOpen(); }}
-                  iconProps={{ iconName: 'Edit' }}
-                  styles={{ root: { marginRight: 12, bottom: 3, fontFamily: 'Horde Open Sans SemiBold !important' } }}>
-                  {`Pools`}
-               </CommandButton>}
-               {selectedButton}
-               <PoolEditorModal></PoolEditorModal>
-               <PoolSelectionModal></PoolSelectionModal>
-            </Stack>
+                  }
+               }}
+            />
          </Stack>
-         }
+         <Stack>
+            <Checkbox label={"Exact Match"} checked={localState.filterExactMatch} onChange={(ev, checked) => localState.setExactMatch(checked!)} />
+         </Stack>
+         <Stack>
+            <SearchBox
+               showIcon={true}
+               disableAnimation={true}
+               placeholder="Search Agents"
+               value={localState.agentFilter}
+               styles={{ root: { width: 240 } }}
+               onChange={(event?: React.ChangeEvent<HTMLInputElement> | undefined, newValue?: string | undefined) => { localState.setAgentFilter(newValue ?? ""); }}
+               onClear={() => { localState.setAgentFilter(""); }}
+            />
+         </Stack>
+
       </Stack>
    );
 });
