Horde: Add notes for setting up a self-signed cert for testing.

diff --git a/Engine/Source/Programs/Horde/Docs/Deployment/Server.md b/Engine/Source/Programs/Horde/Docs/Deployment/Server.md
--- a/Engine/Source/Programs/Horde/Docs/Deployment/Server.md
+++ b/Engine/Source/Programs/Horde/Docs/Deployment/Server.md
@@ -196,8 +196,45 @@ Settings for port usage are defined in [Server.json](ServerSettings.md):
 
 ### HTTPS
 
-To serve data over HTTPS, set the `HttpsPort` property in the [Server.json](ServerSettings.md) file. The server
-machine must have a correctly configured certificate to serve data correctly.
+To serve data over HTTPS, set the `HttpsPort` property in the [Server.json](ServerSettings.md) file.
+
+Configure the certificate for [Kestrel](https://learn.microsoft.com/en-us/aspnet/core/fundamentals/servers/kestrel?view=aspnetcore-8.0)
+(the NET Core web server) by setting the default certificate in the same file.
+
+Cross-platform:
+
+   ```json
+    "Kestrel": {
+        "Certificates": {
+            "Default": {
+                "Path": "C:\\cert\\test.pfx",
+                "Password": "my-password"
+            }
+        }
+    }
+   ```
+
+Windows (using the system certificate store):
+
+   ```json
+    "Kestrel": {
+        "Certificates": {
+            "Default":
+            {
+                "Subject": "my-domain.com",
+
+                // Use the 'Personal' certificate store on the local machine
+                "Store": "My",
+                "Location": "LocalMachine"
+            }
+        }
+    }
+   ```
+
+> **Note**: The `Kestrel` object must be added at the root scope of the file, not within the `Horde` object.
+
+Other ways to configure certificates for Kestrel are listed on
+[MSDN](https://learn.microsoft.com/en-us/aspnet/core/fundamentals/servers/kestrel/endpoints?view=aspnetcore-8.0#configure-https).
 
 Both HTTP/1.1 and HTTP/2.0 traffic can be served over the `HttpsPort`. Unencrypted traffic can be disabled by
 setting `HttpPort` and `Http2Port` to zero.
@@ -207,6 +244,8 @@ Horde's internal account system, for example), and it's important that these URL
 By default, this URL is derived from the server's reported DNS name, but this can be overridden through the
 `ServerUrl` property.
 
+To set up a self-signed certificate for testing see [Tutorials > Self Signed Certs](..\Tutorials\SelfSignedCerts.md).
+
 ### Monitoring
 
 Horde uses [Serilog](https://serilog.net/) for logging and is configured to generate plain text and JSON log files
diff --git a/Engine/Source/Programs/Horde/Docs/Internals/StructuredLogging.md b/Engine/Source/Programs/Horde/Docs/Internals/StructuredLogging.md
--- a/Engine/Source/Programs/Horde/Docs/Internals/StructuredLogging.md
+++ b/Engine/Source/Programs/Horde/Docs/Internals/StructuredLogging.md
@@ -87,10 +87,6 @@ over plain text output and construct structured log events from them. Some are i
 `Engine/Source/Programs/AutomationTool/AutomationUtils/Matchers/...`).
 These are used by the `LogEventParser` class in EpicGames.Core.
 
-* TODO: ILogCursor - window
-* TODO: ILogMatcher - MUST be stateless
-* TODO: Explain where to implement log parsers
-
 To implement a new matcher for plain-text log output, create a class that implements the `ILogEventMatcher` interface
 from `EpicGames.Core`, and ensure it's registered 
 
diff --git a/Engine/Source/Programs/Horde/Docs/Tutorials/SelfSignedCert.md b/Engine/Source/Programs/Horde/Docs/Tutorials/SelfSignedCert.md
new file mode 100644
--- /dev/null
+++ b/Engine/Source/Programs/Horde/Docs/Tutorials/SelfSignedCert.md
@@ -0,0 +1,70 @@
+[Horde](../../README.md) > Getting Started: Self-Signed Certs
+
+# Getting Started: Self-Signed Certs
+
+## Introduction
+
+When deploying Horde in production environments, it is recommended
+to use a verified signing certificate.
+
+For testing scenarios, if can be useful to install a self-signed
+certificate.
+
+> **Note:** Using self-signed certificates circumvents basic security
+measures. Do not use this technique in production environments.
+
+## Server
+
+1. From an administrator PowerShell prompt, add a certificate to the "Personal" store for the local machine by running:
+
+   ```pwsh
+   New-SelfSignedCertificate -CertStoreLocation 'Cert:\LocalMachine\My' -DnsName 'my-domain.com'
+   ```
+
+2. Open the Certificate Manager MMC snap-in by running `certmgr.msc` from the Windows 'Run' menu. You should see the
+certificate created above in the `Personal\Certificates` section.
+
+   Select the certificate and press Ctrl+C. Navigate to the
+   'Trusted Root Authorities\Certificates' section, and press Ctrl+V to create a copy.
+
+3. Open the [server.json](../Config/Orientation.md) file, and uncomment the `HttpsPort` line:
+
+    ```json
+     "HttpsPort": 13341,
+    ```
+
+   ...as well as the certificate section at the bottom of the file -
+   updating the subject name to the DNS name on the certificate created above.
+
+   ```json
+    "Kestrel":
+    {
+        "Certificates":
+        {
+            "Default":
+            {
+                "Subject": "my-domain.com",
+                "Store": "My",
+                "Location": "LocalMachine"
+            }
+        }
+    }
+    ```
+
+4. Restart the server. You should be able to connect over HTTPS from a browser on the same machine on port 13341.
+
+## Clients
+
+1. Browse to the server specified above on the HTTPS URL. On the warning
+   dialog about the server having an invalid certificate, choose to export
+   it to a file.
+
+   On **Google Chrome**, this can be accessed by clicking on the "Not Secure" button in the address bar, selecting
+   "Certificate is not valid", switching to the "Details" tab in the certificate browse, and chosing "Export". Select
+   'Base-64 Encoded ASCII' as the file type, and save the file.
+
+   The certificate may also be exported directly from the Certificate
+   Manager MMC snap-in.
+
+2. Locate the exported certificate file in Windows Explorer, right click on it, and choose "Install Certificate". When
+   prompted, choose to import the certificate in to the "Trusted Root Certificates" store.
