Horde: Only export nodes to Horde that are referenced by the initial job parameters, unless -AllowTargetChanges is specified. Prevents data being copied to temp storage that will never be used unless the target list changes.

diff --git a/Engine/Source/Programs/Horde/Horde.Agent/Execution/JobExecutor.cs b/Engine/Source/Programs/Horde/Horde.Agent/Execution/JobExecutor.cs
--- a/Engine/Source/Programs/Horde/Horde.Agent/Execution/JobExecutor.cs
+++ b/Engine/Source/Programs/Horde/Horde.Agent/Execution/JobExecutor.cs
@@ -225,6 +225,7 @@ namespace Horde.Agent.Execution
 		protected BeginBatchResponse Batch { get; }
 
 		protected List<string> _additionalArguments = new List<string>();
+		protected bool _allowTargetChanges;
 
 		protected bool _compileAutomationTool = true;
 
@@ -313,6 +314,10 @@ namespace Horde.Agent.Execution
 				{
 					_savePreprocessedScript = true;
 				}
+				else if (argument.Equals("-AllowTargetChanges", StringComparison.OrdinalIgnoreCase))
+				{
+					_allowTargetChanges = true;
+				}
 				else if (argument.StartsWith(TargetArgumentPrefix, StringComparison.OrdinalIgnoreCase))
 				{
 					_targets.Add(argument.Substring(TargetArgumentPrefix.Length));
@@ -485,6 +490,15 @@ namespace Horde.Agent.Execution
 			}
 			arguments.AppendArgument("-HordeExport=", definitionFile.FullName);
 			arguments.AppendArgument("-ListOnly");
+
+			if (!_allowTargetChanges)
+			{
+				foreach (string target in _targets)
+				{
+					arguments.AppendArgument(TargetArgumentPrefix + target);
+				}
+			}
+
 			//Arguments.AppendArgument("-TokenSignature=", JobId.ToString());
 			foreach (string additionalArgument in _additionalArguments)
 			{
