Horde: Resurrect  browsing from job artifacts
#rnx

diff --git a/Engine/Source/Programs/Horde/HordeDashboard/src/components/jobDetailsV2/JobDetailArtifacts.tsx b/Engine/Source/Programs/Horde/HordeDashboard/src/components/jobDetailsV2/JobDetailArtifacts.tsx
--- a/Engine/Source/Programs/Horde/HordeDashboard/src/components/jobDetailsV2/JobDetailArtifacts.tsx
+++ b/Engine/Source/Programs/Horde/HordeDashboard/src/components/jobDetailsV2/JobDetailArtifacts.tsx
@@ -2,7 +2,7 @@
 
 import { DetailsList, DetailsListLayoutMode, DirectionalHint, FontIcon, IColumn, IContextualMenuProps, PrimaryButton, SelectionMode, Stack, Text } from '@fluentui/react';
 import { observer } from 'mobx-react-lite';
-import React, { useEffect } from 'react';
+import React, { useEffect, useState } from 'react';
 import { GetJobArtifactResponse, JobStepState } from '../../backend/Api';
 import dashboard, { StatusColor } from '../../backend/Dashboard';
 import { ISideRailLink } from '../../base/components/SideRail';
@@ -10,6 +10,7 @@ import { getStepETA, getStepFinishTime } from '../../base/utilities/timeUtils';
 import { getHordeStyling } from '../../styles/Styles';
 import { getStepStatusColor } from '../../styles/colors';
 import { JobDataView, JobDetailsV2 } from "./JobDetailsViewCommon";
+import { JobArtifactsModal } from '../artifacts/ArtifactsModal';
 
 const sideRail: ISideRailLink = { text: "Artifacts", url: "rail_artifacts" };
 
@@ -70,6 +71,8 @@ JobDetailsV2.registerDataView("JobArtifactsDataView", (details: JobDetailsV2) =>
 
 export const JobArtifactsPanel: React.FC<{ jobDetails: JobDetailsV2 }> = observer(({ jobDetails }) => {
 
+   const [selected, setSelected] = useState<GetJobArtifactResponse | undefined>(undefined);
+
    const artifactView = jobDetails.getDataView<JobArtifactsDataView>("JobArtifactsDataView");
 
    jobDetails.subscribe();
@@ -99,7 +102,7 @@ export const JobArtifactsPanel: React.FC<{ jobDetails: JobDetailsV2 }> = observe
    const columns: IColumn[] = [
       { key: 'column_desc', name: 'Description', minWidth: 580, isResizable: false, isMultiline: true },
       { key: 'column_time', name: 'Time', minWidth: 120, isResizable: false, isMultiline: true },
-      { key: 'column_download', name: 'Download', minWidth: 148, isResizable: false, isMultiline: true },
+      { key: 'column_download', name: 'Download', minWidth: 300, isResizable: false, isMultiline: true },
    ];
 
    let artifacts = [...jobData.artifacts!];
@@ -135,14 +138,24 @@ export const JobArtifactsPanel: React.FC<{ jobDetails: JobDetailsV2 }> = observe
       };
 
       if (column.key === 'column_download') {
-         return <Stack id="downloadview" horizontalAlign="end" verticalAlign="center" verticalFill={true} style={{ paddingRight: 8 }}>
-            <PrimaryButton split text="Download" menuProps={downloadProps}
-               disabled={!stepFinished}
-               style={{ fontFamily: "Horde Open Sans SemiBold" }}
-               onClick={() => {
-                  window.location.assign(`/api/v2/artifacts/${item.id}/download?format=zip`);
-               }}
-            />
+         return <Stack horizontalAlign="end" verticalAlign="center" verticalFill={true} style={{ paddingRight: 8 }}>
+            <Stack horizontal tokens={{ childrenGap: 18 }}>
+            <PrimaryButton text="Browse" 
+                  disabled={!stepFinished}
+                  style={{ fontFamily: "Horde Open Sans SemiBold" }}
+                  onClick={() => {
+                     setSelected(item)
+                  }}
+               />
+               
+               <PrimaryButton split text="Download as Zip" menuProps={downloadProps}
+                  disabled={!stepFinished}
+                  style={{ fontFamily: "Horde Open Sans SemiBold" }}
+                  onClick={() => {
+                     window.location.assign(`/api/v2/artifacts/${item.id}/download?format=zip`);
+                  }}
+               />
+            </Stack>
          </Stack>
       }
 
@@ -198,6 +211,7 @@ export const JobArtifactsPanel: React.FC<{ jobDetails: JobDetailsV2 }> = observe
 
 
    return (<Stack id={sideRail.url} styles={{ root: { paddingTop: 18, paddingRight: 12 } }}>
+      {!!selected && <JobArtifactsModal jobId={jobData.id} stepId={selected.stepId} contextType={selected.type} onClose={() => setSelected(undefined)} />}
       <Stack className={hordeClasses.raised} >
          <Stack tokens={{ childrenGap: 12 }} grow>
             <Stack horizontal>
