Horde: Add a defaults.json file for backwards compatibility with older server installers. Just includes default.global.json.

diff --git a/Engine/Source/Programs/Horde/Horde.Server/Defaults/defaults.json b/Engine/Source/Programs/Horde/Horde.Server/Defaults/defaults.json
new file mode 100644
--- /dev/null
+++ b/Engine/Source/Programs/Horde/Horde.Server/Defaults/defaults.json
@@ -0,0 +1,8 @@
+// This file exists for backwards compatibility with older server versions. New installations should directly include the default.global.json file.
+{
+	"include": [
+		{
+			"path": "default.global.json"
+		}
+	]
+}
