Horde: Revert merge of RefsController.cs; missing other changes in this branch.

#rnx

diff --git a/Engine/Source/Programs/Horde/Horde.Server/Ddc/RefsController.cs b/Engine/Source/Programs/Horde/Horde.Server/Ddc/RefsController.cs
--- a/Engine/Source/Programs/Horde/Horde.Server/Ddc/RefsController.cs
+++ b/Engine/Source/Programs/Horde/Horde.Server/Ddc/RefsController.cs
@@ -294,7 +294,7 @@ namespace Horde.Server.Ddc
 							byte[] packageBytes;
 							{
 								using TelemetrySpan _ = _tracer.StartActiveSpan("cbpackage.buffer").SetAttribute("operation.name", "cbpackage.buffer");
-								packageBytes = await writer.ToByteArrayAsync();
+								packageBytes = await writer.ToByteArray();
 							}
 							await using BlobContents contents = new BlobContents(packageBytes);
 							await WriteBody(contents, CustomMediaTypeNames.UnrealCompactBinaryPackage);
@@ -741,11 +741,11 @@ namespace Horde.Server.Ddc
 			_diagnosticContext.Set("Content-Length", Request.ContentLength ?? -1);
 
 			byte[] b = await RequestUtil.ReadRawBodyAsync(Request);
-			CbPackageReader packageReader = await CbPackageReader.CreateAsync(new MemoryStream(b));
+			CbPackageReader packageReader = await CbPackageReader.Create(new MemoryStream(b));
 
 			try
 			{
-				await foreach ((CbPackageAttachmentEntry entry, byte[] blob) in packageReader.IterateAttachmentsAsync())
+				await foreach ((CbPackageAttachmentEntry entry, byte[] blob) in packageReader.IterateAttachments())
 				{
 					if (entry.Flags.HasFlag(CbPackageAttachmentFlags.IsError))
 					{
