Added User and Build UBT/UBA Metrics to Horde Analytics

diff --git a/Engine/Source/Programs/Horde/Horde.Server/Defaults/default-metrics.telemetry.json b/Engine/Source/Programs/Horde/Horde.Server/Defaults/default-metrics.telemetry.json
--- a/Engine/Source/Programs/Horde/Horde.Server/Defaults/default-metrics.telemetry.json
+++ b/Engine/Source/Programs/Horde/Horde.Server/Defaults/default-metrics.telemetry.json
@@ -482,6 +482,70 @@
 			"function": "Average",
 			"interval": "01:00:00"
 		},
+		{
+			"id": "ubt-uba-user-duration",
+			"filter": "$[?(@.Payload.EventName == 'Core.UBT.Executor' && $.Payload.Config_IsBuildMachine==false && search(@.Payload.Executor, \"Unreal Build Accelerator\") )]",
+			"groupBy": "$.Payload.Build_BranchName,$.Payload.Target,$.Payload.Mode,$.Payload.Configuration,$.Payload.Platform",
+			"property": "$.Payload.Duration",
+			"function": "Average",
+			"interval": "00:30:00"
+		},
+		{
+			"id": "ubt-uba-user-remote-usage",
+			"filter": "$[?(@.Payload.EventName == 'Core.UBT.Executor' && $.Payload.Config_IsBuildMachine==false && search(@.Payload.Executor, \"Unreal Build Accelerator\") )]",
+			"groupBy": "$.Payload.Build_BranchName,$.Payload.Target,$.Payload.Mode,$.Payload.Configuration,$.Payload.Platform",
+			"property": "$.Payload.RemoteUsage",
+			"function": "Average",
+			"interval": "00:30:00"
+		},
+		{
+			"id": "ubt-uba-user-remote-actions",
+			"filter": "$[?(@.Payload.EventName == 'Core.UBT.Executor' && $.Payload.Config_IsBuildMachine==false && search(@.Payload.Executor, \"Unreal Build Accelerator\") )]",
+			"groupBy": "$.Payload.Build_BranchName,$.Payload.Target,$.Payload.Mode,$.Payload.Configuration,$.Payload.Platform",
+			"property": "$.Payload.RemoteActions",
+			"function": "Sum",
+			"interval": "00:30:00"
+		},
+		{
+			"id": "ubt-uba-user-total-actions",
+			"filter": "$[?(@.Payload.EventName == 'Core.UBT.Executor' && $.Payload.Config_IsBuildMachine==false && search(@.Payload.Executor, \"Unreal Build Accelerator\") )]",
+			"groupBy": "$.Payload.Build_BranchName,$.Payload.Target,$.Payload.Mode,$.Payload.Configuration,$.Payload.Platform",
+			"property": "$.Payload.TotalActions",
+			"function": "Sum",
+			"interval": "00:30:00"
+		},
+		{
+			"id": "ubt-uba-build-duration",
+			"filter": "$[?(@.Payload.EventName == 'Core.UBT.Executor' && $.Payload.Config_IsBuildMachine==true && search(@.Payload.Executor, \"Unreal Build Accelerator\") )]",
+			"groupBy": "$.Payload.Build_BranchName,$.Payload.Target,$.Payload.Mode,$.Payload.Configuration,$.Payload.Platform",
+			"property": "$.Payload.Duration",
+			"function": "Average",
+			"interval": "00:30:00"
+		},
+		{
+			"id": "ubt-uba-build-remote-usage",
+			"filter": "$[?(@.Payload.EventName == 'Core.UBT.Executor' && $.Payload.Config_IsBuildMachine==true && search(@.Payload.Executor, \"Unreal Build Accelerator\") )]",
+			"groupBy": "$.Payload.Build_BranchName,$.Payload.Target,$.Payload.Mode,$.Payload.Configuration,$.Payload.Platform",
+			"property": "$.Payload.RemoteUsage",
+			"function": "Average",
+			"interval": "00:30:00"
+		},
+		{
+			"id": "ubt-uba-build-remote-actions",
+			"filter": "$[?(@.Payload.EventName == 'Core.UBT.Executor' && $.Payload.Config_IsBuildMachine==true && search(@.Payload.Executor, \"Unreal Build Accelerator\") )]",
+			"groupBy": "$.Payload.Build_BranchName,$.Payload.Target,$.Payload.Mode,$.Payload.Configuration,$.Payload.Platform",
+			"property": "$.Payload.RemoteActions",
+			"function": "Sum",
+			"interval": "00:30:00"
+		},
+		{
+			"id": "ubt-uba-build-total-actions",
+			"filter": "$[?(@.Payload.EventName == 'Core.UBT.Executor' && $.Payload.Config_IsBuildMachine==true && search(@.Payload.Executor, \"Unreal Build Accelerator\") )]",
+			"groupBy": "$.Payload.Build_BranchName,$.Payload.Target,$.Payload.Mode,$.Payload.Configuration,$.Payload.Platform",
+			"property": "$.Payload.TotalActions",
+			"function": "Sum",
+			"interval": "00:30:00"
+		},
 		{
 			"id": "ubt-uba-duration",
 			"filter": "$[?(@.Payload.EventName == 'Core.UBT.Executor' && search(@.Payload.Executor, \"Unreal Build Accelerator\") )]",
