Build and User UBA MNetrics Tweak for Horde Analytics

diff --git a/Engine/Source/Programs/Horde/Horde.Server/Defaults/default-metrics.telemetry.json b/Engine/Source/Programs/Horde/Horde.Server/Defaults/default-metrics.telemetry.json
--- a/Engine/Source/Programs/Horde/Horde.Server/Defaults/default-metrics.telemetry.json
+++ b/Engine/Source/Programs/Horde/Horde.Server/Defaults/default-metrics.telemetry.json
@@ -483,7 +483,7 @@
 			"interval": "01:00:00"
 		},
 		{
-			"id": "ubt-uba-user-duration",
+			"id": "ubt-uba-user-duration-avg",
 			"filter": "$[?(@.Payload.EventName == 'Core.UBT.Executor' && $.Payload.Config_IsBuildMachine==false && search(@.Payload.Executor, \"Unreal Build Accelerator\") )]",
 			"groupBy": "$.Payload.Build_BranchName,$.Payload.Target,$.Payload.Mode,$.Payload.Configuration,$.Payload.Platform",
 			"property": "$.Payload.Duration",
@@ -491,7 +491,7 @@
 			"interval": "00:30:00"
 		},
 		{
-			"id": "ubt-uba-user-remote-usage",
+			"id": "ubt-uba-user-remote-usage-avg",
 			"filter": "$[?(@.Payload.EventName == 'Core.UBT.Executor' && $.Payload.Config_IsBuildMachine==false && search(@.Payload.Executor, \"Unreal Build Accelerator\") )]",
 			"groupBy": "$.Payload.Build_BranchName,$.Payload.Target,$.Payload.Mode,$.Payload.Configuration,$.Payload.Platform",
 			"property": "$.Payload.RemoteUsage",
@@ -499,7 +499,7 @@
 			"interval": "00:30:00"
 		},
 		{
-			"id": "ubt-uba-user-remote-actions",
+			"id": "ubt-uba-user-remote-actions-sum",
 			"filter": "$[?(@.Payload.EventName == 'Core.UBT.Executor' && $.Payload.Config_IsBuildMachine==false && search(@.Payload.Executor, \"Unreal Build Accelerator\") )]",
 			"groupBy": "$.Payload.Build_BranchName,$.Payload.Target,$.Payload.Mode,$.Payload.Configuration,$.Payload.Platform",
 			"property": "$.Payload.RemoteActions",
@@ -507,7 +507,7 @@
 			"interval": "00:30:00"
 		},
 		{
-			"id": "ubt-uba-user-total-actions",
+			"id": "ubt-uba-user-total-actions-sum",
 			"filter": "$[?(@.Payload.EventName == 'Core.UBT.Executor' && $.Payload.Config_IsBuildMachine==false && search(@.Payload.Executor, \"Unreal Build Accelerator\") )]",
 			"groupBy": "$.Payload.Build_BranchName,$.Payload.Target,$.Payload.Mode,$.Payload.Configuration,$.Payload.Platform",
 			"property": "$.Payload.TotalActions",
@@ -515,7 +515,7 @@
 			"interval": "00:30:00"
 		},
 		{
-			"id": "ubt-uba-build-duration",
+			"id": "ubt-uba-build-duration-avg",
 			"filter": "$[?(@.Payload.EventName == 'Core.UBT.Executor' && $.Payload.Config_IsBuildMachine==true && search(@.Payload.Executor, \"Unreal Build Accelerator\") )]",
 			"groupBy": "$.Payload.Build_BranchName,$.Payload.Target,$.Payload.Mode,$.Payload.Configuration,$.Payload.Platform",
 			"property": "$.Payload.Duration",
@@ -523,7 +523,7 @@
 			"interval": "00:30:00"
 		},
 		{
-			"id": "ubt-uba-build-remote-usage",
+			"id": "ubt-uba-build-remote-usage-avg",
 			"filter": "$[?(@.Payload.EventName == 'Core.UBT.Executor' && $.Payload.Config_IsBuildMachine==true && search(@.Payload.Executor, \"Unreal Build Accelerator\") )]",
 			"groupBy": "$.Payload.Build_BranchName,$.Payload.Target,$.Payload.Mode,$.Payload.Configuration,$.Payload.Platform",
 			"property": "$.Payload.RemoteUsage",
@@ -531,7 +531,7 @@
 			"interval": "00:30:00"
 		},
 		{
-			"id": "ubt-uba-build-remote-actions",
+			"id": "ubt-uba-build-remote-actions-sum",
 			"filter": "$[?(@.Payload.EventName == 'Core.UBT.Executor' && $.Payload.Config_IsBuildMachine==true && search(@.Payload.Executor, \"Unreal Build Accelerator\") )]",
 			"groupBy": "$.Payload.Build_BranchName,$.Payload.Target,$.Payload.Mode,$.Payload.Configuration,$.Payload.Platform",
 			"property": "$.Payload.RemoteActions",
@@ -539,7 +539,7 @@
 			"interval": "00:30:00"
 		},
 		{
-			"id": "ubt-uba-build-total-actions",
+			"id": "ubt-uba-build-total-actions-sum",
 			"filter": "$[?(@.Payload.EventName == 'Core.UBT.Executor' && $.Payload.Config_IsBuildMachine==true && search(@.Payload.Executor, \"Unreal Build Accelerator\") )]",
 			"groupBy": "$.Payload.Build_BranchName,$.Payload.Target,$.Payload.Mode,$.Payload.Configuration,$.Payload.Platform",
 			"property": "$.Payload.TotalActions",
