Horde: Increase Jira timeout to 30 seconds, and improve cancellation handling.

diff --git a/Engine/Source/Programs/Horde/Horde.Server/Issues/External/JiraService.cs b/Engine/Source/Programs/Horde/Horde.Server/Issues/External/JiraService.cs
--- a/Engine/Source/Programs/Horde/Horde.Server/Issues/External/JiraService.cs
+++ b/Engine/Source/Programs/Horde/Horde.Server/Issues/External/JiraService.cs
@@ -143,7 +143,7 @@ namespace Horde.Server.Issues.External
 			_client = new HttpClient();
 			byte[] authBytes = Encoding.ASCII.GetBytes($"{_settings.JiraUsername}:{_settings.JiraApiToken}");
 			_client.DefaultRequestHeaders.Authorization = new AuthenticationHeaderValue("Basic", Convert.ToBase64String(authBytes));
-			_client.Timeout = TimeSpan.FromSeconds(15.0);
+			_client.Timeout = TimeSpan.FromSeconds(30.0);
 			_retryPolicy = HttpPolicyExtensions.HandleTransientHttpError().WaitAndRetryAsync(3, attempt => TimeSpan.FromSeconds(Math.Pow(2.0, attempt)));
 		}
 
@@ -209,9 +209,13 @@ namespace Horde.Server.Issues.External
 
 				try
 				{
-					response = await _retryPolicy.ExecuteAsync(() => _client.GetAsync(uri));
+					response = await _retryPolicy.ExecuteAsync(ctx => _client.GetAsync(uri, ctx), cancellationToken);
 					response.EnsureSuccessStatusCode();
 				}
+				catch (OperationCanceledException)
+				{
+					throw;
+				}
 				catch (Exception)
 				{
 					_logger.LogWarning("Unable to get project info for {ProjectKey}", projectKey);
@@ -253,7 +257,7 @@ namespace Horde.Server.Issues.External
 				// We do get components from the projects endpoint, though these don't contain the archived property :/
 				try
 				{
-					response = await _retryPolicy.ExecuteAsync(() => _client.GetAsync(new Uri(_jiraUrl, $"/rest/api/2/project/{projectKey}/components"), cancellationToken));
+					response = await _retryPolicy.ExecuteAsync(ctx => _client.GetAsync(new Uri(_jiraUrl, $"/rest/api/2/project/{projectKey}/components"), ctx), cancellationToken);
 					response.EnsureSuccessStatusCode();
 				}
 				catch (Exception)
@@ -397,9 +401,13 @@ namespace Horde.Server.Issues.External
 
 				try
 				{
-					response = await _retryPolicy.ExecuteAsync(() => _client.PostAsync(uri, new StringContent(JsonSerializer.Serialize(externalIssueUser), Encoding.UTF8, "application/json")));
+					response = await _retryPolicy.ExecuteAsync(ctx => _client.PostAsync(uri, new StringContent(JsonSerializer.Serialize(externalIssueUser), Encoding.UTF8, "application/json"), ctx), cancellationToken);
 					response.EnsureSuccessStatusCode();
 				}
+				catch (OperationCanceledException)
+				{
+					throw;
+				}
 				catch (Exception ex)
 				{
 					_logger.LogWarning(ex, "Unable to add {ExternalUser} as watcher to issue {IssueKey}", externalIssueUser, jiraResponse.Key);
