Horde: Disable verbose logging for Perforce service.

diff --git a/Engine/Source/Programs/Horde/Horde.Server/appsettings.json b/Engine/Source/Programs/Horde/Horde.Server/appsettings.json
--- a/Engine/Source/Programs/Horde/Horde.Server/appsettings.json
+++ b/Engine/Source/Programs/Horde/Horde.Server/appsettings.json
@@ -18,7 +18,6 @@
 				"System.Net.Http.HttpClient": "Warning",
 				"Grpc": "Warning",
 				"Horde.Server.Issues.IssueService": "Debug",
-				"Horde.Server.Perforce.PerforceService": "Verbose",
 				"EpicGames.Horde.Storage": "Debug"
 			}
 		}
