Horde: Increase verbosity for AWS object store.

diff --git a/Engine/Source/Programs/Horde/Horde.Server/appsettings.json b/Engine/Source/Programs/Horde/Horde.Server/appsettings.json
--- a/Engine/Source/Programs/Horde/Horde.Server/appsettings.json
+++ b/Engine/Source/Programs/Horde/Horde.Server/appsettings.json
@@ -18,6 +18,7 @@
 				"System.Net.Http.HttpClient": "Warning",
 				"Grpc": "Warning",
 				"Horde.Server.Issues.IssueService": "Debug",
+				"Horde.Server.Storage.ObjectStores.AwsObjectStore": "Debug",
 				"EpicGames.Horde.Storage": "Debug"
 			}
 		}
