Horde: Update release notes.

diff --git a/Engine/Source/Programs/Horde/Docs/ReleaseNotes.md b/Engine/Source/Programs/Horde/Docs/ReleaseNotes.md
--- a/Engine/Source/Programs/Horde/Docs/ReleaseNotes.md
+++ b/Engine/Source/Programs/Horde/Docs/ReleaseNotes.md
@@ -1,5 +1,17 @@
 # Release Notes
 
+## 2024-04-01
+
+* Fix decoding of UTF8 characters as part of unescaping Json strings. (32638382)
+* Increase Jira timeout to 30 seconds, and improve cancellation handling. (32636702)
+* Add a defaults.json file for backwards compatibility with older server installers. Just includes default.global.json. (32619739)
+* Move tool update packages onto a separate "Internal" tab on the dashboard. (32619239)
+* Only export nodes to Horde that are referenced by the initial job parameters, unless -AllowTargetChanges is specified. Prevents data being copied to temp storage that will never be used unless the target list changes. (32616197)
+* Add notes for setting up a self-signed cert for testing. (32613289)
+* Doc updates. (32595767)
+* Comment out the placeholder ticket value in the default config file. (32585713)
+* Improve ordering of nodes when writing large file trees to storage, such that the nodes which are read after each other are adjacent in the storage blobs. Nodes are read in the reverse order to which they are written, depth-first. Arranging nodes in the blob in this order prevents thrashing of the cache and improves performance. (32585474)
+
 ## 2024-03-28
 
 * Remove /userinfo call when authenticating via JWT. If an access token is passed, it's not guaranteed it has permission to access /userinfo from OIDC. ID tokens during normal web-based login does on the other hand. (32539012)
