Horde: Fix static analysis warnings.

diff --git a/Engine/Source/Programs/Horde/Horde.Server.Tests/Tools/ToolTests.cs b/Engine/Source/Programs/Horde/Horde.Server.Tests/Tools/ToolTests.cs
--- a/Engine/Source/Programs/Horde/Horde.Server.Tests/Tools/ToolTests.cs
+++ b/Engine/Source/Programs/Horde/Horde.Server.Tests/Tools/ToolTests.cs
@@ -20,7 +20,7 @@ namespace Horde.Server.Tests.Tools
 	[TestClass]
 	public class ToolTests : TestSetup
 	{
-		ToolId _toolId = new ToolId("ugs");
+		readonly ToolId _toolId = new ToolId("ugs");
 
 		public ToolTests()
 		{
diff --git a/Engine/Source/Programs/Horde/Horde.Server/Accounts/AccountCollection.cs b/Engine/Source/Programs/Horde/Horde.Server/Accounts/AccountCollection.cs
--- a/Engine/Source/Programs/Horde/Horde.Server/Accounts/AccountCollection.cs
+++ b/Engine/Source/Programs/Horde/Horde.Server/Accounts/AccountCollection.cs
@@ -187,7 +187,7 @@ namespace Horde.Server.Accounts
 
 		const int DuplicateKeyErrorCode = 11000;
 
-		private static AccountId s_defaultAdminAccountId = AccountId.Parse("65d4f282ff286703e0609ccd");
+		private static readonly AccountId s_defaultAdminAccountId = AccountId.Parse("65d4f282ff286703e0609ccd");
 
 		private bool _hasCreatedAdminAccount = false;
 		private readonly IMongoCollection<AccountDocument> _accounts;
diff --git a/Engine/Source/Programs/Horde/Horde.Server/Ddc/BlobService.cs b/Engine/Source/Programs/Horde/Horde.Server/Ddc/BlobService.cs
--- a/Engine/Source/Programs/Horde/Horde.Server/Ddc/BlobService.cs
+++ b/Engine/Source/Programs/Horde/Horde.Server/Ddc/BlobService.cs
@@ -14,7 +14,7 @@ namespace Horde.Server.Ddc
 {
 	class BlobService : IBlobService
 	{
-		static BlobType s_rawBlobType = new BlobType(new Guid("{03E6C37B-33C1-491F-8541-D3C401B8B8EF}"), 1);
+		static readonly BlobType s_rawBlobType = new BlobType(new Guid("{03E6C37B-33C1-491F-8541-D3C401B8B8EF}"), 1);
 
 		readonly IStorageClientFactory _storageClientFactory;
 		readonly Tracer _tracer;
diff --git a/Engine/Source/Programs/Horde/Horde.Server/Replicators/PerforceReplicator.cs b/Engine/Source/Programs/Horde/Horde.Server/Replicators/PerforceReplicator.cs
--- a/Engine/Source/Programs/Horde/Horde.Server/Replicators/PerforceReplicator.cs
+++ b/Engine/Source/Programs/Horde/Horde.Server/Replicators/PerforceReplicator.cs
@@ -73,7 +73,7 @@ namespace Horde.Server.Replicators
 
 		class StateNodeConverter : BlobConverter<StateNode>
 		{
-			static BlobType s_blobType = new BlobType("{8C874966-4273-2E89-9FAC-ABA46DC89154}", 2);
+			static readonly BlobType s_blobType = new BlobType("{8C874966-4273-2E89-9FAC-ABA46DC89154}", 2);
 
 			public override StateNode Read(IBlobReader reader, BlobSerializerOptions options)
 			{
