Horde: Add ability to set custom tag suffix for public Docker image

diff --git a/Engine/Source/Programs/Horde/BuildHorde.xml b/Engine/Source/Programs/Horde/BuildHorde.xml
--- a/Engine/Source/Programs/Horde/BuildHorde.xml
+++ b/Engine/Source/Programs/Horde/BuildHorde.xml
@@ -14,6 +14,7 @@
 	<Option Name="ImageRepoUrl" DefaultValue="" Description="URL to Docker/OCI image repository"/>
 	<Option Name="ImageName" DefaultValue="horde-server" Description="Name of Docker/OCI image"/>
 	<Option Name="ImageTag" DefaultValue="$(DefaultImageTag)" Description="Tag to apply to image"/>
+	<Option Name="ImageTagSuffix" DefaultValue="-bundled" Description="Text to append to tag"/>
 	<Option Name="ImageEnvFile" DefaultValue="" Description="Path to environment file for used for Docker invocations"/>
 	<Option Name="ImageRepoAuthFile" DefaultValue="" Description="Path to file with repository auth credentials"/>
 
@@ -299,7 +300,7 @@
 		</Node>
 
 		<Node Name="Push Bundled Docker Image" Requires="Build Bundled Docker Image">
-			<Docker-Push Repository="$(ImageRepoUrl)" Image="horde-server:bundled" TargetImage="$(ImageName):$(ImageTag)-bundled" EnvironmentFile="$(ImageEnvFile)" RepositoryAuthFile="$(ImageRepoAuthFile)" />
+			<Docker-Push Repository="$(ImageRepoUrl)" Image="horde-server:bundled" TargetImage="$(ImageName):$(ImageTag)$(ImageTagSuffix)" EnvironmentFile="$(ImageEnvFile)" RepositoryAuthFile="$(ImageRepoAuthFile)" />
 		</Node>
 	</Agent>
 
