Horde: Set empty default for image tag suffix option

diff --git a/Engine/Source/Programs/Horde/BuildHorde.xml b/Engine/Source/Programs/Horde/BuildHorde.xml
--- a/Engine/Source/Programs/Horde/BuildHorde.xml
+++ b/Engine/Source/Programs/Horde/BuildHorde.xml
@@ -14,7 +14,7 @@
 	<Option Name="ImageRepoUrl" DefaultValue="" Description="URL to Docker/OCI image repository"/>
 	<Option Name="ImageName" DefaultValue="horde-server" Description="Name of Docker/OCI image"/>
 	<Option Name="ImageTag" DefaultValue="$(DefaultImageTag)" Description="Tag to apply to image"/>
-	<Option Name="ImageTagSuffix" DefaultValue="-bundled" Description="Text to append to tag"/>
+	<Option Name="ImageTagSuffix" DefaultValue="" Description="Text to append to tag"/>
 	<Option Name="ImageEnvFile" DefaultValue="" Description="Path to environment file for used for Docker invocations"/>
 	<Option Name="ImageRepoAuthFile" DefaultValue="" Description="Path to file with repository auth credentials"/>
 
