Horde: Build server image on same BuildGraph agent as bundled image

diff --git a/Engine/Source/Programs/Horde/BuildHorde.xml b/Engine/Source/Programs/Horde/BuildHorde.xml
--- a/Engine/Source/Programs/Horde/BuildHorde.xml
+++ b/Engine/Source/Programs/Horde/BuildHorde.xml
@@ -34,15 +34,16 @@
 
 	<!-- SERVER -->
 
+	<Property Name="StagePaths">
+		Engine/Binaries/DotNET/EpicGames.Perforce.Native/...
+		Engine/Source/Programs/Shared/...
+		Engine/Source/Programs/Horde/...
+		Engine/Source/Programs/AutomationTool/AutomationUtils/Matchers/...
+		Engine/Source/Programs/UnrealBuildTool/Matchers/...
+	</Property>
+	
 	<Agent Name="HordeServer" Type="Linux;Win64_Docker">
 		<Property Name="StagingDir" Value="$(RootDir)/Engine/Saved/Horde.Server"/>
-		<Property Name="StagePaths">
-			Engine/Binaries/DotNET/EpicGames.Perforce.Native/...
-			Engine/Source/Programs/Shared/...
-			Engine/Source/Programs/Horde/...
-			Engine/Source/Programs/AutomationTool/AutomationUtils/Matchers/...
-			Engine/Source/Programs/UnrealBuildTool/Matchers/...
-		</Property>
 
 		<Node Name="Check Documentation">
 			<CheckMarkdown Files="$(RootDir)/Engine/Source/Programs/Horde/README.md;$(RootDir)/Engine/Source/Programs/Horde/Docs/..."/>
@@ -292,7 +293,10 @@
 	<Agent Name="HordeInstallerLinux" Type="Linux">
 		<Property Name="StagingDir" Value="$(RootDir)/Staging"/>
 		
-		<Node Name="Build Bundled Docker Image" Requires="Stage Horde Server Installer;Build HordeServer;#ServerWithToolsFiles">
+		<Node Name="Build Bundled Docker Image" Requires="Stage Horde Server Installer;#ServerWithToolsFiles">
+			<Tag Files="$(StagePaths)" Except=".../.vs/...;.../.git/...;.../bin/...;.../obj/..." With="#InputFiles"/>
+			<Docker-Build BaseDir="Engine" Files="#InputFiles" Tag="horde-server-bare" DockerFile="Engine/Source/Programs/Horde/Horde.Server/Dockerfile" Arguments="--build-arg msbuild_args=&quot;$(VersionArguments)&quot;"/>
+
 			<!-- A Horde server image has already been built, and is referenced inside Dockerfile.bundled -->
 			<Copy From="Engine/Source/Programs/Horde/Horde.Server/Dockerfile.bundled" To="$(StagingDir)/Dockerfile.bundled" />
 			<Tag Files="$(StagingDir)/Dockerfile*;$(StagingDir)/Dashboard/...;$(StagingDir)/Server/Tools/...;$(StagingDir)/Server/*.json" With="#DockerFiles"/>
