Horde: Add basic details for configuring Wine for remote execution

diff --git a/Engine/Source/Programs/Horde/Docs/Deployment/Agent.md b/Engine/Source/Programs/Horde/Docs/Deployment/Agent.md
--- a/Engine/Source/Programs/Horde/Docs/Deployment/Agent.md
+++ b/Engine/Source/Programs/Horde/Docs/Deployment/Agent.md
@@ -216,4 +216,38 @@ such as running the agent from a PowerShell script on the desktop or using an ex
 It's important to ensure the agent process is restarted in case of an unexpected exit.
 
 Since most agents won't require GPU access, it's recommended to create a separate pool
-specifically for interactive agents that need this support.
\ No newline at end of file
+specifically for interactive agents that need this support.
+
+### Linux and Wine
+
+Wine is a compatibility layer that enables Windows applications to run on Unix-like systems,
+allowing agents on Linux to execute remote compute tasks under Wine.
+Running workloads on Linux instead of Windows offers several advantages, such as faster OS boot times,
+improved disk I/O performance, and the elimination of license costs.
+
+To enable Wine, the path to `wine64` must be set in the [wineExecutablePath](../Deployment/AgentSettings.md) configuration option,
+as described in the agent settings documentation.
+A wrapper script in Bash can be used to modify any parameters sent to Wine, in this case configuring necessary file paths.
+When running a build, the `-UBAHordeAllowWine` flag for UnrealBuildTool must be set to true, which also is the default value.
+Once these prerequisites are met, agents will attempt to utilize Wine when executing compute tasks,
+leveraging the benefits of running on Linux while maintaining compatibility with Windows applications, such as `cl.exe` and `link.exe`.
+
+```bash
+#!/bin/bash
+# Wrapper script for Wine providing a hook point to modify incoming command-line arguments from the Windows application to run
+# Right now, this script is mostly called by remote execution tasks, which run UnrealBuildAccelerator (UBA)
+export WINEDEBUG=-all
+export WINEARCH=win64
+export WINEPREFIX=/opt/horde/wine-data
+
+# Overwrite UE_HORDE_SHARED_DIR to point to C:\ which will exist inside Wine.
+# This in turn is mounted under WINEPREFIX specified above
+export UE_HORDE_SHARED_DIR="C:\\Uba"
+
+if [ -n "$${UE_HORDE_TERMINATION_SIGNAL_FILE}" ]; then
+  # Rewrite Linux path to be a Windows path under Z:\ which maps to / in Wine (replacing slashes with backslashes)
+  export UE_HORDE_TERMINATION_SIGNAL_FILE="Z:$${UE_HORDE_TERMINATION_SIGNAL_FILE//\//\\}"
+fi
+
+/usr/local/bin/wine64 "$@"
+```
