Horde: Fix compilation error in LoginCommand

diff --git a/Engine/Source/Programs/Horde/Horde/Commands/LoginCommand.cs b/Engine/Source/Programs/Horde/Horde/Commands/LoginCommand.cs
--- a/Engine/Source/Programs/Horde/Horde/Commands/LoginCommand.cs
+++ b/Engine/Source/Programs/Horde/Horde/Commands/LoginCommand.cs
@@ -59,10 +59,10 @@ namespace Horde.Commands
 			{
 				HordeHttpAuthHandlerState state = _serviceProvider.GetRequiredService<HordeHttpAuthHandlerState>();
 
-				AuthenticationHeaderValue? header = await state.TryGetAuthHeaderAsync(CancellationToken.None);
-				if (header != null)
+				string? accessToken = await state.GetAccessTokenAsync(true, CancellationToken.None);
+				if (accessToken != null)
 				{
-					Console.WriteLine("{0} {1}", header.Scheme, header.Parameter);
+					Console.WriteLine($"Access token {accessToken}");
 				}
 			}
 
