Horde: Fix compile warning for LoginCommand

diff --git a/Engine/Source/Programs/Horde/Horde/Commands/LoginCommand.cs b/Engine/Source/Programs/Horde/Horde/Commands/LoginCommand.cs
--- a/Engine/Source/Programs/Horde/Horde/Commands/LoginCommand.cs
+++ b/Engine/Source/Programs/Horde/Horde/Commands/LoginCommand.cs
@@ -1,7 +1,6 @@
 // Copyright Epic Games, Inc. All Rights Reserved.
 
 using System.ComponentModel;
-using System.Net.Http.Headers;
 using EpicGames.Core;
 using EpicGames.Horde;
 using EpicGames.Horde.Server;
