Horde: Do not assign compute resource to same machine as requester

diff --git a/Engine/Source/Programs/Horde/Horde.Server.Tests/Compute/ComputeServiceTest.cs b/Engine/Source/Programs/Horde/Horde.Server.Tests/Compute/ComputeServiceTest.cs
--- a/Engine/Source/Programs/Horde/Horde.Server.Tests/Compute/ComputeServiceTest.cs
+++ b/Engine/Source/Programs/Horde/Horde.Server.Tests/Compute/ComputeServiceTest.cs
@@ -209,6 +209,23 @@ namespace Horde.Server.Tests.Compute
 			Assert.AreEqual(agent4.Id, resource4!.AgentId);
 		}
 
+		[TestMethod]
+		public async Task Assignment_RequesterAlsoRunsAgent_IsNotAssignedAsync()
+		{
+			AllocateResourceParams CreateParams(string ip)
+			{
+				return new AllocateResourceParams(new ClusterId("default"), ComputeProtocol.Latest, new Requirements { Pool = "foo" }) { RequesterIp = IPAddress.Parse(ip) };
+			}
+
+			IAgent agent = await CreateAgentAsync(new PoolId("foo"), properties: ["ComputeIp=11.0.0.1", "ComputePort=5000"]);
+			
+			ComputeResource? resource1 = await ComputeService.TryAllocateResourceAsync(CreateParams("11.0.0.1"), CancellationToken.None);
+			Assert.IsNull(resource1);
+			
+			ComputeResource? resource2 = await ComputeService.TryAllocateResourceAsync(CreateParams("11.0.0.2"), CancellationToken.None);
+			Assert.AreEqual(agent.Id, resource2!.AgentId);
+		}
+
 		[TestMethod]
 		public async Task Connection_Direct_IpConnection_Async()
 		{
diff --git a/Engine/Source/Programs/Horde/Horde.Server/Agents/IAgent.cs b/Engine/Source/Programs/Horde/Horde.Server/Agents/IAgent.cs
--- a/Engine/Source/Programs/Horde/Horde.Server/Agents/IAgent.cs
+++ b/Engine/Source/Programs/Horde/Horde.Server/Agents/IAgent.cs
@@ -452,6 +452,16 @@ namespace Horde.Server.Agents
 		/// </summary>
 		public const string Ram = "RAM";
 
+		/// <summary>
+		/// IP address used for sending compute task payloads
+		/// </summary>
+		public const string ComputeIp = "ComputeIp";
+		
+		/// <summary>
+		/// Port used for sending compute task payloads
+		/// </summary>
+		public const string ComputePort = "ComputePort";
+
 		/// <summary>
 		/// AWS: Instance ID
 		/// </summary>
@@ -815,15 +825,23 @@ namespace Horde.Server.Agents
 		/// <param name="agent">The agent to create a lease for</param>
 		/// <param name="requirements">Requirements for the lease</param>
 		/// <param name="assignedResources">Receives the allocated resources</param>
+		/// <param name="conditions">Condition to check in addition to those in requirements</param>
 		/// <returns>True if the new lease can be granted</returns>
-		public static bool MeetsRequirements(this IAgent agent, Requirements requirements, Dictionary<string, int> assignedResources)
+		public static bool MeetsRequirements(this IAgent agent, Requirements requirements, Dictionary<string, int> assignedResources, List<Condition> conditions)
 		{
 			PoolId? poolId = null;
 			if (!String.IsNullOrEmpty(requirements.Pool))
 			{
 				poolId = new PoolId(requirements.Pool);
 			}
-			return MeetsRequirements(agent, poolId, requirements.Condition, requirements.Resources, requirements.Exclusive, assignedResources);
+
+			List<Condition> combinedConditions = [];
+			if (requirements.Condition != null)
+			{
+				combinedConditions.Add(requirements.Condition);
+			}
+			combinedConditions.AddRange(conditions);
+			return MeetsRequirements(agent, poolId, combinedConditions, requirements.Resources, requirements.Exclusive, assignedResources);
 		}
 
 		/// <summary>
@@ -831,12 +849,12 @@ namespace Horde.Server.Agents
 		/// </summary>
 		/// <param name="agent">The agent to create a lease for</param>
 		/// <param name="poolId">Pool to take the machine from</param>
-		/// <param name="condition">Condition to satisfy</param>
+		/// <param name="conditions">Conditions to satisfy</param>
 		/// <param name="resources">Resources required to execute</param>
 		/// <param name="exclusive">Whether the lease needs to be executed exclusively on the machine</param>
 		/// <param name="assignedResources">Resources allocated to the task</param>
 		/// <returns>True if the new lease can be granted</returns>
-		public static bool MeetsRequirements(this IAgent agent, PoolId? poolId, Condition? condition, Dictionary<string, ResourceRequirements>? resources, bool exclusive, Dictionary<string, int> assignedResources)
+		public static bool MeetsRequirements(this IAgent agent, PoolId? poolId, List<Condition> conditions, Dictionary<string, ResourceRequirements>? resources, bool exclusive, Dictionary<string, int> assignedResources)
 		{
 			if (!agent.Enabled || agent.Status != AgentStatus.Ok)
 			{
@@ -854,7 +872,7 @@ namespace Horde.Server.Agents
 			{
 				return false;
 			}
-			if (condition != null && !agent.SatisfiesCondition(condition))
+			if (conditions.Any(condition => !agent.SatisfiesCondition(condition)))
 			{
 				return false;
 			}
diff --git a/Engine/Source/Programs/Horde/Horde.Server/Compute/ComputeService.cs b/Engine/Source/Programs/Horde/Horde.Server/Compute/ComputeService.cs
--- a/Engine/Source/Programs/Horde/Horde.Server/Compute/ComputeService.cs
+++ b/Engine/Source/Programs/Horde/Horde.Server/Compute/ComputeService.cs
@@ -32,6 +32,7 @@ using Microsoft.Extensions.Logging;
 using Microsoft.Extensions.Options;
 using OpenTelemetry.Trace;
 using StackExchange.Redis;
+using Condition = EpicGames.Horde.Common.Condition;
 
 namespace Horde.Server.Compute
 {
@@ -369,6 +370,15 @@ namespace Horde.Server.Compute
 				certificate = GenerateCert(arp.Encryption); // A no-op if certificate is not required	
 			}
 
+			List<Condition> conditions = [];
+			if (arp.RequesterIp != null)
+			{
+				// Do not assign a compute resource hosted on the same machine as requester
+				// This is often the case when a user is offering their local workstation to participate in a shared compute pool
+				// With UnrealBuildAccelerator, it already makes full use of the local machine.
+				conditions.Add(Condition.Parse($"{KnownPropertyNames.ComputeIp} != '{arp.RequesterIp.ToString()}'"));
+			}
+
 			try
 			{
 				IReadOnlyList<IAgent> agents = await _agentCollection.FindAsync(cancellationToken: cancellationToken);
@@ -376,7 +386,7 @@ namespace Horde.Server.Compute
 				{
 					Dictionary<string, int> assignedResources = new Dictionary<string, int>();
 
-					bool match = agent.MeetsRequirements(arp.Requirements, assignedResources);
+					bool match = agent.MeetsRequirements(arp.Requirements, assignedResources, conditions);
 					if (match)
 					{
 						using TelemetrySpan matchSpan = _tracer.StartActiveSpan("Found match");
diff --git a/Engine/Source/Programs/Shared/EpicGames.Horde/Compute/Requirements.cs b/Engine/Source/Programs/Shared/EpicGames.Horde/Compute/Requirements.cs
--- a/Engine/Source/Programs/Shared/EpicGames.Horde/Compute/Requirements.cs
+++ b/Engine/Source/Programs/Shared/EpicGames.Horde/Compute/Requirements.cs
@@ -8,7 +8,7 @@ using EpicGames.Serialization;
 namespace EpicGames.Horde.Compute
 {
 	/// <summary>
-	/// Stores information about a directory in an action's workspace
+	/// Requirements for a compute task to be assigned an agent
 	/// </summary>
 	public class Requirements
 	{
