Horde: Instrument JobTaskSource ticker

We see occasional slow downs in agent assignment for jobs. Adding some additional tracing should help see what's causing it.

diff --git a/Engine/Source/Programs/Horde/Horde.Server/Jobs/JobTaskSource.cs b/Engine/Source/Programs/Horde/Horde.Server/Jobs/JobTaskSource.cs
--- a/Engine/Source/Programs/Horde/Horde.Server/Jobs/JobTaskSource.cs
+++ b/Engine/Source/Programs/Horde/Horde.Server/Jobs/JobTaskSource.cs
@@ -39,7 +39,7 @@ using HordeCommon.Rpc.Tasks;
 using Microsoft.Extensions.Hosting;
 using Microsoft.Extensions.Logging;
 using Microsoft.Extensions.Options;
-using MongoDB.Driver;
+using OpenTelemetry.Trace;
 
 namespace Horde.Server.Jobs
 {
@@ -204,6 +204,7 @@ namespace Horde.Server.Jobs
 		readonly IOptionsMonitor<ServerSettings> _settings;
 		readonly IOptionsMonitor<GlobalConfig> _globalConfig;
 		readonly ILogger<JobTaskSource> _logger;
+		readonly Tracer _tracer;
 		readonly ITicker _ticker;
 
 		// Object used for ensuring mutual exclusion to the queues
@@ -237,7 +238,7 @@ namespace Horde.Server.Jobs
 		/// <summary>
 		/// Constructor
 		/// </summary>
-		public JobTaskSource(AclService aclService, IAgentCollection agents, IJobCollection jobs, IJobStepRefCollection jobStepRefs, IBisectTaskCollection bisectTasks, IGraphCollection graphs, IPoolCollection pools, PoolService poolService, IUgsMetadataCollection ugsMetadataCollection, IStreamCollection streamCollection, ILogFileService logFileService, PerforceLoadBalancer perforceLoadBalancer, IClock clock, IOptionsMonitor<ServerSettings> settings, IOptionsMonitor<GlobalConfig> globalConfig, ILogger<JobTaskSource> logger)
+		public JobTaskSource(AclService aclService, IAgentCollection agents, IJobCollection jobs, IJobStepRefCollection jobStepRefs, IBisectTaskCollection bisectTasks, IGraphCollection graphs, IPoolCollection pools, PoolService poolService, IUgsMetadataCollection ugsMetadataCollection, IStreamCollection streamCollection, ILogFileService logFileService, PerforceLoadBalancer perforceLoadBalancer, IClock clock, IOptionsMonitor<ServerSettings> settings, IOptionsMonitor<GlobalConfig> globalConfig, Tracer tracer, ILogger<JobTaskSource> logger)
 		{
 			_aclService = aclService;
 			_agentsCollection = agents;
@@ -255,6 +256,7 @@ namespace Horde.Server.Jobs
 			_ticker = clock.AddTicker<JobTaskSource>(s_refreshInterval, TickAsync, logger);
 			_globalConfig = globalConfig;
 			_settings = settings;
+			_tracer = tracer;
 			_logger = logger;
 
 			OnLeaseStartedProperties.Add(nameof(ExecuteJobTask.JobId), x => JobId.Parse(x.JobId)).Add(nameof(ExecuteJobTask.BatchId), x => JobStepBatchId.Parse(x.BatchId)).Add(nameof(ExecuteJobTask.LogId), x => LogId.Parse(x.LogId));
@@ -360,12 +362,16 @@ namespace Horde.Server.Jobs
 		/// <returns>Async task</returns>
 		internal async ValueTask TickAsync(CancellationToken cancellationToken)
 		{
+			using TelemetrySpan span = _tracer.StartActiveSpan($"{nameof(JobTaskSource)}.{nameof(TickAsync)}");
+			
 			// Set the NewBatchIdToQueueItem member, so we capture any updated jobs during the DB query.
 			lock (_lockObject)
 			{
 				_newQueueItemsDuringUpdate = new List<QueueItem>();
 			}
 
+			using TelemetrySpan setupSpan = _tracer.StartActiveSpan($"{nameof(JobTaskSource)}.{nameof(TickAsync)}.Setup");
+			
 			// Query all the current streams
 			GlobalConfig globalConfig = _globalConfig.CurrentValue;
 			IReadOnlyList<IStream> streamsList = await _streamCollection.GetAsync(globalConfig.Streams, cancellationToken);
@@ -382,9 +388,18 @@ namespace Horde.Server.Jobs
 
 			// Query for a new list of jobs for the queue
 			IReadOnlyList<IJob> newJobs = await _jobs.GetDispatchQueueAsync(cancellationToken);
+
+			setupSpan.End();
+			span.SetAttribute("numNewJobs", newJobs.Count);
+			span.SetAttribute("numAgents", agents.Count);
+			span.SetAttribute("numPools", pools.Count);
+			span.SetAttribute("numStreams", streamsList.Count);
+			
 			for (int idx = 0; idx < newJobs.Count; idx++)
 			{
 				IJob? newJob = newJobs[idx];
+				using TelemetrySpan newJobSpan = _tracer.StartActiveSpan($"{nameof(JobTaskSource)}.{nameof(TickAsync)}.NewJob");
+				newJobSpan.SetAttribute("jobId", newJob.Id.ToString());
 
 				if (newJob.GraphHash == null)
 				{
