Horde: Rework commit metadata replication to run a single task for each replicated cluster, for a simpler code flow and better tracing data.

diff --git a/Engine/Source/Programs/Horde/Horde.Server/Perforce/PerforceServiceCache.cs b/Engine/Source/Programs/Horde/Horde.Server/Perforce/PerforceServiceCache.cs
--- a/Engine/Source/Programs/Horde/Horde.Server/Perforce/PerforceServiceCache.cs
+++ b/Engine/Source/Programs/Horde/Horde.Server/Perforce/PerforceServiceCache.cs
@@ -2,7 +2,6 @@
 
 using System;
 using System.Collections.Generic;
-using System.Diagnostics;
 using System.IO;
 using System.Linq;
 using System.Runtime.CompilerServices;
@@ -110,16 +109,25 @@ namespace Horde.Server.Perforce
 			}
 		}
 
-		class ClusterTicker
+		sealed class ClusterTicker : IAsyncDisposable
 		{
 			public string ClusterName { get; }
-			public Task<ClusterState?>? Task { get; set; }
-			public Stopwatch Timer { get; }
+			public BackgroundTask? Task { get; set; }
+			public List<StreamInfo> Streams { get; set; }
 
-			public ClusterTicker(string name)
+			public ClusterTicker(string name, List<StreamInfo> streamInfos)
 			{
 				ClusterName = name;
-				Timer = Stopwatch.StartNew();
+				Streams = streamInfos;
+			}
+
+			public async ValueTask DisposeAsync()
+			{
+				if (Task != null)
+				{
+					await Task.DisposeAsync();
+					Task = null;
+				}
 			}
 		}
 
@@ -272,120 +280,50 @@ namespace Horde.Server.Perforce
 		/// <summary>
 		/// Polls Perforce for submitted changes
 		/// </summary>
-		/// <param name="cancellationToken"></param>
-		/// <returns></returns>
 		async ValueTask UpdateCommitsAsync(CancellationToken cancellationToken)
 		{
-			CacheState state = await _mongoService.GetSingletonAsync<CacheState>(cancellationToken);
-
-			// Get the current list of streams and their views
-			Dictionary<string, List<StreamInfo>> clusters = await CreateStreamInfoAsync(cancellationToken);
-
-			// Task for updating the list of clusters periodically
-			Task<Dictionary<string, List<StreamInfo>>>? clusterTask = null;
-			Stopwatch clusterTimer = Stopwatch.StartNew();
-
-			// Poll each cluster
-			List<ClusterTicker> tickers = new List<ClusterTicker>();
+			Dictionary<string, ClusterTicker> clusterToTicker = new Dictionary<string, ClusterTicker>();
 			try
 			{
+				_logger.LogDebug("Starting commit metadata replication");
 				for (; ; )
 				{
-					// Update the background task for refreshing the list of clusters
-					if (clusterTask != null && clusterTask.IsCompleted)
-					{
-						try
-						{
-							clusters = await clusterTask;
-						}
-						catch (Exception ex)
-						{
-							_logger.LogError(ex, "Exception while updating cluster information: {Message}", ex.Message);
-						}
-						clusterTask = null;
-					}
-
-					// Don't do any updates during downtime; we might just create a bunch of P4 errors.
 					if (!_downtimeService.IsDowntimeActive)
 					{
-						// Check if it's time to start a new cluster update
-						if (clusterTask == null && clusterTimer.Elapsed > TimeSpan.FromSeconds(30.0))
-						{
-							clusterTask = Task.Run(() => CreateStreamInfoAsync(cancellationToken), cancellationToken);
-							clusterTimer.Restart();
-						}
+						Dictionary<string, List<StreamInfo>> clusterToStreams = await CreateStreamInfoAsync(cancellationToken);
 
-						// Remove any state for clusters that are no longer valid
-						bool updateState = false;
-						foreach (string clusterName in state.Clusters.Keys)
+						// Stop any tickers that no longer exist
+						foreach (ClusterTicker ticker in clusterToTicker.Values.ToArray())
 						{
-							if (!clusters.ContainsKey(clusterName))
+							if (!clusterToStreams.ContainsKey(ticker.ClusterName))
 							{
-								state.Clusters.Remove(clusterName);
-								updateState = true;
+								await StopTickerAsync(ticker);
+								clusterToTicker.Remove(ticker.ClusterName);
 							}
 						}
 
-						// Make sure there's a ticker for every cluster
-						foreach (string clusterName in clusters.Keys)
+						// Start any new tickers that don't exist yet
+						foreach ((string clusterName, List<StreamInfo> streamInfos) in clusterToStreams)
 						{
-							if (!tickers.Any(x => x.ClusterName.Equals(clusterName, StringComparison.OrdinalIgnoreCase)))
+							if (!clusterToTicker.ContainsKey(clusterName))
 							{
-								ClusterTicker ticker = new ClusterTicker(clusterName);
-								tickers.Add(ticker);
-							}
-						}
-
-						// Check if it's time to update any tickers
-						for (int idx = 0; idx < tickers.Count; idx++)
-						{
-							ClusterTicker ticker = tickers[idx];
-							if (ticker.Task != null && ticker.Task.IsCompleted)
-							{
-								ClusterState? clusterState = await ticker.Task;
-								if (clusterState != null)
-								{
-									state.Clusters[ticker.ClusterName] = clusterState;
-									updateState = true;
-								}
-								ticker.Task = null;
-							}
-							if (ticker.Task == null)
-							{
-								List<StreamInfo>? streams;
-								if (!clusters.TryGetValue(ticker.ClusterName, out streams))
-								{
-									tickers.RemoveAt(idx--);
-									continue;
-								}
-
-								ClusterState? clusterState;
-								if (!state.Clusters.TryGetValue(ticker.ClusterName, out clusterState))
-								{
-									clusterState = new ClusterState();
-								}
-
-								ticker.Task = Task.Run(() => UpdateClusterGuardedAsync(ticker.ClusterName, streams, clusterState, cancellationToken));
-							}
-						}
-
-						// Apply any updates to the global state
-						if (updateState)
-						{
-							if (!await _mongoService.TryUpdateSingletonAsync(state, cancellationToken))
-							{
-								state = await _mongoService.GetSingletonAsync<CacheState>(cancellationToken);
+								ClusterTicker ticker = StartTicker(clusterName, streamInfos);
+								clusterToTicker.Add(ticker.ClusterName, ticker);
 							}
 						}
 					}
 
-					// Wait before performing the next poll
-					await Task.Delay(TimeSpan.FromSeconds(2.0), cancellationToken);
+					// Wait before performing the next update
+					await Task.Delay(TimeSpan.FromSeconds(30.0), cancellationToken);
 				}
 			}
 			finally
 			{
-				await Task.WhenAll(tickers.Select(x => x.Task).Where(x => x != null)!);
+				_logger.LogDebug("Stopping commit metadata replication");
+				foreach (ClusterTicker ticker in clusterToTicker.Values)
+				{
+					await StopTickerAsync(ticker);
+				}
 			}
 		}
 
@@ -421,33 +359,82 @@ namespace Horde.Server.Perforce
 			}
 		}
 
-		async Task<ClusterState?> UpdateClusterGuardedAsync(string clusterName, List<StreamInfo> streamInfos, ClusterState state, CancellationToken cancellationToken)
+		ClusterTicker StartTicker(string clusterName, List<StreamInfo> streamInfos)
+		{
+			ClusterTicker ticker = new ClusterTicker(clusterName, streamInfos);
+			ticker.Task = BackgroundTask.StartNew(ctx => UpdateClusterAsync(ticker, ctx));
+			return ticker;
+		}
+
+		async ValueTask StopTickerAsync(ClusterTicker ticker)
 		{
 			try
 			{
-				ClusterState? next = await UpdateClusterAsync(clusterName, streamInfos, state, cancellationToken);
-				return next;
+				await ticker.DisposeAsync();
 			}
-			catch (OperationCanceledException)
+			catch (Exception ex)
 			{
-				return null;
+				_logger.LogWarning(ex, "Exception while disposing ticker for cluster '{ClusterName}': {Message}", ticker.ClusterName, ex.Message);
 			}
-			catch (Exception ex)
+		}
+
+		async Task UpdateClusterAsync(ClusterTicker ticker, CancellationToken cancellationToken)
+		{
+			using TelemetrySpan span = _tracer.StartActiveSpan($"{nameof(PerforceServiceCache)}.{nameof(UpdateClusterAsync)}");
+			span.SetAttribute("ClusterName", ticker.ClusterName);
+
+			while (!cancellationToken.IsCancellationRequested)
 			{
-				_logger.LogError(ex, "Exception while updating cluster state: {Message}", ex.Message);
-				return null;
+				// Update the cluster state
+				if (!_downtimeService.IsDowntimeActive)
+				{
+					try
+					{
+						await UpdateClusterInternalAsync(ticker.ClusterName, ticker.Streams, cancellationToken);
+					}
+					catch (OperationCanceledException)
+					{
+						throw;
+					}
+					catch (Exception ex)
+					{
+						_logger.LogError(ex, "Exception while updating cluster state: {Message}", ex.Message);
+					}
+				}
+
+				// Wait before polling again
+				await Task.Delay(TimeSpan.FromSeconds(2.0), cancellationToken);
 			}
 		}
 
-		async Task<ClusterState?> UpdateClusterAsync(string clusterName, List<StreamInfo> streamInfos, ClusterState state, CancellationToken cancellationToken)
+		async Task UpdateClusterInternalAsync(string clusterName, List<StreamInfo> streamInfos, CancellationToken cancellationToken)
 		{
-			const int MaxChanges = 250;
+			using TelemetrySpan span = _tracer.StartActiveSpan($"{nameof(PerforceServiceCache)}.{nameof(UpdateClusterInternalAsync)}");
+
+			// Get the current state of this cluster
+			CacheState cacheState = await _mongoService.GetSingletonAsync<CacheState>(cancellationToken);
+
+			ClusterState? clusterState;
+			if (!cacheState.Clusters.TryGetValue(clusterName, out clusterState))
+			{
+				clusterState = new ClusterState();
+			}
 
-			using TelemetrySpan telemetrySpan = _tracer.StartActiveSpan($"{nameof(PerforceServiceCache)}.{nameof(UpdateClusterAsync)}");
-			telemetrySpan.SetAttribute("Cluster", clusterName);
+			// Update the commits
+			ClusterState? nextClusterState = await UpdateClusterCommitsAsync(clusterName, streamInfos, clusterState, cancellationToken);
+			if (nextClusterState != null)
+			{
+				await _mongoService.UpdateSingletonAsync<CacheState>(state => state.Clusters[clusterName] = nextClusterState, cancellationToken);
+			}
+		}
 
+		async Task<ClusterState?> UpdateClusterCommitsAsync(string clusterName, List<StreamInfo> streamInfos, ClusterState state, CancellationToken cancellationToken)
+		{
+			using TelemetrySpan telemetrySpan = _tracer.StartActiveSpan($"{nameof(PerforceServiceCache)}.{nameof(UpdateClusterCommitsAsync)}");
 			using (IPooledPerforceConnection perforce = await ConnectAsync(clusterName, null, cancellationToken))
 			{
+				const int MaxChanges = 250;
+
 				// If the hash of any stream definition has changed, invalidate the replicated changes.
 				bool modified = false;
 				foreach (StreamInfo streamInfo in streamInfos)
@@ -474,6 +461,8 @@ namespace Horde.Server.Perforce
 					spec = $"@{state.MaxChange + 1},@now";
 				}
 
+				_logger.LogDebug("Replicating changes for cluster {ClusterName} matching {FileSpec}", clusterName, spec);
+
 				// Find the changes within that range, and abort if there's nothing new
 				List<ChangesRecord> changes = await perforce.GetChangesAsync(ChangesOptions.None, MaxChanges, ChangeStatus.Submitted, spec, cancellationToken);
 				if (changes.Count > 0)
