Horde: Reduce log spam for Perforce metadata updates.

diff --git a/Engine/Source/Programs/Horde/Horde.Server/Perforce/PerforceServiceCache.cs b/Engine/Source/Programs/Horde/Horde.Server/Perforce/PerforceServiceCache.cs
--- a/Engine/Source/Programs/Horde/Horde.Server/Perforce/PerforceServiceCache.cs
+++ b/Engine/Source/Programs/Horde/Horde.Server/Perforce/PerforceServiceCache.cs
@@ -461,8 +461,6 @@ namespace Horde.Server.Perforce
 					spec = $"@{state.MaxChange + 1},@now";
 				}
 
-				_logger.LogDebug("Replicating changes for cluster {ClusterName} matching {FileSpec}", clusterName, spec);
-
 				// Find the changes within that range, and abort if there's nothing new
 				List<ChangesRecord> changes = await perforce.GetChangesAsync(ChangesOptions.None, MaxChanges, ChangeStatus.Submitted, spec, cancellationToken);
 				if (changes.Count > 0)
@@ -482,6 +480,8 @@ namespace Horde.Server.Perforce
 					return modified ? state : null;
 				}
 
+				_logger.LogDebug("Replicating {NumChanges} changes for cluster {ClusterName} matching {FileSpec}, {NumUpdatedChanges} changes to refresh", changes.Count, clusterName, spec, refreshNumbers.Length);
+
 				// If we've retrieved the maximum number of changes from the server, we no longer have a complete chronological cache and need to reset it.
 				bool reset = false;
 				if (changes.Count >= MaxChanges)
