Horde: Explicitly set parent context for telemetry spans in Perforce updates, to work around async context mismatches.

diff --git a/Engine/Source/Programs/Horde/Horde.Server/Perforce/PerforceServiceCache.cs b/Engine/Source/Programs/Horde/Horde.Server/Perforce/PerforceServiceCache.cs
--- a/Engine/Source/Programs/Horde/Horde.Server/Perforce/PerforceServiceCache.cs
+++ b/Engine/Source/Programs/Horde/Horde.Server/Perforce/PerforceServiceCache.cs
@@ -390,7 +390,7 @@ namespace Horde.Server.Perforce
 				{
 					try
 					{
-						await UpdateClusterInternalAsync(ticker.ClusterName, ticker.Streams, cancellationToken);
+						await UpdateClusterInternalAsync(ticker.ClusterName, ticker.Streams, span.Context, cancellationToken);
 					}
 					catch (OperationCanceledException)
 					{
@@ -407,7 +407,7 @@ namespace Horde.Server.Perforce
 			}
 		}
 
-		async Task UpdateClusterInternalAsync(string clusterName, List<StreamInfo> streamInfos, CancellationToken cancellationToken)
+		async Task UpdateClusterInternalAsync(string clusterName, List<StreamInfo> streamInfos, SpanContext parentContext, CancellationToken cancellationToken)
 		{
 			using TelemetrySpan span = _tracer.StartActiveSpan($"{nameof(PerforceServiceCache)}.{nameof(UpdateClusterInternalAsync)}");
 
@@ -421,16 +421,16 @@ namespace Horde.Server.Perforce
 			}
 
 			// Update the commits
-			ClusterState? nextClusterState = await UpdateClusterCommitsAsync(clusterName, streamInfos, clusterState, cancellationToken);
+			ClusterState? nextClusterState = await UpdateClusterCommitsAsync(clusterName, streamInfos, clusterState, parentContext, cancellationToken);
 			if (nextClusterState != null)
 			{
 				await _mongoService.UpdateSingletonAsync<CacheState>(state => state.Clusters[clusterName] = nextClusterState, cancellationToken);
 			}
 		}
 
-		async Task<ClusterState?> UpdateClusterCommitsAsync(string clusterName, List<StreamInfo> streamInfos, ClusterState state, CancellationToken cancellationToken)
+		async Task<ClusterState?> UpdateClusterCommitsAsync(string clusterName, List<StreamInfo> streamInfos, ClusterState state, SpanContext parentContext, CancellationToken cancellationToken)
 		{
-			using TelemetrySpan telemetrySpan = _tracer.StartActiveSpan($"{nameof(PerforceServiceCache)}.{nameof(UpdateClusterCommitsAsync)}");
+			using TelemetrySpan telemetrySpan = _tracer.StartActiveSpan($"{nameof(PerforceServiceCache)}.{nameof(UpdateClusterCommitsAsync)}", parentContext: parentContext);
 			using (IPooledPerforceConnection perforce = await ConnectAsync(clusterName, null, cancellationToken))
 			{
 				const int MaxChanges = 250;
