Horde: Missing parentContext argument.

diff --git a/Engine/Source/Programs/Horde/Horde.Server/Perforce/PerforceServiceCache.cs b/Engine/Source/Programs/Horde/Horde.Server/Perforce/PerforceServiceCache.cs
--- a/Engine/Source/Programs/Horde/Horde.Server/Perforce/PerforceServiceCache.cs
+++ b/Engine/Source/Programs/Horde/Horde.Server/Perforce/PerforceServiceCache.cs
@@ -409,7 +409,7 @@ namespace Horde.Server.Perforce
 
 		async Task UpdateClusterInternalAsync(string clusterName, List<StreamInfo> streamInfos, SpanContext parentContext, CancellationToken cancellationToken)
 		{
-			using TelemetrySpan span = _tracer.StartActiveSpan($"{nameof(PerforceServiceCache)}.{nameof(UpdateClusterInternalAsync)}");
+			using TelemetrySpan span = _tracer.StartActiveSpan($"{nameof(PerforceServiceCache)}.{nameof(UpdateClusterInternalAsync)}", parentContext: parentContext);
 
 			// Get the current state of this cluster
 			CacheState cacheState = await _mongoService.GetSingletonAsync<CacheState>(cancellationToken);
