Horde: Update derived data for issues when they are closed by timeout.

diff --git a/Engine/Source/Programs/Horde/Horde.Server/Issues/IssueService.cs b/Engine/Source/Programs/Horde/Horde.Server/Issues/IssueService.cs
--- a/Engine/Source/Programs/Horde/Horde.Server/Issues/IssueService.cs
+++ b/Engine/Source/Programs/Horde/Horde.Server/Issues/IssueService.cs
@@ -339,6 +339,7 @@ namespace Horde.Server.Issues
 				if (openIssue.LastSeenAt < utcNow - TimeSpan.FromDays(7.0))
 				{
 					await _issueCollection.TryUpdateIssueAsync(openIssue, null, newResolvedById: IIssue.ResolvedByTimeoutId, cancellationToken: cancellationToken);
+					await UpdateIssueDerivedDataAsync(openIssue, cancellationToken);
 					openIssues.RemoveAt(idx--);
 					continue;
 				}
