Horde: Fix a parent telemetry span context.

diff --git a/Engine/Source/Programs/Horde/Horde.Server/Perforce/PerforceServiceCache.cs b/Engine/Source/Programs/Horde/Horde.Server/Perforce/PerforceServiceCache.cs
--- a/Engine/Source/Programs/Horde/Horde.Server/Perforce/PerforceServiceCache.cs
+++ b/Engine/Source/Programs/Horde/Horde.Server/Perforce/PerforceServiceCache.cs
@@ -421,7 +421,7 @@ namespace Horde.Server.Perforce
 			}
 
 			// Update the commits
-			ClusterState? nextClusterState = await UpdateClusterCommitsAsync(clusterName, streamInfos, clusterState, parentContext, cancellationToken);
+			ClusterState? nextClusterState = await UpdateClusterCommitsAsync(clusterName, streamInfos, clusterState, span.Context, cancellationToken);
 			if (nextClusterState != null)
 			{
 				await _mongoService.UpdateSingletonAsync<CacheState>(state => state.Clusters[clusterName] = nextClusterState, cancellationToken);
