Horde: Prevent log appearing empty if tailing is enabled but no new data has been received before the existing tail data is expired. Rpc log sink was waiting forever for new log tail data, so server was discarding the data already received.

diff --git a/Engine/Source/Programs/Horde/Horde.Agent/Utility/JsonRpcLogSink.cs b/Engine/Source/Programs/Horde/Horde.Agent/Utility/JsonRpcLogSink.cs
--- a/Engine/Source/Programs/Horde/Horde.Agent/Utility/JsonRpcLogSink.cs
+++ b/Engine/Source/Programs/Horde/Horde.Agent/Utility/JsonRpcLogSink.cs
@@ -108,6 +108,7 @@ namespace Horde.Agent.Utility
 		async Task TickTailInternalAsync()
 		{
 			int tailNext = -1;
+			Task tickTask = Task.CompletedTask;
 			while (!_tailTaskStop.IsSet())
 			{
 				Task newTailDataTask = _newTailDataEvent.Task;
@@ -120,11 +121,13 @@ namespace Horde.Agent.Utility
 					(tailNext, tailData) = _builder.ReadTailData(tailNext, 16 * 1024);
 				}
 
-				// If we don't have any updates for the server, wait until we do.
-				if (tailNext != -1 && tailData.IsEmpty && tailNext == initialTailNext)
+				// If we don't have any updates for the server, wait until we do. We need to ensure
+				// we keep pumping the RPC with the server in case the requested tail next value changes,
+				// and to make sure that we don't expire the existing tail data.
+				if (tailNext != -1 && tailData.IsEmpty && tailNext == initialTailNext && !tickTask.IsCompleted)
 				{
 					_logger.LogInformation("No tail data available for log {LogId} after line {TailNext}; waiting for more...", _logId, tailNext);
-					await newTailDataTask;
+					await Task.WhenAny(newTailDataTask, tickTask);
 					continue;
 				}
 
@@ -146,6 +149,8 @@ namespace Horde.Agent.Utility
 					tailNext = newTailNext;
 					_logger.LogInformation("Modified tail position for log {LogId} to {TailNext}", _logId, tailNext);
 				}
+
+				tickTask = Task.Delay(TimeSpan.FromSeconds(10.0));
 			}
 			_logger.LogInformation("Finishing log tail task");
 		}
