Horde: Update release notes.

#rnx

diff --git a/Engine/Source/Programs/Horde/Docs/ReleaseNotes.md b/Engine/Source/Programs/Horde/Docs/ReleaseNotes.md
--- a/Engine/Source/Programs/Horde/Docs/ReleaseNotes.md
+++ b/Engine/Source/Programs/Horde/Docs/ReleaseNotes.md
@@ -1,5 +1,17 @@
 # Release Notes
 
+## 2024-04-05
+
+* Handle socket shutdown errors gracefully rather than throwing exceptions. (32756972)
+* Add log message whenever native P4 connection buffer increases in size above 16mb. (32751407)
+* Prevent log appearing empty if tailing is enabled but no new data has been received before the existing tail data is expired. Rpc log sink was waiting forever for new log tail data, so server was discarding the data already received. (32749401)
+* Fix default settings for the analytics dashboard, and update documentation to specify the correct telemetry store name. (32741193)
+* Update derived data for issues when they are closed by timeout. (32733991)
+* Explicitly set parent context for telemetry spans in Perforce updates, to work around async context mismatches. (32732888)
+* Rework commit metadata replication to run a single task for each replicated cluster, for a simpler code flow and better tracing data. (32730286)
+* Instrument JobTaskSource ticker. We see occasional slow downs in agent assignment for jobs. Adding some additional tracing should help see what's causing it. (32727003)
+* Do not assign compute resource to same machine as requester (32723450)
+
 ## 2024-04-01
 
 * Fix decoding of UTF8 characters as part of unescaping Json strings. (32638382)
