Horde: Fix compile error in Login command.

diff --git a/Engine/Source/Programs/Horde/Horde/Commands/LoginCommand.cs b/Engine/Source/Programs/Horde/Horde/Commands/LoginCommand.cs
--- a/Engine/Source/Programs/Horde/Horde/Commands/LoginCommand.cs
+++ b/Engine/Source/Programs/Horde/Horde/Commands/LoginCommand.cs
@@ -61,7 +61,7 @@ namespace Horde.Commands
 				string? accessToken = await state.GetAccessTokenAsync(true, CancellationToken.None);
 				if (accessToken != null)
 				{
-					Console.WriteLine($"Access token {accessToken}");
+					Console.WriteLine($"Bearer {accessToken}");
 				}
 			}
 
