Horde: Add more logging to identify Perforce load balancing lookup failures

diff --git a/Engine/Source/Programs/Horde/Horde.Server/Perforce/PerforceLoadBalancer.cs b/Engine/Source/Programs/Horde/Horde.Server/Perforce/PerforceLoadBalancer.cs
--- a/Engine/Source/Programs/Horde/Horde.Server/Perforce/PerforceLoadBalancer.cs
+++ b/Engine/Source/Programs/Horde/Horde.Server/Perforce/PerforceLoadBalancer.cs
@@ -6,6 +6,7 @@ using System.Globalization;
 using System.Linq;
 using System.Net;
 using System.Net.Http;
+using System.Text;
 using System.Text.Json;
 using System.Threading;
 using System.Threading.Tasks;
@@ -342,6 +343,13 @@ namespace Horde.Server.Perforce
 				}
 			}
 
+			// Count number of servers per status before filtering
+			Dictionary<string, object> statusCounts = new();
+			foreach (PerforceServerStatus status in System.Enum.GetValues<PerforceServerStatus>())
+			{
+				statusCounts[status.ToString()] = candidates.Count(x => x.Status == status);
+			}
+
 			// Remove any servers that are unhealthy
 			if (candidates.Any(x => x.Status == PerforceServerStatus.Healthy))
 			{
@@ -354,6 +362,7 @@ namespace Horde.Server.Perforce
 
 			if (candidates.Count == 0)
 			{
+				using IDisposable? logScope = _logger.BeginScope(statusCounts);
 				_logger.LogWarning("Unable to find any healthy Perforce server in cluster {ClusterName}", cluster.Name);
 				return null;
 			}
@@ -457,6 +466,7 @@ namespace Horde.Server.Perforce
 			}
 			await Task.WhenAll(tasks);
 
+			// Update Horde's health monitor for subsystems
 			list = await _serverListSingleton.GetAsync(cancellationToken);
 			(HealthStatus health, string message) = GetPerforceHealth(list.Servers);
 			span.SetAttribute("health.status", health.ToString());
@@ -592,7 +602,7 @@ namespace Horde.Server.Perforce
 			}
 			else
 			{
-				await ResolveServersAsync(initialHostName, hostNames, cancellationToken);
+				await ResolveDnsHostnameAsync(initialHostName, hostNames, cancellationToken);
 			}
 
 			foreach (string hostName in hostNames)
@@ -606,15 +616,20 @@ namespace Horde.Server.Perforce
 			}
 		}
 
-		async Task ResolveServersAsync(string hostName, List<string> hostNames, CancellationToken cancellationToken)
+		/// <summary>
+		/// Find all the IP addresses for a hostname
+		/// </summary>
+		/// <param name="hostName">DNS hostname</param>
+		/// <param name="ipAddresses">List of IP addresses</param>
+		/// <param name="cancellationToken">Cancellation token</param>
+		async Task ResolveDnsHostnameAsync(string hostName, List<string> ipAddresses, CancellationToken cancellationToken)
 		{
-			// Find all the addresses of the hosts
 			IPHostEntry entry = await Dns.GetHostEntryAsync(hostName, cancellationToken);
 			foreach (IPAddress address in entry.AddressList)
 			{
 				try
 				{
-					hostNames.Add(address.ToString());
+					ipAddresses.Add(address.ToString());
 				}
 				catch (Exception ex)
 				{
@@ -699,10 +714,13 @@ namespace Horde.Server.Perforce
 					}
 				}
 
+				string dataStr = Encoding.UTF8.GetString(data).Substring(0, 200);
+				_logger.LogWarning("Unable to parse health check output. Data {Data} URL {HealthCheckUrl}", dataStr, healthCheckUrl);
 				return new ServerHealth(PerforceServerStatus.Unknown, "Unable to parse health check output");
 			}
-			catch
+			catch (Exception e)
 			{
+				_logger.LogWarning("Failed checking Perforce server health via HTTP. URL {HealthCheckUrl} Reason {Reason}", healthCheckUrl, e.Message);
 				return new ServerHealth(PerforceServerStatus.Unhealthy, $"Failed to query status at {healthCheckUrl}");
 			}
 		}
