Horde: Allow clients to upload the list of references for each blob, removing the need for the server to do it later.

diff --git a/Engine/Source/Programs/Horde/Horde.Server.Tests/Accounts/AccountControllerTests.cs b/Engine/Source/Programs/Horde/Horde.Server.Tests/Accounts/AccountControllerTests.cs
--- a/Engine/Source/Programs/Horde/Horde.Server.Tests/Accounts/AccountControllerTests.cs
+++ b/Engine/Source/Programs/Horde/Horde.Server.Tests/Accounts/AccountControllerTests.cs
@@ -30,7 +30,7 @@ public class AccountControllerTest : IAsyncDisposable
 		{
 			{ "Horde:AuthMethod", "Horde" },
 		};
-		_app = new FakeHordeWebApp(settings, allowAutoRedirect: false);
+		_app = new FakeHordeWebApp(settings);
 		_accountCollection = _app.ServiceProvider.GetRequiredService<IAccountCollection>();
 	}
 
@@ -54,7 +54,8 @@ public class AccountControllerTest : IAsyncDisposable
 	[TestMethod]
 	public async Task NotLoggedIn_AccountPageShowsLoginLinkAsync()
 	{
-		HttpResponseMessage res = await _app.HttpClient.GetAsync("account");
+		using HttpClient httpClient = _app.CreateHttpClient(allowAutoRedirect: false);
+		HttpResponseMessage res = await httpClient.GetAsync("account");
 		Assert.AreEqual(HttpStatusCode.OK, res.StatusCode);
 		Assert.IsTrue((await res.Content.ReadAsStringAsync()).Contains("<b>Login</b>"));
 	}
@@ -62,7 +63,8 @@ public class AccountControllerTest : IAsyncDisposable
 	[TestMethod]
 	public async Task NotLoggedIn_DefaultLoginRedirectsToUserPassFormAsync()
 	{
-		HttpResponseMessage res = await _app.HttpClient.GetAsync("account/login");
+		using HttpClient httpClient = _app.CreateHttpClient(allowAutoRedirect: false);
+		HttpResponseMessage res = await httpClient.GetAsync("account/login");
 		Assert.AreEqual(HttpStatusCode.Found, res.StatusCode);
 		Assert.AreEqual("/account/login/horde?ReturnUrl=%2Faccount", res.Headers.Location!.PathAndQuery);
 	}
@@ -70,7 +72,8 @@ public class AccountControllerTest : IAsyncDisposable
 	[TestMethod]
 	public async Task NotLoggedIn_UserPassFormAsync()
 	{
-		HttpResponseMessage res = await _app.HttpClient.GetAsync("account/login/horde");
+		using HttpClient httpClient = _app.CreateHttpClient(allowAutoRedirect: false);
+		HttpResponseMessage res = await httpClient.GetAsync("account/login/horde");
 		Assert.AreEqual(HttpStatusCode.OK, res.StatusCode);
 		Assert.IsTrue((await res.Content.ReadAsStringAsync()).Contains("<form action"));
 	}
@@ -78,14 +81,15 @@ public class AccountControllerTest : IAsyncDisposable
 	[TestMethod]
 	public async Task NotLoggedIn_BadCredentialsShowsErrorAsync()
 	{
+		using HttpClient httpClient = _app.CreateHttpClient(allowAutoRedirect: false);
 		{
-			HttpResponseMessage res = await LoginAsync("does-not-exist", "foo");
+			HttpResponseMessage res = await LoginAsync(httpClient, "does-not-exist", "foo");
 			Assert.AreEqual(HttpStatusCode.BadRequest, res.StatusCode);
 			Assert.IsTrue((await res.Content.ReadAsStringAsync()).Contains("Invalid username or password"));
 		}
 
 		{
-			HttpResponseMessage res = await LoginAsync(_account1.Login, "wrong-password");
+			HttpResponseMessage res = await LoginAsync(httpClient, _account1.Login, "wrong-password");
 			Assert.AreEqual(HttpStatusCode.BadRequest, res.StatusCode);
 			Assert.IsTrue((await res.Content.ReadAsStringAsync()).Contains("Invalid username or password"));
 		}
@@ -94,7 +98,8 @@ public class AccountControllerTest : IAsyncDisposable
 	[TestMethod]
 	public async Task NotLoggedIn_CorrectCredentialsLogsInAsync()
 	{
-		HttpResponseMessage res = await LoginAsync(_account1.Login, "pass1");
+		using HttpClient httpClient = _app.CreateHttpClient(allowAutoRedirect: false);
+		HttpResponseMessage res = await LoginAsync(httpClient, _account1.Login, "pass1");
 		Assert.AreEqual(HttpStatusCode.Redirect, res.StatusCode);
 		Assert.AreEqual("/", res.Headers.Location!.ToString());
 	}
@@ -102,8 +107,9 @@ public class AccountControllerTest : IAsyncDisposable
 	[TestMethod]
 	public async Task LoggedIn_ShowsCredentialsAsync()
 	{
-		await LoginAsync(_account1.Login, "pass1");
-		HttpResponseMessage res2 = await _app.HttpClient.GetAsync("account");
+		using HttpClient httpClient = _app.CreateHttpClient(allowAutoRedirect: false);
+		await LoginAsync(httpClient, _account1.Login, "pass1");
+		HttpResponseMessage res2 = await httpClient.GetAsync("account");
 		Assert.AreEqual(HttpStatusCode.OK, res2.StatusCode);
 
 		string content = await res2.Content.ReadAsStringAsync();
@@ -111,7 +117,7 @@ public class AccountControllerTest : IAsyncDisposable
 		Assert.IsTrue(content.Contains("myClaimType2") && content.Contains("myClaimValue2"));
 	}
 
-	private async Task<HttpResponseMessage> LoginAsync(string username, string password, string? redirectUrl = null)
+	private async Task<HttpResponseMessage> LoginAsync(HttpClient httpClient, string username, string password, string? redirectUrl = null)
 	{
 		using StringContent sc = new($"username={username}&password={password}", Encoding.UTF8, "application/x-www-form-urlencoded");
 		string url = "account/login/horde";
@@ -119,7 +125,7 @@ public class AccountControllerTest : IAsyncDisposable
 		{
 			url = QueryHelpers.AddQueryString("account/login/horde", "ReturnUrl", redirectUrl);
 		}
-		HttpResponseMessage response = await _app.HttpClient.PostAsync(url, sc);
+		HttpResponseMessage response = await httpClient.PostAsync(url, sc);
 		return response;
 	}
 }
\ No newline at end of file
diff --git a/Engine/Source/Programs/Horde/Horde.Server.Tests/ControllerIntegrationTest.cs b/Engine/Source/Programs/Horde/Horde.Server.Tests/ControllerIntegrationTest.cs
--- a/Engine/Source/Programs/Horde/Horde.Server.Tests/ControllerIntegrationTest.cs
+++ b/Engine/Source/Programs/Horde/Horde.Server.Tests/ControllerIntegrationTest.cs
@@ -88,17 +88,20 @@ public class FakeHordeWebApp : IAsyncDisposable
 {
 	public MongoInstance MongoInstance { get; }
 	public RedisInstance RedisInstance { get; }
-	public HttpClient HttpClient { get; }
 	public IServiceProvider ServiceProvider => Factory.Services;
 	private TestWebApplicationFactory<Startup> Factory { get; }
 
-	public FakeHordeWebApp(Dictionary<string, string>? settings = null, bool allowAutoRedirect = true)
+	public FakeHordeWebApp(Dictionary<string, string>? settings = null)
 	{
 		MongoInstance = new MongoInstance();
 		RedisInstance = new RedisInstance();
 		Factory = new TestWebApplicationFactory<Startup>(MongoInstance, RedisInstance, settings);
+	}
+
+	public HttpClient CreateHttpClient(bool allowAutoRedirect = true)
+	{
 		WebApplicationFactoryClientOptions opts = new() { AllowAutoRedirect = allowAutoRedirect };
-		HttpClient = Factory.CreateClient(opts);
+		return Factory.CreateClient(opts);
 	}
 
 	public async ValueTask DisposeAsync()
@@ -121,7 +124,7 @@ public class FakeHordeWebApp : IAsyncDisposable
 
 public class ControllerIntegrationTest : IAsyncDisposable
 {
-	protected HttpClient Client => _app.HttpClient;
+	protected HttpClient Client { get; }
 	protected IServiceProvider ServiceProvider => _app.ServiceProvider;
 	private readonly Lazy<Task<Fixture>> _fixture;
 	private readonly FakeHordeWebApp _app;
@@ -130,6 +133,7 @@ public class ControllerIntegrationTest : IAsyncDisposable
 	{
 		_app = new FakeHordeWebApp();
 		_fixture = new Lazy<Task<Fixture>>(CreateFixtureAsync);
+		Client = _app.CreateHttpClient();
 	}
 
 	public virtual async ValueTask DisposeAsync()
@@ -138,6 +142,9 @@ public class ControllerIntegrationTest : IAsyncDisposable
 		GC.SuppressFinalize(this);
 	}
 
+	public HttpClient CreateHttpClient(bool allowAutoRedirect = true)
+		=> _app.CreateHttpClient(allowAutoRedirect);
+
 	public Task<Fixture> GetFixtureAsync()
 	{
 		return _fixture.Value;
diff --git a/Engine/Source/Programs/Horde/Horde.Server.Tests/Storage/StorageControllerTests.cs b/Engine/Source/Programs/Horde/Horde.Server.Tests/Storage/StorageControllerTests.cs
new file mode 100644
--- /dev/null
+++ b/Engine/Source/Programs/Horde/Horde.Server.Tests/Storage/StorageControllerTests.cs
@@ -0,0 +1,64 @@
+﻿// Copyright Epic Games, Inc. All Rights Reserved.
+
+using System;
+using System.Collections.Generic;
+using System.Linq;
+using System.Text;
+using System.Threading.Tasks;
+using EpicGames.Core;
+using EpicGames.Horde.Storage;
+using EpicGames.Horde.Storage.Backends;
+using EpicGames.Horde.Storage.Bundles;
+using Horde.Server.Server;
+using Horde.Server.Storage;
+using Microsoft.Extensions.DependencyInjection;
+using Microsoft.Extensions.Logging;
+using Microsoft.VisualStudio.TestTools.UnitTesting;
+using MongoDB.Bson;
+using MongoDB.Driver;
+
+namespace Horde.Server.Tests.Storage
+{
+	[TestClass]
+	public class StorageControllerTests : ControllerIntegrationTest
+	{
+		static readonly Guid s_guid = Guid.Parse("{8B9701C9-9677-4E3C-895F-2D654267C4CB}");
+
+		[TestMethod]
+		public async Task TestBlobImportsAsync()
+		{
+			ILogger logger = ServiceProvider.GetRequiredService<ILogger<StorageControllerTests>>();
+			HttpStorageBackend backend = new HttpStorageBackend("api/v1/storage/default", () => CreateHttpClient(), () => CreateHttpClient(), logger);
+			await using BundleCache cache = new BundleCache();
+			using BundleStorageClient client = new BundleStorageClient(backend, cache, null, logger);
+
+			IBlobRef blobRef1;
+			IBlobRef blobRef2;
+			await using (IBlobWriter writer = client.CreateBlobWriter())
+			{
+				writer.WriteInt32(1234);
+				blobRef1 = await writer.CompleteAsync(new BlobType(s_guid, 0));
+				await writer.FlushAsync();
+
+				writer.WriteBlobRef(blobRef1);
+				blobRef2 = await writer.CompleteAsync(new BlobType(s_guid, 0));
+			}
+
+			BlobLocator locator1 = blobRef1.GetLocator();
+			BlobLocator locator2 = blobRef2.GetLocator();
+
+			StorageService storageService = ServiceProvider.GetRequiredService<StorageService>();
+			IMongoCollection<StorageService.BlobInfo> blobCollection = storageService._blobCollection;
+
+			StorageService.BlobInfo blob2 = await blobCollection.Find(x => x.Path == locator2.BaseLocator.ToString()).FirstAsync();
+			Assert.IsNotNull(blob2.Imports);
+			Assert.AreEqual(1, blob2.Imports.Count);
+			Assert.AreEqual(blob2.Path, locator2.BaseLocator.ToString());
+
+			StorageService.BlobInfo blob1 = await blobCollection.Find(x => x.Id == blob2.Imports![0]).FirstAsync();
+			Assert.IsNotNull(blob1.Imports);
+			Assert.AreEqual(0, blob1.Imports.Count);
+			Assert.AreEqual(blob1.Path, locator1.BaseLocator.ToString());
+		}
+	}
+}
diff --git a/Engine/Source/Programs/Horde/Horde.Server/Logs/LogsController.cs b/Engine/Source/Programs/Horde/Horde.Server/Logs/LogsController.cs
--- a/Engine/Source/Programs/Horde/Horde.Server/Logs/LogsController.cs
+++ b/Engine/Source/Programs/Horde/Horde.Server/Logs/LogsController.cs
@@ -100,13 +100,13 @@ namespace Horde.Server.Logs
 		/// Uploads a blob for a log file. See /api/v1/storage/XXX/blobs.
 		/// </summary>
 		/// <param name="logFileId">Id of the log file to get information about</param>
-		/// <param name="file">Data for the blob</param>
+		/// <param name="request">Request for the upload</param>
 		/// <param name="cancellationToken">Cancellation token for the request</param>
 		/// <returns>Information about the requested project</returns>
 		[HttpPost]
 		[Route("/api/v1/logs/{logFileId}/blobs")]
 		[ProducesResponseType(typeof(WriteBlobResponse), 200)]
-		public async Task<ActionResult<WriteBlobResponse>> WriteLogBlobAsync(LogId logFileId, IFormFile? file, CancellationToken cancellationToken = default)
+		public async Task<ActionResult<WriteBlobResponse>> WriteLogBlobAsync(LogId logFileId, WriteBlobRequest request, CancellationToken cancellationToken = default)
 		{
 			ILogFile? logFile = await _logFileService.GetLogFileAsync(logFileId, cancellationToken);
 			if (logFile == null)
@@ -117,9 +117,15 @@ namespace Horde.Server.Logs
 			{
 				return Forbid();
 			}
+			if (!String.IsNullOrEmpty(request.Prefix))
+			{
+				return BadRequest("Cannot specify prefix for logs");
+			}
+
+			request.Prefix = $"{logFile.RefName}";
 
 			IStorageBackend storageBackend = _storageService.CreateBackend(Namespace.Logs);
-			return await StorageController.WriteBlobAsync(storageBackend, file, $"{logFile.RefName}", cancellationToken);
+			return await StorageController.WriteBlobAsync(storageBackend, request, cancellationToken);
 		}
 
 		/// <summary>
diff --git a/Engine/Source/Programs/Horde/Horde.Server/Storage/StorageController.cs b/Engine/Source/Programs/Horde/Horde.Server/Storage/StorageController.cs
--- a/Engine/Source/Programs/Horde/Horde.Server/Storage/StorageController.cs
+++ b/Engine/Source/Programs/Horde/Horde.Server/Storage/StorageController.cs
@@ -57,39 +57,37 @@ namespace Horde.Server.Storage
 		/// Uploads data to the storage service. 
 		/// </summary>
 		/// <param name="namespaceId">Namespace to fetch from</param>
-		/// <param name="file">Data to be uploaded. May be null, in which case the server may return a separate url.</param>
-		/// <param name="prefix">Prefix for the uploaded file</param>
+		/// <param name="request">Information about the blob to write</param>
 		/// <param name="cancellationToken">Cancellation token for the operation</param>
 		[HttpPost]
 		[Route("/api/v1/storage/{namespaceId}/blobs")]
-		public async Task<ActionResult<WriteBlobResponse>> WriteBlobAsync(NamespaceId namespaceId, IFormFile? file, [FromForm] string? prefix = default, CancellationToken cancellationToken = default)
+		public async Task<ActionResult<WriteBlobResponse>> WriteBlobAsync(NamespaceId namespaceId, WriteBlobRequest request, CancellationToken cancellationToken = default)
 		{
 			IStorageBackend? storageBackend = _storageService.TryCreateBackend(namespaceId);
 			if (storageBackend == null)
 			{
 				return NotFound(namespaceId);
 			}
-			if (!Authorize(namespaceId, StorageAclAction.WriteBlobs) && !HasPathClaim(User, HordeClaimTypes.WriteNamespace, namespaceId, prefix ?? String.Empty))
+			if (!Authorize(namespaceId, StorageAclAction.WriteBlobs) && !HasPathClaim(User, HordeClaimTypes.WriteNamespace, namespaceId, request.Prefix ?? String.Empty))
 			{
 				return Forbid(StorageAclAction.WriteBlobs, namespaceId);
 			}
 
-			return await WriteBlobAsync(storageBackend, file, prefix, cancellationToken);
+			return await WriteBlobAsync(storageBackend, request, cancellationToken);
 		}
 
 		/// <summary>
 		/// Writes a blob to storage. Exposed as a public utility method to allow other routes with their own authentication methods to wrap their own authentication/redirection.
 		/// </summary>
 		/// <param name="storageBackend">The backend to write to</param>
-		/// <param name="file">File to be written</param>
-		/// <param name="prefix">Prefix for uploaded blobs</param>
+		/// <param name="request">Information about the blob to write</param>
 		/// <param name="cancellationToken">Cancellation token</param>
 		/// <returns>Information about the written blob, or redirect information</returns>
-		public static async Task<ActionResult<WriteBlobResponse>> WriteBlobAsync(IStorageBackend storageBackend, IFormFile? file, [FromForm] string? prefix = default, CancellationToken cancellationToken = default)
+		public static async Task<ActionResult<WriteBlobResponse>> WriteBlobAsync(IStorageBackend storageBackend, WriteBlobRequest request, CancellationToken cancellationToken = default)
 		{
-			if (file == null)
+			if (request.File == null)
 			{
-				(BlobLocator Path, Uri UploadUrl)? result = await storageBackend.TryGetBlobWriteRedirectAsync(prefix ?? String.Empty, cancellationToken);
+				(BlobLocator Path, Uri UploadUrl)? result = await storageBackend.TryGetBlobWriteRedirectAsync(request.Imports, request.Prefix ?? String.Empty, cancellationToken);
 				if (result == null)
 				{
 					return new WriteBlobResponse { SupportsRedirects = false };
@@ -99,8 +97,8 @@ namespace Horde.Server.Storage
 			}
 			else
 			{
-				using Stream stream = file.OpenReadStream();
-				BlobLocator locator = await storageBackend.WriteBlobAsync(stream, prefix, cancellationToken);
+				using Stream stream = request.File.OpenReadStream();
+				BlobLocator locator = await storageBackend.WriteBlobAsync(stream, request.Imports, request.Prefix, cancellationToken);
 				return new WriteBlobResponse { Blob = locator.ToString(), SupportsRedirects = storageBackend.SupportsRedirects };
 			}
 		}
@@ -532,4 +530,38 @@ namespace Horde.Server.Storage
 
 		static string GetNodeLink(NamespaceId namespaceId, BlobLocator locator) => $"/api/v1/storage/{namespaceId}/nodes/{locator.BaseLocator}?{locator.Fragment}";
 	}
+
+	/// <summary>
+	/// Request to upload a blob
+	/// </summary>
+	public class WriteBlobRequest
+	{
+		/// <summary>
+		/// Set to indicate that the blob does not have any imports. Used to distinguish from an "unknown" set of imports.
+		/// </summary>
+		[FromForm(Name = "leaf")]
+		public bool Leaf
+		{
+			get => Imports != null && Imports.Count == 0;
+			set => Imports = (value ? new List<BlobLocator>() : Imports);
+		}
+
+		/// <summary>
+		/// Imported blobs
+		/// </summary>
+		[FromForm(Name = "import")]
+		public List<BlobLocator>? Imports { get; set; }
+
+		/// <summary>
+		/// Prefix for the uploaded file
+		/// </summary>
+		[FromForm(Name = "prefix")]
+		public string? Prefix { get; set; }
+
+		/// <summary>
+		/// Data to be uploaded. May be null, in which case the server may return a separate url.
+		/// </summary>
+		[FromForm(Name = "file")]
+		public IFormFile? File { get; set; }
+	}
 }
diff --git a/Engine/Source/Programs/Horde/Horde.Server/Storage/StorageService.cs b/Engine/Source/Programs/Horde/Horde.Server/Storage/StorageService.cs
--- a/Engine/Source/Programs/Horde/Horde.Server/Storage/StorageService.cs
+++ b/Engine/Source/Programs/Horde/Horde.Server/Storage/StorageService.cs
@@ -77,12 +77,12 @@ namespace Horde.Server.Storage
 				=> _store.ReadAsync(GetObjectKey(locator), offset, length, cancellationToken);
 
 			/// <inheritdoc/>
-			public async Task<BlobLocator> WriteBlobAsync(Stream stream, string? basePath = null, CancellationToken cancellationToken = default)
+			public async Task<BlobLocator> WriteBlobAsync(Stream stream, IReadOnlyCollection<BlobLocator>? imports, string? basePath = null, CancellationToken cancellationToken = default)
 			{
 				BlobLocator locator = StorageHelpers.CreateUniqueLocator(basePath);
 
 				await _store.WriteAsync(GetObjectKey(locator), stream, cancellationToken);
-				await _outer.AddBlobAsync(NamespaceId, locator, null, cancellationToken);
+				await _outer.AddBlobAsync(NamespaceId, locator, imports, null, cancellationToken);
 
 				return locator;
 			}
@@ -94,7 +94,7 @@ namespace Horde.Server.Storage
 			}
 
 			/// <inheritdoc/>
-			public async ValueTask<(BlobLocator, Uri)?> TryGetBlobWriteRedirectAsync(string? prefix = null, CancellationToken cancellationToken = default)
+			public async ValueTask<(BlobLocator, Uri)?> TryGetBlobWriteRedirectAsync(IReadOnlyCollection<BlobLocator>? imports = null, string? prefix = null, CancellationToken cancellationToken = default)
 			{
 				if (!_store.SupportsRedirects)
 				{
@@ -109,7 +109,7 @@ namespace Horde.Server.Storage
 					return null;
 				}
 
-				await _outer.AddBlobAsync(NamespaceId, locator, null, cancellationToken);
+				await _outer.AddBlobAsync(NamespaceId, locator, imports, null, cancellationToken);
 
 				return (locator, url);
 			}
@@ -202,7 +202,7 @@ namespace Horde.Server.Storage
 			}
 		}
 
-		class AliasInfo
+		internal class AliasInfo
 		{
 			[BsonElement("nam")]
 			public string Name { get; set; }
@@ -232,7 +232,7 @@ namespace Horde.Server.Storage
 			}
 		}
 
-		class BlobInfo
+		internal class BlobInfo
 		{
 			public ObjectId Id { get; set; }
 
@@ -264,7 +264,7 @@ namespace Horde.Server.Storage
 			}
 		}
 
-		class RefInfo : ISupportInitialize
+		internal class RefInfo : ISupportInitialize
 		{
 			[BsonIgnoreIfDefault]
 			public ObjectId Id { get; set; }
@@ -362,8 +362,8 @@ namespace Horde.Server.Storage
 		readonly Tracer _tracer;
 		readonly ILogger _logger;
 
-		readonly IMongoCollection<BlobInfo> _blobCollection;
-		readonly IMongoCollection<RefInfo> _refCollection;
+		internal readonly IMongoCollection<BlobInfo> _blobCollection;
+		internal readonly IMongoCollection<RefInfo> _refCollection;
 
 		readonly ITicker _blobTicker;
 		readonly ITicker _refTicker;
@@ -544,11 +544,31 @@ namespace Horde.Server.Storage
 		#region Blobs
 
 		/// <inheritdoc/>
-		async Task AddBlobAsync(NamespaceId namespaceId, BlobLocator locator, List<AliasInfo>? exports = null, CancellationToken cancellationToken = default)
+		async Task AddBlobAsync(NamespaceId namespaceId, BlobLocator locator, IReadOnlyCollection<BlobLocator>? imports = null, List<AliasInfo>? exports = null, CancellationToken cancellationToken = default)
 		{
 			ObjectId id = ObjectId.GenerateNewId(_clock.UtcNow);
 			BlobInfo blobInfo = new BlobInfo(id, namespaceId, locator);
 			blobInfo.Aliases = exports;
+
+			if (imports != null)
+			{
+				if (imports.Count == 0)
+				{
+					blobInfo.Imports = new List<ObjectId>();
+				}
+				else
+				{
+					List<string> paths = imports.Select(x => x.BaseLocator.ToString()).ToList();
+
+					FilterDefinition<BlobInfo> filter =
+						Builders<BlobInfo>.Filter.Eq(x => x.NamespaceId, namespaceId) &
+						Builders<BlobInfo>.Filter.In(x => x.Path, paths);
+
+					blobInfo.Imports = await _blobCollection.Find(filter).Project(x => x.Id).ToListAsync(cancellationToken);
+					blobInfo.Imports.Sort();
+				}
+			}
+
 			await _blobCollection.InsertOneAsync(blobInfo, new InsertOneOptions { }, cancellationToken);
 		}
 
@@ -648,6 +668,24 @@ namespace Horde.Server.Storage
 
 				importInfoIds.Add(blobInfoDoc.Id);
 			}
+
+			importInfoIds.Sort();
+
+			if (blobInfo.Imports != null && !Enumerable.SequenceEqual(importInfoIds, blobInfo.Imports))
+			{
+				List<ObjectId> missing = importInfoIds.Except(blobInfo.Imports).ToList();
+				if (missing.Count > 0)
+				{
+					_logger.LogWarning("Missing imports for blob {Locator}: {Missing}", String.Join(", ", missing.Select(x => x.ToString())));
+				}
+
+				List<ObjectId> extra = blobInfo.Imports.Except(importInfoIds).ToList();
+				if (extra.Count > 0)
+				{
+					_logger.LogWarning("Extra imports for blob {Locator}: {Extra}", String.Join(", ", extra.Select(x => x.ToString())));
+				}
+			}
+
 			await _blobCollection.UpdateOneAsync(x => x.Id == blobInfo.Id, Builders<BlobInfo>.Update.Set(x => x.Imports, importInfoIds), null, cancellationToken);
 
 			AddGcCheckRecord(blobInfo.NamespaceId, blobInfo.Id);
diff --git a/Engine/Source/Programs/Horde/Horde.Server/Tools/ToolsController.cs b/Engine/Source/Programs/Horde/Horde.Server/Tools/ToolsController.cs
--- a/Engine/Source/Programs/Horde/Horde.Server/Tools/ToolsController.cs
+++ b/Engine/Source/Programs/Horde/Horde.Server/Tools/ToolsController.cs
@@ -46,11 +46,11 @@ namespace Horde.Server.Tools
 		/// Uploads blob data for a new tool deployment.
 		/// </summary>
 		/// <param name="id">Identifier of the tool to upload</param>
-		/// <param name="file">Data to be uploaded. May be null, in which case the server may return a separate url.</param>
+		/// <param name="request">Upload request</param>
 		/// <param name="cancellationToken">Cancellation token for the operation</param>
 		[HttpPost]
 		[Route("/api/v1/tools/{id}/blobs")]
-		public async Task<ActionResult<WriteBlobResponse>> WriteBlobAsync(ToolId id, IFormFile? file, CancellationToken cancellationToken = default)
+		public async Task<ActionResult<WriteBlobResponse>> WriteBlobAsync(ToolId id, WriteBlobRequest request, CancellationToken cancellationToken = default)
 		{
 			ITool? tool = await _toolCollection.GetAsync(id, _globalConfig.Value, cancellationToken);
 			if (tool == null)
@@ -63,7 +63,7 @@ namespace Horde.Server.Tools
 			}
 
 			IStorageBackend storageBackend = _toolCollection.CreateStorageBackend(tool);
-			return await StorageController.WriteBlobAsync(storageBackend, file, cancellationToken: cancellationToken);
+			return await StorageController.WriteBlobAsync(storageBackend, request, cancellationToken);
 		}
 
 		/// <summary>
diff --git a/Engine/Source/Programs/Shared/EpicGames.Horde/Compute/AgentStorageClient.cs b/Engine/Source/Programs/Shared/EpicGames.Horde/Compute/AgentStorageClient.cs
--- a/Engine/Source/Programs/Shared/EpicGames.Horde/Compute/AgentStorageClient.cs
+++ b/Engine/Source/Programs/Shared/EpicGames.Horde/Compute/AgentStorageClient.cs
@@ -1,6 +1,7 @@
 // Copyright Epic Games, Inc. All Rights Reserved.
 
 using System;
+using System.Collections.Generic;
 using System.IO;
 using System.Threading;
 using System.Threading.Tasks;
@@ -67,13 +68,16 @@ namespace EpicGames.Horde.Compute
 		}
 
 		/// <inheritdoc/>
-		public Task<BlobLocator> WriteBlobAsync(Stream stream, string? basePath = null, CancellationToken cancellationToken = default) => throw new NotSupportedException();
+		public Task<BlobLocator> WriteBlobAsync(Stream stream, IReadOnlyCollection<BlobLocator>? imports, string? basePath = null, CancellationToken cancellationToken = default) 
+			=> throw new NotSupportedException();
 
 		/// <inheritdoc/>
-		public ValueTask<Uri?> TryGetBlobReadRedirectAsync(BlobLocator locator, CancellationToken cancellationToken = default) => default;
+		public ValueTask<Uri?> TryGetBlobReadRedirectAsync(BlobLocator locator, CancellationToken cancellationToken = default) 
+			=> default;
 
 		/// <inheritdoc/>
-		public ValueTask<(BlobLocator, Uri)?> TryGetBlobWriteRedirectAsync(string? prefix = null, CancellationToken cancellationToken = default) => default;
+		public ValueTask<(BlobLocator, Uri)?> TryGetBlobWriteRedirectAsync(IReadOnlyCollection<BlobLocator>? imports = null, string? prefix = null, CancellationToken cancellationToken = default) 
+			=> default;
 
 		#endregion
 
diff --git a/Engine/Source/Programs/Shared/EpicGames.Horde/Storage/Backends/FileStorageBackend.cs b/Engine/Source/Programs/Shared/EpicGames.Horde/Storage/Backends/FileStorageBackend.cs
--- a/Engine/Source/Programs/Shared/EpicGames.Horde/Storage/Backends/FileStorageBackend.cs
+++ b/Engine/Source/Programs/Shared/EpicGames.Horde/Storage/Backends/FileStorageBackend.cs
@@ -69,7 +69,7 @@ namespace EpicGames.Horde.Storage.Backends
 			=> _objectStore.ReadAsync(GetBlobFile(locator), offset, length, cancellationToken);
 
 		/// <inheritdoc/>
-		public async Task<BlobLocator> WriteBlobAsync(Stream stream, string? prefix = null, CancellationToken cancellationToken = default)
+		public async Task<BlobLocator> WriteBlobAsync(Stream stream, IReadOnlyCollection<BlobLocator>? imports, string? prefix = null, CancellationToken cancellationToken = default)
 		{
 			BlobLocator locator = StorageHelpers.CreateUniqueLocator(prefix);
 			await _objectStore.WriteAsync(GetBlobFile(locator), stream, cancellationToken);
@@ -80,7 +80,7 @@ namespace EpicGames.Horde.Storage.Backends
 		public ValueTask<Uri?> TryGetBlobReadRedirectAsync(BlobLocator path, CancellationToken cancellationToken = default) => default;
 
 		/// <inheritdoc/>
-		public ValueTask<(BlobLocator, Uri)?> TryGetBlobWriteRedirectAsync(string? prefix = null, CancellationToken cancellationToken = default) => default;
+		public ValueTask<(BlobLocator, Uri)?> TryGetBlobWriteRedirectAsync(IReadOnlyCollection<BlobLocator>? imports = null, string? prefix = null, CancellationToken cancellationToken = default) => default;
 
 		#endregion
 
diff --git a/Engine/Source/Programs/Shared/EpicGames.Horde/Storage/Backends/HttpStorageBackend.cs b/Engine/Source/Programs/Shared/EpicGames.Horde/Storage/Backends/HttpStorageBackend.cs
--- a/Engine/Source/Programs/Shared/EpicGames.Horde/Storage/Backends/HttpStorageBackend.cs
+++ b/Engine/Source/Programs/Shared/EpicGames.Horde/Storage/Backends/HttpStorageBackend.cs
@@ -1,6 +1,7 @@
 // Copyright Epic Games, Inc. All Rights Reserved.
 
 using System;
+using System.Collections.Generic;
 using System.Diagnostics;
 using System.IO;
 using System.Net;
@@ -136,13 +137,13 @@ namespace EpicGames.Horde.Storage.Backends
 		}
 
 		/// <inheritdoc/>
-		public async Task<BlobLocator> WriteBlobAsync(Stream stream, string? prefix = null, CancellationToken cancellationToken = default)
+		public async Task<BlobLocator> WriteBlobAsync(Stream stream, IReadOnlyCollection<BlobLocator>? imports, string? prefix = null, CancellationToken cancellationToken = default)
 		{
 			using StreamContent streamContent = new StreamContent(stream);
 
 			if (_supportsUploadRedirects)
 			{
-				WriteBlobResponse redirectResponse = await SendWriteRequestAsync(null, prefix, cancellationToken);
+				WriteBlobResponse redirectResponse = await SendWriteRequestAsync(null, imports, prefix, cancellationToken);
 				if (redirectResponse.UploadUrl != null)
 				{
 					using HttpClient uploadRedirectClient = _createUploadRedirectClient();
@@ -158,26 +159,35 @@ namespace EpicGames.Horde.Storage.Backends
 				}
 			}
 
-			WriteBlobResponse response = await SendWriteRequestAsync(streamContent, prefix, cancellationToken);
+			WriteBlobResponse response = await SendWriteRequestAsync(streamContent, imports, prefix, cancellationToken);
 			_supportsUploadRedirects = response.SupportsRedirects ?? false;
 			_logger.LogDebug("Written {Locator} (direct)", response.Blob);
 			return new BlobLocator(response.Blob);
 		}
 
-		async Task<WriteBlobResponse> SendWriteRequestAsync(StreamContent? streamContent, string? prefix = null, CancellationToken cancellationToken = default)
+		async Task<WriteBlobResponse> SendWriteRequestAsync(StreamContent? streamContent, IReadOnlyCollection<BlobLocator>? imports = null, string? prefix = null, CancellationToken cancellationToken = default)
 		{
 			using (HttpClient httpClient = _createClient())
 			{
 				using (HttpRequestMessage request = new HttpRequestMessage(HttpMethod.Post, $"{_basePath}/blobs"))
 				{
-					using StringContent stringContent = new StringContent(prefix ?? String.Empty);
-
-					MultipartFormDataContent form = new MultipartFormDataContent();
+					using MultipartFormDataContent form = new MultipartFormDataContent();
 					if (streamContent != null)
 					{
 						form.Add(streamContent, "file", "filename");
 					}
-					form.Add(stringContent, "prefix");
+					if (imports != null)
+					{
+						foreach (BlobLocator import in imports)
+						{
+							form.Add(new StringContent(import.ToString()), "import");
+						}
+						if (imports.Count == 0)
+						{
+							form.Add(new StringContent("true"), "leaf");
+						}
+					}
+					form.Add(new StringContent(prefix ?? String.Empty), "prefix");
 
 					request.Content = form;
 					using (HttpResponseMessage response = await httpClient.SendAsync(request, cancellationToken))
@@ -203,11 +213,11 @@ namespace EpicGames.Horde.Storage.Backends
 		}
 
 		/// <inheritdoc/>
-		public async ValueTask<(BlobLocator, Uri)?> TryGetBlobWriteRedirectAsync(string? prefix = null, CancellationToken cancellationToken = default)
+		public async ValueTask<(BlobLocator, Uri)?> TryGetBlobWriteRedirectAsync(IReadOnlyCollection<BlobLocator>? imports = null, string? prefix = null, CancellationToken cancellationToken = default)
 		{
 			if (_supportsUploadRedirects)
 			{
-				WriteBlobResponse redirectResponse = await SendWriteRequestAsync(null, prefix, cancellationToken);
+				WriteBlobResponse redirectResponse = await SendWriteRequestAsync(null, imports, prefix, cancellationToken);
 				if (redirectResponse.UploadUrl != null)
 				{
 					return (new BlobLocator(redirectResponse.Blob), redirectResponse.UploadUrl);
diff --git a/Engine/Source/Programs/Shared/EpicGames.Horde/Storage/Backends/MemoryStorageBackend.cs b/Engine/Source/Programs/Shared/EpicGames.Horde/Storage/Backends/MemoryStorageBackend.cs
--- a/Engine/Source/Programs/Shared/EpicGames.Horde/Storage/Backends/MemoryStorageBackend.cs
+++ b/Engine/Source/Programs/Shared/EpicGames.Horde/Storage/Backends/MemoryStorageBackend.cs
@@ -65,7 +65,7 @@ namespace EpicGames.Horde.Storage.Backends
 		}
 
 		/// <inheritdoc/>
-		public async Task<BlobLocator> WriteBlobAsync(Stream stream, string? prefix = null, CancellationToken cancellationToken = default)
+		public async Task<BlobLocator> WriteBlobAsync(Stream stream, IReadOnlyCollection<BlobLocator>? imports, string? prefix = null, CancellationToken cancellationToken = default)
 		{
 			BlobLocator locator = StorageHelpers.CreateUniqueLocator(prefix);
 			_blobs[locator] = await stream.ReadAllBytesAsync(cancellationToken);
@@ -77,7 +77,7 @@ namespace EpicGames.Horde.Storage.Backends
 			=> default;
 
 		/// <inheritdoc/>
-		public ValueTask<(BlobLocator, Uri)?> TryGetBlobWriteRedirectAsync(string? prefix = null, CancellationToken cancellationToken = default)
+		public ValueTask<(BlobLocator, Uri)?> TryGetBlobWriteRedirectAsync(IReadOnlyCollection<BlobLocator>? imports = null, string? prefix = null, CancellationToken cancellationToken = default)
 			=> default;
 
 		#endregion
diff --git a/Engine/Source/Programs/Shared/EpicGames.Horde/Storage/Bundles/V1/BundleWriter.cs b/Engine/Source/Programs/Shared/EpicGames.Horde/Storage/Bundles/V1/BundleWriter.cs
--- a/Engine/Source/Programs/Shared/EpicGames.Horde/Storage/Bundles/V1/BundleWriter.cs
+++ b/Engine/Source/Programs/Shared/EpicGames.Horde/Storage/Bundles/V1/BundleWriter.cs
@@ -436,10 +436,10 @@ namespace EpicGames.Horde.Storage.Bundles.V1
 					(BundleHeader header, List<ReadOnlyMemory<byte>> packets) = CreateBundle();
 
 					// Write the bundle to storage
-					BundleHandle[] imports = new BundleHandle[header.Imports.Count];
-					for (int idx = 0; idx < header.Imports.Count; idx++)
+					HashSet<BlobLocator> imports = new HashSet<BlobLocator>();
+					foreach (BlobLocator import in header.Imports)
 					{
-						imports[idx] = new FlushedBundleHandle(store, new BlobLocator(header.Imports[idx].Path));
+						imports.Add(import.BaseLocator);
 					}
 
 					// Create the output sequence
@@ -455,7 +455,7 @@ namespace EpicGames.Horde.Storage.Bundles.V1
 					BlobLocator locator;
 					using (ReadOnlySequenceStream stream = new ReadOnlySequenceStream(sequence.Construct()))
 					{
-						locator = await store.Backend.WriteBlobAsync(stream, basePath, CancellationToken.None);
+						locator = await store.Backend.WriteBlobAsync(stream, imports, basePath, CancellationToken.None);
 					}
 
 					traceLogger?.LogInformation("Written bundle {BundleId} as {Locator}", BundleId, locator);
diff --git a/Engine/Source/Programs/Shared/EpicGames.Horde/Storage/Bundles/V2/BundleWriter.cs b/Engine/Source/Programs/Shared/EpicGames.Horde/Storage/Bundles/V2/BundleWriter.cs
--- a/Engine/Source/Programs/Shared/EpicGames.Horde/Storage/Bundles/V2/BundleWriter.cs
+++ b/Engine/Source/Programs/Shared/EpicGames.Horde/Storage/Bundles/V2/BundleWriter.cs
@@ -67,6 +67,7 @@ namespace EpicGames.Horde.Storage.Bundles.V2
 			PendingPacketHandle? _packetHandle;
 			List<(ExportHandle, AliasInfo)>? _pendingExportAliases;
 			RefCountedMemoryWriter? _encodedPacketWriter;
+			HashSet<BlobLocator> _bundleImports = new HashSet<BlobLocator>();
 
 			/// <summary>
 			/// Object used for locking access to this bundle's state
@@ -179,6 +180,8 @@ namespace EpicGames.Horde.Storage.Bundles.V2
 					}
 				}
 
+				_bundleImports.UnionWith(_packetWriter.BundleImports);
+
 				_packetWriter.Dispose();
 				_packetWriter = null;
 
@@ -208,7 +211,7 @@ namespace EpicGames.Horde.Storage.Bundles.V2
 				FlushedBundleHandle flushedHandle;
 				using (ReadOnlySequenceStream stream = new ReadOnlySequenceStream(_encodedPacketWriter.AsSequence()))
 				{
-					BlobLocator locator = await _storageClient.Backend.WriteBlobAsync(stream, _basePath, cancellationToken);
+					BlobLocator locator = await _storageClient.Backend.WriteBlobAsync(stream, _bundleImports, _basePath, cancellationToken);
 					flushedHandle = new FlushedBundleHandle(_storageClient, locator);
 				}
 
diff --git a/Engine/Source/Programs/Shared/EpicGames.Horde/Storage/Bundles/V2/PacketWriter.cs b/Engine/Source/Programs/Shared/EpicGames.Horde/Storage/Bundles/V2/PacketWriter.cs
--- a/Engine/Source/Programs/Shared/EpicGames.Horde/Storage/Bundles/V2/PacketWriter.cs
+++ b/Engine/Source/Programs/Shared/EpicGames.Horde/Storage/Bundles/V2/PacketWriter.cs
@@ -4,6 +4,7 @@ using System;
 using System.Buffers.Binary;
 using System.Collections.Generic;
 using System.Diagnostics;
+using System.Linq;
 using System.Runtime.InteropServices;
 using EpicGames.Core;
 
@@ -31,12 +32,19 @@ namespace EpicGames.Horde.Storage.Bundles.V2
 		readonly List<object> _importHandles = new List<object>();
 		readonly Dictionary<object, int> _importMap = new Dictionary<object, int>();
 		readonly List<int> _exportOffsets = new List<int>();
+		readonly HashSet<int> _bundleImportIndices = new HashSet<int>();
 
 		/// <summary>
 		/// Current length of the packet
 		/// </summary>
 		public int Length => _length;
 
+		/// <summary>
+		/// Enumerate all imported bundle locators
+		/// </summary>
+		public IEnumerable<BlobLocator> BundleImports
+			=> _bundleImportIndices.Select(x => GetImportLocator(x));
+
 		/// <summary>
 		/// Constructor
 		/// </summary>
@@ -229,7 +237,10 @@ namespace EpicGames.Horde.Storage.Bundles.V2
 				{
 					throw new NotSupportedException("Referenced blob has not been flushed to storage yet");
 				}
-				return AddImport(-1, locator.Path, handle);
+
+				int bundleImportIdx = AddImport(-1, locator.Path, handle);
+				_bundleImportIndices.Add(bundleImportIdx);
+				return bundleImportIdx;
 			}
 			else if (handle is IBlobHandle blobHandle)
 			{
@@ -252,6 +263,26 @@ namespace EpicGames.Horde.Storage.Bundles.V2
 		public IBlobHandle GetImport(int importIdx)
 			=> (_importHandles[importIdx + PacketImport.Bias] as IBlobHandle) ?? throw new InvalidOperationException("Import is not a blob handle");
 
+		/// <summary>
+		/// Gets the import assigned to a particular index
+		/// </summary>
+		public BlobLocator GetImportLocator(int importIdx)
+		{
+			Utf8StringBuilder builder = new Utf8StringBuilder();
+			BuildImportLocator(importIdx, builder);
+			return new BlobLocator(builder.ToUtf8String());
+		}
+
+		void BuildImportLocator(int importIdx, Utf8StringBuilder builder)
+		{
+			PacketImport import = _imports[importIdx];
+			if (import.BaseIdx != -1)
+			{
+				BuildImportLocator(import.BaseIdx, builder);
+			}
+			builder.Append(import.Fragment);
+		}
+
 		/// <summary>
 		/// Gets data to write new export
 		/// </summary>
diff --git a/Engine/Source/Programs/Shared/EpicGames.Horde/Storage/Clients/KeyValueStorageClient.cs b/Engine/Source/Programs/Shared/EpicGames.Horde/Storage/Clients/KeyValueStorageClient.cs
--- a/Engine/Source/Programs/Shared/EpicGames.Horde/Storage/Clients/KeyValueStorageClient.cs
+++ b/Engine/Source/Programs/Shared/EpicGames.Horde/Storage/Clients/KeyValueStorageClient.cs
@@ -177,7 +177,7 @@ namespace EpicGames.Horde.Storage.Clients
 			byte[] encodedData = EncodedBlobData.Create(type, locators, payload);
 
 			using ReadOnlyMemoryStream encodedStream = new ReadOnlyMemoryStream(encodedData);
-			BlobLocator locator = await _backend.WriteBlobAsync(encodedStream, basePath, cancellationToken);
+			BlobLocator locator = await _backend.WriteBlobAsync(encodedStream, locators, basePath, cancellationToken);
 
 			IoHash hash = IoHash.Compute(payload);
 			return BlobRef.Create(hash, CreateBlobHandle(locator));
diff --git a/Engine/Source/Programs/Shared/EpicGames.Horde/Storage/IStorageBackend.cs b/Engine/Source/Programs/Shared/EpicGames.Horde/Storage/IStorageBackend.cs
--- a/Engine/Source/Programs/Shared/EpicGames.Horde/Storage/IStorageBackend.cs
+++ b/Engine/Source/Programs/Shared/EpicGames.Horde/Storage/IStorageBackend.cs
@@ -1,6 +1,8 @@
 // Copyright Epic Games, Inc. All Rights Reserved.
 
 using System;
+using System.Buffers;
+using System.Collections.Generic;
 using System.IO;
 using System.Text;
 using System.Threading;
@@ -44,11 +46,12 @@ namespace EpicGames.Horde.Storage
 		/// <summary>
 		/// Writes a stream to the storage backend. If the stream throws an exception during read, the write will be aborted.
 		/// </summary>
-		/// <param name="stream">Stream to write</param>
+		/// <param name="stream">Data stream</param>
+		/// <param name="imports">Imported blobs. If omitted, the backend will parse them from the stream data.</param>
 		/// <param name="prefix">Path prefix for the uploaded data</param>
 		/// <param name="cancellationToken">Cancellation token for the operation</param>
 		/// <returns>Path to the uploaded object</returns>
-		Task<BlobLocator> WriteBlobAsync(Stream stream, string? prefix = null, CancellationToken cancellationToken = default);
+		Task<BlobLocator> WriteBlobAsync(Stream stream, IReadOnlyCollection<BlobLocator>? imports = null, string? prefix = null, CancellationToken cancellationToken = default);
 
 		/// <summary>
 		/// Gets a HTTP redirect for a read request
@@ -61,10 +64,11 @@ namespace EpicGames.Horde.Storage
 		/// <summary>
 		/// Gets a HTTP redirect for a write request
 		/// </summary>
+		/// <param name="imports">Imports for this blob</param>
 		/// <param name="prefix">Prefix for the uploaded data</param>
 		/// <param name="cancellationToken">Cancellation token for the operation</param>
 		/// <returns>Path for retrieval, and URI to upload the data to</returns>
-		ValueTask<(BlobLocator, Uri)?> TryGetBlobWriteRedirectAsync(string? prefix = null, CancellationToken cancellationToken = default);
+		ValueTask<(BlobLocator, Uri)?> TryGetBlobWriteRedirectAsync(IReadOnlyCollection<BlobLocator>? imports = null, string? prefix = null, CancellationToken cancellationToken = default);
 
 		#endregion
 
@@ -209,13 +213,14 @@ namespace EpicGames.Horde.Storage
 		/// </summary>
 		/// <param name="storageBackend">Backend to read from</param>
 		/// <param name="data">Data to be written</param>
+		/// <param name="imports"></param>
 		/// <param name="prefix">Prefix for the uploaded data</param>
 		/// <param name="cancellationToken">Cancellation token for the operation</param>
-		public static async Task<BlobLocator> WriteBytesAsync(this IStorageBackend storageBackend, ReadOnlyMemory<byte> data, string? prefix = null, CancellationToken cancellationToken = default)
+		public static async Task<BlobLocator> WriteBytesAsync(this IStorageBackend storageBackend, ReadOnlyMemory<byte> data, IReadOnlyCollection<BlobLocator>? imports = null, string? prefix = null, CancellationToken cancellationToken = default)
 		{
 			using (ReadOnlyMemoryStream stream = new ReadOnlyMemoryStream(data))
 			{
-				return await storageBackend.WriteBlobAsync(stream, prefix, cancellationToken);
+				return await storageBackend.WriteBlobAsync(stream, imports, prefix, cancellationToken);
 			}
 		}
 	}
diff --git a/Engine/Source/Programs/Shared/EpicGames.Horde/Storage/StorageBackendCache.cs b/Engine/Source/Programs/Shared/EpicGames.Horde/Storage/StorageBackendCache.cs
--- a/Engine/Source/Programs/Shared/EpicGames.Horde/Storage/StorageBackendCache.cs
+++ b/Engine/Source/Programs/Shared/EpicGames.Horde/Storage/StorageBackendCache.cs
@@ -130,9 +130,14 @@ namespace EpicGames.Horde.Storage
 				return storageObject.Slice(offset, length);
 			}
 
-			public Task<BlobLocator> WriteBlobAsync(Stream stream, string? prefix = null, CancellationToken cancellationToken = default) => _inner.WriteBlobAsync(stream, prefix, cancellationToken);
-			public ValueTask<Uri?> TryGetBlobReadRedirectAsync(BlobLocator locator, CancellationToken cancellationToken = default) => _inner.TryGetBlobReadRedirectAsync(locator, cancellationToken);
-			public ValueTask<(BlobLocator, Uri)?> TryGetBlobWriteRedirectAsync(string? prefix = null, CancellationToken cancellationToken = default) => _inner.TryGetBlobWriteRedirectAsync(prefix, cancellationToken);
+			public Task<BlobLocator> WriteBlobAsync(Stream stream, IReadOnlyCollection<BlobLocator>? imports, string? prefix = null, CancellationToken cancellationToken = default) 
+				=> _inner.WriteBlobAsync(stream, imports, prefix, cancellationToken);
+
+			public ValueTask<Uri?> TryGetBlobReadRedirectAsync(BlobLocator locator, CancellationToken cancellationToken = default) 
+				=> _inner.TryGetBlobReadRedirectAsync(locator, cancellationToken);
+
+			public ValueTask<(BlobLocator, Uri)?> TryGetBlobWriteRedirectAsync(IReadOnlyCollection<BlobLocator>? imports = null, string? prefix = null, CancellationToken cancellationToken = default) 
+				=> _inner.TryGetBlobWriteRedirectAsync(imports, prefix, cancellationToken);
 
 			public void GetStats(StorageStats stats)
 			{
