Horde: Fix some static analysis warnings.

diff --git a/Engine/Source/Programs/Horde/Horde.Server.Tests/Accounts/AccountControllerTests.cs b/Engine/Source/Programs/Horde/Horde.Server.Tests/Accounts/AccountControllerTests.cs
--- a/Engine/Source/Programs/Horde/Horde.Server.Tests/Accounts/AccountControllerTests.cs
+++ b/Engine/Source/Programs/Horde/Horde.Server.Tests/Accounts/AccountControllerTests.cs
@@ -117,7 +117,7 @@ public class AccountControllerTest : IAsyncDisposable
 		Assert.IsTrue(content.Contains("myClaimType2") && content.Contains("myClaimValue2"));
 	}
 
-	private async Task<HttpResponseMessage> LoginAsync(HttpClient httpClient, string username, string password, string? redirectUrl = null)
+	private static async Task<HttpResponseMessage> LoginAsync(HttpClient httpClient, string username, string password, string? redirectUrl = null)
 	{
 		using StringContent sc = new($"username={username}&password={password}", Encoding.UTF8, "application/x-www-form-urlencoded");
 		string url = "account/login/horde";
diff --git a/Engine/Source/Programs/Horde/Horde.Server.Tests/Storage/StorageControllerTests.cs b/Engine/Source/Programs/Horde/Horde.Server.Tests/Storage/StorageControllerTests.cs
--- a/Engine/Source/Programs/Horde/Horde.Server.Tests/Storage/StorageControllerTests.cs
+++ b/Engine/Source/Programs/Horde/Horde.Server.Tests/Storage/StorageControllerTests.cs
@@ -1,20 +1,16 @@
 ﻿// Copyright Epic Games, Inc. All Rights Reserved.
 
 using System;
-using System.Collections.Generic;
 using System.Linq;
-using System.Text;
 using System.Threading.Tasks;
 using EpicGames.Core;
 using EpicGames.Horde.Storage;
 using EpicGames.Horde.Storage.Backends;
 using EpicGames.Horde.Storage.Bundles;
-using Horde.Server.Server;
 using Horde.Server.Storage;
 using Microsoft.Extensions.DependencyInjection;
 using Microsoft.Extensions.Logging;
 using Microsoft.VisualStudio.TestTools.UnitTesting;
-using MongoDB.Bson;
 using MongoDB.Driver;
 
 namespace Horde.Server.Tests.Storage
@@ -48,7 +44,7 @@ namespace Horde.Server.Tests.Storage
 			BlobLocator locator2 = blobRef2.GetLocator();
 
 			StorageService storageService = ServiceProvider.GetRequiredService<StorageService>();
-			IMongoCollection<StorageService.BlobInfo> blobCollection = storageService._blobCollection;
+			IMongoCollection<StorageService.BlobInfo> blobCollection = storageService.BlobCollection;
 
 			StorageService.BlobInfo blob2 = await blobCollection.Find(x => x.Path == locator2.BaseLocator.ToString()).FirstAsync();
 			Assert.IsNotNull(blob2.Imports);
diff --git a/Engine/Source/Programs/Horde/Horde.Server/Logs/LogsController.cs b/Engine/Source/Programs/Horde/Horde.Server/Logs/LogsController.cs
--- a/Engine/Source/Programs/Horde/Horde.Server/Logs/LogsController.cs
+++ b/Engine/Source/Programs/Horde/Horde.Server/Logs/LogsController.cs
@@ -21,7 +21,6 @@ using Horde.Server.Server;
 using Horde.Server.Storage;
 using Horde.Server.Utilities;
 using Microsoft.AspNetCore.Authorization;
-using Microsoft.AspNetCore.Http;
 using Microsoft.AspNetCore.Mvc;
 using Microsoft.Extensions.Options;
 using MongoDB.Bson;
diff --git a/Engine/Source/Programs/Horde/Horde.Server/Storage/StorageService.cs b/Engine/Source/Programs/Horde/Horde.Server/Storage/StorageService.cs
--- a/Engine/Source/Programs/Horde/Horde.Server/Storage/StorageService.cs
+++ b/Engine/Source/Programs/Horde/Horde.Server/Storage/StorageService.cs
@@ -362,8 +362,8 @@ namespace Horde.Server.Storage
 		readonly Tracer _tracer;
 		readonly ILogger _logger;
 
-		internal readonly IMongoCollection<BlobInfo> _blobCollection;
-		internal readonly IMongoCollection<RefInfo> _refCollection;
+		readonly IMongoCollection<BlobInfo> _blobCollection;
+		readonly IMongoCollection<RefInfo> _refCollection;
 
 		readonly ITicker _blobTicker;
 		readonly ITicker _refTicker;
@@ -376,6 +376,8 @@ namespace Horde.Server.Storage
 		string? _lastConfigRevision;
 		State? _lastState;
 
+		internal IMongoCollection<BlobInfo> BlobCollection => _blobCollection;
+
 		/// <summary>
 		/// Constructor
 		/// </summary>
@@ -676,13 +678,13 @@ namespace Horde.Server.Storage
 				List<ObjectId> missing = importInfoIds.Except(blobInfo.Imports).ToList();
 				if (missing.Count > 0)
 				{
-					_logger.LogWarning("Missing imports for blob {Locator}: {Missing}", String.Join(", ", missing.Select(x => x.ToString())));
+					_logger.LogWarning("Missing imports for blob {Locator}: {Missing}", blobInfo.Path, String.Join(", ", missing.Select(x => x.ToString())));
 				}
 
 				List<ObjectId> extra = blobInfo.Imports.Except(importInfoIds).ToList();
 				if (extra.Count > 0)
 				{
-					_logger.LogWarning("Extra imports for blob {Locator}: {Extra}", String.Join(", ", extra.Select(x => x.ToString())));
+					_logger.LogWarning("Extra imports for blob {Locator}: {Extra}", blobInfo.Path, String.Join(", ", extra.Select(x => x.ToString())));
 				}
 			}
 
diff --git a/Engine/Source/Programs/Shared/EpicGames.Horde/Storage/Backends/HttpStorageBackend.cs b/Engine/Source/Programs/Shared/EpicGames.Horde/Storage/Backends/HttpStorageBackend.cs
--- a/Engine/Source/Programs/Shared/EpicGames.Horde/Storage/Backends/HttpStorageBackend.cs
+++ b/Engine/Source/Programs/Shared/EpicGames.Horde/Storage/Backends/HttpStorageBackend.cs
@@ -176,6 +176,7 @@ namespace EpicGames.Horde.Storage.Backends
 					{
 						form.Add(streamContent, "file", "filename");
 					}
+#pragma warning disable CA2000 // Disposed by form
 					if (imports != null)
 					{
 						foreach (BlobLocator import in imports)
@@ -188,6 +189,7 @@ namespace EpicGames.Horde.Storage.Backends
 						}
 					}
 					form.Add(new StringContent(prefix ?? String.Empty), "prefix");
+#pragma warning restore CA2000
 
 					request.Content = form;
 					using (HttpResponseMessage response = await httpClient.SendAsync(request, cancellationToken))
diff --git a/Engine/Source/Programs/Shared/EpicGames.Horde/Storage/Bundles/V2/BundleWriter.cs b/Engine/Source/Programs/Shared/EpicGames.Horde/Storage/Bundles/V2/BundleWriter.cs
--- a/Engine/Source/Programs/Shared/EpicGames.Horde/Storage/Bundles/V2/BundleWriter.cs
+++ b/Engine/Source/Programs/Shared/EpicGames.Horde/Storage/Bundles/V2/BundleWriter.cs
@@ -67,7 +67,7 @@ namespace EpicGames.Horde.Storage.Bundles.V2
 			PendingPacketHandle? _packetHandle;
 			List<(ExportHandle, AliasInfo)>? _pendingExportAliases;
 			RefCountedMemoryWriter? _encodedPacketWriter;
-			HashSet<BlobLocator> _bundleImports = new HashSet<BlobLocator>();
+			readonly HashSet<BlobLocator> _bundleImports = new HashSet<BlobLocator>();
 
 			/// <summary>
 			/// Object used for locking access to this bundle's state
diff --git a/Engine/Source/Programs/Shared/EpicGames.Horde/Storage/IStorageBackend.cs b/Engine/Source/Programs/Shared/EpicGames.Horde/Storage/IStorageBackend.cs
--- a/Engine/Source/Programs/Shared/EpicGames.Horde/Storage/IStorageBackend.cs
+++ b/Engine/Source/Programs/Shared/EpicGames.Horde/Storage/IStorageBackend.cs
@@ -1,7 +1,6 @@
 // Copyright Epic Games, Inc. All Rights Reserved.
 
 using System;
-using System.Buffers;
 using System.Collections.Generic;
 using System.IO;
 using System.Text;
@@ -9,6 +8,8 @@ using System.Threading;
 using System.Threading.Tasks;
 using EpicGames.Core;
 
+#pragma warning disable CA1716
+
 namespace EpicGames.Horde.Storage
 {
 	/// <summary>
