Horde: Move JobStepState, JobStepOutcome and Priority enums into native C# code (rather than protobuf enums)

#rnx

diff --git a/Engine/Source/Programs/Horde/Horde.Agent.Tests/Helper.cs b/Engine/Source/Programs/Horde/Horde.Agent.Tests/Helper.cs
--- a/Engine/Source/Programs/Horde/Horde.Agent.Tests/Helper.cs
+++ b/Engine/Source/Programs/Horde/Horde.Agent.Tests/Helper.cs
@@ -4,13 +4,13 @@ using System;
 using System.Collections.Generic;
 using System.Threading;
 using System.Threading.Tasks;
+using EpicGames.Horde.Jobs;
 using Google.Protobuf.WellKnownTypes;
 using Grpc.Core;
 using Grpc.Net.Client;
 using Horde.Agent.Execution;
 using Horde.Agent.Utility;
 using Horde.Common.Rpc;
-using HordeCommon;
 using HordeCommon.Rpc;
 using HordeCommon.Rpc.Messages;
 using Microsoft.Extensions.Logging;
diff --git a/Engine/Source/Programs/Horde/Horde.Agent.Tests/WorkerServiceTest.cs b/Engine/Source/Programs/Horde/Horde.Agent.Tests/WorkerServiceTest.cs
--- a/Engine/Source/Programs/Horde/Horde.Agent.Tests/WorkerServiceTest.cs
+++ b/Engine/Source/Programs/Horde/Horde.Agent.Tests/WorkerServiceTest.cs
@@ -26,7 +26,6 @@ using Horde.Agent.Leases.Handlers;
 using Horde.Agent.Services;
 using Horde.Agent.Utility;
 using Horde.Common.Rpc;
-using HordeCommon;
 using HordeCommon.Rpc;
 using HordeCommon.Rpc.Messages;
 using HordeCommon.Rpc.Tasks;
@@ -197,12 +196,12 @@ namespace Horde.Agent.Tests
 
 			Assert.AreEqual(LeaseOutcome.Success, outcome);
 			Assert.AreEqual(3, client.UpdateStepRequests.Count);
-			Assert.AreEqual(JobStepOutcome.Success, client.UpdateStepRequests[0].Outcome);
-			Assert.AreEqual(JobStepState.Completed, client.UpdateStepRequests[0].State);
-			Assert.AreEqual(JobStepOutcome.Failure, client.UpdateStepRequests[1].Outcome);
-			Assert.AreEqual(JobStepState.Aborted, client.UpdateStepRequests[1].State);
-			Assert.AreEqual(JobStepOutcome.Success, client.UpdateStepRequests[2].Outcome);
-			Assert.AreEqual(JobStepState.Completed, client.UpdateStepRequests[2].State);
+			Assert.AreEqual(JobStepOutcome.Success, (JobStepOutcome)client.UpdateStepRequests[0].Outcome);
+			Assert.AreEqual(JobStepState.Completed, (JobStepState)client.UpdateStepRequests[0].State);
+			Assert.AreEqual(JobStepOutcome.Failure, (JobStepOutcome)client.UpdateStepRequests[1].Outcome);
+			Assert.AreEqual(JobStepState.Aborted, (JobStepState)client.UpdateStepRequests[1].State);
+			Assert.AreEqual(JobStepOutcome.Success, (JobStepOutcome)client.UpdateStepRequests[2].Outcome);
+			Assert.AreEqual(JobStepState.Completed, (JobStepState)client.UpdateStepRequests[2].State);
 		}
 
 		[TestMethod]
diff --git a/Engine/Source/Programs/Horde/Horde.Agent/Commands/Utilities/BuildGraphCommand.cs b/Engine/Source/Programs/Horde/Horde.Agent/Commands/Utilities/BuildGraphCommand.cs
--- a/Engine/Source/Programs/Horde/Horde.Agent/Commands/Utilities/BuildGraphCommand.cs
+++ b/Engine/Source/Programs/Horde/Horde.Agent/Commands/Utilities/BuildGraphCommand.cs
@@ -5,9 +5,9 @@ using System.Text;
 using System.Text.Json;
 using System.Text.Json.Serialization;
 using EpicGames.Core;
+using EpicGames.Horde.Jobs;
 using Horde.Agent.Parser;
 using Horde.Agent.Utility;
-using HordeCommon;
 using HordeCommon.Rpc;
 using Microsoft.Extensions.Logging;
 
diff --git a/Engine/Source/Programs/Horde/Horde.Agent/Execution/JobExecutor.cs b/Engine/Source/Programs/Horde/Horde.Agent/Execution/JobExecutor.cs
--- a/Engine/Source/Programs/Horde/Horde.Agent/Execution/JobExecutor.cs
+++ b/Engine/Source/Programs/Horde/Horde.Agent/Execution/JobExecutor.cs
@@ -669,7 +669,7 @@ namespace Horde.Agent.Execution
 					}
 					createNode.RunEarly = exportedNode.RunEarly;
 					createNode.Warnings = exportedNode.Warnings;
-					createNode.Priority = Priority.Normal;
+					createNode.Priority = (int)Priority.Normal;
 					createNode.Annotations.Add(exportedNode.Annotations);
 					createGroup.Nodes.Add(createNode);
 				}
diff --git a/Engine/Source/Programs/Horde/Horde.Agent/Execution/TestExecutor.cs b/Engine/Source/Programs/Horde/Horde.Agent/Execution/TestExecutor.cs
--- a/Engine/Source/Programs/Horde/Horde.Agent/Execution/TestExecutor.cs
+++ b/Engine/Source/Programs/Horde/Horde.Agent/Execution/TestExecutor.cs
@@ -2,10 +2,10 @@
 
 using System.Diagnostics.CodeAnalysis;
 using EpicGames.Core;
+using EpicGames.Horde.Jobs;
 using Horde.Agent.Parser;
 using Horde.Agent.Utility;
 using Horde.Common.Rpc;
-using HordeCommon;
 using HordeCommon.Rpc;
 using HordeCommon.Rpc.Messages;
 using Microsoft.Extensions.Logging;
diff --git a/Engine/Source/Programs/Horde/Horde.Agent/Horde.Agent.csproj b/Engine/Source/Programs/Horde/Horde.Agent/Horde.Agent.csproj
--- a/Engine/Source/Programs/Horde/Horde.Agent/Horde.Agent.csproj
+++ b/Engine/Source/Programs/Horde/Horde.Agent/Horde.Agent.csproj
@@ -157,13 +157,10 @@
     <Protobuf Include="..\Horde.Server\Protos\horde\common\agent_status.proto" ProtoRoot="..\Horde.Server\Protos" Access="Public" />
     <Protobuf Include="..\Horde.Server\Protos\horde\common\event_severity.proto" ProtoRoot="..\Horde.Server\Protos" Access="Public" />
     <Protobuf Include="..\Horde.Server\Protos\horde\common\io_hash.proto" ProtoRoot="..\Horde.Server\Protos" Access="Public" />
-    <Protobuf Include="..\Horde.Server\Protos\horde\common\jobstep_outcome.proto" ProtoRoot="..\Horde.Server\Protos" Access="Public" />
-    <Protobuf Include="..\Horde.Server\Protos\horde\common\jobstep_state.proto" ProtoRoot="..\Horde.Server\Protos" Access="Public" />
     <Protobuf Include="..\Horde.Server\Protos\horde\common\label_change.proto" ProtoRoot="..\Horde.Server\Protos" Access="Public" />
     <Protobuf Include="..\Horde.Server\Protos\horde\common\lease_outcome.proto" ProtoRoot="..\Horde.Server\Protos" Access="Public" />
     <Protobuf Include="..\Horde.Server\Protos\horde\common\lease_state.proto" ProtoRoot="..\Horde.Server\Protos" Access="Public" />
     <Protobuf Include="..\Horde.Server\Protos\horde\common\namespace_id.proto" ProtoRoot="..\Horde.Server\Protos" Access="Public" />
-    <Protobuf Include="..\Horde.Server\Protos\horde\common\priority.proto" ProtoRoot="..\Horde.Server\Protos" Access="Public" />
     <Protobuf Include="..\Horde.Server\Protos\horde\messages\agent_capabilities.proto" ProtoRoot="..\Horde.Server\Protos" Access="Public" />
     <Protobuf Include="..\Horde.Server\Protos\horde\messages\agent_requirements.proto" ProtoRoot="..\Horde.Server\Protos" Access="Public" />
     <Protobuf Include="..\Horde.Server\Protos\horde\messages\agent_workspace.proto" ProtoRoot="..\Horde.Server\Protos" Access="Public" />
diff --git a/Engine/Source/Programs/Horde/Horde.Agent/Leases/Handlers/JobHandler.cs b/Engine/Source/Programs/Horde/Horde.Agent/Leases/Handlers/JobHandler.cs
--- a/Engine/Source/Programs/Horde/Horde.Agent/Leases/Handlers/JobHandler.cs
+++ b/Engine/Source/Programs/Horde/Horde.Agent/Leases/Handlers/JobHandler.cs
@@ -14,7 +14,6 @@ using Horde.Agent.Execution;
 using Horde.Agent.Services;
 using Horde.Agent.Utility;
 using Horde.Common.Rpc;
-using HordeCommon;
 using HordeCommon.Rpc;
 using HordeCommon.Rpc.Messages;
 using HordeCommon.Rpc.Tasks;
diff --git a/Engine/Source/Programs/Horde/Horde.Agent/Utilities/HordeRpc.cs b/Engine/Source/Programs/Horde/Horde.Agent/Utilities/HordeRpc.cs
--- a/Engine/Source/Programs/Horde/Horde.Agent/Utilities/HordeRpc.cs
+++ b/Engine/Source/Programs/Horde/Horde.Agent/Utilities/HordeRpc.cs
@@ -124,8 +124,8 @@ namespace HordeCommon.Rpc
 			JobId = jobId.ToString();
 			BatchId = batchId.ToString();
 			StepId = stepId.ToString();
-			State = state;
-			Outcome = outcome;
+			State = (int)state;
+			Outcome = (int)outcome;
 		}
 	}
 
@@ -143,8 +143,8 @@ namespace HordeCommon.Rpc
 	{
 		public GetStepResponse(JobStepOutcome outcome, JobStepState state, bool abortRequested)
 		{
-			Outcome = outcome;
-			State = state;
+			Outcome = (int)outcome;
+			State = (int)state;
 			AbortRequested = abortRequested;
 		}
 	}
diff --git a/Engine/Source/Programs/Horde/Horde.Agent/Utility/JsonRpcLogSink.cs b/Engine/Source/Programs/Horde/Horde.Agent/Utility/JsonRpcLogSink.cs
--- a/Engine/Source/Programs/Horde/Horde.Agent/Utility/JsonRpcLogSink.cs
+++ b/Engine/Source/Programs/Horde/Horde.Agent/Utility/JsonRpcLogSink.cs
@@ -8,7 +8,6 @@ using EpicGames.Horde.Storage;
 using Google.Protobuf;
 using Grpc.Core;
 using Horde.Common.Rpc;
-using HordeCommon;
 using HordeCommon.Rpc;
 using Microsoft.Extensions.Logging;
 using static Horde.Common.Rpc.LogRpc;
diff --git a/Engine/Source/Programs/Horde/Horde.Agent/Utility/ServerLogger.cs b/Engine/Source/Programs/Horde/Horde.Agent/Utility/ServerLogger.cs
--- a/Engine/Source/Programs/Horde/Horde.Agent/Utility/ServerLogger.cs
+++ b/Engine/Source/Programs/Horde/Horde.Agent/Utility/ServerLogger.cs
@@ -3,6 +3,7 @@
 using System.Text;
 using System.Threading.Channels;
 using EpicGames.Core;
+using EpicGames.Horde.Jobs;
 using EpicGames.Horde.Logs;
 using Google.Protobuf;
 using HordeCommon;
diff --git a/Engine/Source/Programs/Horde/Horde.Server.Tests/Devices/DeviceServiceTests.cs b/Engine/Source/Programs/Horde/Horde.Server.Tests/Devices/DeviceServiceTests.cs
--- a/Engine/Source/Programs/Horde/Horde.Server.Tests/Devices/DeviceServiceTests.cs
+++ b/Engine/Source/Programs/Horde/Horde.Server.Tests/Devices/DeviceServiceTests.cs
@@ -18,7 +18,6 @@ using Horde.Server.Jobs.Templates;
 using Horde.Server.Server;
 using Horde.Server.Users;
 using Horde.Server.Utilities;
-using HordeCommon;
 using Microsoft.AspNetCore.Http;
 using Microsoft.AspNetCore.Mvc;
 using Microsoft.VisualStudio.TestTools.UnitTesting;
diff --git a/Engine/Source/Programs/Horde/Horde.Server.Tests/Jobs/Bisect/BisectTests.cs b/Engine/Source/Programs/Horde/Horde.Server.Tests/Jobs/Bisect/BisectTests.cs
--- a/Engine/Source/Programs/Horde/Horde.Server.Tests/Jobs/Bisect/BisectTests.cs
+++ b/Engine/Source/Programs/Horde/Horde.Server.Tests/Jobs/Bisect/BisectTests.cs
@@ -7,6 +7,7 @@ using System.Security.Claims;
 using System.Threading;
 using System.Threading.Tasks;
 using EpicGames.Core;
+using EpicGames.Horde.Jobs;
 using EpicGames.Horde.Jobs.Templates;
 using EpicGames.Horde.Projects;
 using EpicGames.Horde.Streams;
@@ -19,7 +20,6 @@ using Horde.Server.Projects;
 using Horde.Server.Streams;
 using Horde.Server.Users;
 using Horde.Server.Utilities;
-using HordeCommon;
 using Microsoft.AspNetCore.Http;
 using Microsoft.AspNetCore.Mvc;
 using Microsoft.Extensions.DependencyInjection;
diff --git a/Engine/Source/Programs/Horde/Horde.Server.Tests/Jobs/GetLabelStatesTests.cs b/Engine/Source/Programs/Horde/Horde.Server.Tests/Jobs/GetLabelStatesTests.cs
--- a/Engine/Source/Programs/Horde/Horde.Server.Tests/Jobs/GetLabelStatesTests.cs
+++ b/Engine/Source/Programs/Horde/Horde.Server.Tests/Jobs/GetLabelStatesTests.cs
@@ -4,7 +4,6 @@ using System.Collections.Generic;
 using EpicGames.Horde.Jobs;
 using Horde.Server.Jobs;
 using Horde.Server.Jobs.Graphs;
-using HordeCommon;
 using Microsoft.VisualStudio.TestTools.UnitTesting;
 using Moq;
 
diff --git a/Engine/Source/Programs/Horde/Horde.Server.Tests/Jobs/JobCollectionTests.cs b/Engine/Source/Programs/Horde/Horde.Server.Tests/Jobs/JobCollectionTests.cs
--- a/Engine/Source/Programs/Horde/Horde.Server.Tests/Jobs/JobCollectionTests.cs
+++ b/Engine/Source/Programs/Horde/Horde.Server.Tests/Jobs/JobCollectionTests.cs
@@ -21,7 +21,6 @@ using Horde.Server.Jobs.Graphs;
 using Horde.Server.Jobs.Templates;
 using Horde.Server.Logs;
 using Horde.Server.Utilities;
-using HordeCommon;
 using Microsoft.VisualStudio.TestTools.UnitTesting;
 using Moq;
 
diff --git a/Engine/Source/Programs/Horde/Horde.Server.Tests/Jobs/JobServiceTests.cs b/Engine/Source/Programs/Horde/Horde.Server.Tests/Jobs/JobServiceTests.cs
--- a/Engine/Source/Programs/Horde/Horde.Server.Tests/Jobs/JobServiceTests.cs
+++ b/Engine/Source/Programs/Horde/Horde.Server.Tests/Jobs/JobServiceTests.cs
@@ -8,6 +8,7 @@ using System.Threading.Tasks;
 using EpicGames.Core;
 using EpicGames.Horde.Agents;
 using EpicGames.Horde.Agents.Pools;
+using EpicGames.Horde.Jobs;
 using EpicGames.Horde.Jobs.Templates;
 using EpicGames.Horde.Streams;
 using Horde.Server.Agents;
@@ -20,7 +21,6 @@ using Horde.Server.Projects;
 using Horde.Server.Server;
 using Horde.Server.Streams;
 using Horde.Server.Users;
-using HordeCommon;
 using Microsoft.CodeAnalysis;
 using Microsoft.VisualStudio.TestTools.UnitTesting;
 
diff --git a/Engine/Source/Programs/Horde/Horde.Server.Tests/Jobs/JobTaskSourceTests.cs b/Engine/Source/Programs/Horde/Horde.Server.Tests/Jobs/JobTaskSourceTests.cs
--- a/Engine/Source/Programs/Horde/Horde.Server.Tests/Jobs/JobTaskSourceTests.cs
+++ b/Engine/Source/Programs/Horde/Horde.Server.Tests/Jobs/JobTaskSourceTests.cs
@@ -7,6 +7,7 @@ using System.Threading;
 using System.Threading.Tasks;
 using EpicGames.Horde.Agents;
 using EpicGames.Horde.Agents.Pools;
+using EpicGames.Horde.Jobs;
 using EpicGames.Horde.Users;
 using Horde.Server.Agents;
 using Horde.Server.Agents.Pools;
@@ -14,7 +15,6 @@ using Horde.Server.Jobs;
 using Horde.Server.Jobs.Graphs;
 using Horde.Server.Streams;
 using Horde.Server.Utilities;
-using HordeCommon;
 using Microsoft.VisualStudio.TestTools.UnitTesting;
 
 namespace Horde.Server.Tests.Jobs
diff --git a/Engine/Source/Programs/Horde/Horde.Server.Tests/Jobs/Schedules/SchedulerTests.cs b/Engine/Source/Programs/Horde/Horde.Server.Tests/Jobs/Schedules/SchedulerTests.cs
--- a/Engine/Source/Programs/Horde/Horde.Server.Tests/Jobs/Schedules/SchedulerTests.cs
+++ b/Engine/Source/Programs/Horde/Horde.Server.Tests/Jobs/Schedules/SchedulerTests.cs
@@ -18,7 +18,6 @@ using Horde.Server.Projects;
 using Horde.Server.Server;
 using Horde.Server.Streams;
 using Horde.Server.Users;
-using HordeCommon;
 using Microsoft.Extensions.DependencyInjection;
 using Microsoft.VisualStudio.TestTools.UnitTesting;
 
diff --git a/Engine/Source/Programs/Horde/Horde.Server.Tests/Notifications/NotificationServiceTests.cs b/Engine/Source/Programs/Horde/Horde.Server.Tests/Notifications/NotificationServiceTests.cs
--- a/Engine/Source/Programs/Horde/Horde.Server.Tests/Notifications/NotificationServiceTests.cs
+++ b/Engine/Source/Programs/Horde/Horde.Server.Tests/Notifications/NotificationServiceTests.cs
@@ -16,7 +16,6 @@ using Horde.Server.Logs;
 using Horde.Server.Notifications;
 using Horde.Server.Streams;
 using Horde.Server.Users;
-using HordeCommon;
 using Microsoft.Extensions.Caching.Memory;
 using Microsoft.Extensions.DependencyInjection;
 using Microsoft.VisualStudio.TestTools.UnitTesting;
diff --git a/Engine/Source/Programs/Horde/Horde.Server/Acls/AclService.cs b/Engine/Source/Programs/Horde/Horde.Server/Acls/AclService.cs
--- a/Engine/Source/Programs/Horde/Horde.Server/Acls/AclService.cs
+++ b/Engine/Source/Programs/Horde/Horde.Server/Acls/AclService.cs
@@ -136,7 +136,7 @@ namespace Horde.Server.Acls
 							IJobStepBatch? batch = job.Batches.FirstOrDefault(x => x.LeaseId == leaseId);
 							if (batch != null && batch.State == JobStepBatchState.Running)
 							{
-								IJobStep? step = batch.Steps.FirstOrDefault(x => x.State == HordeCommon.JobStepState.Running);
+								IJobStep? step = batch.Steps.FirstOrDefault(x => x.State == JobStepState.Running);
 								if (step != null)
 								{
 									return (job, step);
diff --git a/Engine/Source/Programs/Horde/Horde.Server/Artifacts/ArtifactsController.cs b/Engine/Source/Programs/Horde/Horde.Server/Artifacts/ArtifactsController.cs
--- a/Engine/Source/Programs/Horde/Horde.Server/Artifacts/ArtifactsController.cs
+++ b/Engine/Source/Programs/Horde/Horde.Server/Artifacts/ArtifactsController.cs
@@ -114,7 +114,7 @@ namespace Horde.Server.Artifacts
 				List<string> keys = new List<string>(request.Keys);
 				keys.Add(job.GetArtifactKey());
 
-				IJobStep? step = batch.Steps.FirstOrDefault(x => x.State == HordeCommon.JobStepState.Running);
+				IJobStep? step = batch.Steps.FirstOrDefault(x => x.State == JobStepState.Running);
 				if (step != null)
 				{
 					keys.Add(job.GetArtifactKey(step));
diff --git a/Engine/Source/Programs/Horde/Horde.Server/Horde.Server.csproj b/Engine/Source/Programs/Horde/Horde.Server/Horde.Server.csproj
--- a/Engine/Source/Programs/Horde/Horde.Server/Horde.Server.csproj
+++ b/Engine/Source/Programs/Horde/Horde.Server/Horde.Server.csproj
@@ -106,13 +106,10 @@
     <Protobuf Include="Protos\horde\common\agent_status.proto" ProtoRoot="Protos" />
     <Protobuf Include="Protos\horde\common\event_severity.proto" ProtoRoot="Protos" />
     <Protobuf Include="Protos\horde\common\io_hash.proto" ProtoRoot="Protos" />
-    <Protobuf Include="Protos\horde\common\jobstep_outcome.proto" ProtoRoot="Protos" />
-    <Protobuf Include="Protos\horde\common\jobstep_state.proto" ProtoRoot="Protos" />
     <Protobuf Include="Protos\horde\common\label_change.proto" ProtoRoot="Protos" />
     <Protobuf Include="Protos\horde\common\lease_outcome.proto" ProtoRoot="Protos" />
     <Protobuf Include="Protos\horde\common\lease_state.proto" ProtoRoot="Protos" />
     <Protobuf Include="Protos\horde\common\namespace_id.proto" ProtoRoot="Protos" />
-    <Protobuf Include="Protos\horde\common\priority.proto" ProtoRoot="Protos" />
     <Protobuf Include="Protos\horde\messages\agent_capabilities.proto" ProtoRoot="Protos" />
     <Protobuf Include="Protos\horde\messages\agent_requirements.proto" ProtoRoot="Protos" />
     <Protobuf Include="Protos\horde\messages\agent_workspace.proto" ProtoRoot="Protos" />
diff --git a/Engine/Source/Programs/Horde/Horde.Server/Issues/IssueReportService.cs b/Engine/Source/Programs/Horde/Horde.Server/Issues/IssueReportService.cs
--- a/Engine/Source/Programs/Horde/Horde.Server/Issues/IssueReportService.cs
+++ b/Engine/Source/Programs/Horde/Horde.Server/Issues/IssueReportService.cs
@@ -7,6 +7,7 @@ using System.Threading;
 using System.Threading.Tasks;
 using EpicGames.Core;
 using EpicGames.Horde.Issues;
+using EpicGames.Horde.Jobs;
 using EpicGames.Horde.Jobs.Templates;
 using EpicGames.Horde.Streams;
 using Horde.Server.Jobs;
diff --git a/Engine/Source/Programs/Horde/Horde.Server/Issues/IssuesController.cs b/Engine/Source/Programs/Horde/Horde.Server/Issues/IssuesController.cs
--- a/Engine/Source/Programs/Horde/Horde.Server/Issues/IssuesController.cs
+++ b/Engine/Source/Programs/Horde/Horde.Server/Issues/IssuesController.cs
@@ -21,7 +21,6 @@ using Horde.Server.Server;
 using Horde.Server.Streams;
 using Horde.Server.Users;
 using Horde.Server.Utilities;
-using HordeCommon;
 using Microsoft.AspNetCore.Authorization;
 using Microsoft.AspNetCore.Mvc;
 using Microsoft.Extensions.Logging;
diff --git a/Engine/Source/Programs/Horde/Horde.Server/Jobs/Bisect/BisectController.cs b/Engine/Source/Programs/Horde/Horde.Server/Jobs/Bisect/BisectController.cs
--- a/Engine/Source/Programs/Horde/Horde.Server/Jobs/Bisect/BisectController.cs
+++ b/Engine/Source/Programs/Horde/Horde.Server/Jobs/Bisect/BisectController.cs
@@ -16,7 +16,6 @@ using Horde.Server.Server;
 using Horde.Server.Streams;
 using Horde.Server.Users;
 using Horde.Server.Utilities;
-using HordeCommon;
 using Microsoft.AspNetCore.Authorization;
 using Microsoft.AspNetCore.Mvc;
 using Microsoft.Extensions.Logging;
diff --git a/Engine/Source/Programs/Horde/Horde.Server/Jobs/Bisect/BisectService.cs b/Engine/Source/Programs/Horde/Horde.Server/Jobs/Bisect/BisectService.cs
--- a/Engine/Source/Programs/Horde/Horde.Server/Jobs/Bisect/BisectService.cs
+++ b/Engine/Source/Programs/Horde/Horde.Server/Jobs/Bisect/BisectService.cs
@@ -6,6 +6,7 @@ using System.Linq;
 using System.Threading;
 using System.Threading.Tasks;
 using EpicGames.Core;
+using EpicGames.Horde.Jobs;
 using Horde.Server.Jobs.Graphs;
 using Horde.Server.Jobs.Templates;
 using Horde.Server.Perforce;
diff --git a/Engine/Source/Programs/Horde/Horde.Server/Jobs/Bisect/BisectTaskCollection.cs b/Engine/Source/Programs/Horde/Horde.Server/Jobs/Bisect/BisectTaskCollection.cs
--- a/Engine/Source/Programs/Horde/Horde.Server/Jobs/Bisect/BisectTaskCollection.cs
+++ b/Engine/Source/Programs/Horde/Horde.Server/Jobs/Bisect/BisectTaskCollection.cs
@@ -13,7 +13,6 @@ using EpicGames.Horde.Users;
 using Horde.Server.Perforce;
 using Horde.Server.Server;
 using Horde.Server.Utilities;
-using HordeCommon;
 using MongoDB.Bson;
 using MongoDB.Bson.Serialization.Attributes;
 using MongoDB.Driver;
diff --git a/Engine/Source/Programs/Horde/Horde.Server/Jobs/Bisect/IBisectTask.cs b/Engine/Source/Programs/Horde/Horde.Server/Jobs/Bisect/IBisectTask.cs
--- a/Engine/Source/Programs/Horde/Horde.Server/Jobs/Bisect/IBisectTask.cs
+++ b/Engine/Source/Programs/Horde/Horde.Server/Jobs/Bisect/IBisectTask.cs
@@ -7,7 +7,6 @@ using EpicGames.Horde.Jobs.Templates;
 using EpicGames.Horde.Streams;
 using EpicGames.Horde.Users;
 using Horde.Server.Perforce;
-using HordeCommon;
 
 namespace Horde.Server.Jobs.Bisect
 {
diff --git a/Engine/Source/Programs/Horde/Horde.Server/Jobs/Bisect/IBisectTaskCollection.cs b/Engine/Source/Programs/Horde/Horde.Server/Jobs/Bisect/IBisectTaskCollection.cs
--- a/Engine/Source/Programs/Horde/Horde.Server/Jobs/Bisect/IBisectTaskCollection.cs
+++ b/Engine/Source/Programs/Horde/Horde.Server/Jobs/Bisect/IBisectTaskCollection.cs
@@ -9,7 +9,6 @@ using EpicGames.Horde.Jobs.Bisect;
 using EpicGames.Horde.Users;
 using Horde.Server.Jobs.Graphs;
 using Horde.Server.Perforce;
-using HordeCommon;
 using Microsoft.Extensions.Logging;
 
 namespace Horde.Server.Jobs.Bisect
diff --git a/Engine/Source/Programs/Horde/Horde.Server/Jobs/Graphs/GraphCollection.cs b/Engine/Source/Programs/Horde/Horde.Server/Jobs/Graphs/GraphCollection.cs
--- a/Engine/Source/Programs/Horde/Horde.Server/Jobs/Graphs/GraphCollection.cs
+++ b/Engine/Source/Programs/Horde/Horde.Server/Jobs/Graphs/GraphCollection.cs
@@ -8,6 +8,7 @@ using System.Threading;
 using System.Threading.Tasks;
 using EpicGames.Core;
 using EpicGames.Horde.Artifacts;
+using EpicGames.Horde.Jobs;
 using Horde.Server.Jobs.Templates;
 using Horde.Server.Server;
 using HordeCommon;
diff --git a/Engine/Source/Programs/Horde/Horde.Server/Jobs/Graphs/GraphMessages.cs b/Engine/Source/Programs/Horde/Horde.Server/Jobs/Graphs/GraphMessages.cs
--- a/Engine/Source/Programs/Horde/Horde.Server/Jobs/Graphs/GraphMessages.cs
+++ b/Engine/Source/Programs/Horde/Horde.Server/Jobs/Graphs/GraphMessages.cs
@@ -3,7 +3,7 @@
 using System.Collections.Generic;
 using System.Linq;
 using EpicGames.Core;
-using HordeCommon;
+using EpicGames.Horde.Jobs;
 
 namespace Horde.Server.Jobs.Graphs
 {
diff --git a/Engine/Source/Programs/Horde/Horde.Server/Jobs/Graphs/IGraph.cs b/Engine/Source/Programs/Horde/Horde.Server/Jobs/Graphs/IGraph.cs
--- a/Engine/Source/Programs/Horde/Horde.Server/Jobs/Graphs/IGraph.cs
+++ b/Engine/Source/Programs/Horde/Horde.Server/Jobs/Graphs/IGraph.cs
@@ -7,6 +7,7 @@ using System.Diagnostics.CodeAnalysis;
 using System.Linq;
 using EpicGames.Core;
 using EpicGames.Horde.Artifacts;
+using EpicGames.Horde.Jobs;
 using HordeCommon;
 using MongoDB.Bson.Serialization.Attributes;
 
diff --git a/Engine/Source/Programs/Horde/Horde.Server/Jobs/IJob.cs b/Engine/Source/Programs/Horde/Horde.Server/Jobs/IJob.cs
--- a/Engine/Source/Programs/Horde/Horde.Server/Jobs/IJob.cs
+++ b/Engine/Source/Programs/Horde/Horde.Server/Jobs/IJob.cs
@@ -23,7 +23,6 @@ using Horde.Server.Jobs.Timing;
 using Horde.Server.Server;
 using Horde.Server.Streams;
 using Horde.Server.Ugs;
-using HordeCommon;
 using HordeCommon.Rpc.Tasks;
 using MongoDB.Bson;
 using MongoDB.Driver;
diff --git a/Engine/Source/Programs/Horde/Horde.Server/Jobs/IJobCollection.cs b/Engine/Source/Programs/Horde/Horde.Server/Jobs/IJobCollection.cs
--- a/Engine/Source/Programs/Horde/Horde.Server/Jobs/IJobCollection.cs
+++ b/Engine/Source/Programs/Horde/Horde.Server/Jobs/IJobCollection.cs
@@ -18,7 +18,6 @@ using EpicGames.Horde.Users;
 using Horde.Server.Acls;
 using Horde.Server.Jobs.Graphs;
 using Horde.Server.Streams;
-using HordeCommon;
 using HordeCommon.Rpc.Tasks;
 using MongoDB.Bson;
 
diff --git a/Engine/Source/Programs/Horde/Horde.Server/Jobs/IJobStepRef.cs b/Engine/Source/Programs/Horde/Horde.Server/Jobs/IJobStepRef.cs
--- a/Engine/Source/Programs/Horde/Horde.Server/Jobs/IJobStepRef.cs
+++ b/Engine/Source/Programs/Horde/Horde.Server/Jobs/IJobStepRef.cs
@@ -9,7 +9,6 @@ using EpicGames.Horde.Jobs;
 using EpicGames.Horde.Jobs.Templates;
 using EpicGames.Horde.Logs;
 using EpicGames.Horde.Streams;
-using HordeCommon;
 using MongoDB.Bson.Serialization;
 using MongoDB.Bson.Serialization.Attributes;
 
diff --git a/Engine/Source/Programs/Horde/Horde.Server/Jobs/IJobStepRefCollection.cs b/Engine/Source/Programs/Horde/Horde.Server/Jobs/IJobStepRefCollection.cs
--- a/Engine/Source/Programs/Horde/Horde.Server/Jobs/IJobStepRefCollection.cs
+++ b/Engine/Source/Programs/Horde/Horde.Server/Jobs/IJobStepRefCollection.cs
@@ -11,7 +11,6 @@ using EpicGames.Horde.Jobs.Templates;
 using EpicGames.Horde.Logs;
 using EpicGames.Horde.Streams;
 using Horde.Server.Jobs.Graphs;
-using HordeCommon;
 using Microsoft.Extensions.Logging;
 
 namespace Horde.Server.Jobs
diff --git a/Engine/Source/Programs/Horde/Horde.Server/Jobs/JobCollection.cs b/Engine/Source/Programs/Horde/Horde.Server/Jobs/JobCollection.cs
--- a/Engine/Source/Programs/Horde/Horde.Server/Jobs/JobCollection.cs
+++ b/Engine/Source/Programs/Horde/Horde.Server/Jobs/JobCollection.cs
@@ -25,7 +25,6 @@ using Horde.Server.Server;
 using Horde.Server.Streams;
 using Horde.Server.Telemetry;
 using Horde.Server.Utilities;
-using HordeCommon;
 using HordeCommon.Rpc.Tasks;
 using Microsoft.Extensions.Logging;
 using Microsoft.Extensions.Options;
@@ -352,7 +351,7 @@ namespace Horde.Server.Jobs
 				PreflightDescription = options.PreflightDescription;
 				StartedByUserId = options.StartedByUserId;
 				StartedByBisectTaskId = options.StartedByBisectTaskId;
-				Priority = options.Priority ?? HordeCommon.Priority.Normal;
+				Priority = options.Priority ?? Priority.Normal;
 				AutoSubmit = options.AutoSubmit ?? false;
 				UpdateIssues = options.UpdateIssues ?? (options.StartedByUserId == null && (options.PreflightChange == 0 || options.PreflightChange == null));
 				PromoteIssuesByDefault = options.PromoteIssuesByDefault ?? false;
diff --git a/Engine/Source/Programs/Horde/Horde.Server/Jobs/JobMessages.cs b/Engine/Source/Programs/Horde/Horde.Server/Jobs/JobMessages.cs
--- a/Engine/Source/Programs/Horde/Horde.Server/Jobs/JobMessages.cs
+++ b/Engine/Source/Programs/Horde/Horde.Server/Jobs/JobMessages.cs
@@ -14,7 +14,7 @@ using EpicGames.Horde.Streams;
 using EpicGames.Horde.Users;
 using Horde.Server.Jobs.Graphs;
 using Horde.Server.Streams;
-using HordeCommon;
+//using HordeCommon;
 using HordeCommon.Rpc.Tasks;
 
 namespace Horde.Server.Jobs
diff --git a/Engine/Source/Programs/Horde/Horde.Server/Jobs/JobRpcService.cs b/Engine/Source/Programs/Horde/Horde.Server/Jobs/JobRpcService.cs
--- a/Engine/Source/Programs/Horde/Horde.Server/Jobs/JobRpcService.cs
+++ b/Engine/Source/Programs/Horde/Horde.Server/Jobs/JobRpcService.cs
@@ -31,7 +31,6 @@ using Horde.Server.Logs;
 using Horde.Server.Server;
 using Horde.Server.Streams;
 using Horde.Server.Utilities;
-using HordeCommon;
 using HordeCommon.Rpc;
 using Microsoft.AspNetCore.Authorization;
 using Microsoft.Extensions.Logging;
@@ -769,7 +768,7 @@ namespace Horde.Server.Jobs
 				error = JobStepError.TimedOut;
 			}
 
-			await _jobService.UpdateStepAsync(job, batch.Id, JobStepId.Parse(request.StepId), streamConfig, request.State, request.Outcome, error, null, null, null, null, null);
+			await _jobService.UpdateStepAsync(job, batch.Id, JobStepId.Parse(request.StepId), streamConfig, (JobStepState)request.State, (JobStepOutcome)request.Outcome, error, null, null, null, null, null);
 			return new Empty();
 		}
 
@@ -790,7 +789,7 @@ namespace Horde.Server.Jobs
 				throw new StructuredRpcException(StatusCode.NotFound, "Unable to find step {JobId}:{BatchId}:{StepId}", job.Id, batch.Id, stepId);
 			}
 
-			return new RpcGetStepResponse { Outcome = step.Outcome, State = step.State, AbortRequested = step.AbortRequested || step.HasTimedOut(_clock.UtcNow) };
+			return new RpcGetStepResponse { Outcome = (int)step.Outcome, State = (int)step.State, AbortRequested = step.AbortRequested || step.HasTimedOut(_clock.UtcNow) };
 		}
 
 		/// <summary>
@@ -824,7 +823,7 @@ namespace Horde.Server.Jobs
 					List<NewNode> newNodes = new List<NewNode>();
 					foreach (CreateNodeRequest node in group.Nodes)
 					{
-						NewNode newNode = new NewNode(node.Name, node.Inputs.ToList(), node.Outputs.ToList(), node.InputDependencies.ToList(), node.OrderDependencies.ToList(), node.Priority, node.AllowRetry, node.RunEarly, node.Warnings, new Dictionary<string, string>(node.Credentials), new Dictionary<string, string>(node.Properties), new NodeAnnotations(node.Annotations));
+						NewNode newNode = new NewNode(node.Name, node.Inputs.ToList(), node.Outputs.ToList(), node.InputDependencies.ToList(), node.OrderDependencies.ToList(), (Priority)node.Priority, node.AllowRetry, node.RunEarly, node.Warnings, new Dictionary<string, string>(node.Credentials), new Dictionary<string, string>(node.Properties), new NodeAnnotations(node.Annotations));
 						newNodes.Add(newNode);
 					}
 					newGroups.Add(new NewGroup(group.AgentType, newNodes));
diff --git a/Engine/Source/Programs/Horde/Horde.Server/Jobs/JobService.cs b/Engine/Source/Programs/Horde/Horde.Server/Jobs/JobService.cs
--- a/Engine/Source/Programs/Horde/Horde.Server/Jobs/JobService.cs
+++ b/Engine/Source/Programs/Horde/Horde.Server/Jobs/JobService.cs
@@ -30,7 +30,6 @@ using Horde.Server.Server;
 using Horde.Server.Streams;
 using Horde.Server.Users;
 using Horde.Server.Utilities;
-using HordeCommon;
 using Microsoft.Extensions.Logging;
 using MongoDB.Bson;
 using OpenTelemetry.Trace;
diff --git a/Engine/Source/Programs/Horde/Horde.Server/Jobs/JobStepRefCollection.cs b/Engine/Source/Programs/Horde/Horde.Server/Jobs/JobStepRefCollection.cs
--- a/Engine/Source/Programs/Horde/Horde.Server/Jobs/JobStepRefCollection.cs
+++ b/Engine/Source/Programs/Horde/Horde.Server/Jobs/JobStepRefCollection.cs
@@ -15,7 +15,6 @@ using Horde.Server.Server;
 using Horde.Server.Streams;
 using Horde.Server.Telemetry;
 using Horde.Server.Utilities;
-using HordeCommon;
 using Microsoft.Extensions.Options;
 using MongoDB.Bson.Serialization.Attributes;
 using MongoDB.Driver;
diff --git a/Engine/Source/Programs/Horde/Horde.Server/Jobs/JobsController.cs b/Engine/Source/Programs/Horde/Horde.Server/Jobs/JobsController.cs
--- a/Engine/Source/Programs/Horde/Horde.Server/Jobs/JobsController.cs
+++ b/Engine/Source/Programs/Horde/Horde.Server/Jobs/JobsController.cs
@@ -25,7 +25,6 @@ using Horde.Server.Server;
 using Horde.Server.Streams;
 using Horde.Server.Users;
 using Horde.Server.Utilities;
-using HordeCommon;
 using Microsoft.AspNetCore.Authorization;
 using Microsoft.AspNetCore.Mvc;
 using Microsoft.Extensions.Options;
diff --git a/Engine/Source/Programs/Horde/Horde.Server/Jobs/Templates/ITemplate.cs b/Engine/Source/Programs/Horde/Horde.Server/Jobs/Templates/ITemplate.cs
--- a/Engine/Source/Programs/Horde/Horde.Server/Jobs/Templates/ITemplate.cs
+++ b/Engine/Source/Programs/Horde/Horde.Server/Jobs/Templates/ITemplate.cs
@@ -3,7 +3,7 @@
 using System;
 using System.Collections.Generic;
 using EpicGames.Core;
-using HordeCommon;
+using EpicGames.Horde.Jobs;
 using MongoDB.Bson.Serialization.Attributes;
 
 namespace Horde.Server.Jobs.Templates
diff --git a/Engine/Source/Programs/Horde/Horde.Server/Jobs/Templates/TemplateCollection.cs b/Engine/Source/Programs/Horde/Horde.Server/Jobs/Templates/TemplateCollection.cs
--- a/Engine/Source/Programs/Horde/Horde.Server/Jobs/Templates/TemplateCollection.cs
+++ b/Engine/Source/Programs/Horde/Horde.Server/Jobs/Templates/TemplateCollection.cs
@@ -5,8 +5,8 @@ using System.Collections.Generic;
 using System.Linq;
 using System.Threading.Tasks;
 using EpicGames.Core;
+using EpicGames.Horde.Jobs;
 using Horde.Server.Server;
-using HordeCommon;
 using Microsoft.Extensions.Caching.Memory;
 using MongoDB.Bson;
 using MongoDB.Bson.Serialization.Attributes;
diff --git a/Engine/Source/Programs/Horde/Horde.Server/Jobs/Templates/TemplateConfig.cs b/Engine/Source/Programs/Horde/Horde.Server/Jobs/Templates/TemplateConfig.cs
--- a/Engine/Source/Programs/Horde/Horde.Server/Jobs/Templates/TemplateConfig.cs
+++ b/Engine/Source/Programs/Horde/Horde.Server/Jobs/Templates/TemplateConfig.cs
@@ -6,9 +6,9 @@ using System.ComponentModel.DataAnnotations;
 using System.IO;
 using System.Text.Json.Serialization;
 using EpicGames.Core;
+using EpicGames.Horde.Jobs;
 using Horde.Server.Configuration;
 using Horde.Server.Streams;
-using HordeCommon;
 using HordeCommon.Rpc.Tasks;
 
 namespace Horde.Server.Jobs.Templates
diff --git a/Engine/Source/Programs/Horde/Horde.Server/Jobs/Templates/TemplateMessages.cs b/Engine/Source/Programs/Horde/Horde.Server/Jobs/Templates/TemplateMessages.cs
--- a/Engine/Source/Programs/Horde/Horde.Server/Jobs/Templates/TemplateMessages.cs
+++ b/Engine/Source/Programs/Horde/Horde.Server/Jobs/Templates/TemplateMessages.cs
@@ -2,7 +2,7 @@
 
 using System.Collections.Generic;
 using EpicGames.Core;
-using HordeCommon;
+using EpicGames.Horde.Jobs;
 
 namespace Horde.Server.Jobs.Templates
 {
diff --git a/Engine/Source/Programs/Horde/Horde.Server/Notifications/INotificationSink.cs b/Engine/Source/Programs/Horde/Horde.Server/Notifications/INotificationSink.cs
--- a/Engine/Source/Programs/Horde/Horde.Server/Notifications/INotificationSink.cs
+++ b/Engine/Source/Programs/Horde/Horde.Server/Notifications/INotificationSink.cs
@@ -4,6 +4,7 @@ using System;
 using System.Collections.Generic;
 using System.Threading;
 using System.Threading.Tasks;
+using EpicGames.Horde.Jobs;
 using Horde.Server.Agents;
 using Horde.Server.Devices;
 using Horde.Server.Issues;
@@ -12,7 +13,6 @@ using Horde.Server.Jobs.Graphs;
 using Horde.Server.Logs;
 using Horde.Server.Streams;
 using Horde.Server.Users;
-using HordeCommon;
 
 namespace Horde.Server.Notifications
 {
diff --git a/Engine/Source/Programs/Horde/Horde.Server/Notifications/SubscriptionCollection.cs b/Engine/Source/Programs/Horde/Horde.Server/Notifications/SubscriptionCollection.cs
--- a/Engine/Source/Programs/Horde/Horde.Server/Notifications/SubscriptionCollection.cs
+++ b/Engine/Source/Programs/Horde/Horde.Server/Notifications/SubscriptionCollection.cs
@@ -7,13 +7,13 @@ using System.Linq;
 using System.Text;
 using System.Threading.Tasks;
 using EpicGames.Core;
+using EpicGames.Horde.Jobs;
 using EpicGames.Horde.Jobs.Templates;
 using EpicGames.Horde.Streams;
 using EpicGames.Horde.Users;
 using Horde.Server.Jobs;
 using Horde.Server.Server;
 using Horde.Server.Utilities;
-using HordeCommon;
 using MongoDB.Bson.Serialization.Attributes;
 using MongoDB.Driver;
 
diff --git a/Engine/Source/Programs/Horde/Horde.Server/Notifications/SubscriptionMessages.cs b/Engine/Source/Programs/Horde/Horde.Server/Notifications/SubscriptionMessages.cs
--- a/Engine/Source/Programs/Horde/Horde.Server/Notifications/SubscriptionMessages.cs
+++ b/Engine/Source/Programs/Horde/Horde.Server/Notifications/SubscriptionMessages.cs
@@ -3,10 +3,10 @@
 using System;
 using System.ComponentModel.DataAnnotations;
 using EpicGames.Core;
+using EpicGames.Horde.Jobs;
 using EpicGames.Horde.Jobs.Templates;
 using EpicGames.Horde.Streams;
 using Horde.Server.Jobs;
-using HordeCommon;
 
 namespace Horde.Server.Notifications
 {
diff --git a/Engine/Source/Programs/Horde/Horde.Server/Streams/StreamConfig.cs b/Engine/Source/Programs/Horde/Horde.Server/Streams/StreamConfig.cs
--- a/Engine/Source/Programs/Horde/Horde.Server/Streams/StreamConfig.cs
+++ b/Engine/Source/Programs/Horde/Horde.Server/Streams/StreamConfig.cs
@@ -16,6 +16,7 @@ using EpicGames.Horde.Acls;
 using EpicGames.Horde.Agents.Pools;
 using EpicGames.Horde.Common;
 using EpicGames.Horde.Issues;
+using EpicGames.Horde.Jobs;
 using EpicGames.Horde.Jobs.Templates;
 using EpicGames.Horde.Replicators;
 using EpicGames.Horde.Streams;
@@ -30,7 +31,6 @@ using Horde.Server.Projects;
 using Horde.Server.Replicators;
 using Horde.Server.Server;
 using Horde.Server.Utilities;
-using HordeCommon;
 using HordeCommon.Rpc.Tasks;
 
 namespace Horde.Server.Streams
diff --git a/Engine/Source/Programs/Shared/EpicGames.Horde/Jobs/JobStepOutcome.cs b/Engine/Source/Programs/Shared/EpicGames.Horde/Jobs/JobStepOutcome.cs
new file mode 100644
--- /dev/null
+++ b/Engine/Source/Programs/Shared/EpicGames.Horde/Jobs/JobStepOutcome.cs
@@ -0,0 +1,30 @@
+// Copyright Epic Games, Inc. All Rights Reserved.
+
+namespace EpicGames.Horde.Jobs
+{
+	/// <summary>
+	/// Outcome for a jobstep
+	/// </summary>
+	public enum JobStepOutcome
+	{
+		/// <summary>
+		/// Outcome is not known
+		/// </summary>
+		Unspecified = 0,
+
+		/// <summary>
+		/// Step failed
+		/// </summary>
+		Failure = 1,
+
+		/// <summary>
+		/// Step completed with warnings
+		/// </summary>
+		Warnings = 2,
+
+		/// <summary>
+		/// Step succeeded
+		/// </summary>
+		Success = 3
+	}
+}
diff --git a/Engine/Source/Programs/Shared/EpicGames.Horde/Jobs/JobStepState.cs b/Engine/Source/Programs/Shared/EpicGames.Horde/Jobs/JobStepState.cs
new file mode 100644
--- /dev/null
+++ b/Engine/Source/Programs/Shared/EpicGames.Horde/Jobs/JobStepState.cs
@@ -0,0 +1,45 @@
+// Copyright Epic Games, Inc. All Rights Reserved.
+
+namespace EpicGames.Horde.Jobs
+{
+	/// <summary>
+	/// State of a job step
+	/// </summary>
+	public enum JobStepState
+	{
+		/// <summary>
+		/// Unspecified
+		/// </summary>
+		Unspecified = 0,
+
+		/// <summary>
+		/// Waiting for dependencies of this step to complete (or paused)
+		/// </summary>
+		Waiting = 1,
+
+		/// <summary>
+		/// Ready to run, but has not been scheduled yet
+		/// </summary>
+		Ready = 2,
+
+		/// <summary>
+		/// Dependencies of this step failed, so it cannot be executed
+		/// </summary>
+		Skipped = 3,
+
+		/// <summary>
+		/// There is an active instance of this step running
+		/// </summary>
+		Running = 4,
+
+		/// <summary>
+		/// This step has been run
+		/// </summary>
+		Completed = 5,
+
+		/// <summary>
+		/// This step started to execute, but was aborted
+		/// </summary>
+		Aborted = 6
+	}
+}
diff --git a/Engine/Source/Programs/Shared/EpicGames.Horde/Jobs/Priority.cs b/Engine/Source/Programs/Shared/EpicGames.Horde/Jobs/Priority.cs
new file mode 100644
--- /dev/null
+++ b/Engine/Source/Programs/Shared/EpicGames.Horde/Jobs/Priority.cs
@@ -0,0 +1,45 @@
+﻿// Copyright Epic Games, Inc. All Rights Reserved.
+
+namespace EpicGames.Horde.Jobs
+{
+	/// <summary>
+	/// Priority of a job or step
+	/// </summary>
+	public enum Priority
+	{
+		/// <summary>
+		/// Not specified
+		/// </summary>
+		Unspecified = 0,
+
+		/// <summary>
+		/// Lowest priority
+		/// </summary>
+		Lowest = 1,
+
+		/// <summary>
+		/// Below normal priority
+		/// </summary>
+		BelowNormal = 2,
+
+		/// <summary>
+		/// Normal priority
+		/// </summary>
+		Normal = 3,
+
+		/// <summary>
+		/// Above normal priority
+		/// </summary>
+		AboveNormal = 4,
+
+		/// <summary>
+		/// High priority
+		/// </summary>
+		High = 5,
+
+		/// <summary>
+		/// Highest priority
+		/// </summary>
+		Highest = 6
+	}
+}
