Horde: Add a native JobOptions class, so we can use definitions in shared libraries without adding a gRPC dependency.

diff --git a/Engine/Source/Programs/Horde/Horde.Agent.Tests/Execution/WorkspaceExecutorTest.cs b/Engine/Source/Programs/Horde/Horde.Agent.Tests/Execution/WorkspaceExecutorTest.cs
--- a/Engine/Source/Programs/Horde/Horde.Agent.Tests/Execution/WorkspaceExecutorTest.cs
+++ b/Engine/Source/Programs/Horde/Horde.Agent.Tests/Execution/WorkspaceExecutorTest.cs
@@ -53,7 +53,7 @@ public sealed class WorkspaceExecutorTest : IAsyncDisposable
 		_workspace.SetFile(1, "foo/bar/baz.h", "baz");
 
 		BeginBatchResponse batch = new BeginBatchResponse { Change = 1 };
-		JobExecutorOptions executorOptions = new JobExecutorOptions(_session, null!, _jobId, _batchId, batch, null!, new JobOptions());
+		JobExecutorOptions executorOptions = new JobExecutorOptions(_session, null!, _jobId, _batchId, batch, null!, new RpcJobOptions());
 		_executor = new(executorOptions, _workspace, null, NullLogger.Instance);
 	}
 
@@ -79,7 +79,7 @@ public sealed class WorkspaceExecutorTest : IAsyncDisposable
 	public async Task RegularAndAutoSdkWorkspaceAsync()
 	{
 		BeginBatchResponse batch = new BeginBatchResponse { Change = 1 };
-		JobExecutorOptions executorOptions = new JobExecutorOptions(_session, null!, _jobId, _batchId, batch, null!, new JobOptions());
+		JobExecutorOptions executorOptions = new JobExecutorOptions(_session, null!, _jobId, _batchId, batch, null!, new RpcJobOptions());
 		using WorkspaceExecutor executor = new(executorOptions, _workspace, _autoSdkWorkspace, NullLogger.Instance);
 
 		await executor.InitializeAsync(_logger, CancellationToken.None);
@@ -93,7 +93,7 @@ public sealed class WorkspaceExecutorTest : IAsyncDisposable
 	public async Task EnvVarsAsync()
 	{
 		BeginBatchResponse batch = new BeginBatchResponse { Change = 1, StreamName = "//UE5/Main" };
-		JobExecutorOptions executorOptions = new JobExecutorOptions(_session, null!, _jobId, _batchId, batch, null!, new JobOptions());
+		JobExecutorOptions executorOptions = new JobExecutorOptions(_session, null!, _jobId, _batchId, batch, null!, new RpcJobOptions());
 		using WorkspaceExecutor executor = new(executorOptions, _workspace, _autoSdkWorkspace, NullLogger.Instance);
 		await executor.InitializeAsync(_logger, CancellationToken.None);
 
@@ -120,7 +120,7 @@ public sealed class WorkspaceExecutorTest : IAsyncDisposable
 		_workspace.SetFile(1000, "New/Feature/Foo.cs", "foo");
 
 		BeginBatchResponse batch = new BeginBatchResponse { Change = 1, PreflightChange = 1000 };
-		JobExecutorOptions executorOptions = new JobExecutorOptions(_session, null!, preflightJobId, _batchId, batch, null!, new JobOptions());
+		JobExecutorOptions executorOptions = new JobExecutorOptions(_session, null!, preflightJobId, _batchId, batch, null!, new RpcJobOptions());
 		using WorkspaceExecutor executor = new(executorOptions, _workspace, null, NullLogger.Instance);
 
 		await executor.InitializeAsync(_logger, CancellationToken.None);
@@ -137,7 +137,7 @@ public sealed class WorkspaceExecutorTest : IAsyncDisposable
 
 		_server.AddJob(noChangeJobId, _streamId, 0, 0);
 		BeginBatchResponse batch = new BeginBatchResponse { };
-		JobExecutorOptions executorOptions = new JobExecutorOptions(_session, null!, noChangeJobId, _batchId, batch, null!, new JobOptions());
+		JobExecutorOptions executorOptions = new JobExecutorOptions(_session, null!, noChangeJobId, _batchId, batch, null!, new RpcJobOptions());
 		using WorkspaceExecutor executor = new(executorOptions, _workspace, null, NullLogger.Instance);
 		await Assert.ThrowsExceptionAsync<WorkspaceMaterializationException>(() => executor.InitializeAsync(_logger, CancellationToken.None));
 	}
diff --git a/Engine/Source/Programs/Horde/Horde.Agent.Tests/WorkerServiceTest.cs b/Engine/Source/Programs/Horde/Horde.Agent.Tests/WorkerServiceTest.cs
--- a/Engine/Source/Programs/Horde/Horde.Agent.Tests/WorkerServiceTest.cs
+++ b/Engine/Source/Programs/Horde/Horde.Agent.Tests/WorkerServiceTest.cs
@@ -161,7 +161,7 @@ namespace Horde.Agent.Tests
 			executeJobTask.BatchId = _batchId.ToString();
 			executeJobTask.LogId = _logId.ToString();
 			executeJobTask.JobName = "jobName1";
-			executeJobTask.JobOptions = new JobOptions { Executor = SimpleTestExecutor.Name };
+			executeJobTask.JobOptions = new RpcJobOptions { Executor = SimpleTestExecutor.Name };
 			executeJobTask.AutoSdkWorkspace = new AgentWorkspace();
 			executeJobTask.Workspace = new AgentWorkspace();
 
diff --git a/Engine/Source/Programs/Horde/Horde.Agent/Execution/JobExecutor.cs b/Engine/Source/Programs/Horde/Horde.Agent/Execution/JobExecutor.cs
--- a/Engine/Source/Programs/Horde/Horde.Agent/Execution/JobExecutor.cs
+++ b/Engine/Source/Programs/Horde/Horde.Agent/Execution/JobExecutor.cs
@@ -66,9 +66,9 @@ namespace Horde.Agent.Execution
 		public JobStepBatchId BatchId { get; }
 		public BeginBatchResponse Batch { get; }
 		public string Token { get; }
-		public JobOptions JobOptions { get; }
+		public RpcJobOptions JobOptions { get; }
 
-		public JobExecutorOptions(ISession session, HttpStorageClientFactory storageFactory, JobId jobId, JobStepBatchId batchId, BeginBatchResponse batch, string token, JobOptions jobOptions)
+		public JobExecutorOptions(ISession session, HttpStorageClientFactory storageFactory, JobId jobId, JobStepBatchId batchId, BeginBatchResponse batch, string token, RpcJobOptions jobOptions)
 		{
 			Session = session;
 			StorageFactory = storageFactory;
@@ -231,7 +231,7 @@ namespace Horde.Agent.Execution
 
 		protected ISession Session { get; }
 		protected HttpStorageClientFactory StorageFactory { get; }
-		protected JobOptions JobOptions { get; }
+		protected RpcJobOptions JobOptions { get; }
 
 		protected IRpcConnection RpcConnection => Session.RpcConnection;
 		protected Dictionary<string, string> _remapAgentTypes = new Dictionary<string, string>(StringComparer.OrdinalIgnoreCase);
diff --git a/Engine/Source/Programs/Horde/Horde.Agent/Leases/Handlers/JobHandler.cs b/Engine/Source/Programs/Horde/Horde.Agent/Leases/Handlers/JobHandler.cs
--- a/Engine/Source/Programs/Horde/Horde.Agent/Leases/Handlers/JobHandler.cs
+++ b/Engine/Source/Programs/Horde/Horde.Agent/Leases/Handlers/JobHandler.cs
@@ -24,7 +24,7 @@ using OpenTracing.Util;
 
 namespace Horde.Agent.Leases.Handlers
 {
-	record class JobTaskInfo(string JobName, JobId JobId, JobStepBatchId BatchId, JobOptions JobOptions, string Token, LogId LogId, AgentWorkspace Workspace, AgentWorkspace? AutoSdkWorkspace);
+	record class JobTaskInfo(string JobName, JobId JobId, JobStepBatchId BatchId, RpcJobOptions JobOptions, string Token, LogId LogId, AgentWorkspace Workspace, AgentWorkspace? AutoSdkWorkspace);
 
 	class JobHandler : LeaseHandler<ExecuteJobTask>
 	{
@@ -79,7 +79,7 @@ namespace Horde.Agent.Leases.Handlers
 				CurrentJobId = executeTask.JobId;
 				CurrentBatchId = executeTask.BatchId;
 
-				executeTask.JobOptions ??= new JobOptions();
+				executeTask.JobOptions ??= new RpcJobOptions();
 
 				if (executeTask.JobOptions.RunInSeparateProcess ?? false)
 				{
@@ -155,7 +155,7 @@ namespace Horde.Agent.Leases.Handlers
 		internal async Task<LeaseResult> ExecuteInternalAsync(ISession session, LeaseId leaseId, JobTaskInfo executeTask, ILogger localLogger, CancellationToken cancellationToken)
 		{
 			// Create a storage client for this session
-			JobOptions jobOptions = executeTask.JobOptions;
+			RpcJobOptions jobOptions = executeTask.JobOptions;
 			await using IServerLogger logger = _serverLoggerFactory.CreateLogger(session, executeTask.LogId, localLogger, executeTask.JobId, executeTask.BatchId, null, null);
 
 			logger.LogInformation("Executing job \"{JobName}\", jobId {JobId}, batchId {BatchId}, leaseId {LeaseId}, agentVersion {AgentVersion}", executeTask.JobName, executeTask.JobId, executeTask.BatchId, leaseId, AgentApp.Version);
diff --git a/Engine/Source/Programs/Horde/Horde.Server/Jobs/IJob.cs b/Engine/Source/Programs/Horde/Horde.Server/Jobs/IJob.cs
--- a/Engine/Source/Programs/Horde/Horde.Server/Jobs/IJob.cs
+++ b/Engine/Source/Programs/Horde/Horde.Server/Jobs/IJob.cs
@@ -23,7 +23,6 @@ using Horde.Server.Jobs.Timing;
 using Horde.Server.Server;
 using Horde.Server.Streams;
 using Horde.Server.Ugs;
-using HordeCommon.Rpc.Tasks;
 using MongoDB.Bson;
 using MongoDB.Driver;
 
diff --git a/Engine/Source/Programs/Horde/Horde.Server/Jobs/IJobCollection.cs b/Engine/Source/Programs/Horde/Horde.Server/Jobs/IJobCollection.cs
--- a/Engine/Source/Programs/Horde/Horde.Server/Jobs/IJobCollection.cs
+++ b/Engine/Source/Programs/Horde/Horde.Server/Jobs/IJobCollection.cs
@@ -18,7 +18,6 @@ using EpicGames.Horde.Users;
 using Horde.Server.Acls;
 using Horde.Server.Jobs.Graphs;
 using Horde.Server.Streams;
-using HordeCommon.Rpc.Tasks;
 using MongoDB.Bson;
 
 namespace Horde.Server.Jobs
diff --git a/Engine/Source/Programs/Horde/Horde.Server/Jobs/JobCollection.cs b/Engine/Source/Programs/Horde/Horde.Server/Jobs/JobCollection.cs
--- a/Engine/Source/Programs/Horde/Horde.Server/Jobs/JobCollection.cs
+++ b/Engine/Source/Programs/Horde/Horde.Server/Jobs/JobCollection.cs
@@ -25,7 +25,6 @@ using Horde.Server.Server;
 using Horde.Server.Streams;
 using Horde.Server.Telemetry;
 using Horde.Server.Utilities;
-using HordeCommon.Rpc.Tasks;
 using Microsoft.Extensions.Logging;
 using Microsoft.Extensions.Options;
 using MongoDB.Bson;
diff --git a/Engine/Source/Programs/Horde/Horde.Server/Jobs/JobMessages.cs b/Engine/Source/Programs/Horde/Horde.Server/Jobs/JobMessages.cs
--- a/Engine/Source/Programs/Horde/Horde.Server/Jobs/JobMessages.cs
+++ b/Engine/Source/Programs/Horde/Horde.Server/Jobs/JobMessages.cs
@@ -14,8 +14,6 @@ using EpicGames.Horde.Streams;
 using EpicGames.Horde.Users;
 using Horde.Server.Jobs.Graphs;
 using Horde.Server.Streams;
-//using HordeCommon;
-using HordeCommon.Rpc.Tasks;
 
 namespace Horde.Server.Jobs
 {
@@ -40,6 +38,78 @@ namespace Horde.Server.Jobs
 		Complete,
 	}
 
+	/// <summary>
+	/// Options for executing a job
+	/// </summary>
+	public class JobOptions
+	{
+		/// <summary>
+		/// Name of the executor to use
+		/// </summary>
+		public string? Executor { get; set; }
+
+		/// <summary>
+		/// Whether to use the new temp storage backend
+		/// </summary>
+		public bool? UseNewTempStorage { get; set; }
+
+		/// <summary>
+		/// Whether to execute using Wine emulation on Linux
+		/// </summary>
+		public bool? UseWine { get; set; }
+
+		/// <summary>
+		/// Executes the job lease in a separate process
+		/// </summary>
+		public bool? RunInSeparateProcess { get; set; }
+
+		/// <summary>
+		/// What workspace materializer to use in WorkspaceExecutor. Will override any value from workspace config.
+		/// </summary>
+		public string? WorkspaceMaterializer { get; set; }
+
+		/// <summary>
+		/// Options for executing a job inside a container
+		/// </summary>
+		public JobContainerOptions? Container { get; set; }
+
+		/// <summary>
+		/// Merge defaults from another options object
+		/// </summary>
+		/// <param name="other"></param>
+		public void MergeDefaults(JobOptions other)
+		{
+			Executor ??= other.Executor;
+			UseNewTempStorage ??= other.UseNewTempStorage;
+		}
+	}
+
+	/// <summary>
+	/// Options for executing a job inside a container
+	/// </summary>
+	public class JobContainerOptions
+	{
+		/// <summary>
+		/// Whether to execute job inside a container
+		/// </summary>
+		public bool? Enabled { get; set; }
+
+		/// <summary>
+		/// Image URL to container, such as "quay.io/podman/hello"
+		/// </summary>
+		public string? ImageUrl { get; set; }
+
+		/// <summary>
+		/// Container engine executable (docker or with full path like /usr/bin/podman)
+		/// </summary>
+		public string? ContainerEngineExecutable { get; set; }
+
+		/// <summary>
+		/// Additional arguments to pass to container engine
+		/// </summary>
+		public string? ExtraArguments { get; set; }
+	}
+
 	/// <summary>
 	/// Parameters required to create a job
 	/// </summary>
diff --git a/Engine/Source/Programs/Horde/Horde.Server/Jobs/JobTaskSource.cs b/Engine/Source/Programs/Horde/Horde.Server/Jobs/JobTaskSource.cs
--- a/Engine/Source/Programs/Horde/Horde.Server/Jobs/JobTaskSource.cs
+++ b/Engine/Source/Programs/Horde/Horde.Server/Jobs/JobTaskSource.cs
@@ -903,7 +903,7 @@ namespace Horde.Server.Jobs
 			task.BatchId = batch.Id.ToString();
 			task.LogId = logId.ToString();
 			task.JobName = leaseName.ToString();
-			task.JobOptions = job.JobOptions;
+			task.JobOptions = (job.JobOptions != null)? GetRpcJobOptions(job.JobOptions) : null;
 			task.NamespaceId = namespaceId.ToString();
 			task.StoragePrefix = storagePrefix;
 			task.Token = await _aclService.IssueBearerTokenAsync(claims, null, cancellationToken);
@@ -938,6 +938,31 @@ namespace Horde.Server.Jobs
 			return task;
 		}
 
+		static RpcJobOptions GetRpcJobOptions(JobOptions options)
+		{
+			RpcJobOptions rpcOptions = new RpcJobOptions();
+			rpcOptions.Executor = options.Executor;
+			rpcOptions.UseNewTempStorage = options.UseNewTempStorage;
+			rpcOptions.UseWine = options.UseWine;
+			rpcOptions.RunInSeparateProcess = options.RunInSeparateProcess;
+			rpcOptions.WorkspaceMaterializer = options.WorkspaceMaterializer;
+			if (options.Container != null)
+			{
+				rpcOptions.Container = GetRpcJobContainerOptions(options.Container);
+			}
+			return rpcOptions;
+		}
+
+		static RpcJobContainerOptions GetRpcJobContainerOptions(JobContainerOptions options)
+		{
+			RpcJobContainerOptions rpcContainerOptions = new RpcJobContainerOptions();
+			rpcContainerOptions.Enabled = options.Enabled;
+			rpcContainerOptions.ImageUrl = options.ImageUrl;
+			rpcContainerOptions.ContainerEngineExecutable = options.ContainerEngineExecutable;
+			rpcContainerOptions.ExtraArguments = options.ExtraArguments;
+			return rpcContainerOptions;
+		}
+
 		/// <summary>
 		/// Send any badge updates for this job
 		/// </summary>
diff --git a/Engine/Source/Programs/Horde/Horde.Server/Jobs/Templates/TemplateConfig.cs b/Engine/Source/Programs/Horde/Horde.Server/Jobs/Templates/TemplateConfig.cs
--- a/Engine/Source/Programs/Horde/Horde.Server/Jobs/Templates/TemplateConfig.cs
+++ b/Engine/Source/Programs/Horde/Horde.Server/Jobs/Templates/TemplateConfig.cs
@@ -9,7 +9,6 @@ using EpicGames.Core;
 using EpicGames.Horde.Jobs;
 using Horde.Server.Configuration;
 using Horde.Server.Streams;
-using HordeCommon.Rpc.Tasks;
 
 namespace Horde.Server.Jobs.Templates
 {
diff --git a/Engine/Source/Programs/Horde/Horde.Server/Projects/ProjectConfig.cs b/Engine/Source/Programs/Horde/Horde.Server/Projects/ProjectConfig.cs
--- a/Engine/Source/Programs/Horde/Horde.Server/Projects/ProjectConfig.cs
+++ b/Engine/Source/Programs/Horde/Horde.Server/Projects/ProjectConfig.cs
@@ -13,9 +13,9 @@ using EpicGames.Horde.Telemetry;
 using Horde.Server.Acls;
 using Horde.Server.Agents.Pools;
 using Horde.Server.Configuration;
+using Horde.Server.Jobs;
 using Horde.Server.Server;
 using Horde.Server.Streams;
-using HordeCommon.Rpc.Tasks;
 
 namespace Horde.Server.Projects
 {
diff --git a/Engine/Source/Programs/Horde/Horde.Server/Protos/horde/tasks/job_task.cs b/Engine/Source/Programs/Horde/Horde.Server/Protos/horde/tasks/job_task.cs
--- a/Engine/Source/Programs/Horde/Horde.Server/Protos/horde/tasks/job_task.cs
+++ b/Engine/Source/Programs/Horde/Horde.Server/Protos/horde/tasks/job_task.cs
@@ -2,15 +2,7 @@
 
 namespace HordeCommon.Rpc.Tasks
 {
-	partial class JobOptions
+	partial class RpcJobOptions
 	{
-		/// <summary>
-		/// Merge settings from another JobOptions object
-		/// </summary>
-		public void MergeDefaults(JobOptions other)
-		{
-			Executor ??= other.Executor;
-			UseNewTempStorage ??= other.UseNewTempStorage;
-		}
 	}
 }
diff --git a/Engine/Source/Programs/Horde/Horde.Server/Server/DebugController.cs b/Engine/Source/Programs/Horde/Horde.Server/Server/DebugController.cs
--- a/Engine/Source/Programs/Horde/Horde.Server/Server/DebugController.cs
+++ b/Engine/Source/Programs/Horde/Horde.Server/Server/DebugController.cs
@@ -28,7 +28,6 @@ using Horde.Server.Logs;
 using Horde.Server.Projects;
 using Horde.Server.Streams;
 using Horde.Server.Utilities;
-using HordeCommon.Rpc.Tasks;
 using JetBrains.Profiler.SelfApi;
 using Microsoft.AspNetCore.Authorization;
 using Microsoft.AspNetCore.Http;
diff --git a/Engine/Source/Programs/Horde/Horde.Server/Streams/StreamConfig.cs b/Engine/Source/Programs/Horde/Horde.Server/Streams/StreamConfig.cs
--- a/Engine/Source/Programs/Horde/Horde.Server/Streams/StreamConfig.cs
+++ b/Engine/Source/Programs/Horde/Horde.Server/Streams/StreamConfig.cs
@@ -24,6 +24,7 @@ using EpicGames.Horde.Telemetry;
 using Horde.Server.Acls;
 using Horde.Server.Configuration;
 using Horde.Server.Issues;
+using Horde.Server.Jobs;
 using Horde.Server.Jobs.Graphs;
 using Horde.Server.Jobs.Templates;
 using Horde.Server.Perforce;
@@ -31,7 +32,6 @@ using Horde.Server.Projects;
 using Horde.Server.Replicators;
 using Horde.Server.Server;
 using Horde.Server.Utilities;
-using HordeCommon.Rpc.Tasks;
 
 namespace Horde.Server.Streams
 {
