Horde: Fix test failures due to modified error strings.

diff --git a/Engine/Source/Programs/Horde/Horde.Server.Tests/Issues/IssueServiceTests.cs b/Engine/Source/Programs/Horde/Horde.Server.Tests/Issues/IssueServiceTests.cs
--- a/Engine/Source/Programs/Horde/Horde.Server.Tests/Issues/IssueServiceTests.cs
+++ b/Engine/Source/Programs/Horde/Horde.Server.Tests/Issues/IssueServiceTests.cs
@@ -1737,7 +1737,7 @@ namespace Horde.Server.Tests.Issues
 				Assert.AreEqual(1, issues.Count);
 				Assert.AreEqual("Gauntlet", issues[0].Fingerprints[0].Type);
 				Assert.AreEqual(new IssueKey("test:Bar.Foo.Test", IssueKeyType.None), issues[0].Fingerprints[0].Keys.First());
-				Assert.AreEqual("Gauntlet test errors with Bar.Foo.Test", issues[0].Summary);
+				Assert.AreEqual("Automation test errors with Bar.Foo.Test", issues[0].Summary);
 			}
 			// #2
 			// Scenario: Gauntlet device event with Name property
@@ -1754,7 +1754,7 @@ namespace Horde.Server.Tests.Issues
 				Assert.AreEqual(1, issues.Count);
 				Assert.AreEqual("Gauntlet", issues[0].Fingerprints[0].Type);
 				Assert.AreEqual(new IssueKey("device:Foo", IssueKeyType.None), issues[0].Fingerprints[0].Keys.First());
-				Assert.AreEqual("Gauntlet device warnings with Foo", issues[0].Summary);
+				Assert.AreEqual("Automation device warnings with Foo", issues[0].Summary);
 			}
 			// #3
 			// Scenario: Gauntlet build drop event with File and Directory property
@@ -1772,7 +1772,7 @@ namespace Horde.Server.Tests.Issues
 				Assert.AreEqual(1, issues.Count);
 				Assert.AreEqual("Gauntlet", issues[0].Fingerprints[0].Type);
 				Assert.AreEqual(new IssueKey("access:/Bar/Foo.txt", IssueKeyType.None), issues[0].Fingerprints[0].Keys.First());
-				Assert.AreEqual("Gauntlet access errors with /Bar/Foo.txt and with /Bar/Foo", issues[0].Summary);
+				Assert.AreEqual("Automation access errors with /Bar/Foo.txt and with /Bar/Foo", issues[0].Summary);
 			}
 			// #4
 			// Scenario: Gauntlet Fatal event
@@ -1816,7 +1816,7 @@ namespace Horde.Server.Tests.Issues
 				Assert.AreEqual(1, issues.Count);
 				Assert.AreEqual("Gauntlet", issues[0].Fingerprints[0].Type);
 				Assert.AreEqual("hash:", issues[0].Fingerprints[0].Keys.First().Name.Substring(0, 5));
-				Assert.AreEqual("Gauntlet fatal errors in Update Version Files", issues[0].Summary);
+				Assert.AreEqual("Automation fatal errors in Update Version Files", issues[0].Summary);
 			}
 			// #5
 			// Scenario: Gauntlet Test event
@@ -1850,7 +1850,7 @@ namespace Horde.Server.Tests.Issues
 				Assert.AreEqual(1, issues.Count);
 				Assert.AreEqual("Gauntlet", issues[0].Fingerprints[0].Type);
 				Assert.AreEqual("hash:", issues[0].Fingerprints[0].Keys.First().Name.Substring(0, 5));
-				Assert.AreEqual("Gauntlet test errors in Update Version Files", issues[0].Summary);
+				Assert.AreEqual("Automation test errors in Update Version Files", issues[0].Summary);
 			}
 		}
 
