Horde: Separate gRPC definition of ReportPlacement from native definition.

#rnx

diff --git a/Engine/Source/Programs/Horde/Horde.Agent/Execution/JobExecutor.cs b/Engine/Source/Programs/Horde/Horde.Agent/Execution/JobExecutor.cs
--- a/Engine/Source/Programs/Horde/Horde.Agent/Execution/JobExecutor.cs
+++ b/Engine/Source/Programs/Horde/Horde.Agent/Execution/JobExecutor.cs
@@ -198,7 +198,7 @@ namespace Horde.Agent.Execution
 		class ReportData
 		{
 			public ReportScope Scope { get; set; }
-			public ReportPlacement Placement { get; set; }
+			public RpcReportPlacement Placement { get; set; }
 			public string Name { get; set; } = String.Empty;
 			public string Content { get; set; } = String.Empty;
 			public string FileName { get; set; } = String.Empty;
diff --git a/Engine/Source/Programs/Horde/Horde.Server/Jobs/IJob.cs b/Engine/Source/Programs/Horde/Horde.Server/Jobs/IJob.cs
--- a/Engine/Source/Programs/Horde/Horde.Server/Jobs/IJob.cs
+++ b/Engine/Source/Programs/Horde/Horde.Server/Jobs/IJob.cs
@@ -28,8 +28,6 @@ using MongoDB.Driver;
 
 namespace Horde.Server.Jobs
 {
-	using ReportPlacement = HordeCommon.Rpc.ReportPlacement;
-
 	/// <summary>
 	/// Report for a job or jobstep
 	/// </summary>
diff --git a/Engine/Source/Programs/Horde/Horde.Server/Jobs/JobMessages.cs b/Engine/Source/Programs/Horde/Horde.Server/Jobs/JobMessages.cs
--- a/Engine/Source/Programs/Horde/Horde.Server/Jobs/JobMessages.cs
+++ b/Engine/Source/Programs/Horde/Horde.Server/Jobs/JobMessages.cs
@@ -243,6 +243,22 @@ namespace Horde.Server.Jobs
 		public List<string>? Arguments { get; set; }
 	}
 
+	/// <summary>
+	/// Placement for a job report
+	/// </summary>
+	public enum ReportPlacement
+	{
+		/// <summary>
+		/// On a panel of its own
+		/// </summary>
+		Panel = 0,
+
+		/// <summary>
+		/// In the summary panel
+		/// </summary>
+		Summary = 1
+	}
+
 	/// <summary>
 	/// Information about a report associated with a job
 	/// </summary>
@@ -256,7 +272,7 @@ namespace Horde.Server.Jobs
 		/// <summary>
 		/// Report placement
 		/// </summary>
-		public HordeCommon.Rpc.ReportPlacement Placement { get; set; }
+		public ReportPlacement Placement { get; set; }
 
 		/// <summary>
 		/// The artifact id
diff --git a/Engine/Source/Programs/Horde/Horde.Server/Jobs/JobRpcService.cs b/Engine/Source/Programs/Horde/Horde.Server/Jobs/JobRpcService.cs
--- a/Engine/Source/Programs/Horde/Horde.Server/Jobs/JobRpcService.cs
+++ b/Engine/Source/Programs/Horde/Horde.Server/Jobs/JobRpcService.cs
@@ -1028,7 +1028,7 @@ namespace Horde.Server.Jobs
 
 			IJobStepBatch batch = AuthorizeBatch(job, JobStepBatchId.Parse(request.BatchId), context);
 
-			Report newReport = new Report { Name = request.Name, Placement = request.Placement, Content = request.Content };
+			Report newReport = new Report { Name = request.Name, Placement = (ReportPlacement)request.Placement, Content = request.Content };
 			if (request.Scope == ReportScope.Job)
 			{
 				_logger.LogDebug("Adding report to job {JobId}: {Name} -> {Content}", job.Id, request.Name, request.Content);
