Horde: Move JobStepError into EpicGames.Horde.

#rnx

diff --git a/Engine/Source/Programs/Horde/Horde.Server/Jobs/IJob.cs b/Engine/Source/Programs/Horde/Horde.Server/Jobs/IJob.cs
--- a/Engine/Source/Programs/Horde/Horde.Server/Jobs/IJob.cs
+++ b/Engine/Source/Programs/Horde/Horde.Server/Jobs/IJob.cs
@@ -158,32 +158,6 @@ namespace Horde.Server.Jobs
 		public Dictionary<string, string>? Properties { get; }
 	}
 
-	/// <summary>
-	/// Systemic error codes for a job step failing
-	/// </summary>
-	public enum JobStepError
-	{
-		/// <summary>
-		/// No systemic error
-		/// </summary>
-		None,
-
-		/// <summary>
-		/// Step did not complete in the required amount of time
-		/// </summary>
-		TimedOut,
-
-		/// <summary>
-		/// Step is in is paused state so was skipped
-		/// </summary>
-		Paused,
-
-		/// <summary>
-		/// Step did not complete because the batch exited
-		/// </summary>
-		Incomplete
-	}
-
 	/// <summary>
 	/// Extension methods for job steps
 	/// </summary>
diff --git a/Engine/Source/Programs/Shared/EpicGames.Horde/Jobs/JobStepError.cs b/Engine/Source/Programs/Shared/EpicGames.Horde/Jobs/JobStepError.cs
new file mode 100644
--- /dev/null
+++ b/Engine/Source/Programs/Shared/EpicGames.Horde/Jobs/JobStepError.cs
@@ -0,0 +1,30 @@
+// Copyright Epic Games, Inc. All Rights Reserved.
+
+namespace EpicGames.Horde.Jobs
+{
+	/// <summary>
+	/// Systemic error codes for a job step failing
+	/// </summary>
+	public enum JobStepError
+	{
+		/// <summary>
+		/// No systemic error
+		/// </summary>
+		None,
+
+		/// <summary>
+		/// Step did not complete in the required amount of time
+		/// </summary>
+		TimedOut,
+
+		/// <summary>
+		/// Step is in is paused state so was skipped
+		/// </summary>
+		Paused,
+
+		/// <summary>
+		/// Step did not complete because the batch exited
+		/// </summary>
+		Incomplete
+	}
+}
