Horde: Move commit classes into Horde.Server.Commits namespace.

#rnx

diff --git a/Engine/Source/Programs/Horde/Horde.Server.Tests/Jobs/JobServiceTests.cs b/Engine/Source/Programs/Horde/Horde.Server.Tests/Jobs/JobServiceTests.cs
--- a/Engine/Source/Programs/Horde/Horde.Server.Tests/Jobs/JobServiceTests.cs
+++ b/Engine/Source/Programs/Horde/Horde.Server.Tests/Jobs/JobServiceTests.cs
@@ -13,11 +13,11 @@ using EpicGames.Horde.Jobs;
 using EpicGames.Horde.Jobs.Templates;
 using EpicGames.Horde.Streams;
 using Horde.Server.Agents;
+using Horde.Server.Commits;
 using Horde.Server.Jobs;
 using Horde.Server.Jobs.Graphs;
 using Horde.Server.Jobs.Templates;
 using Horde.Server.Logs;
-using Horde.Server.Perforce;
 using Horde.Server.Projects;
 using Horde.Server.Server;
 using Horde.Server.Streams;
diff --git a/Engine/Source/Programs/Horde/Horde.Server.Tests/Stubs/Services/PerforceServiceStub.cs b/Engine/Source/Programs/Horde/Horde.Server.Tests/Stubs/Services/PerforceServiceStub.cs
--- a/Engine/Source/Programs/Horde/Horde.Server.Tests/Stubs/Services/PerforceServiceStub.cs
+++ b/Engine/Source/Programs/Horde/Horde.Server.Tests/Stubs/Services/PerforceServiceStub.cs
@@ -10,6 +10,7 @@ using EpicGames.Core;
 using EpicGames.Horde.Commits;
 using EpicGames.Horde.Streams;
 using EpicGames.Horde.Users;
+using Horde.Server.Commits;
 using Horde.Server.Perforce;
 using Horde.Server.Streams;
 using Horde.Server.Users;
diff --git a/Engine/Source/Programs/Horde/Horde.Server.Tests/TestSetup.cs b/Engine/Source/Programs/Horde/Horde.Server.Tests/TestSetup.cs
--- a/Engine/Source/Programs/Horde/Horde.Server.Tests/TestSetup.cs
+++ b/Engine/Source/Programs/Horde/Horde.Server.Tests/TestSetup.cs
@@ -29,6 +29,7 @@ using Horde.Server.Agents.Sessions;
 using Horde.Server.Agents.Utilization;
 using Horde.Server.Artifacts;
 using Horde.Server.Auditing;
+using Horde.Server.Commits;
 using Horde.Server.Compute;
 using Horde.Server.Configuration;
 using Horde.Server.Dashboard;
diff --git a/Engine/Source/Programs/Horde/Horde.Server/Commits/CommitService.cs b/Engine/Source/Programs/Horde/Horde.Server/Commits/CommitService.cs
new file mode 100644
--- /dev/null
+++ b/Engine/Source/Programs/Horde/Horde.Server/Commits/CommitService.cs
@@ -0,0 +1,26 @@
+// Copyright Epic Games, Inc. All Rights Reserved.
+
+using Horde.Server.Perforce;
+using Horde.Server.Streams;
+
+namespace Horde.Server.Commits
+{
+	/// <summary>
+	/// Provides commit information for streams
+	/// </summary>
+	class CommitService : ICommitService
+	{
+		readonly IPerforceService _perforceService;
+
+		/// <summary>
+		/// Constructor
+		/// </summary>
+		public CommitService(IPerforceService perforceService)
+		{
+			_perforceService = perforceService;
+		}
+
+		/// <inheritdoc/>
+		public ICommitCollection GetCollection(StreamConfig streamConfig) => _perforceService.GetCommits(streamConfig);
+	}
+}
diff --git a/Engine/Source/Programs/Horde/Horde.Server/Commits/ICommit.cs b/Engine/Source/Programs/Horde/Horde.Server/Commits/ICommit.cs
new file mode 100644
--- /dev/null
+++ b/Engine/Source/Programs/Horde/Horde.Server/Commits/ICommit.cs
@@ -0,0 +1,82 @@
+﻿// Copyright Epic Games, Inc. All Rights Reserved.
+
+using System;
+using System.Collections.Generic;
+using System.Threading;
+using System.Threading.Tasks;
+using EpicGames.Core;
+using EpicGames.Horde.Commits;
+using EpicGames.Horde.Streams;
+using EpicGames.Horde.Users;
+
+namespace Horde.Server.Commits
+{
+	/// <summary>
+	/// Stores metadata about a commit
+	/// </summary>
+	public interface ICommit
+	{
+		/// <summary>
+		/// Stream containing the commit
+		/// </summary>
+		StreamId StreamId { get; }
+
+		/// <summary>
+		/// The changelist number
+		/// </summary>
+		int Number { get; }
+
+		/// <summary>
+		/// The change that this commit originates from
+		/// </summary>
+		int OriginalChange { get; }
+
+		/// <summary>
+		/// The author user id
+		/// </summary>
+		UserId AuthorId { get; }
+
+		/// <summary>
+		/// The owner of this change, if different from the author (due to Robomerge)
+		/// </summary>
+		UserId OwnerId { get; }
+
+		/// <summary>
+		/// Changelist description
+		/// </summary>
+		string Description { get; }
+
+		/// <summary>
+		/// Base path for all files in the change
+		/// </summary>
+		string BasePath { get; }
+
+		/// <summary>
+		/// Date/time that change was committed
+		/// </summary>
+		DateTime DateUtc { get; }
+
+		/// <summary>
+		/// Gets the list of tags for the commit
+		/// </summary>
+		/// <param name="cancellationToken">Cancellation token for the operation</param>
+		/// <returns>True if the commit has the given tag</returns>
+		ValueTask<IReadOnlyList<CommitTag>> GetTagsAsync(CancellationToken cancellationToken);
+
+		/// <summary>
+		/// Determine if this commit matches the given filter. Prefer using commit tags rather than this method; the results can be cached.
+		/// </summary>
+		/// <param name="filter">Filter to test</param>
+		/// <param name="cancellationToken">Cancellation token for the operation</param>
+		/// <returns></returns>
+		ValueTask<bool> MatchesFilterAsync(FileFilter filter, CancellationToken cancellationToken);
+
+		/// <summary>
+		/// Gets the files for this change, relative to the root of the stream
+		/// </summary>
+		/// <param name="maxFiles">Maximum number of files to return. Querying large number of files may cause performance issues with merge commits.</param>
+		/// <param name="cancellationToken">Cancellation token for the operation</param>
+		/// <returns>List of files modified by this commit</returns>
+		ValueTask<IReadOnlyList<string>> GetFilesAsync(int maxFiles, CancellationToken cancellationToken);
+	}
+}
diff --git a/Engine/Source/Programs/Horde/Horde.Server/Commits/ICommitCollection.cs b/Engine/Source/Programs/Horde/Horde.Server/Commits/ICommitCollection.cs
new file mode 100644
--- /dev/null
+++ b/Engine/Source/Programs/Horde/Horde.Server/Commits/ICommitCollection.cs
@@ -0,0 +1,115 @@
+﻿// Copyright Epic Games, Inc. All Rights Reserved.
+
+using System;
+using System.Collections.Generic;
+using System.Linq;
+using System.Threading;
+using System.Threading.Tasks;
+using EpicGames.Horde.Commits;
+using EpicGames.Perforce;
+using Horde.Server.Jobs.Templates;
+
+namespace Horde.Server.Commits
+{
+	/// <summary>
+	/// VCS abstraction. Provides information about commits to a particular stream.
+	/// </summary>
+	public interface ICommitCollection
+	{
+		/// <summary>
+		/// Creates a new change
+		/// </summary>
+		/// <param name="path">Path to modify in the change</param>
+		/// <param name="description">Description of the change</param>
+		/// <param name="cancellationToken">Cancellation token for the operation</param>
+		/// <returns>New commit information</returns>
+		Task<int> CreateNewAsync(string path, string description, CancellationToken cancellationToken = default);
+
+		/// <summary>
+		/// Gets the latest change for a particular stream
+		/// </summary>
+		/// <param name="changeNumber">Change numbers to query</param>
+		/// <param name="cancellationToken">Cancellation token for the operation</param>
+		/// <returns>Commit details</returns>
+		Task<ICommit> GetAsync(int changeNumber, CancellationToken cancellationToken = default);
+
+		/// <summary>
+		/// Finds changes submitted to a stream, in reverse order.
+		/// </summary>
+		/// <param name="minChange">The minimum changelist number</param>
+		/// <param name="maxChange">The maximum changelist number</param>
+		/// <param name="maxResults">Maximum number of results to return</param>
+		/// <param name="tags">Tags for the commits to return</param>
+		/// <param name="cancellationToken">Cancellation token for the operation</param>
+		/// <returns>Changelist information</returns>
+		IAsyncEnumerable<ICommit> FindAsync(int? minChange = null, int? maxChange = null, int? maxResults = null, IReadOnlyList<CommitTag>? tags = null, CancellationToken cancellationToken = default);
+
+		/// <summary>
+		/// Subscribes to changes from this commit source
+		/// </summary>
+		/// <param name="minChange">Minimum changelist number (exclusive)</param>
+		/// <param name="tags">Tags for the commit to return</param>
+		/// <param name="cancellationToken">Cancellation token for the operation</param>
+		/// <returns>New change information</returns>
+		IAsyncEnumerable<ICommit> SubscribeAsync(int minChange, IReadOnlyList<CommitTag>? tags = null, CancellationToken cancellationToken = default);
+	}
+
+	/// <summary>
+	/// Extension methods for <see cref="ICommitCollection"/>
+	/// </summary>
+	public static class CommitCollectionExtensions
+	{
+		/// <summary>
+		/// Creates a new change for a template
+		/// </summary>
+		/// <param name="perforce">The Perforce service instance</param>
+		/// <param name="template">The template being built</param>
+		/// <param name="cancellationToken">Cancellation token for the operation</param>
+		/// <returns>New changelist number</returns>
+		public static Task<int> CreateNewAsync(this ICommitCollection perforce, ITemplate template, CancellationToken cancellationToken)
+		{
+			string description = (template.SubmitDescription ?? "[Horde] New change for $(TemplateName)").Replace("$(TemplateName)", template.Name, StringComparison.OrdinalIgnoreCase);
+			return perforce.CreateNewAsync(template.SubmitNewChange!, description, cancellationToken);
+		}
+
+		/// <summary>
+		/// Gets the last code code equal or before the given change number
+		/// </summary>
+		/// <param name="source">The commit source to query</param>
+		/// <param name="maxChange">Maximum code change to query</param>
+		/// <param name="cancellationToken">Cancellation token for the operation</param>
+		/// <returns>The last code change</returns>
+		public static async ValueTask<ICommit?> GetLastCodeChangeAsync(this ICommitCollection source, int? maxChange, CancellationToken cancellationToken = default)
+		{
+			return await source.FindAsync(null, maxChange, 1, new[] { CommitTag.Code }, cancellationToken).FirstOrDefaultAsync(cancellationToken);
+		}
+
+		/// <summary>
+		/// Finds the latest commit from a source
+		/// </summary>
+		/// <param name="source">The commit source to query</param>
+		/// <param name="cancellationToken">Cancellation token for the operation</param>
+		/// <returns>The latest commit</returns>
+		public static async Task<ICommit> GetLatestAsync(this ICommitCollection source, CancellationToken cancellationToken = default)
+		{
+			ICommit? commit = await source.FindAsync(null, null, 1, null, cancellationToken).FirstOrDefaultAsync(cancellationToken);
+			if (commit == null)
+			{
+				throw new PerforceException("No changes found for stream.");
+			}
+			return commit;
+		}
+
+		/// <summary>
+		/// Finds the latest commit from a source
+		/// </summary>
+		/// <param name="source">The commit source to query</param>
+		/// <param name="cancellationToken">Cancellation token for the operation</param>
+		/// <returns>The latest commit</returns>
+		public static async Task<int> GetLatestNumberAsync(this ICommitCollection source, CancellationToken cancellationToken = default)
+		{
+			ICommit commit = await GetLatestAsync(source, cancellationToken);
+			return commit.Number;
+		}
+	}
+}
diff --git a/Engine/Source/Programs/Horde/Horde.Server/Commits/ICommitService.cs b/Engine/Source/Programs/Horde/Horde.Server/Commits/ICommitService.cs
new file mode 100644
--- /dev/null
+++ b/Engine/Source/Programs/Horde/Horde.Server/Commits/ICommitService.cs
@@ -0,0 +1,19 @@
+// Copyright Epic Games, Inc. All Rights Reserved.
+
+using Horde.Server.Streams;
+
+namespace Horde.Server.Commits
+{
+	/// <summary>
+	/// Provides information about commits to a stream
+	/// </summary>
+	public interface ICommitService
+	{
+		/// <summary>
+		/// Gets a commit collection for the given stream
+		/// </summary>
+		/// <param name="streamConfig">Stream to get commits for</param>
+		/// <returns>Collection object</returns>
+		ICommitCollection GetCollection(StreamConfig streamConfig);
+	}
+}
diff --git a/Engine/Source/Programs/Horde/Horde.Server/Issues/IssueService.cs b/Engine/Source/Programs/Horde/Horde.Server/Issues/IssueService.cs
--- a/Engine/Source/Programs/Horde/Horde.Server/Issues/IssueService.cs
+++ b/Engine/Source/Programs/Horde/Horde.Server/Issues/IssueService.cs
@@ -15,10 +15,10 @@ using EpicGames.Horde.Jobs;
 using EpicGames.Horde.Jobs.Templates;
 using EpicGames.Horde.Streams;
 using EpicGames.Horde.Users;
+using Horde.Server.Commits;
 using Horde.Server.Jobs;
 using Horde.Server.Jobs.Graphs;
 using Horde.Server.Logs;
-using Horde.Server.Perforce;
 using Horde.Server.Server;
 using Horde.Server.Streams;
 using Horde.Server.Users;
diff --git a/Engine/Source/Programs/Horde/Horde.Server/Jobs/Bisect/BisectService.cs b/Engine/Source/Programs/Horde/Horde.Server/Jobs/Bisect/BisectService.cs
--- a/Engine/Source/Programs/Horde/Horde.Server/Jobs/Bisect/BisectService.cs
+++ b/Engine/Source/Programs/Horde/Horde.Server/Jobs/Bisect/BisectService.cs
@@ -7,9 +7,9 @@ using System.Threading;
 using System.Threading.Tasks;
 using EpicGames.Core;
 using EpicGames.Horde.Jobs;
+using Horde.Server.Commits;
 using Horde.Server.Jobs.Graphs;
 using Horde.Server.Jobs.Templates;
-using Horde.Server.Perforce;
 using Horde.Server.Server;
 using Horde.Server.Streams;
 using HordeCommon;
diff --git a/Engine/Source/Programs/Horde/Horde.Server/Jobs/JobService.cs b/Engine/Source/Programs/Horde/Horde.Server/Jobs/JobService.cs
--- a/Engine/Source/Programs/Horde/Horde.Server/Jobs/JobService.cs
+++ b/Engine/Source/Programs/Horde/Horde.Server/Jobs/JobService.cs
@@ -20,6 +20,7 @@ using EpicGames.Horde.Logs;
 using EpicGames.Horde.Streams;
 using EpicGames.Horde.Users;
 using Horde.Server.Agents;
+using Horde.Server.Commits;
 using Horde.Server.Issues;
 using Horde.Server.Jobs.Bisect;
 using Horde.Server.Jobs.Graphs;
diff --git a/Engine/Source/Programs/Horde/Horde.Server/Jobs/JobsController.cs b/Engine/Source/Programs/Horde/Horde.Server/Jobs/JobsController.cs
--- a/Engine/Source/Programs/Horde/Horde.Server/Jobs/JobsController.cs
+++ b/Engine/Source/Programs/Horde/Horde.Server/Jobs/JobsController.cs
@@ -15,6 +15,7 @@ using EpicGames.Horde.Users;
 using Horde.Server.Acls;
 using Horde.Server.Agents;
 using Horde.Server.Artifacts;
+using Horde.Server.Commits;
 using Horde.Server.Jobs.Artifacts;
 using Horde.Server.Jobs.Graphs;
 using Horde.Server.Jobs.Templates;
diff --git a/Engine/Source/Programs/Horde/Horde.Server/Jobs/Schedules/ScheduleService.cs b/Engine/Source/Programs/Horde/Horde.Server/Jobs/Schedules/ScheduleService.cs
--- a/Engine/Source/Programs/Horde/Horde.Server/Jobs/Schedules/ScheduleService.cs
+++ b/Engine/Source/Programs/Horde/Horde.Server/Jobs/Schedules/ScheduleService.cs
@@ -15,9 +15,9 @@ using EpicGames.Horde.Streams;
 using EpicGames.Redis;
 using EpicGames.Redis.Utility;
 using EpicGames.Serialization;
+using Horde.Server.Commits;
 using Horde.Server.Jobs.Graphs;
 using Horde.Server.Jobs.Templates;
-using Horde.Server.Perforce;
 using Horde.Server.Server;
 using Horde.Server.Streams;
 using Horde.Server.Utilities;
diff --git a/Engine/Source/Programs/Horde/Horde.Server/Perforce/CommitService.cs b/Engine/Source/Programs/Horde/Horde.Server/Perforce/CommitService.cs
deleted file mode 100644
diff --git a/Engine/Source/Programs/Horde/Horde.Server/Perforce/ICommit.cs b/Engine/Source/Programs/Horde/Horde.Server/Perforce/ICommit.cs
deleted file mode 100644
diff --git a/Engine/Source/Programs/Horde/Horde.Server/Perforce/ICommitCollection.cs b/Engine/Source/Programs/Horde/Horde.Server/Perforce/ICommitCollection.cs
deleted file mode 100644
diff --git a/Engine/Source/Programs/Horde/Horde.Server/Perforce/ICommitService.cs b/Engine/Source/Programs/Horde/Horde.Server/Perforce/ICommitService.cs
deleted file mode 100644
diff --git a/Engine/Source/Programs/Horde/Horde.Server/Perforce/IPerforceService.cs b/Engine/Source/Programs/Horde/Horde.Server/Perforce/IPerforceService.cs
--- a/Engine/Source/Programs/Horde/Horde.Server/Perforce/IPerforceService.cs
+++ b/Engine/Source/Programs/Horde/Horde.Server/Perforce/IPerforceService.cs
@@ -7,6 +7,7 @@ using System.Threading;
 using System.Threading.Tasks;
 using EpicGames.Horde.Commits;
 using EpicGames.Perforce;
+using Horde.Server.Commits;
 using Horde.Server.Streams;
 using Horde.Server.Users;
 
diff --git a/Engine/Source/Programs/Horde/Horde.Server/Perforce/PerforceService.cs b/Engine/Source/Programs/Horde/Horde.Server/Perforce/PerforceService.cs
--- a/Engine/Source/Programs/Horde/Horde.Server/Perforce/PerforceService.cs
+++ b/Engine/Source/Programs/Horde/Horde.Server/Perforce/PerforceService.cs
@@ -18,6 +18,7 @@ using EpicGames.Horde.Commits;
 using EpicGames.Horde.Streams;
 using EpicGames.Horde.Users;
 using EpicGames.Perforce;
+using Horde.Server.Commits;
 using Horde.Server.Server;
 using Horde.Server.Streams;
 using Horde.Server.Users;
diff --git a/Engine/Source/Programs/Horde/Horde.Server/Perforce/PerforceServiceCache.cs b/Engine/Source/Programs/Horde/Horde.Server/Perforce/PerforceServiceCache.cs
--- a/Engine/Source/Programs/Horde/Horde.Server/Perforce/PerforceServiceCache.cs
+++ b/Engine/Source/Programs/Horde/Horde.Server/Perforce/PerforceServiceCache.cs
@@ -14,6 +14,7 @@ using EpicGames.Horde.Streams;
 using EpicGames.Horde.Users;
 using EpicGames.Perforce;
 using EpicGames.Redis;
+using Horde.Server.Commits;
 using Horde.Server.Server;
 using Horde.Server.Streams;
 using Horde.Server.Users;
diff --git a/Engine/Source/Programs/Horde/Horde.Server/Replicators/PerforceReplicator.cs b/Engine/Source/Programs/Horde/Horde.Server/Replicators/PerforceReplicator.cs
--- a/Engine/Source/Programs/Horde/Horde.Server/Replicators/PerforceReplicator.cs
+++ b/Engine/Source/Programs/Horde/Horde.Server/Replicators/PerforceReplicator.cs
@@ -16,6 +16,7 @@ using EpicGames.Horde.Storage.Bundles;
 using EpicGames.Horde.Storage.Nodes;
 using EpicGames.Horde.Streams;
 using EpicGames.Perforce;
+using Horde.Server.Commits;
 using Horde.Server.Perforce;
 using Horde.Server.Server;
 using Horde.Server.Storage;
diff --git a/Engine/Source/Programs/Horde/Horde.Server/Startup.cs b/Engine/Source/Programs/Horde/Horde.Server/Startup.cs
--- a/Engine/Source/Programs/Horde/Horde.Server/Startup.cs
+++ b/Engine/Source/Programs/Horde/Horde.Server/Startup.cs
@@ -53,6 +53,7 @@ using Horde.Server.Artifacts;
 using Horde.Server.Auditing;
 using Horde.Server.Authentication;
 using Horde.Server.Aws;
+using Horde.Server.Commits;
 using Horde.Server.Compute;
 using Horde.Server.Configuration;
 using Horde.Server.Dashboard;
diff --git a/Engine/Source/Programs/Horde/Horde.Server/Streams/StreamMessages.cs b/Engine/Source/Programs/Horde/Horde.Server/Streams/StreamMessages.cs
--- a/Engine/Source/Programs/Horde/Horde.Server/Streams/StreamMessages.cs
+++ b/Engine/Source/Programs/Horde/Horde.Server/Streams/StreamMessages.cs
@@ -10,10 +10,10 @@ using EpicGames.Horde.Jobs.Templates;
 using EpicGames.Horde.Projects;
 using EpicGames.Horde.Streams;
 using EpicGames.Horde.Users;
+using Horde.Server.Commits;
 using Horde.Server.Issues;
 using Horde.Server.Jobs.Schedules;
 using Horde.Server.Jobs.Templates;
-using Horde.Server.Perforce;
 using Horde.Server.Projects;
 using Horde.Server.Users;
 
diff --git a/Engine/Source/Programs/Horde/Horde.Server/Streams/StreamsController.cs b/Engine/Source/Programs/Horde/Horde.Server/Streams/StreamsController.cs
--- a/Engine/Source/Programs/Horde/Horde.Server/Streams/StreamsController.cs
+++ b/Engine/Source/Programs/Horde/Horde.Server/Streams/StreamsController.cs
@@ -10,9 +10,9 @@ using EpicGames.Horde.Jobs.Templates;
 using EpicGames.Horde.Projects;
 using EpicGames.Horde.Streams;
 using EpicGames.Horde.Users;
+using Horde.Server.Commits;
 using Horde.Server.Jobs;
 using Horde.Server.Jobs.Templates;
-using Horde.Server.Perforce;
 using Horde.Server.Projects;
 using Horde.Server.Server;
 using Horde.Server.Users;
