Horde: Move internal object references out of JobMessages.cs

#rnx

diff --git a/Engine/Source/Programs/Horde/Horde.Server.Tests/Commits/CommitTagTests.cs b/Engine/Source/Programs/Horde/Horde.Server.Tests/Commits/CommitTagTests.cs
new file mode 100644
--- /dev/null
+++ b/Engine/Source/Programs/Horde/Horde.Server.Tests/Commits/CommitTagTests.cs
@@ -0,0 +1,73 @@
+// Copyright Epic Games, Inc. All Rights Reserved.
+
+using System.Collections.Generic;
+using EpicGames.Core;
+using EpicGames.Horde.Commits;
+using Horde.Server.Streams;
+using Microsoft.VisualStudio.TestTools.UnitTesting;
+
+namespace Horde.Server.Tests.Commits
+{
+	[TestClass]
+	public class CommitTagTests
+	{
+		[TestMethod]
+		public void TestCode()
+		{
+			StreamConfig config = new StreamConfig();
+
+			FileFilter? filter;
+			Assert.IsTrue(config.TryGetCommitTagFilter(CommitTag.Code, out filter));
+
+			Assert.IsTrue(filter.Matches("foo.cpp"));
+			Assert.IsTrue(filter.Matches("foo.h"));
+			Assert.IsTrue(filter.Matches("foo.cs"));
+			Assert.IsFalse(filter.Matches("foo.uasset"));
+			Assert.IsFalse(filter.Matches("foo.png"));
+		}
+
+		[TestMethod]
+		public void TestContent()
+		{
+			StreamConfig config = new StreamConfig();
+
+			FileFilter? filter;
+			Assert.IsTrue(config.TryGetCommitTagFilter(CommitTag.Content, out filter));
+
+			Assert.IsFalse(filter.Matches("foo.cpp"));
+			Assert.IsFalse(filter.Matches("foo.h"));
+			Assert.IsFalse(filter.Matches("foo.cs"));
+			Assert.IsTrue(filter.Matches("foo.uasset"));
+			Assert.IsTrue(filter.Matches("foo.png"));
+		}
+
+		[TestMethod]
+		public void TestCustom()
+		{
+			StreamConfig config = new StreamConfig();
+			config.CommitTags.Add(new CommitTagConfig { Name = new CommitTag("test"), Filter = new List<string> { "*TEST*", "-*TEST2*" } });
+
+			FileFilter? filter;
+			Assert.IsTrue(config.TryGetCommitTagFilter(new CommitTag("test"), out filter));
+
+			Assert.IsFalse(filter.Matches("somefile"));
+			Assert.IsTrue(filter.Matches("sometestfile"));
+			Assert.IsFalse(filter.Matches("sometest2file"));
+		}
+
+		[TestMethod]
+		public void TestInheritancee()
+		{
+			StreamConfig config = new StreamConfig();
+			config.CommitTags.Add(new CommitTagConfig { Name = new CommitTag("test"), Base = CommitTag.Code, Filter = new List<string> { "*TEST*", "-*TEST2*" } });
+
+			FileFilter? filter;
+			Assert.IsTrue(config.TryGetCommitTagFilter(new CommitTag("test"), out filter));
+
+			Assert.IsTrue(filter.Matches("somefile.cpp"));
+			Assert.IsFalse(filter.Matches("somefile.uasset"));
+			Assert.IsTrue(filter.Matches("sometestfile.uasset"));
+			Assert.IsFalse(filter.Matches("sometest2file.uasset"));
+		}
+	}
+}
diff --git a/Engine/Source/Programs/Horde/Horde.Server.Tests/Perforce/CommitTagTests.cs b/Engine/Source/Programs/Horde/Horde.Server.Tests/Perforce/CommitTagTests.cs
deleted file mode 100644
diff --git a/Engine/Source/Programs/Horde/Horde.Server/Agents/Fleet/JobQueueStrategy.cs b/Engine/Source/Programs/Horde/Horde.Server/Agents/Fleet/JobQueueStrategy.cs
--- a/Engine/Source/Programs/Horde/Horde.Server/Agents/Fleet/JobQueueStrategy.cs
+++ b/Engine/Source/Programs/Horde/Horde.Server/Agents/Fleet/JobQueueStrategy.cs
@@ -9,6 +9,7 @@ using System.Threading;
 using System.Threading.Tasks;
 using EpicGames.Core;
 using EpicGames.Horde.Agents.Pools;
+using EpicGames.Horde.Jobs;
 using EpicGames.Horde.Streams;
 using Horde.Server.Agents.Pools;
 using Horde.Server.Jobs;
diff --git a/Engine/Source/Programs/Horde/Horde.Server/Jobs/Bisect/BisectController.cs b/Engine/Source/Programs/Horde/Horde.Server/Jobs/Bisect/BisectController.cs
--- a/Engine/Source/Programs/Horde/Horde.Server/Jobs/Bisect/BisectController.cs
+++ b/Engine/Source/Programs/Horde/Horde.Server/Jobs/Bisect/BisectController.cs
@@ -151,7 +151,7 @@ namespace Horde.Server.Jobs.Bisect
 		{
 			_bisectTask = bisectTask;
 			Owner = owner;
-			Steps = steps.Select(s => new GetJobStepRefResponse(s)).ToList();
+			Steps = steps.Select(s => StreamsController.CreateGetJobStepRefResponse(s)).ToList();
 			NextJobId = nextJob?.Id;
 			NextJobChange = nextJob?.Change;
 		}
diff --git a/Engine/Source/Programs/Horde/Horde.Server/Jobs/IJob.cs b/Engine/Source/Programs/Horde/Horde.Server/Jobs/IJob.cs
--- a/Engine/Source/Programs/Horde/Horde.Server/Jobs/IJob.cs
+++ b/Engine/Source/Programs/Horde/Horde.Server/Jobs/IJob.cs
@@ -490,6 +490,20 @@ namespace Horde.Server.Jobs
 			timingInfo.WaitForAll(others);
 			return timingInfo;
 		}
+
+		/// <summary>
+		/// Copies this info to a repsonse object
+		/// </summary>
+		public void CopyToResponse(GetTimingInfoResponse response)
+		{
+			response.TotalWaitTime = (float?)TotalWaitTime?.TotalSeconds;
+			response.TotalInitTime = (float?)TotalInitTime?.TotalSeconds;
+			response.TotalTimeToComplete = (float?)TotalTimeToComplete?.TotalSeconds;
+
+			response.AverageTotalWaitTime = (float?)AverageTotalWaitTime?.TotalSeconds;
+			response.AverageTotalInitTime = (float?)AverageTotalInitTime?.TotalSeconds;
+			response.AverageTotalTimeToComplete = (float?)AverageTotalTimeToComplete?.TotalSeconds;
+		}
 	}
 
 	/// <summary>
diff --git a/Engine/Source/Programs/Horde/Horde.Server/Jobs/JobMessages.cs b/Engine/Source/Programs/Horde/Horde.Server/Jobs/JobMessages.cs
--- a/Engine/Source/Programs/Horde/Horde.Server/Jobs/JobMessages.cs
+++ b/Engine/Source/Programs/Horde/Horde.Server/Jobs/JobMessages.cs
@@ -7,13 +7,13 @@ using System.Linq;
 using System.Text.Json.Serialization;
 using EpicGames.Core;
 using EpicGames.Horde.Artifacts;
+using EpicGames.Horde.Commits;
+using EpicGames.Horde.Common;
 using EpicGames.Horde.Jobs;
 using EpicGames.Horde.Jobs.Bisect;
 using EpicGames.Horde.Jobs.Templates;
 using EpicGames.Horde.Streams;
 using EpicGames.Horde.Users;
-using Horde.Server.Jobs.Graphs;
-using Horde.Server.Streams;
 
 namespace Horde.Server.Jobs
 {
@@ -110,6 +110,42 @@ namespace Horde.Server.Jobs
 		public string? ExtraArguments { get; set; }
 	}
 
+	/// <summary>
+	/// Query selecting the base changelist to use
+	/// </summary>
+	public class ChangeQueryConfig
+	{
+		/// <summary>
+		/// Name of this query, for display on the dashboard.
+		/// </summary>
+		public string? Name { get; set; }
+
+		/// <summary>
+		/// Condition to evaluate before deciding to use this query. May query tags in a preflight.
+		/// </summary>
+		public Condition? Condition { get; set; }
+
+		/// <summary>
+		/// The template id to query
+		/// </summary>
+		public TemplateId? TemplateId { get; set; }
+
+		/// <summary>
+		/// The target to query
+		/// </summary>
+		public string? Target { get; set; }
+
+		/// <summary>
+		/// Whether to match a job that produced warnings
+		/// </summary>
+		public List<JobStepOutcome>? Outcomes { get; set; }
+
+		/// <summary>
+		/// Finds the last commit with this tag
+		/// </summary>
+		public CommitTag? CommitTag { get; set; }
+	}
+
 	/// <summary>
 	/// Parameters required to create a job
 	/// </summary>
@@ -287,13 +323,10 @@ namespace Horde.Server.Jobs
 		/// <summary>
 		/// Constructor
 		/// </summary>
-		/// <param name="report"></param>
-		public GetReportResponse(IReport report)
+		public GetReportResponse(string name, ReportPlacement placement)
 		{
-			Name = report.Name;
-			Placement = report.Placement;
-			ArtifactId = report.ArtifactId?.ToString();
-			Content = report.Content;
+			Name = name;
+			Placement = placement;
 		}
 	}
 
@@ -305,12 +338,12 @@ namespace Horde.Server.Jobs
 		/// <summary>
 		/// Unique Id for the job
 		/// </summary>
-		public string Id { get; set; }
+		public JobId Id { get; set; }
 
 		/// <summary>
 		/// Unique id of the stream containing this job
 		/// </summary>
-		public string StreamId { get; set; }
+		public StreamId StreamId { get; set; }
 
 		/// <summary>
 		/// Name of the job
@@ -345,12 +378,12 @@ namespace Horde.Server.Jobs
 		/// <summary>
 		/// The template type
 		/// </summary>
-		public string TemplateId { get; set; }
+		public TemplateId TemplateId { get; set; }
 
 		/// <summary>
 		/// Hash of the actual template data
 		/// </summary>
-		public string TemplateHash { get; set; }
+		public string? TemplateHash { get; set; }
 
 		/// <summary>
 		/// Hash of the graph for this job
@@ -400,12 +433,12 @@ namespace Horde.Server.Jobs
 		/// <summary>
 		/// The submitted changelist number
 		/// </summary>
-		public int? AutoSubmitChange { get; }
+		public int? AutoSubmitChange { get; set; }
 
 		/// <summary>
 		/// Message produced by trying to auto-submit the change
 		/// </summary>
-		public string? AutoSubmitMessage { get; }
+		public string? AutoSubmitMessage { get; set; }
 
 		/// <summary>
 		/// Time that the job was created
@@ -445,7 +478,7 @@ namespace Horde.Server.Jobs
 		/// <summary>
 		/// Parameters for the job
 		/// </summary>
-		public List<string> Arguments { get; set; }
+		public List<string> Arguments { get; set; } = new List<string>();
 
 		/// <summary>
 		/// The last update time for this job
@@ -465,39 +498,12 @@ namespace Horde.Server.Jobs
 		/// <summary>
 		/// Constructor
 		/// </summary>
-		/// <param name="job">Job to create a response for</param>
-		/// <param name="startedByUserInfo">User that started this job</param>
-		/// <param name="abortedByUserInfo">User that aborted this job</param>
-		public GetJobResponse(IJob job, GetThinUserInfoResponse? startedByUserInfo, GetThinUserInfoResponse? abortedByUserInfo)
+		public GetJobResponse(JobId jobId, StreamId streamId, TemplateId templateId, string name)
 		{
-			Id = job.Id.ToString();
-			StreamId = job.StreamId.ToString();
-			Name = job.Name;
-			Change = job.Change;
-			CodeChange = (job.CodeChange != 0) ? (int?)job.CodeChange : null;
-			PreflightChange = (job.PreflightChange != 0) ? (int?)job.PreflightChange : null;
-			ClonedPreflightChange = (job.ClonedPreflightChange != 0) ? (int?)job.ClonedPreflightChange : null;
-			PreflightDescription = job.PreflightDescription;
-			TemplateId = job.TemplateId.ToString();
-			TemplateHash = job.TemplateHash?.ToString() ?? String.Empty;
-			GraphHash = job.GraphHash.ToString();
-			StartedByUserId = job.StartedByUserId?.ToString();
-			StartedByUser = startedByUserInfo?.Login;
-			StartedByUserInfo = startedByUserInfo;
-			StartedByBisectTaskId = job.StartedByBisectTaskId;
-			AbortedByUser = abortedByUserInfo?.Login;
-			AbortedByUserInfo = abortedByUserInfo;
-			CreateTime = new DateTimeOffset(job.CreateTimeUtc);
-			State = job.GetState();
-			Priority = job.Priority;
-			AutoSubmit = job.AutoSubmit;
-			AutoSubmitChange = job.AutoSubmitChange;
-			AutoSubmitMessage = job.AutoSubmitMessage;
-			Reports = job.Reports?.ConvertAll(x => new GetReportResponse(x));
-			Arguments = job.Arguments.ToList();
-			UpdateTime = new DateTimeOffset(job.UpdateTimeUtc);
-			UseArtifactsV2 = job.JobOptions?.UseNewTempStorage ?? true;
-			UpdateIssues = job.UpdateIssues;
+			Id = jobId;
+			StreamId = streamId;
+			TemplateId = templateId;
+			Name = name;
 		}
 	}
 
@@ -639,7 +645,7 @@ namespace Horde.Server.Jobs
 		/// <summary>
 		/// The unique id of the step
 		/// </summary>
-		public string Id { get; set; }
+		public JobStepId Id { get; set; }
 
 		/// <summary>
 		/// Index of the node which this jobstep is to execute
@@ -715,36 +721,6 @@ namespace Horde.Server.Jobs
 		/// User-defined properties for this jobstep.
 		/// </summary>
 		public Dictionary<string, string>? Properties { get; set; }
-
-		/// <summary>
-		/// Constructor
-		/// </summary>
-		/// <param name="step">The step to construct from</param>
-		/// <param name="abortedByUserInfo">User that aborted this step</param>
-		/// <param name="retriedByUserInfo">User that retried this step</param>
-		public GetStepResponse(IJobStep step, GetThinUserInfoResponse? abortedByUserInfo, GetThinUserInfoResponse? retriedByUserInfo)
-		{
-			Id = step.Id.ToString();
-			NodeIdx = step.NodeIdx;
-			State = step.State;
-			Outcome = step.Outcome;
-			Error = step.Error;
-			AbortRequested = step.AbortRequested;
-			AbortByUser = abortedByUserInfo?.Login;
-			AbortedByUserInfo = abortedByUserInfo;
-			RetryByUser = retriedByUserInfo?.Login;
-			RetriedByUserInfo = retriedByUserInfo;
-			LogId = step.LogId?.ToString();
-			ReadyTime = step.ReadyTimeUtc;
-			StartTime = step.StartTimeUtc;
-			FinishTime = step.FinishTimeUtc;
-			Reports = step.Reports?.ConvertAll(x => new GetReportResponse(x));
-
-			if (step.Properties != null && step.Properties.Count > 0)
-			{
-				Properties = step.Properties;
-			}
-		}
 	}
 
 	/// <summary>
@@ -893,7 +869,7 @@ namespace Horde.Server.Jobs
 		/// <summary>
 		/// Unique id for this batch
 		/// </summary>
-		public string Id { get; set; }
+		public JobStepBatchId Id { get; set; }
 
 		/// <summary>
 		/// The unique log file id
@@ -918,7 +894,7 @@ namespace Horde.Server.Jobs
 		/// <summary>
 		/// Steps within this run
 		/// </summary>
-		public List<GetStepResponse> Steps { get; set; }
+		public List<GetStepResponse> Steps { get; set; } = new List<GetStepResponse>();
 
 		/// <summary>
 		/// The agent assigned to execute this group
@@ -959,31 +935,6 @@ namespace Horde.Server.Jobs
 		/// Time at which the group became ready (UTC).
 		/// </summary>
 		public DateTimeOffset? ReadyTime { get; set; }
-
-		/// <summary>
-		/// Converts this batch into a public response object
-		/// </summary>
-		/// <param name="batch">The batch to construct from</param>
-		/// <param name="steps">Steps in this batch</param>
-		/// <param name="agentRate">Rate for this agent</param>
-		/// <returns>Response instance</returns>
-		public GetBatchResponse(IJobStepBatch batch, List<GetStepResponse> steps, double? agentRate)
-		{
-			Id = batch.Id.ToString();
-			LogId = batch.LogId?.ToString();
-			GroupIdx = batch.GroupIdx;
-			State = batch.State;
-			Error = batch.Error;
-			Steps = steps;
-			AgentId = batch.AgentId?.ToString();
-			AgentRate = agentRate;
-			SessionId = batch.SessionId?.ToString();
-			LeaseId = batch.LeaseId?.ToString();
-			WeightedPriority = batch.SchedulePriority;
-			StartTime = batch.StartTimeUtc;
-			FinishTime = batch.FinishTimeUtc;
-			ReadyTime = batch.ReadyTimeUtc;
-		}
 	}
 
 	/// <summary>
@@ -1117,21 +1068,6 @@ namespace Horde.Server.Jobs
 		/// Average duration to this point
 		/// </summary>
 		public float? AverageTotalTimeToComplete { get; set; }
-
-		/// <summary>
-		/// Constructor
-		/// </summary>
-		/// <param name="timingInfo">Timing info to construct from</param>
-		public GetTimingInfoResponse(TimingInfo timingInfo)
-		{
-			TotalWaitTime = (float?)timingInfo.TotalWaitTime?.TotalSeconds;
-			TotalInitTime = (float?)timingInfo.TotalInitTime?.TotalSeconds;
-			TotalTimeToComplete = (float?)timingInfo.TotalTimeToComplete?.TotalSeconds;
-
-			AverageTotalWaitTime = (float?)timingInfo.AverageTotalWaitTime?.TotalSeconds;
-			AverageTotalInitTime = (float?)timingInfo.AverageTotalInitTime?.TotalSeconds;
-			AverageTotalTimeToComplete = (float?)timingInfo.AverageTotalTimeToComplete?.TotalSeconds;
-		}
 	}
 
 	/// <summary>
@@ -1158,20 +1094,6 @@ namespace Horde.Server.Jobs
 		/// Average duration for this step
 		/// </summary>
 		public float? AverageStepDuration { get; set; }
-
-		/// <summary>
-		/// Constructor
-		/// </summary>
-		/// <param name="name">Name of the node</param>
-		/// <param name="jobTimingInfo">Timing info to construct from</param>
-		public GetStepTimingInfoResponse(string? name, TimingInfo jobTimingInfo)
-			: base(jobTimingInfo)
-		{
-			Name = name;
-			AverageStepWaitTime = jobTimingInfo.StepTiming?.AverageWaitTime;
-			AverageStepInitTime = jobTimingInfo.StepTiming?.AverageInitTime;
-			AverageStepDuration = jobTimingInfo.StepTiming?.AverageDuration;
-		}
 	}
 
 	/// <summary>
@@ -1210,20 +1132,6 @@ namespace Horde.Server.Jobs
 		/// Category for the label
 		/// </summary>
 		public string? UgsProject { get; set; }
-
-		/// <summary>
-		/// Constructor
-		/// </summary>
-		/// <param name="label">The label to construct from</param>
-		/// <param name="timingInfo">Timing info to construct from</param>
-		public GetLabelTimingInfoResponse(ILabel label, TimingInfo timingInfo)
-			: base(timingInfo)
-		{
-			DashboardName = label.DashboardName;
-			DashboardCategory = label.DashboardCategory;
-			UgsName = label.UgsName;
-			UgsProject = label.UgsProject;
-		}
 	}
 
 	/// <summary>
@@ -1234,17 +1142,17 @@ namespace Horde.Server.Jobs
 		/// <summary>
 		/// The job id
 		/// </summary>
-		public string JobId { get; set; }
+		public JobId JobId { get; set; }
 
 		/// <summary>
 		/// The batch containing the step
 		/// </summary>
-		public string BatchId { get; set; }
+		public JobStepBatchId BatchId { get; set; }
 
 		/// <summary>
 		/// The step identifier
 		/// </summary>
-		public string StepId { get; set; }
+		public JobStepId StepId { get; set; }
 
 		/// <summary>
 		/// The change number being built
@@ -1274,35 +1182,16 @@ namespace Horde.Server.Jobs
 		/// <summary>
 		/// The issues which affected this step
 		/// </summary>
-		public List<int>? IssueIds { get; }
+		public List<int>? IssueIds { get; set; }
 
 		/// <summary>
 		/// Time at which the step started.
 		/// </summary>
-		public DateTimeOffset StartTime { get; }
+		public DateTimeOffset StartTime { get; set; }
 
 		/// <summary>
 		/// Time at which the step finished.
 		/// </summary>
-		public DateTimeOffset? FinishTime { get; }
-
-		/// <summary>
-		/// Constructor
-		/// </summary>
-		/// <param name="jobStepRef">The jobstep ref to construct from</param>
-		public GetJobStepRefResponse(IJobStepRef jobStepRef)
-		{
-			JobId = jobStepRef.Id.JobId.ToString();
-			BatchId = jobStepRef.Id.BatchId.ToString();
-			StepId = jobStepRef.Id.StepId.ToString();
-			Change = jobStepRef.Change;
-			LogId = jobStepRef.LogId.ToString();
-			PoolId = jobStepRef.PoolId?.ToString();
-			AgentId = jobStepRef.AgentId?.ToString();
-			Outcome = jobStepRef.Outcome;
-			IssueIds = jobStepRef.IssueIds?.Select(id => id).ToList();
-			StartTime = jobStepRef.StartTimeUtc;
-			FinishTime = jobStepRef.FinishTimeUtc;
-		}
+		public DateTimeOffset? FinishTime { get; set; }
 	}
 }
diff --git a/Engine/Source/Programs/Horde/Horde.Server/Jobs/JobsController.cs b/Engine/Source/Programs/Horde/Horde.Server/Jobs/JobsController.cs
--- a/Engine/Source/Programs/Horde/Horde.Server/Jobs/JobsController.cs
+++ b/Engine/Source/Programs/Horde/Horde.Server/Jobs/JobsController.cs
@@ -475,7 +475,7 @@ namespace Horde.Server.Jobs
 				abortedByUserInfo = (await _userCollection.GetCachedUserAsync(job.AbortedByUserId.Value, cancellationToken))?.ToThinApiResponse();
 			}
 
-			GetJobResponse response = new GetJobResponse(job, startedByUserInfo, abortedByUserInfo);
+			GetJobResponse response = CreateGetJobResponse(job, startedByUserInfo, abortedByUserInfo);
 			if (includeBatches || includeLabels)
 			{
 				if (includeBatches)
@@ -547,6 +547,44 @@ namespace Horde.Server.Jobs
 				&& artifact.Type != ArtifactType.StepTestData;
 		}
 
+		static GetJobResponse CreateGetJobResponse(IJob job, GetThinUserInfoResponse? startedByUserInfo, GetThinUserInfoResponse? abortedByUserInfo)
+		{
+			GetJobResponse response = new GetJobResponse(job.Id, job.StreamId, job.TemplateId, job.Name);
+			response.Change = job.Change;
+			response.CodeChange = (job.CodeChange != 0) ? (int?)job.CodeChange : null;
+			response.PreflightChange = (job.PreflightChange != 0) ? (int?)job.PreflightChange : null;
+			response.ClonedPreflightChange = (job.ClonedPreflightChange != 0) ? (int?)job.ClonedPreflightChange : null;
+			response.PreflightDescription = job.PreflightDescription;
+			response.TemplateHash = job.TemplateHash?.ToString() ?? String.Empty;
+			response.GraphHash = job.GraphHash.ToString();
+			response.StartedByUserId = job.StartedByUserId?.ToString();
+			response.StartedByUser = startedByUserInfo?.Login;
+			response.StartedByUserInfo = startedByUserInfo;
+			response.StartedByBisectTaskId = job.StartedByBisectTaskId;
+			response.AbortedByUser = abortedByUserInfo?.Login;
+			response.AbortedByUserInfo = abortedByUserInfo;
+			response.CreateTime = new DateTimeOffset(job.CreateTimeUtc);
+			response.State = job.GetState();
+			response.Priority = job.Priority;
+			response.AutoSubmit = job.AutoSubmit;
+			response.AutoSubmitChange = job.AutoSubmitChange;
+			response.AutoSubmitMessage = job.AutoSubmitMessage;
+			response.Reports = job.Reports?.ConvertAll(x => CreateGetReportResponse(x));
+			response.Arguments = job.Arguments.ToList();
+			response.UpdateTime = new DateTimeOffset(job.UpdateTimeUtc);
+			response.UseArtifactsV2 = job.JobOptions?.UseNewTempStorage ?? true;
+			response.UpdateIssues = job.UpdateIssues;
+			return response;
+		}
+
+		static GetReportResponse CreateGetReportResponse(IReport report)
+		{
+			GetReportResponse response = new GetReportResponse(report.Name, report.Placement);
+			response.ArtifactId = report.ArtifactId?.ToString();
+			response.Content = report.Content;
+			return response;
+		}
+
 		/// <summary>
 		/// Get the response object for a batch
 		/// </summary>
@@ -568,7 +606,34 @@ namespace Horde.Server.Jobs
 				agentRate = await _agentService.GetRateAsync(batch.AgentId.Value, cancellationToken);
 			}
 
-			return new GetBatchResponse(batch, steps, agentRate);
+			return CreateGetBatchResponse(batch, steps, agentRate);
+		}
+
+		/// <summary>
+		/// Converts this batch into a public response object
+		/// </summary>
+		/// <param name="batch">The batch to construct from</param>
+		/// <param name="steps">Steps in this batch</param>
+		/// <param name="agentRate">Rate for this agent</param>
+		/// <returns>Response instance</returns>
+		static GetBatchResponse CreateGetBatchResponse(IJobStepBatch batch, List<GetStepResponse> steps, double? agentRate)
+		{
+			GetBatchResponse response = new GetBatchResponse();
+			response.Id = batch.Id;
+			response.LogId = batch.LogId?.ToString();
+			response.GroupIdx = batch.GroupIdx;
+			response.State = batch.State;
+			response.Error = batch.Error;
+			response.Steps.AddRange(steps);
+			response.AgentId = batch.AgentId?.ToString();
+			response.AgentRate = agentRate;
+			response.SessionId = batch.SessionId?.ToString();
+			response.LeaseId = batch.LeaseId?.ToString();
+			response.WeightedPriority = batch.SchedulePriority;
+			response.StartTime = batch.StartTimeUtc;
+			response.FinishTime = batch.FinishTimeUtc;
+			response.ReadyTime = batch.ReadyTimeUtc;
+			return response;
 		}
 
 		/// <summary>
@@ -591,7 +656,39 @@ namespace Horde.Server.Jobs
 				retriedByUserInfo = (await _userCollection.GetCachedUserAsync(step.RetriedByUserId.Value, cancellationToken))?.ToThinApiResponse();
 			}
 
-			return new GetStepResponse(step, abortedByUserInfo, retriedByUserInfo);
+			return CreateGetStepResponse(step, abortedByUserInfo, retriedByUserInfo);
+		}
+
+		/// <summary>
+		/// Constructor
+		/// </summary>
+		/// <param name="step">The step to construct from</param>
+		/// <param name="abortedByUserInfo">User that aborted this step</param>
+		/// <param name="retriedByUserInfo">User that retried this step</param>
+		static GetStepResponse CreateGetStepResponse(IJobStep step, GetThinUserInfoResponse? abortedByUserInfo, GetThinUserInfoResponse? retriedByUserInfo)
+		{
+			GetStepResponse response = new GetStepResponse();
+			response.Id = step.Id;
+			response.NodeIdx = step.NodeIdx;
+			response.State = step.State;
+			response.Outcome = step.Outcome;
+			response.Error = step.Error;
+			response.AbortRequested = step.AbortRequested;
+			response.AbortByUser = abortedByUserInfo?.Login;
+			response.AbortedByUserInfo = abortedByUserInfo;
+			response.RetryByUser = retriedByUserInfo?.Login;
+			response.RetriedByUserInfo = retriedByUserInfo;
+			response.LogId = step.LogId?.ToString();
+			response.ReadyTime = step.ReadyTimeUtc;
+			response.StartTime = step.StartTimeUtc;
+			response.FinishTime = step.FinishTimeUtc;
+			response.Reports = step.Reports?.ConvertAll(x => CreateGetReportResponse(x));
+
+			if (step.Properties != null && step.Properties.Count > 0)
+			{
+				response.Properties = step.Properties;
+			}
+			return response;
 		}
 
 		/// <summary>
@@ -668,7 +765,7 @@ namespace Horde.Server.Jobs
 				foreach (IJobStep step in batch.Steps)
 				{
 					INode node = graph.Groups[batch.GroupIdx].Nodes[step.NodeIdx];
-					steps[step.Id.ToString()] = new GetStepTimingInfoResponse(node.Name, nodeToTimingInfo[node]);
+					steps[step.Id.ToString()] = CreateGetStepTimingInfoResponse(node.Name, nodeToTimingInfo[node]);
 				}
 			}
 
@@ -676,7 +773,7 @@ namespace Horde.Server.Jobs
 			foreach (ILabel label in graph.Labels)
 			{
 				TimingInfo timingInfo = TimingInfo.Max(label.GetDependencies(graph.Groups).Select(x => nodeToTimingInfo[x]));
-				labels.Add(new GetLabelTimingInfoResponse(label, timingInfo));
+				labels.Add(CreateGetLabelTimingInfoResponse(label, timingInfo));
 			}
 
 			GetJobResponse? jobResponse = null;
@@ -688,6 +785,30 @@ namespace Horde.Server.Jobs
 			return new GetJobTimingResponse(jobResponse, steps, labels);
 		}
 
+		static GetLabelTimingInfoResponse CreateGetLabelTimingInfoResponse(ILabel label, TimingInfo timingInfo)
+		{
+			GetLabelTimingInfoResponse response = new GetLabelTimingInfoResponse();
+			timingInfo.CopyToResponse(response);
+
+			response.DashboardName = label.DashboardName;
+			response.DashboardCategory = label.DashboardCategory;
+			response.UgsName = label.UgsName;
+			response.UgsProject = label.UgsProject;
+			return response;
+		}
+
+		static GetStepTimingInfoResponse CreateGetStepTimingInfoResponse(string? name, TimingInfo jobTimingInfo)
+		{
+			GetStepTimingInfoResponse response = new GetStepTimingInfoResponse();
+			jobTimingInfo.CopyToResponse(response);
+
+			response.Name = name;
+			response.AverageStepWaitTime = jobTimingInfo.StepTiming?.AverageWaitTime;
+			response.AverageStepInitTime = jobTimingInfo.StepTiming?.AverageInitTime;
+			response.AverageStepDuration = jobTimingInfo.StepTiming?.AverageDuration;
+			return response;
+		}
+
 		/// <summary>
 		/// Find timing information about the graph for multiple jobs
 		/// </summary>
diff --git a/Engine/Source/Programs/Horde/Horde.Server/Jobs/Templates/TemplateConfig.cs b/Engine/Source/Programs/Horde/Horde.Server/Jobs/Templates/TemplateConfig.cs
--- a/Engine/Source/Programs/Horde/Horde.Server/Jobs/Templates/TemplateConfig.cs
+++ b/Engine/Source/Programs/Horde/Horde.Server/Jobs/Templates/TemplateConfig.cs
@@ -8,7 +8,6 @@ using System.Text.Json.Serialization;
 using EpicGames.Core;
 using EpicGames.Horde.Jobs;
 using Horde.Server.Configuration;
-using Horde.Server.Streams;
 
 namespace Horde.Server.Jobs.Templates
 {
diff --git a/Engine/Source/Programs/Horde/Horde.Server/Projects/ProjectConfig.cs b/Engine/Source/Programs/Horde/Horde.Server/Projects/ProjectConfig.cs
--- a/Engine/Source/Programs/Horde/Horde.Server/Projects/ProjectConfig.cs
+++ b/Engine/Source/Programs/Horde/Horde.Server/Projects/ProjectConfig.cs
@@ -8,6 +8,7 @@ using System.Security.Claims;
 using System.Text.Json.Serialization;
 using EpicGames.Core;
 using EpicGames.Horde.Acls;
+using EpicGames.Horde.Jobs;
 using EpicGames.Horde.Projects;
 using EpicGames.Horde.Telemetry;
 using Horde.Server.Acls;
diff --git a/Engine/Source/Programs/Horde/Horde.Server/Streams/StreamConfig.cs b/Engine/Source/Programs/Horde/Horde.Server/Streams/StreamConfig.cs
--- a/Engine/Source/Programs/Horde/Horde.Server/Streams/StreamConfig.cs
+++ b/Engine/Source/Programs/Horde/Horde.Server/Streams/StreamConfig.cs
@@ -15,7 +15,6 @@ using EpicGames.Horde;
 using EpicGames.Horde.Acls;
 using EpicGames.Horde.Agents.Pools;
 using EpicGames.Horde.Commits;
-using EpicGames.Horde.Common;
 using EpicGames.Horde.Issues;
 using EpicGames.Horde.Jobs;
 using EpicGames.Horde.Jobs.Templates;
@@ -711,42 +710,6 @@ namespace Horde.Server.Streams
 		public string? Method { get; set; } = null;
 	}
 
-	/// <summary>
-	/// Query selecting the base changelist to use
-	/// </summary>
-	public class ChangeQueryConfig
-	{
-		/// <summary>
-		/// Name of this query, for display on the dashboard.
-		/// </summary>
-		public string? Name { get; set; }
-
-		/// <summary>
-		/// Condition to evaluate before deciding to use this query. May query tags in a preflight.
-		/// </summary>
-		public Condition? Condition { get; set; }
-
-		/// <summary>
-		/// The template id to query
-		/// </summary>
-		public TemplateId? TemplateId { get; set; }
-
-		/// <summary>
-		/// The target to query
-		/// </summary>
-		public string? Target { get; set; }
-
-		/// <summary>
-		/// Whether to match a job that produced warnings
-		/// </summary>
-		public List<JobStepOutcome>? Outcomes { get; set; }
-
-		/// <summary>
-		/// Finds the last commit with this tag
-		/// </summary>
-		public CommitTag? CommitTag { get; set; }
-	}
-
 	/// <summary>
 	/// Specifies defaults for running a preflight
 	/// </summary>
diff --git a/Engine/Source/Programs/Horde/Horde.Server/Streams/StreamMessages.cs b/Engine/Source/Programs/Horde/Horde.Server/Streams/StreamMessages.cs
--- a/Engine/Source/Programs/Horde/Horde.Server/Streams/StreamMessages.cs
+++ b/Engine/Source/Programs/Horde/Horde.Server/Streams/StreamMessages.cs
@@ -12,6 +12,7 @@ using EpicGames.Horde.Streams;
 using EpicGames.Horde.Users;
 using Horde.Server.Commits;
 using Horde.Server.Issues;
+using Horde.Server.Jobs;
 using Horde.Server.Jobs.Schedules;
 using Horde.Server.Jobs.Templates;
 using Horde.Server.Projects;
diff --git a/Engine/Source/Programs/Horde/Horde.Server/Streams/StreamsController.cs b/Engine/Source/Programs/Horde/Horde.Server/Streams/StreamsController.cs
--- a/Engine/Source/Programs/Horde/Horde.Server/Streams/StreamsController.cs
+++ b/Engine/Source/Programs/Horde/Horde.Server/Streams/StreamsController.cs
@@ -6,6 +6,7 @@ using System.Linq;
 using System.Threading;
 using System.Threading.Tasks;
 using EpicGames.Horde.Commits;
+using EpicGames.Horde.Jobs;
 using EpicGames.Horde.Jobs.Templates;
 using EpicGames.Horde.Projects;
 using EpicGames.Horde.Streams;
@@ -347,7 +348,28 @@ namespace Horde.Server.Streams
 			TemplateId templateIdValue = new TemplateId(templateId);
 
 			List<IJobStepRef> steps = await _jobStepRefCollection.GetStepsForNodeAsync(streamId, templateIdValue, step, change, true, count);
-			return steps.ConvertAll(x => PropertyFilter.Apply(new GetJobStepRefResponse(x), filter));
+			return steps.ConvertAll(x => PropertyFilter.Apply(CreateGetJobStepRefResponse(x), filter));
+		}
+
+		/// <summary>
+		/// Constructor
+		/// </summary>
+		/// <param name="jobStepRef">The jobstep ref to construct from</param>
+		internal static GetJobStepRefResponse CreateGetJobStepRefResponse(IJobStepRef jobStepRef)
+		{
+			GetJobStepRefResponse response = new GetJobStepRefResponse();
+			response.JobId = jobStepRef.Id.JobId;
+			response.BatchId = jobStepRef.Id.BatchId;
+			response.StepId = jobStepRef.Id.StepId;
+			response.Change = jobStepRef.Change;
+			response.LogId = jobStepRef.LogId.ToString();
+			response.PoolId = jobStepRef.PoolId?.ToString();
+			response.AgentId = jobStepRef.AgentId?.ToString();
+			response.Outcome = jobStepRef.Outcome;
+			response.IssueIds = jobStepRef.IssueIds?.Select(id => id).ToList();
+			response.StartTime = jobStepRef.StartTimeUtc;
+			response.FinishTime = jobStepRef.FinishTimeUtc;
+			return response;
 		}
 
 		/// <summary>
