Horde: Prefix all RPC messages with "Rpc", to distinguish from HTTP request/responses by the same name.

diff --git a/Engine/Source/Programs/Horde/Horde.Agent.Tests/Execution/WorkspaceExecutorTest.cs b/Engine/Source/Programs/Horde/Horde.Agent.Tests/Execution/WorkspaceExecutorTest.cs
--- a/Engine/Source/Programs/Horde/Horde.Agent.Tests/Execution/WorkspaceExecutorTest.cs
+++ b/Engine/Source/Programs/Horde/Horde.Agent.Tests/Execution/WorkspaceExecutorTest.cs
@@ -52,7 +52,7 @@ public sealed class WorkspaceExecutorTest : IAsyncDisposable
 		_workspace.SetFile(1, "main.cpp", "main");
 		_workspace.SetFile(1, "foo/bar/baz.h", "baz");
 
-		BeginBatchResponse batch = new BeginBatchResponse { Change = 1 };
+		RpcBeginBatchResponse batch = new RpcBeginBatchResponse { Change = 1 };
 		JobExecutorOptions executorOptions = new JobExecutorOptions(_session, null!, _jobId, _batchId, batch, null!, new RpcJobOptions());
 		_executor = new(executorOptions, _workspace, null, NullLogger.Instance);
 	}
@@ -78,7 +78,7 @@ public sealed class WorkspaceExecutorTest : IAsyncDisposable
 	[TestMethod]
 	public async Task RegularAndAutoSdkWorkspaceAsync()
 	{
-		BeginBatchResponse batch = new BeginBatchResponse { Change = 1 };
+		RpcBeginBatchResponse batch = new RpcBeginBatchResponse { Change = 1 };
 		JobExecutorOptions executorOptions = new JobExecutorOptions(_session, null!, _jobId, _batchId, batch, null!, new RpcJobOptions());
 		using WorkspaceExecutor executor = new(executorOptions, _workspace, _autoSdkWorkspace, NullLogger.Instance);
 
@@ -92,7 +92,7 @@ public sealed class WorkspaceExecutorTest : IAsyncDisposable
 	[TestMethod]
 	public async Task EnvVarsAsync()
 	{
-		BeginBatchResponse batch = new BeginBatchResponse { Change = 1, StreamName = "//UE5/Main" };
+		RpcBeginBatchResponse batch = new RpcBeginBatchResponse { Change = 1, StreamName = "//UE5/Main" };
 		JobExecutorOptions executorOptions = new JobExecutorOptions(_session, null!, _jobId, _batchId, batch, null!, new RpcJobOptions());
 		using WorkspaceExecutor executor = new(executorOptions, _workspace, _autoSdkWorkspace, NullLogger.Instance);
 		await executor.InitializeAsync(_logger, CancellationToken.None);
@@ -119,7 +119,7 @@ public sealed class WorkspaceExecutorTest : IAsyncDisposable
 		_server.AddJob(preflightJobId, _streamId, 1, 1000);
 		_workspace.SetFile(1000, "New/Feature/Foo.cs", "foo");
 
-		BeginBatchResponse batch = new BeginBatchResponse { Change = 1, PreflightChange = 1000 };
+		RpcBeginBatchResponse batch = new RpcBeginBatchResponse { Change = 1, PreflightChange = 1000 };
 		JobExecutorOptions executorOptions = new JobExecutorOptions(_session, null!, preflightJobId, _batchId, batch, null!, new RpcJobOptions());
 		using WorkspaceExecutor executor = new(executorOptions, _workspace, null, NullLogger.Instance);
 
@@ -136,7 +136,7 @@ public sealed class WorkspaceExecutorTest : IAsyncDisposable
 		JobId noChangeJobId = JobId.Parse("65bd0655591b5d5d7d047b5a");
 
 		_server.AddJob(noChangeJobId, _streamId, 0, 0);
-		BeginBatchResponse batch = new BeginBatchResponse { };
+		RpcBeginBatchResponse batch = new RpcBeginBatchResponse { };
 		JobExecutorOptions executorOptions = new JobExecutorOptions(_session, null!, noChangeJobId, _batchId, batch, null!, new RpcJobOptions());
 		using WorkspaceExecutor executor = new(executorOptions, _workspace, null, NullLogger.Instance);
 		await Assert.ThrowsExceptionAsync<WorkspaceMaterializationException>(() => executor.InitializeAsync(_logger, CancellationToken.None));
diff --git a/Engine/Source/Programs/Horde/Horde.Agent.Tests/Helper.cs b/Engine/Source/Programs/Horde/Horde.Agent.Tests/Helper.cs
--- a/Engine/Source/Programs/Horde/Horde.Agent.Tests/Helper.cs
+++ b/Engine/Source/Programs/Horde/Horde.Agent.Tests/Helper.cs
@@ -103,10 +103,10 @@ namespace Horde.Agent.Tests
 
 	class JobRpcClientStub : JobRpc.JobRpcClient
 	{
-		public readonly Queue<BeginStepResponse> BeginStepResponses = new Queue<BeginStepResponse>();
-		public readonly List<UpdateStepRequest> UpdateStepRequests = new List<UpdateStepRequest>();
-		public readonly Dictionary<GetStepRequest, GetStepResponse> GetStepResponses = new Dictionary<GetStepRequest, GetStepResponse>();
-		public Func<GetStepRequest, GetStepResponse>? _getStepFunc = null;
+		public readonly Queue<RpcBeginStepResponse> BeginStepResponses = new Queue<RpcBeginStepResponse>();
+		public readonly List<RpcUpdateStepRequest> UpdateStepRequests = new List<RpcUpdateStepRequest>();
+		public readonly Dictionary<RpcGetStepRequest, RpcGetStepResponse> GetStepResponses = new Dictionary<RpcGetStepRequest, RpcGetStepResponse>();
+		public Func<RpcGetStepRequest, RpcGetStepResponse>? _getStepFunc = null;
 		private readonly ILogger _logger;
 
 		public JobRpcClientStub(ILogger logger)
@@ -114,11 +114,11 @@ namespace Horde.Agent.Tests
 			_logger = logger;
 		}
 
-		public override AsyncUnaryCall<BeginBatchResponse> BeginBatchAsync(BeginBatchRequest request,
+		public override AsyncUnaryCall<RpcBeginBatchResponse> BeginBatchAsync(RpcBeginBatchRequest request,
 			CallOptions options)
 		{
 			_logger.LogDebug("HordeRpcClientStub.BeginBatchAsync()");
-			BeginBatchResponse res = new BeginBatchResponse();
+			RpcBeginBatchResponse res = new RpcBeginBatchResponse();
 
 			res.AgentType = "agentType1";
 			res.LogId = "logId1";
@@ -127,39 +127,39 @@ namespace Horde.Agent.Tests
 			return Wrap(res);
 		}
 
-		public override AsyncUnaryCall<Empty> FinishBatchAsync(FinishBatchRequest request, CallOptions options)
+		public override AsyncUnaryCall<Empty> FinishBatchAsync(RpcFinishBatchRequest request, CallOptions options)
 		{
 			Empty res = new Empty();
 			return Wrap(res);
 		}
 
-		public override AsyncUnaryCall<GetStreamResponse> GetStreamAsync(GetStreamRequest request, CallOptions options)
+		public override AsyncUnaryCall<RpcGetStreamResponse> GetStreamAsync(RpcGetStreamRequest request, CallOptions options)
 		{
-			GetStreamResponse res = new GetStreamResponse();
+			RpcGetStreamResponse res = new RpcGetStreamResponse();
 			return Wrap(res);
 		}
 
-		public override AsyncUnaryCall<GetJobResponse> GetJobAsync(GetJobRequest request, CallOptions options)
+		public override AsyncUnaryCall<RpcGetJobResponse> GetJobAsync(RpcGetJobRequest request, CallOptions options)
 		{
-			GetJobResponse res = new GetJobResponse();
+			RpcGetJobResponse res = new RpcGetJobResponse();
 			return Wrap(res);
 		}
 
-		public override AsyncUnaryCall<BeginStepResponse> BeginStepAsync(BeginStepRequest request, CallOptions options)
+		public override AsyncUnaryCall<RpcBeginStepResponse> BeginStepAsync(RpcBeginStepRequest request, CallOptions options)
 		{
 			if (BeginStepResponses.Count == 0)
 			{
-				BeginStepResponse completeRes = new BeginStepResponse();
-				completeRes.State = BeginStepResponse.Types.Result.Complete;
+				RpcBeginStepResponse completeRes = new RpcBeginStepResponse();
+				completeRes.State = RpcBeginStepResponse.Types.Result.Complete;
 				return Wrap(completeRes);
 			}
 
-			BeginStepResponse res = BeginStepResponses.Dequeue();
-			res.State = BeginStepResponse.Types.Result.Ready;
+			RpcBeginStepResponse res = BeginStepResponses.Dequeue();
+			res.State = RpcBeginStepResponse.Types.Result.Ready;
 			return Wrap(res);
 		}
 
-		public override AsyncUnaryCall<Empty> UpdateStepAsync(UpdateStepRequest request, CallOptions options)
+		public override AsyncUnaryCall<Empty> UpdateStepAsync(RpcUpdateStepRequest request, CallOptions options)
 		{
 			_logger.LogDebug("UpdateStepAsync(Request: {Request})", request);
 			UpdateStepRequests.Add(request);
@@ -167,26 +167,26 @@ namespace Horde.Agent.Tests
 			return Wrap(res);
 		}
 
-		public override AsyncUnaryCall<Empty> CreateEventsAsync(CreateEventsRequest request, CallOptions options)
+		public override AsyncUnaryCall<Empty> CreateEventsAsync(RpcCreateEventsRequest request, CallOptions options)
 		{
 			_logger.LogDebug("CreateEventsAsync: {Request}", request);
 			Empty res = new Empty();
 			return Wrap(res);
 		}
 
-		public override AsyncUnaryCall<GetStepResponse> GetStepAsync(GetStepRequest request, CallOptions options)
+		public override AsyncUnaryCall<RpcGetStepResponse> GetStepAsync(RpcGetStepRequest request, CallOptions options)
 		{
 			if (_getStepFunc != null)
 			{
 				return Wrap(_getStepFunc(request));
 			}
 
-			if (GetStepResponses.TryGetValue(request, out GetStepResponse? res))
+			if (GetStepResponses.TryGetValue(request, out RpcGetStepResponse? res))
 			{
 				return Wrap(res);
 			}
 
-			return Wrap(new GetStepResponse());
+			return Wrap(new RpcGetStepResponse());
 		}
 
 		public static AsyncUnaryCall<T> Wrap<T>(T res)
@@ -241,6 +241,6 @@ namespace Horde.Agent.Tests
 			_executor = executor;
 		}
 
-		public IJobExecutor CreateExecutor(AgentWorkspace? workspaceInfo, AgentWorkspace? autoSdkWorkspaceInfo, JobExecutorOptions options) => _executor;
+		public IJobExecutor CreateExecutor(RpcAgentWorkspace? workspaceInfo, RpcAgentWorkspace? autoSdkWorkspaceInfo, JobExecutorOptions options) => _executor;
 	}
 }
\ No newline at end of file
diff --git a/Engine/Source/Programs/Horde/Horde.Agent.Tests/WorkerServiceTest.cs b/Engine/Source/Programs/Horde/Horde.Agent.Tests/WorkerServiceTest.cs
--- a/Engine/Source/Programs/Horde/Horde.Agent.Tests/WorkerServiceTest.cs
+++ b/Engine/Source/Programs/Horde/Horde.Agent.Tests/WorkerServiceTest.cs
@@ -162,8 +162,8 @@ namespace Horde.Agent.Tests
 			executeJobTask.LogId = _logId.ToString();
 			executeJobTask.JobName = "jobName1";
 			executeJobTask.JobOptions = new RpcJobOptions { Executor = SimpleTestExecutor.Name };
-			executeJobTask.AutoSdkWorkspace = new AgentWorkspace();
-			executeJobTask.Workspace = new AgentWorkspace();
+			executeJobTask.AutoSdkWorkspace = new RpcAgentWorkspace();
+			executeJobTask.Workspace = new RpcAgentWorkspace();
 
 			JobRpcClientStub client = new JobRpcClientStub(NullLogger.Instance);
 			await using RpcConnectionStub rpcConnection = new RpcConnectionStub(null!, null!, client);
@@ -171,12 +171,12 @@ namespace Horde.Agent.Tests
 			await using FakeHordeRpcServer fakeServer = new();
 			await using ISession session = FakeServerSessionFactory.CreateSession(rpcConnection, fakeServer.GetGrpcChannel());
 
-			client.BeginStepResponses.Enqueue(new BeginStepResponse { Name = "stepName1", StepId = _stepId1.ToString() });
-			client.BeginStepResponses.Enqueue(new BeginStepResponse { Name = "stepName2", StepId = _stepId2.ToString() });
-			client.BeginStepResponses.Enqueue(new BeginStepResponse { Name = "stepName3", StepId = _stepId3.ToString() });
+			client.BeginStepResponses.Enqueue(new RpcBeginStepResponse { Name = "stepName1", StepId = _stepId1.ToString() });
+			client.BeginStepResponses.Enqueue(new RpcBeginStepResponse { Name = "stepName2", StepId = _stepId2.ToString() });
+			client.BeginStepResponses.Enqueue(new RpcBeginStepResponse { Name = "stepName3", StepId = _stepId3.ToString() });
 
-			GetStepRequest step2Req = new GetStepRequest(_jobId, _batchId, _stepId2);
-			GetStepResponse step2Res = new GetStepResponse(JobStepOutcome.Unspecified, JobStepState.Unspecified, true);
+			RpcGetStepRequest step2Req = new RpcGetStepRequest(_jobId, _batchId, _stepId2);
+			RpcGetStepResponse step2Res = new RpcGetStepResponse(JobStepOutcome.Unspecified, JobStepState.Unspecified, true);
 			client.GetStepResponses[step2Req] = step2Res;
 
 			using SimpleTestExecutor executor = new SimpleTestExecutor(async (step, logger, cancelToken) =>
@@ -227,9 +227,9 @@ namespace Horde.Agent.Tests
 			{
 				return ++c switch
 				{
-					1 => new GetStepResponse { AbortRequested = false },
+					1 => new RpcGetStepResponse { AbortRequested = false },
 					2 => throw new RpcException(new Status(StatusCode.Cancelled, "Fake cancel from test")),
-					3 => new GetStepResponse { AbortRequested = true },
+					3 => new RpcGetStepResponse { AbortRequested = true },
 					_ => throw new Exception("Should never reach here")
 				};
 			};
@@ -307,10 +307,10 @@ namespace Horde.Agent.Tests
 	{
 		private readonly string _serverName;
 		private readonly bool _isStopping = false;
-		private readonly Dictionary<string, Lease> _leases = new();
+		private readonly Dictionary<string, RpcLease> _leases = new();
 
-		private readonly Dictionary<StreamId, GetStreamResponse> _streamIdToStreamResponse = new();
-		private readonly Dictionary<JobId, GetJobResponse> _jobIdToJobResponse = new();
+		private readonly Dictionary<StreamId, RpcGetStreamResponse> _streamIdToStreamResponse = new();
+		private readonly Dictionary<JobId, RpcGetJobResponse> _jobIdToJobResponse = new();
 		private readonly Mock<IRpcClientRef<HordeRpc.HordeRpcClient>> _mockClientRef;
 		private readonly Mock<IRpcConnection> _mockConnection;
 		private readonly ILogger<FakeHordeRpcServer> _logger;
@@ -330,7 +330,7 @@ namespace Horde.Agent.Tests
 				_outer = outer;
 			}
 
-			public override AsyncDuplexStreamingCall<UpdateSessionRequest, UpdateSessionResponse> UpdateSession(Metadata headers = null!, DateTime? deadline = null, CancellationToken cancellationToken = default)
+			public override AsyncDuplexStreamingCall<RpcUpdateSessionRequest, RpcUpdateSessionResponse> UpdateSession(Metadata headers = null!, DateTime? deadline = null, CancellationToken cancellationToken = default)
 			{
 				return _outer.GetUpdateSessionCall(CancellationToken.None);
 			}
@@ -345,9 +345,9 @@ namespace Horde.Agent.Tests
 				_outer = outer;
 			}
 
-			public override AsyncUnaryCall<GetStreamResponse> GetStreamAsync(GetStreamRequest request, CallOptions options)
+			public override AsyncUnaryCall<RpcGetStreamResponse> GetStreamAsync(RpcGetStreamRequest request, CallOptions options)
 			{
-				if (_outer._streamIdToStreamResponse.TryGetValue(new StreamId(request.StreamId), out GetStreamResponse? streamResponse))
+				if (_outer._streamIdToStreamResponse.TryGetValue(new StreamId(request.StreamId), out RpcGetStreamResponse? streamResponse))
 				{
 					return JobRpcClientStub.Wrap(streamResponse);
 				}
@@ -355,9 +355,9 @@ namespace Horde.Agent.Tests
 				throw new RpcException(new Status(StatusCode.NotFound, $"Stream ID {request.StreamId} not found"));
 			}
 
-			public override AsyncUnaryCall<GetJobResponse> GetJobAsync(GetJobRequest request, CallOptions options)
+			public override AsyncUnaryCall<RpcGetJobResponse> GetJobAsync(RpcGetJobRequest request, CallOptions options)
 			{
-				if (_outer._jobIdToJobResponse.TryGetValue(JobId.Parse(request.JobId), out GetJobResponse? jobResponse))
+				if (_outer._jobIdToJobResponse.TryGetValue(JobId.Parse(request.JobId), out RpcGetJobResponse? jobResponse))
 				{
 					return JobRpcClientStub.Wrap(jobResponse);
 				}
@@ -400,7 +400,7 @@ namespace Horde.Agent.Tests
 			}
 
 			TestTask testTask = new();
-			_leases[leaseId] = new Lease
+			_leases[leaseId] = new RpcLease
 			{
 				Id = leaseId,
 				State = RpcLeaseState.Pending,
@@ -408,7 +408,7 @@ namespace Horde.Agent.Tests
 			};
 		}
 
-		public Lease GetLease(string leaseId)
+		public RpcLease GetLease(string leaseId)
 		{
 			return _leases[leaseId];
 		}
@@ -420,7 +420,7 @@ namespace Horde.Agent.Tests
 				throw new Exception($"Stream ID {streamId} already added");
 			}
 
-			_streamIdToStreamResponse[streamId] = new GetStreamResponse
+			_streamIdToStreamResponse[streamId] = new RpcGetStreamResponse
 			{
 				Name = streamName,
 			};
@@ -428,7 +428,7 @@ namespace Horde.Agent.Tests
 
 		public void AddAgentType(StreamId streamId, string agentType)
 		{
-			if (!_streamIdToStreamResponse.TryGetValue(streamId, out GetStreamResponse? streamResponse))
+			if (!_streamIdToStreamResponse.TryGetValue(streamId, out RpcGetStreamResponse? streamResponse))
 			{
 				throw new Exception($"Stream ID {streamId} not found");
 			}
@@ -436,7 +436,7 @@ namespace Horde.Agent.Tests
 			string tempDir = Path.Join(Path.GetTempPath(), $"horde-agent-type-{agentType}-" + Guid.NewGuid().ToString()[..8]);
 			Directory.CreateDirectory(tempDir);
 
-			streamResponse.AgentTypes[agentType] = new GetAgentTypeResponse
+			streamResponse.AgentTypes[agentType] = new RpcGetAgentTypeResponse
 			{
 				TempStorageDir = tempDir
 			};
@@ -449,7 +449,7 @@ namespace Horde.Agent.Tests
 				throw new Exception($"Stream ID {streamId} not found");
 			}
 
-			_jobIdToJobResponse[jobId] = new GetJobResponse
+			_jobIdToJobResponse[jobId] = new RpcGetJobResponse
 			{
 				StreamId = streamId.ToString(),
 				Change = change,
@@ -467,11 +467,11 @@ namespace Horde.Agent.Tests
 			return _grpcChannel;
 		}
 
-		public CreateSessionResponse OnCreateSessionRequest(CreateSessionRequest request)
+		public RpcCreateSessionResponse OnCreateSessionRequest(RpcCreateSessionRequest request)
 		{
 			CreateSessionReceived.TrySetResult(true);
 			_logger.LogInformation("OnCreateSessionRequest: {AgentId} {Status}", request.Id, request.Status);
-			CreateSessionResponse response = new()
+			RpcCreateSessionResponse response = new()
 			{
 				AgentId = "bogusAgentId",
 				Token = "bogusToken",
@@ -482,16 +482,16 @@ namespace Horde.Agent.Tests
 			return response;
 		}
 
-		public AsyncDuplexStreamingCall<QueryServerStateRequest, QueryServerStateResponse> GetQueryServerStateCall(CancellationToken cancellationToken)
+		public AsyncDuplexStreamingCall<RpcQueryServerStateRequest, RpcQueryServerStateResponse> GetQueryServerStateCall(CancellationToken cancellationToken)
 		{
-			FakeAsyncStreamReader<QueryServerStateResponse> responseStream = new(cancellationToken);
-			FakeClientStreamWriter<QueryServerStateRequest> requestStream = new(onComplete: () =>
+			FakeAsyncStreamReader<RpcQueryServerStateResponse> responseStream = new(cancellationToken);
+			FakeClientStreamWriter<RpcQueryServerStateRequest> requestStream = new(onComplete: () =>
 			{
 				responseStream.Complete();
 				return Task.CompletedTask;
 			});
 
-			responseStream.Write(new QueryServerStateResponse { Name = _serverName, Stopping = _isStopping });
+			responseStream.Write(new RpcQueryServerStateResponse { Name = _serverName, Stopping = _isStopping });
 
 			return new(
 				requestStream,
@@ -502,17 +502,17 @@ namespace Horde.Agent.Tests
 				() => { /*isDisposed = true;*/ });
 		}
 
-		public AsyncDuplexStreamingCall<UpdateSessionRequest, UpdateSessionResponse> GetUpdateSessionCall(CancellationToken cancellationToken)
+		public AsyncDuplexStreamingCall<RpcUpdateSessionRequest, RpcUpdateSessionResponse> GetUpdateSessionCall(CancellationToken cancellationToken)
 		{
-			FakeAsyncStreamReader<UpdateSessionResponse> responseStream = new(cancellationToken);
+			FakeAsyncStreamReader<RpcUpdateSessionResponse> responseStream = new(cancellationToken);
 
-			async Task OnRequest(UpdateSessionRequest request)
+			async Task OnRequest(RpcUpdateSessionRequest request)
 			{
 				UpdateSessionReceived.TrySetResult(true);
 
-				foreach (Lease agentLease in request.Leases)
+				foreach (RpcLease agentLease in request.Leases)
 				{
-					Lease serverLease = _leases[agentLease.Id];
+					RpcLease serverLease = _leases[agentLease.Id];
 					serverLease.State = agentLease.State;
 					serverLease.Outcome = agentLease.Outcome;
 					serverLease.Output = agentLease.Output;
@@ -520,12 +520,12 @@ namespace Horde.Agent.Tests
 
 				_logger.LogInformation("OnUpdateSessionRequest: {AgentId} {SessionId} {Status}", request.AgentId, request.SessionId, request.Status);
 				await Task.Delay(100, cancellationToken);
-				UpdateSessionResponse response = new() { ExpiryTime = Timestamp.FromDateTime(DateTime.UtcNow + TimeSpan.FromMinutes(120)) };
+				RpcUpdateSessionResponse response = new() { ExpiryTime = Timestamp.FromDateTime(DateTime.UtcNow + TimeSpan.FromMinutes(120)) };
 				response.Leases.AddRange(_leases.Values.Where(x => x.State != RpcLeaseState.Completed));
 				await responseStream.Write(response);
 			}
 
-			FakeClientStreamWriter<UpdateSessionRequest> requestStream = new(OnRequest, () =>
+			FakeClientStreamWriter<RpcUpdateSessionRequest> requestStream = new(OnRequest, () =>
 			{
 				responseStream.Complete();
 				return Task.CompletedTask;
@@ -555,9 +555,9 @@ namespace Horde.Agent.Tests
 			await _connection.DisposeAsync();
 			_grpcChannel.Dispose();
 
-			foreach (GetStreamResponse stream in _streamIdToStreamResponse.Values)
+			foreach (RpcGetStreamResponse stream in _streamIdToStreamResponse.Values)
 			{
-				foreach (GetAgentTypeResponse agentType in stream.AgentTypes.Values)
+				foreach (RpcGetAgentTypeResponse agentType in stream.AgentTypes.Values)
 				{
 					if (Directory.Exists(agentType.TempStorageDir))
 					{
diff --git a/Engine/Source/Programs/Horde/Horde.Agent.Tests/WorkspaceInfoTest.cs b/Engine/Source/Programs/Horde/Horde.Agent.Tests/WorkspaceInfoTest.cs
--- a/Engine/Source/Programs/Horde/Horde.Agent.Tests/WorkspaceInfoTest.cs
+++ b/Engine/Source/Programs/Horde/Horde.Agent.Tests/WorkspaceInfoTest.cs
@@ -43,7 +43,7 @@ public class WorkspaceInfoTest
 
 	static ManagedWorkspaceOptions GetMwOptions(string? method)
 	{
-		HordeCommon.Rpc.Messages.AgentWorkspace workspace = new HordeCommon.Rpc.Messages.AgentWorkspace();
+		HordeCommon.Rpc.Messages.RpcAgentWorkspace workspace = new HordeCommon.Rpc.Messages.RpcAgentWorkspace();
 		workspace.Method = method ?? "";
 		return WorkspaceInfo.GetMwOptions(workspace);
 	}
diff --git a/Engine/Source/Programs/Horde/Horde.Agent/Commands/Utilities/BuildGraphCommand.cs b/Engine/Source/Programs/Horde/Horde.Agent/Commands/Utilities/BuildGraphCommand.cs
--- a/Engine/Source/Programs/Horde/Horde.Agent/Commands/Utilities/BuildGraphCommand.cs
+++ b/Engine/Source/Programs/Horde/Horde.Agent/Commands/Utilities/BuildGraphCommand.cs
@@ -52,9 +52,9 @@ namespace Horde.Agent.Commands.Utilities
 
 			public Task SetOutcomeAsync(JobStepOutcome outcome, CancellationToken cancellationToken) => Task.CompletedTask;
 
-			public async Task WriteEventsAsync(List<CreateEventRequest> events, CancellationToken cancellationToken)
+			public async Task WriteEventsAsync(List<RpcCreateEventRequest> events, CancellationToken cancellationToken)
 			{
-				foreach (CreateEventRequest request in events)
+				foreach (RpcCreateEventRequest request in events)
 				{
 					JsonSerializerOptions options = new JsonSerializerOptions();
 					options.Converters.Add(new JsonStringEnumConverter());
@@ -63,7 +63,7 @@ namespace Horde.Agent.Commands.Utilities
 				}
 			}
 
-			public async Task WriteOutputAsync(WriteOutputRequest request, CancellationToken cancellationToken)
+			public async Task WriteOutputAsync(RpcWriteOutputRequest request, CancellationToken cancellationToken)
 			{
 				PrintJson(request.Data.Span);
 				await _outputStream.WriteAsync(request.Data.Memory, cancellationToken);
diff --git a/Engine/Source/Programs/Horde/Horde.Agent/Commands/Utilities/CapsCommand.cs b/Engine/Source/Programs/Horde/Horde.Agent/Commands/Utilities/CapsCommand.cs
--- a/Engine/Source/Programs/Horde/Horde.Agent/Commands/Utilities/CapsCommand.cs
+++ b/Engine/Source/Programs/Horde/Horde.Agent/Commands/Utilities/CapsCommand.cs
@@ -26,8 +26,8 @@ namespace Horde.Agent.Commands.Utilities
 		/// <inheritdoc/>
 		public override async Task<int> ExecuteAsync(ILogger logger)
 		{
-			AgentCapabilities capabilities = await _capabilitiesService.GetCapabilitiesAsync(null);
-			foreach (DeviceCapabilities device in capabilities.Devices)
+			RpcAgentCapabilities capabilities = await _capabilitiesService.GetCapabilitiesAsync(null);
+			foreach (RpcDeviceCapabilities device in capabilities.Devices)
 			{
 				logger.LogInformation("Device: {Name}", device.Handle);
 				foreach (string property in device.Properties)
diff --git a/Engine/Source/Programs/Horde/Horde.Agent/Execution/JobExecutor.cs b/Engine/Source/Programs/Horde/Horde.Agent/Execution/JobExecutor.cs
--- a/Engine/Source/Programs/Horde/Horde.Agent/Execution/JobExecutor.cs
+++ b/Engine/Source/Programs/Horde/Horde.Agent/Execution/JobExecutor.cs
@@ -42,10 +42,10 @@ namespace Horde.Agent.Execution
 		IReadOnlyList<string> Inputs,
 		IReadOnlyList<string> OutputNames,
 		IList<int> PublishOutputs,
-		IReadOnlyList<CreateGraphArtifactRequest> Artifacts
+		IReadOnlyList<RpcCreateGraphArtifactRequest> Artifacts
 	)
 	{
-		public JobStepInfo(BeginStepResponse response)
+		public JobStepInfo(RpcBeginStepResponse response)
 			: this(JobStepId.Parse(response.StepId), LogId.Parse(response.LogId), response.Name, response.Credentials, response.Properties, response.EnvVars, response.Warnings, response.Inputs, response.OutputNames, response.PublishOutputs, response.Artifacts)
 		{
 		}
@@ -64,11 +64,11 @@ namespace Horde.Agent.Execution
 		public HttpStorageClientFactory StorageFactory { get; }
 		public JobId JobId { get; }
 		public JobStepBatchId BatchId { get; }
-		public BeginBatchResponse Batch { get; }
+		public RpcBeginBatchResponse Batch { get; }
 		public string Token { get; }
 		public RpcJobOptions JobOptions { get; }
 
-		public JobExecutorOptions(ISession session, HttpStorageClientFactory storageFactory, JobId jobId, JobStepBatchId batchId, BeginBatchResponse batch, string token, RpcJobOptions jobOptions)
+		public JobExecutorOptions(ISession session, HttpStorageClientFactory storageFactory, JobId jobId, JobStepBatchId batchId, RpcBeginBatchResponse batch, string token, RpcJobOptions jobOptions)
 		{
 			Session = session;
 			StorageFactory = storageFactory;
@@ -112,7 +112,7 @@ namespace Horde.Agent.Execution
 			public string? Category { get; set; }
 			public string? UgsBadge { get; set; }
 			public string? UgsProject { get; set; }
-			public LabelChange Change { get; set; } = LabelChange.Current;
+			public RpcLabelChange Change { get; set; } = RpcLabelChange.Current;
 			public List<string> RequiredNodes { get; set; } = new List<string>();
 			public List<string> IncludedNodes { get; set; } = new List<string>();
 		}
@@ -197,7 +197,7 @@ namespace Horde.Agent.Execution
 
 		class ReportData
 		{
-			public ReportScope Scope { get; set; }
+			public RpcReportScope Scope { get; set; }
 			public RpcReportPlacement Placement { get; set; }
 			public string Name { get; set; } = String.Empty;
 			public string Content { get; set; } = String.Empty;
@@ -222,7 +222,7 @@ namespace Horde.Agent.Execution
 
 		protected JobId JobId { get; }
 		protected JobStepBatchId BatchId { get; }
-		protected BeginBatchResponse Batch { get; }
+		protected RpcBeginBatchResponse Batch { get; }
 
 		protected List<string> _additionalArguments = new List<string>();
 		protected bool _allowTargetChanges;
@@ -548,13 +548,13 @@ namespace Horde.Agent.Execution
 					ArtifactName artifactName = TempStorage.GetArtifactNameForNode(SetupStepName);
 					ArtifactType artifactType = ArtifactType.StepOutput;
 
-					CreateJobArtifactRequestV2 artifactRequest = new CreateJobArtifactRequestV2();
+					RpcCreateJobArtifactRequestV2 artifactRequest = new RpcCreateJobArtifactRequestV2();
 					artifactRequest.JobId = JobId.ToString();
 					artifactRequest.StepId = step.StepId.ToString();
 					artifactRequest.Name = artifactName.ToString();
 					artifactRequest.Type = artifactType.ToString();
 
-					CreateJobArtifactResponseV2 artifact = await jobRpc.Client.CreateArtifactV2Async(artifactRequest, cancellationToken: cancellationToken);
+					RpcCreateJobArtifactResponseV2 artifact = await jobRpc.Client.CreateArtifactV2Async(artifactRequest, cancellationToken: cancellationToken);
 					logger.LogInformation("Creating output artifact {ArtifactId} '{ArtifactName}' ({ArtifactType}) with ref {RefName} in namespace {NamespaceId}", artifact.Id, artifactName, artifactType, artifact.RefName, artifact.NamespaceId);
 
 					// Write the data
@@ -586,7 +586,7 @@ namespace Horde.Agent.Execution
 				await StorePreprocessedFileAsync(preprocessedSchemaFile, step.StepId, sharedStorageDir, logger, cancellationToken);
 			}
 
-			UpdateGraphRequest updateGraph = await ParseGraphUpdateAsync(definitionFile, logger, cancellationToken);
+			RpcUpdateGraphRequest updateGraph = await ParseGraphUpdateAsync(definitionFile, logger, cancellationToken);
 			await RpcConnection.InvokeAsync((JobRpc.JobRpcClient x) => x.UpdateGraphAsync(updateGraph, null, null, cancellationToken), cancellationToken);
 
 			HashSet<string> validTargets = new HashSet<string>(StringComparer.OrdinalIgnoreCase);
@@ -605,7 +605,7 @@ namespace Horde.Agent.Execution
 			return true;
 		}
 
-		private async Task<UpdateGraphRequest> ParseGraphUpdateAsync(FileReference definitionFile, ILogger logger, CancellationToken cancellationToken)
+		private async Task<RpcUpdateGraphRequest> ParseGraphUpdateAsync(FileReference definitionFile, ILogger logger, CancellationToken cancellationToken)
 		{
 			JsonSerializerOptions options = new JsonSerializerOptions();
 			options.PropertyNameCaseInsensitive = true;
@@ -613,7 +613,7 @@ namespace Horde.Agent.Execution
 
 			ExportedGraph graph = JsonSerializer.Deserialize<ExportedGraph>(await FileReference.ReadAllBytesAsync(definitionFile, cancellationToken), options)!;
 
-			UpdateGraphRequest updateGraph = new UpdateGraphRequest();
+			RpcUpdateGraphRequest updateGraph = new RpcUpdateGraphRequest();
 			updateGraph.JobId = JobId.ToString();
 
 			List<string> missingAgentTypes = new List<string>();
@@ -644,12 +644,12 @@ namespace Horde.Agent.Execution
 					}
 				}
 
-				CreateGroupRequest createGroup = new CreateGroupRequest();
+				RpcCreateGroupRequest createGroup = new RpcCreateGroupRequest();
 				createGroup.AgentType = agentTypeName;
 
 				foreach (ExportedNode exportedNode in exportedGroup.Nodes)
 				{
-					CreateNodeRequest createNode = new CreateNodeRequest();
+					RpcCreateNodeRequest createNode = new RpcCreateNodeRequest();
 					createNode.Name = exportedNode.Name;
 					if (exportedNode.Inputs != null)
 					{
@@ -687,7 +687,7 @@ namespace Horde.Agent.Execution
 
 			foreach (ExportedAggregate exportedAggregate in graph.Aggregates)
 			{
-				CreateAggregateRequest createAggregate = new CreateAggregateRequest();
+				RpcCreateAggregateRequest createAggregate = new RpcCreateAggregateRequest();
 				createAggregate.Name = exportedAggregate.Name;
 				createAggregate.Nodes.AddRange(exportedAggregate.Nodes);
 				updateGraph.Aggregates.Add(createAggregate);
@@ -695,7 +695,7 @@ namespace Horde.Agent.Execution
 
 			foreach (ExportedLabel exportedLabel in graph.Labels)
 			{
-				CreateLabelRequest createLabel = new CreateLabelRequest();
+				RpcCreateLabelRequest createLabel = new RpcCreateLabelRequest();
 				if (exportedLabel.Name != null)
 				{
 					createLabel.DashboardName = exportedLabel.Name;
@@ -721,7 +721,7 @@ namespace Horde.Agent.Execution
 			Dictionary<string, ExportedNode> nameToNode = graph.Groups.SelectMany(x => x.Nodes).ToDictionary(x => x.Name, x => x);
 			foreach (ExportedBadge exportedBadge in graph.Badges)
 			{
-				CreateLabelRequest createLabel = new CreateLabelRequest();
+				RpcCreateLabelRequest createLabel = new RpcCreateLabelRequest();
 				createLabel.UgsName = exportedBadge.Name;
 
 				string? project = exportedBadge.Project;
@@ -740,11 +740,11 @@ namespace Horde.Agent.Execution
 
 				if (exportedBadge.Change == Batch.Change || exportedBadge.Change == 0)
 				{
-					createLabel.Change = LabelChange.Current;
+					createLabel.Change = RpcLabelChange.Current;
 				}
 				else if (exportedBadge.Change == Batch.CodeChange)
 				{
-					createLabel.Change = LabelChange.Code;
+					createLabel.Change = RpcLabelChange.Code;
 				}
 				else
 				{
@@ -767,7 +767,7 @@ namespace Horde.Agent.Execution
 
 			foreach (ExportedArtifact exportedArtifact in graph.Artifacts)
 			{
-				CreateGraphArtifactRequest createArtifact = new CreateGraphArtifactRequest();
+				RpcCreateGraphArtifactRequest createArtifact = new RpcCreateGraphArtifactRequest();
 
 				createArtifact.Name = exportedArtifact.Name;
 				createArtifact.Type = exportedArtifact.Type ?? String.Empty;
@@ -812,13 +812,13 @@ namespace Horde.Agent.Execution
 
 					using IRpcClientRef<JobRpc.JobRpcClient> jobRpc = await RpcConnection.GetClientRefAsync<JobRpc.JobRpcClient>(cancellationToken);
 
-					GetJobArtifactRequest artifactRequest = new GetJobArtifactRequest();
+					RpcGetJobArtifactRequest artifactRequest = new RpcGetJobArtifactRequest();
 					artifactRequest.JobId = JobId.ToString();
 					artifactRequest.StepId = step.StepId.ToString();
 					artifactRequest.Name = artifactName.ToString();
 					artifactRequest.Type = ArtifactType.StepOutput.ToString();
 
-					GetJobArtifactResponse artifact = await jobRpc.Client.GetArtifactAsync(artifactRequest, cancellationToken: cancellationToken);
+					RpcGetJobArtifactResponse artifact = await jobRpc.Client.GetArtifactAsync(artifactRequest, cancellationToken: cancellationToken);
 
 					NamespaceId namespaceId = new NamespaceId(artifact.NamespaceId);
 					RefName refName = new RefName(artifact.RefName);
@@ -896,13 +896,13 @@ namespace Horde.Agent.Execution
 			{
 				using IRpcClientRef<JobRpc.JobRpcClient> jobRpc = await RpcConnection.GetClientRefAsync<JobRpc.JobRpcClient>(cancellationToken);
 
-				CreateJobArtifactRequestV2 artifactRequest = new CreateJobArtifactRequestV2();
+				RpcCreateJobArtifactRequestV2 artifactRequest = new RpcCreateJobArtifactRequestV2();
 				artifactRequest.JobId = JobId.ToString();
 				artifactRequest.StepId = stepId.ToString();
 				artifactRequest.Name = name.ToString();
 				artifactRequest.Type = type.ToString();
 
-				CreateJobArtifactResponseV2 artifact = await jobRpc.Client.CreateArtifactV2Async(artifactRequest, cancellationToken: cancellationToken);
+				RpcCreateJobArtifactResponseV2 artifact = await jobRpc.Client.CreateArtifactV2Async(artifactRequest, cancellationToken: cancellationToken);
 				Logger.LogInformation("Created artifact {ArtifactId} '{ArtifactName}' ({ArtifactType}) with ref {RefName} ({Link})", artifact.Id, name, type, artifact.RefName, $"{Session.ServerUrl}/api/v1/storage/{artifact.NamespaceId}/refs/{artifact.RefName}");
 
 				using IStorageClient storage = CreateStorageClient(new NamespaceId(artifact.NamespaceId), artifact.Token);
@@ -1096,13 +1096,13 @@ namespace Horde.Agent.Execution
 			using (GlobalTracer.Instance.BuildSpan("TempStorage").WithTag("resource", "Write").StartActive())
 			{
 				// Create the artifact
-				CreateJobArtifactRequestV2 artifactRequest = new CreateJobArtifactRequestV2();
+				RpcCreateJobArtifactRequestV2 artifactRequest = new RpcCreateJobArtifactRequestV2();
 				artifactRequest.JobId = JobId.ToString();
 				artifactRequest.StepId = step.StepId.ToString();
 				artifactRequest.Name = TempStorage.GetArtifactNameForNode(step.Name).ToString();
 				artifactRequest.Type = ArtifactType.StepOutput.ToString();
 
-				CreateJobArtifactResponseV2 artifact = await jobRpc.Client.CreateArtifactV2Async(artifactRequest, cancellationToken: cancellationToken);
+				RpcCreateJobArtifactResponseV2 artifact = await jobRpc.Client.CreateArtifactV2Async(artifactRequest, cancellationToken: cancellationToken);
 				logger.LogInformation("Created artifact {ArtifactId} '{ArtifactName}' ({ArtifactType}) with ref {RefName} ({RefUrl})", artifact.Id, artifactRequest.Name, ArtifactType.StepOutput, artifact.RefName, $"{Session.ServerUrl.ToString().TrimEnd('/')}/api/v1/storage/{artifact.NamespaceId}/refs/{artifact.RefName}");
 
 				using IStorageClient storage = CreateStorageClient(new NamespaceId(artifact.NamespaceId), artifact.Token);
@@ -1157,7 +1157,7 @@ namespace Horde.Agent.Execution
 
 			// Create all the named artifacts. TODO: Merge this with regular temp storage artifacts?
 			logger.LogInformation("Uploading {NumArtifacts} artifacts", step.Artifacts.Count);
-			foreach (CreateGraphArtifactRequest graphArtifact in step.Artifacts)
+			foreach (RpcCreateGraphArtifactRequest graphArtifact in step.Artifacts)
 			{
 				logger.LogInformation("Uploading artifact {Name} using output {Output}", graphArtifact.Name, graphArtifact.OutputName);
 
@@ -1171,7 +1171,7 @@ namespace Horde.Agent.Execution
 				using (GlobalTracer.Instance.BuildSpan("Artifact").WithTag("name", graphArtifact.Name).WithTag("resource", "Write").StartActive())
 				{
 					// Create the artifact
-					CreateJobArtifactRequestV2 artifactRequest = new CreateJobArtifactRequestV2();
+					RpcCreateJobArtifactRequestV2 artifactRequest = new RpcCreateJobArtifactRequestV2();
 					artifactRequest.JobId = JobId.ToString();
 					artifactRequest.StepId = step.StepId.ToString();
 					artifactRequest.Name = graphArtifact.Name;
@@ -1180,7 +1180,7 @@ namespace Horde.Agent.Execution
 					artifactRequest.Keys.AddRange(graphArtifact.Keys);
 					artifactRequest.Metadata.AddRange(graphArtifact.Metadata);
 
-					CreateJobArtifactResponseV2 artifact = await jobRpc.Client.CreateArtifactV2Async(artifactRequest, cancellationToken: cancellationToken);
+					RpcCreateJobArtifactResponseV2 artifact = await jobRpc.Client.CreateArtifactV2Async(artifactRequest, cancellationToken: cancellationToken);
 					logger.LogInformation("Created artifact {ArtifactId} '{ArtifactName}' ({ArtifactType}) with ref {RefName} ({RefUrl})", artifact.Id, artifactRequest.Name, ArtifactType.StepOutput, artifact.RefName, $"{Session.ServerUrl.ToString().TrimEnd('/')}/api/v1/storage/{artifact.NamespaceId}/refs/{artifact.RefName}");
 
 					using IStorageClient storage = CreateStorageClient(new NamespaceId(artifact.NamespaceId), artifact.Token);
@@ -1737,11 +1737,11 @@ namespace Horde.Agent.Execution
 			{
 				jobLogger.LogInformation("Parsing graph update from {File}", graphUpdateFile);
 
-				UpdateGraphRequest updateGraph = await ParseGraphUpdateAsync(graphUpdateFile, jobLogger, cancellationToken);
-				foreach (CreateGroupRequest group in updateGraph.Groups)
+				RpcUpdateGraphRequest updateGraph = await ParseGraphUpdateAsync(graphUpdateFile, jobLogger, cancellationToken);
+				foreach (RpcCreateGroupRequest group in updateGraph.Groups)
 				{
 					jobLogger.LogInformation("  AgentType: {Name}", group.AgentType);
-					foreach (CreateNodeRequest node in group.Nodes)
+					foreach (RpcCreateNodeRequest node in group.Nodes)
 					{
 						jobLogger.LogInformation("    Node: {Name}", node.Name);
 					}
@@ -1917,7 +1917,7 @@ namespace Horde.Agent.Execution
 				report.Content = await FileReference.ReadAllTextAsync(reportDataFile);
 			}
 
-			CreateReportRequest request = new CreateReportRequest();
+			RpcCreateReportRequest request = new RpcCreateReportRequest();
 			request.JobId = JobId.ToString();
 			request.BatchId = BatchId.ToString();
 			request.StepId = stepId.ToString();
@@ -1966,7 +1966,7 @@ namespace Horde.Agent.Execution
 
 		async Task<bool> UploadTestDataAsync(JobRpc.JobRpcClient rpcClient, JobStepId jobStepId, IEnumerable<KeyValuePair<string, object>> pairs)
 		{
-			using (AsyncClientStreamingCall<UploadTestDataRequest, UploadTestDataResponse> call = rpcClient.UploadTestData())
+			using (AsyncClientStreamingCall<RpcUploadTestDataRequest, RpcUploadTestDataResponse> call = rpcClient.UploadTestData())
 			{
 				foreach (KeyValuePair<string, object> pair in pairs)
 				{
@@ -1975,7 +1975,7 @@ namespace Horde.Agent.Execution
 					options.Converters.Add(new JsonStringEnumConverter());
 					byte[] data = JsonSerializer.SerializeToUtf8Bytes(pair.Value, options);
 
-					UploadTestDataRequest request = new UploadTestDataRequest();
+					RpcUploadTestDataRequest request = new RpcUploadTestDataRequest();
 					request.JobId = JobId.ToString();
 					request.JobStepId = jobStepId.ToString();
 					request.Key = pair.Key;
@@ -1993,6 +1993,6 @@ namespace Horde.Agent.Execution
 	{
 		string Name { get; }
 
-		IJobExecutor CreateExecutor(AgentWorkspace workspaceInfo, AgentWorkspace? autoSdkWorkspaceInfo, JobExecutorOptions options);
+		IJobExecutor CreateExecutor(RpcAgentWorkspace workspaceInfo, RpcAgentWorkspace? autoSdkWorkspaceInfo, JobExecutorOptions options);
 	}
 }
diff --git a/Engine/Source/Programs/Horde/Horde.Agent/Execution/LocalExecutor.cs b/Engine/Source/Programs/Horde/Horde.Agent/Execution/LocalExecutor.cs
--- a/Engine/Source/Programs/Horde/Horde.Agent/Execution/LocalExecutor.cs
+++ b/Engine/Source/Programs/Horde/Horde.Agent/Execution/LocalExecutor.cs
@@ -75,7 +75,7 @@ namespace Horde.Agent.Execution
 			_logger = logger;
 		}
 
-		public IJobExecutor CreateExecutor(AgentWorkspace workspaceInfo, AgentWorkspace? autoSdkWorkspaceInfo, JobExecutorOptions options)
+		public IJobExecutor CreateExecutor(RpcAgentWorkspace workspaceInfo, RpcAgentWorkspace? autoSdkWorkspaceInfo, JobExecutorOptions options)
 		{
 			return new LocalExecutor(options, _settings, _logger);
 		}
diff --git a/Engine/Source/Programs/Horde/Horde.Agent/Execution/ManagedWorkspaceMaterializer.cs b/Engine/Source/Programs/Horde/Horde.Agent/Execution/ManagedWorkspaceMaterializer.cs
--- a/Engine/Source/Programs/Horde/Horde.Agent/Execution/ManagedWorkspaceMaterializer.cs
+++ b/Engine/Source/Programs/Horde/Horde.Agent/Execution/ManagedWorkspaceMaterializer.cs
@@ -15,7 +15,7 @@ namespace Horde.Agent.Execution;
 /// </summary>
 public sealed class ManagedWorkspaceMaterializer : IWorkspaceMaterializer
 {
-	private readonly AgentWorkspace _agentWorkspace;
+	private readonly RpcAgentWorkspace _agentWorkspace;
 	private readonly DirectoryReference _workingDir;
 	private readonly bool _useCacheFile;
 	private WorkspaceInfo? _workspace;
@@ -27,7 +27,7 @@ public sealed class ManagedWorkspaceMaterializer : IWorkspaceMaterializer
 	/// <param name="workingDir">Where to put synced Perforce files and any cached data/metadata</param>
 	/// <param name="useCacheFile">Whether to use a cache file during syncs</param>
 	public ManagedWorkspaceMaterializer(
-		AgentWorkspace agentWorkspace,
+		RpcAgentWorkspace agentWorkspace,
 		DirectoryReference workingDir,
 		bool useCacheFile)
 	{
diff --git a/Engine/Source/Programs/Horde/Horde.Agent/Execution/PerforceExecutor.cs b/Engine/Source/Programs/Horde/Horde.Agent/Execution/PerforceExecutor.cs
--- a/Engine/Source/Programs/Horde/Horde.Agent/Execution/PerforceExecutor.cs
+++ b/Engine/Source/Programs/Horde/Horde.Agent/Execution/PerforceExecutor.cs
@@ -20,15 +20,15 @@ namespace Horde.Agent.Execution
 	{
 		public const string Name = "Perforce";
 
-		protected AgentWorkspace _workspaceInfo;
-		protected AgentWorkspace? _autoSdkWorkspaceInfo;
+		protected RpcAgentWorkspace _workspaceInfo;
+		protected RpcAgentWorkspace? _autoSdkWorkspaceInfo;
 		protected DirectoryReference _rootDir;
 		protected DirectoryReference? _sharedStorageDir;
 
 		protected WorkspaceInfo? _autoSdkWorkspace;
 		protected WorkspaceInfo _workspace;
 
-		public PerforceExecutor(AgentWorkspace workspaceInfo, AgentWorkspace? autoSdkWorkspaceInfo, JobExecutorOptions options, ILogger logger)
+		public PerforceExecutor(RpcAgentWorkspace workspaceInfo, RpcAgentWorkspace? autoSdkWorkspaceInfo, JobExecutorOptions options, ILogger logger)
 			: base(options, logger)
 		{
 			_workspaceInfo = workspaceInfo;
@@ -109,7 +109,7 @@ namespace Horde.Agent.Execution
 					List<ChangesRecord> changes = await _workspace.PerforceClient.GetChangesAsync(ChangesOptions.None, 1, ChangeStatus.Submitted, new[] { Batch.StreamName + "/..." }, cancellationToken);
 					Batch.Change = changes[0].Number;
 
-					UpdateJobRequest updateJobRequest = new UpdateJobRequest();
+					RpcUpdateJobRequest updateJobRequest = new RpcUpdateJobRequest();
 					updateJobRequest.JobId = JobId.ToString();
 					updateJobRequest.Change = Batch.Change;
 					await RpcConnection.InvokeAsync((JobRpc.JobRpcClient x) => x.UpdateJobAsync(updateJobRequest, null, null, cancellationToken), cancellationToken);
@@ -241,7 +241,7 @@ namespace Horde.Agent.Execution
 			await base.FinalizeAsync(logger, cancellationToken);
 		}
 
-		public static async Task ConformAsync(DirectoryReference rootDir, IList<AgentWorkspace> pendingWorkspaces, bool removeUntrackedFiles, ILogger logger, CancellationToken cancellationToken)
+		public static async Task ConformAsync(DirectoryReference rootDir, IList<RpcAgentWorkspace> pendingWorkspaces, bool removeUntrackedFiles, ILogger logger, CancellationToken cancellationToken)
 		{
 			using IScope scope = GlobalTracer.Instance.BuildSpan("Conform").StartActive();
 			scope.Span.SetTag("workspaces", String.Join(',', pendingWorkspaces.Select(x => x.Identifier)));
@@ -249,7 +249,7 @@ namespace Horde.Agent.Execution
 
 			// Print out all the workspaces we're going to sync
 			logger.LogInformation("Workspaces:");
-			foreach (AgentWorkspace pendingWorkspace in pendingWorkspaces)
+			foreach (RpcAgentWorkspace pendingWorkspace in pendingWorkspaces)
 			{
 				logger.LogInformation("  Identifier={Identifier}, Stream={StreamName}, Incremental={Incremental} Method={Method} Partitioned={Partitioned}",
 					pendingWorkspace.Identifier, pendingWorkspace.Stream, pendingWorkspace.Incremental, pendingWorkspace.Method, pendingWorkspace.Partitioned);
@@ -261,7 +261,7 @@ namespace Horde.Agent.Execution
 			try
 			{
 				// Set up all the workspaces
-				foreach (AgentWorkspace pendingWorkspace in pendingWorkspaces)
+				foreach (RpcAgentWorkspace pendingWorkspace in pendingWorkspaces)
 				{
 					ManagedWorkspaceOptions options = WorkspaceInfo.GetMwOptions(pendingWorkspace);
 					WorkspaceInfo workspace = await WorkspaceInfo.SetupWorkspaceAsync(pendingWorkspace, rootDir, options, logger, cancellationToken);
@@ -435,7 +435,7 @@ namespace Horde.Agent.Execution
 			_logger = logger;
 		}
 
-		public IJobExecutor CreateExecutor(AgentWorkspace workspaceInfo, AgentWorkspace? autoSdkWorkspaceInfo, JobExecutorOptions options)
+		public IJobExecutor CreateExecutor(RpcAgentWorkspace workspaceInfo, RpcAgentWorkspace? autoSdkWorkspaceInfo, JobExecutorOptions options)
 		{
 			return new PerforceExecutor(workspaceInfo, autoSdkWorkspaceInfo, options, _logger);
 		}
diff --git a/Engine/Source/Programs/Horde/Horde.Agent/Execution/TestExecutor.cs b/Engine/Source/Programs/Horde/Horde.Agent/Execution/TestExecutor.cs
--- a/Engine/Source/Programs/Horde/Horde.Agent/Execution/TestExecutor.cs
+++ b/Engine/Source/Programs/Horde/Horde.Agent/Execution/TestExecutor.cs
@@ -79,31 +79,31 @@ namespace Horde.Agent.Execution
 
 			string projectName = GetArgument("Project", "UnknownProject");
 
-			UpdateGraphRequest updateGraph = new UpdateGraphRequest();
+			RpcUpdateGraphRequest updateGraph = new RpcUpdateGraphRequest();
 			updateGraph.JobId = JobId.ToString();
 
-			CreateGroupRequest winEditorGroup = CreateGroup("Win64");
+			RpcCreateGroupRequest winEditorGroup = CreateGroup("Win64");
 			winEditorGroup.Nodes.Add(CreateNode("Update Version Files", Array.Empty<string>(), JobStepOutcome.Success));
 			winEditorGroup.Nodes.Add(CreateNode("Compile UnrealHeaderTool Win64", new string[] { "Update Version Files" }, JobStepOutcome.Success));
 			winEditorGroup.Nodes.Add(CreateNode("Compile UnrealEditor Win64", new string[] { "Compile UnrealHeaderTool Win64" }, JobStepOutcome.Success));
 			winEditorGroup.Nodes.Add(CreateNode($"Compile {projectName}Editor Win64", new string[] { "Compile UnrealHeaderTool Win64", "Compile UnrealEditor Win64" }, JobStepOutcome.Success));
 			updateGraph.Groups.Add(winEditorGroup);
 
-			CreateGroupRequest winToolsGroup = CreateGroup("Win64");
+			RpcCreateGroupRequest winToolsGroup = CreateGroup("Win64");
 			winToolsGroup.Nodes.Add(CreateNode("Compile Tools Win64", new string[] { "Compile UnrealHeaderTool Win64" }, warningOutcome));
 			updateGraph.Groups.Add(winToolsGroup);
 
-			CreateGroupRequest winClientsGroup = CreateGroup("Win64");
+			RpcCreateGroupRequest winClientsGroup = CreateGroup("Win64");
 			winClientsGroup.Nodes.Add(CreateNode($"Compile {projectName}Client Win64", new string[] { "Compile UnrealHeaderTool Win64" }, JobStepOutcome.Success));
 			updateGraph.Groups.Add(winClientsGroup);
 
-			CreateGroupRequest winCooksGroup = CreateGroup("Win64");
+			RpcCreateGroupRequest winCooksGroup = CreateGroup("Win64");
 			winCooksGroup.Nodes.Add(CreateNode($"Cook {projectName}Client Win64", new string[] { $"Compile {projectName}Editor Win64", "Compile Tools Win64" }, warningOutcome));
 			winCooksGroup.Nodes.Add(CreateNode($"Stage {projectName}Client Win64", new string[] { $"Cook {projectName}Client Win64", "Compile Tools Win64" }, failureOutcome));
 			winCooksGroup.Nodes.Add(CreateNode($"Publish {projectName}Client Win64", new string[] { $"Stage {projectName}Client Win64" }, JobStepOutcome.Success));
 			updateGraph.Groups.Add(winCooksGroup);
 
-			CreateAggregateRequest aggregate = new CreateAggregateRequest();
+			RpcCreateAggregateRequest aggregate = new RpcCreateAggregateRequest();
 			aggregate.Name = "Full Build";
 			aggregate.Nodes.Add($"Publish {projectName}Client Win64");
 			updateGraph.Aggregates.Add(aggregate);
@@ -119,16 +119,16 @@ namespace Horde.Agent.Execution
 			return true;
 		}
 
-		static CreateGroupRequest CreateGroup(string agentType)
+		static RpcCreateGroupRequest CreateGroup(string agentType)
 		{
-			CreateGroupRequest request = new CreateGroupRequest();
+			RpcCreateGroupRequest request = new RpcCreateGroupRequest();
 			request.AgentType = agentType;
 			return request;
 		}
 
-		static CreateNodeRequest CreateNode(string name, string[] inputDependencies, JobStepOutcome outcome)
+		static RpcCreateNodeRequest CreateNode(string name, string[] inputDependencies, JobStepOutcome outcome)
 		{
-			CreateNodeRequest request = new CreateNodeRequest();
+			RpcCreateNodeRequest request = new RpcCreateNodeRequest();
 			request.Name = name;
 			request.InputDependencies.AddRange(inputDependencies);
 			request.Properties.Add("Action", "Build");
@@ -136,9 +136,9 @@ namespace Horde.Agent.Execution
 			return request;
 		}
 
-		static CreateLabelRequest CreateLabel(string category, string name, string[] requiredNodes, string[] includedNodes, Dictionary<string, string[]> dependencyMap)
+		static RpcCreateLabelRequest CreateLabel(string category, string name, string[] requiredNodes, string[] includedNodes, Dictionary<string, string[]> dependencyMap)
 		{
-			CreateLabelRequest request = new CreateLabelRequest();
+			RpcCreateLabelRequest request = new RpcCreateLabelRequest();
 			request.DashboardName = name;
 			request.DashboardCategory = category;
 			request.RequiredNodes.AddRange(requiredNodes);
@@ -146,12 +146,12 @@ namespace Horde.Agent.Execution
 			return request;
 		}
 
-		static Dictionary<string, string[]> CreateDependencyMap(IEnumerable<CreateGroupRequest> groups)
+		static Dictionary<string, string[]> CreateDependencyMap(IEnumerable<RpcCreateGroupRequest> groups)
 		{
 			Dictionary<string, string[]> nameToDependencyNames = new Dictionary<string, string[]>();
-			foreach (CreateGroupRequest group in groups)
+			foreach (RpcCreateGroupRequest group in groups)
 			{
-				foreach (CreateNodeRequest node in group.Nodes)
+				foreach (RpcCreateNodeRequest node in group.Nodes)
 				{
 					HashSet<string> dependencyNames = new HashSet<string> { node.Name };
 
@@ -237,7 +237,7 @@ namespace Horde.Agent.Execution
 			_logger = logger;
 		}
 
-		public IJobExecutor CreateExecutor(AgentWorkspace workspaceInfo, AgentWorkspace? autoSdkWorkspaceInfo, JobExecutorOptions options)
+		public IJobExecutor CreateExecutor(RpcAgentWorkspace workspaceInfo, RpcAgentWorkspace? autoSdkWorkspaceInfo, JobExecutorOptions options)
 		{
 			return new TestExecutor(options, _logger);
 		}
diff --git a/Engine/Source/Programs/Horde/Horde.Agent/Execution/WorkspaceExecutor.cs b/Engine/Source/Programs/Horde/Horde.Agent/Execution/WorkspaceExecutor.cs
--- a/Engine/Source/Programs/Horde/Horde.Agent/Execution/WorkspaceExecutor.cs
+++ b/Engine/Source/Programs/Horde/Horde.Agent/Execution/WorkspaceExecutor.cs
@@ -170,7 +170,7 @@ namespace Horde.Agent.Execution
 			_loggerFactory = loggerFactory;
 		}
 
-		public IJobExecutor CreateExecutor(AgentWorkspace workspaceInfo, AgentWorkspace? autoSdkWorkspaceInfo, JobExecutorOptions options)
+		public IJobExecutor CreateExecutor(RpcAgentWorkspace workspaceInfo, RpcAgentWorkspace? autoSdkWorkspaceInfo, JobExecutorOptions options)
 		{
 			IWorkspaceMaterializer? workspaceMaterializer = null;
 			IWorkspaceMaterializer? autoSdkMaterializer = null;
diff --git a/Engine/Source/Programs/Horde/Horde.Agent/Execution/WorkspaceMaterializer.cs b/Engine/Source/Programs/Horde/Horde.Agent/Execution/WorkspaceMaterializer.cs
--- a/Engine/Source/Programs/Horde/Horde.Agent/Execution/WorkspaceMaterializer.cs
+++ b/Engine/Source/Programs/Horde/Horde.Agent/Execution/WorkspaceMaterializer.cs
@@ -161,13 +161,13 @@ interface IWorkspaceMaterializerFactory
 	/// <param name="options">Job options</param>
 	/// <param name="forAutoSdk">Whether intended for AutoSDK materialization</param>
 	/// <returns>A new workspace materializer instance</returns>
-	IWorkspaceMaterializer CreateMaterializer(WorkspaceMaterializerType type, AgentWorkspace workspaceInfo, JobExecutorOptions options, bool forAutoSdk = false);
+	IWorkspaceMaterializer CreateMaterializer(WorkspaceMaterializerType type, RpcAgentWorkspace workspaceInfo, JobExecutorOptions options, bool forAutoSdk = false);
 }
 
 class WorkspaceMaterializerFactory : IWorkspaceMaterializerFactory
 {
 	/// <inheritdoc/>
-	public IWorkspaceMaterializer CreateMaterializer(WorkspaceMaterializerType type, AgentWorkspace workspaceInfo, JobExecutorOptions options, bool forAutoSdk)
+	public IWorkspaceMaterializer CreateMaterializer(WorkspaceMaterializerType type, RpcAgentWorkspace workspaceInfo, JobExecutorOptions options, bool forAutoSdk)
 	{
 		switch (type)
 		{
diff --git a/Engine/Source/Programs/Horde/Horde.Agent/Leases/Handlers/ConformHandler.cs b/Engine/Source/Programs/Horde/Horde.Agent/Leases/Handlers/ConformHandler.cs
--- a/Engine/Source/Programs/Horde/Horde.Agent/Leases/Handlers/ConformHandler.cs
+++ b/Engine/Source/Programs/Horde/Horde.Agent/Leases/Handlers/ConformHandler.cs
@@ -49,7 +49,7 @@ namespace Horde.Agent.Leases.Handlers
 			await session.TerminateProcessesAsync(TerminateCondition.BeforeConform, logger, cancellationToken);
 
 			bool removeUntrackedFiles = conformTask.RemoveUntrackedFiles;
-			IList<AgentWorkspace> pendingWorkspaces = conformTask.Workspaces;
+			IList<RpcAgentWorkspace> pendingWorkspaces = conformTask.Workspaces;
 			for (; ; )
 			{
 				bool isPerforceExecutor = _settings.Executor.Equals(PerforceExecutor.Name, StringComparison.OrdinalIgnoreCase);
@@ -72,12 +72,12 @@ namespace Horde.Agent.Leases.Handlers
 				}
 
 				// Update the new set of workspaces
-				UpdateAgentWorkspacesRequest request = new UpdateAgentWorkspacesRequest();
+				RpcUpdateAgentWorkspacesRequest request = new RpcUpdateAgentWorkspacesRequest();
 				request.AgentId = session.AgentId.ToString();
 				request.Workspaces.AddRange(pendingWorkspaces);
 				request.RemoveUntrackedFiles = removeUntrackedFiles;
 
-				UpdateAgentWorkspacesResponse response = await session.RpcConnection.InvokeAsync((HordeRpc.HordeRpcClient x) => x.UpdateAgentWorkspacesAsync(request, null, null, cancellationToken), cancellationToken);
+				RpcUpdateAgentWorkspacesResponse response = await session.RpcConnection.InvokeAsync((HordeRpc.HordeRpcClient x) => x.UpdateAgentWorkspacesAsync(request, null, null, cancellationToken), cancellationToken);
 				if (!response.Retry)
 				{
 					logger.LogInformation("Conform finished");
diff --git a/Engine/Source/Programs/Horde/Horde.Agent/Leases/Handlers/JobHandler.cs b/Engine/Source/Programs/Horde/Horde.Agent/Leases/Handlers/JobHandler.cs
--- a/Engine/Source/Programs/Horde/Horde.Agent/Leases/Handlers/JobHandler.cs
+++ b/Engine/Source/Programs/Horde/Horde.Agent/Leases/Handlers/JobHandler.cs
@@ -24,7 +24,7 @@ using OpenTracing.Util;
 
 namespace Horde.Agent.Leases.Handlers
 {
-	record class JobTaskInfo(string JobName, JobId JobId, JobStepBatchId BatchId, RpcJobOptions JobOptions, string Token, LogId LogId, AgentWorkspace Workspace, AgentWorkspace? AutoSdkWorkspace);
+	record class JobTaskInfo(string JobName, JobId JobId, JobStepBatchId BatchId, RpcJobOptions JobOptions, string Token, LogId LogId, RpcAgentWorkspace Workspace, RpcAgentWorkspace? AutoSdkWorkspace);
 
 	class JobHandler : LeaseHandler<ExecuteJobTask>
 	{
@@ -167,7 +167,7 @@ namespace Horde.Agent.Leases.Handlers
 			logger.LogInformation("Executor: {Name}, UseNewTempStorage: {UseNewTempStorage}", jobOptions.Executor, jobOptions.UseNewTempStorage ?? false);
 
 			// Start executing the current batch
-			BeginBatchResponse batch = await session.RpcConnection.InvokeAsync<JobRpc.JobRpcClient, BeginBatchResponse>(x => x.BeginBatchAsync(new BeginBatchRequest(executeTask.JobId, executeTask.BatchId, leaseId), null, null, cancellationToken), cancellationToken);
+			RpcBeginBatchResponse batch = await session.RpcConnection.InvokeAsync<JobRpc.JobRpcClient, RpcBeginBatchResponse>(x => x.BeginBatchAsync(new RpcBeginBatchRequest(executeTask.JobId, executeTask.BatchId, leaseId), null, null, cancellationToken), cancellationToken);
 			try
 			{
 				JobExecutorOptions options = new JobExecutorOptions(session, _serverStorageFactory, executeTask.JobId, executeTask.BatchId, batch, executeTask.Token, jobOptions);
@@ -201,7 +201,7 @@ namespace Horde.Agent.Leases.Handlers
 			}
 
 			// Mark the batch as complete
-			await session.RpcConnection.InvokeAsync((JobRpc.JobRpcClient x) => x.FinishBatchAsync(new FinishBatchRequest(executeTask.JobId, executeTask.BatchId, leaseId), null, null, cancellationToken), cancellationToken);
+			await session.RpcConnection.InvokeAsync((JobRpc.JobRpcClient x) => x.FinishBatchAsync(new RpcFinishBatchRequest(executeTask.JobId, executeTask.BatchId, leaseId), null, null, cancellationToken), cancellationToken);
 			logger.LogInformation("Done.");
 
 			return LeaseResult.Success;
@@ -210,7 +210,7 @@ namespace Horde.Agent.Leases.Handlers
 		/// <summary>
 		/// Executes a batch
 		/// </summary>
-		async Task ExecuteBatchAsync(ISession session, LeaseId leaseId, AgentWorkspace workspaceInfo, AgentWorkspace? autoSdkWorkspaceInfo, JobExecutorOptions options, ILogger logger, ILogger localLogger, CancellationToken cancellationToken)
+		async Task ExecuteBatchAsync(ISession session, LeaseId leaseId, RpcAgentWorkspace workspaceInfo, RpcAgentWorkspace? autoSdkWorkspaceInfo, JobExecutorOptions options, ILogger logger, ILogger localLogger, CancellationToken cancellationToken)
 		{
 			IRpcConnection rpcClient = session.RpcConnection;
 
@@ -244,19 +244,19 @@ namespace Horde.Agent.Leases.Handlers
 				for (; ; )
 				{
 					// Get the next step to execute
-					BeginStepResponse stepResponse = await rpcClient.InvokeAsync((JobRpc.JobRpcClient x) => x.BeginStepAsync(new BeginStepRequest(options.JobId, options.BatchId, leaseId), null, null, cancellationToken), cancellationToken);
-					if (stepResponse.State == BeginStepResponse.Types.Result.Waiting)
+					RpcBeginStepResponse stepResponse = await rpcClient.InvokeAsync((JobRpc.JobRpcClient x) => x.BeginStepAsync(new RpcBeginStepRequest(options.JobId, options.BatchId, leaseId), null, null, cancellationToken), cancellationToken);
+					if (stepResponse.State == RpcBeginStepResponse.Types.Result.Waiting)
 					{
 						logger.LogInformation("Waiting for dependency to be ready");
 						await Task.Delay(TimeSpan.FromSeconds(20.0), cancellationToken);
 						continue;
 					}
-					else if (stepResponse.State == BeginStepResponse.Types.Result.Complete)
+					else if (stepResponse.State == RpcBeginStepResponse.Types.Result.Complete)
 					{
 						logger.LogInformation("No more steps to execute; finalizing lease.");
 						break;
 					}
-					else if (stepResponse.State != BeginStepResponse.Types.Result.Ready)
+					else if (stepResponse.State != RpcBeginStepResponse.Types.Result.Ready)
 					{
 						logger.LogError("Unexpected step state: {StepState}", stepResponse.State);
 						break;
@@ -344,7 +344,7 @@ namespace Horde.Agent.Leases.Handlers
 
 						// Update the server with the outcome from the step
 						logger.LogInformation("Marking step as complete (Outcome={Outcome}, State={StepState})", stepOutcome, stepState);
-						await rpcClient.InvokeAsync((JobRpc.JobRpcClient x) => x.UpdateStepAsync(new UpdateStepRequest(options.JobId, options.BatchId, step.StepId, stepState, stepOutcome), null, null, cancellationToken), cancellationToken);
+						await rpcClient.InvokeAsync((JobRpc.JobRpcClient x) => x.UpdateStepAsync(new RpcUpdateStepRequest(options.JobId, options.BatchId, step.StepId, stepState, stepOutcome), null, null, cancellationToken), cancellationToken);
 					}
 
 					// Print the finishing state
@@ -433,7 +433,7 @@ namespace Horde.Agent.Leases.Handlers
 				TimeSpan waitTime = _stepAbortPollInterval;
 				try
 				{
-					GetStepResponse res = await rpcClient.InvokeAsync((JobRpc.JobRpcClient x) => x.GetStepAsync(new GetStepRequest(jobId, batchId, stepId), null, null, cancellationToken), cancellationToken);
+					RpcGetStepResponse res = await rpcClient.InvokeAsync((JobRpc.JobRpcClient x) => x.GetStepAsync(new RpcGetStepRequest(jobId, batchId, stepId), null, null, cancellationToken), cancellationToken);
 					if (res.AbortRequested)
 					{
 						leaseLogger.LogDebug("Step was aborted by server (JobId={JobId} BatchId={BatchId} StepId={StepId})", jobId, batchId, stepId);
diff --git a/Engine/Source/Programs/Horde/Horde.Agent/Leases/Handlers/UpgradeHandler.cs b/Engine/Source/Programs/Horde/Horde.Agent/Leases/Handlers/UpgradeHandler.cs
--- a/Engine/Source/Programs/Horde/Horde.Agent/Leases/Handlers/UpgradeHandler.cs
+++ b/Engine/Source/Programs/Horde/Horde.Agent/Leases/Handlers/UpgradeHandler.cs
@@ -45,7 +45,7 @@ namespace Horde.Agent.Leases.Handlers
 				// Download the new software
 				FileInfo outputFile = new FileInfo(Path.Combine(upgradeDir.FullName, "Agent.zip"));
 				using (IRpcClientRef<HordeRpc.HordeRpcClient> rpcClientRef = await session.RpcConnection.GetClientRefAsync<HordeRpc.HordeRpcClient>(cancellationToken))
-				using (AsyncServerStreamingCall<DownloadSoftwareResponse> cursor = rpcClientRef.Client.DownloadSoftware(new DownloadSoftwareRequest(requiredVersion), null, null, cancellationToken))
+				using (AsyncServerStreamingCall<RpcDownloadSoftwareResponse> cursor = rpcClientRef.Client.DownloadSoftware(new RpcDownloadSoftwareRequest(requiredVersion), null, null, cancellationToken))
 				{
 					await using Stream outputStream = outputFile.Open(FileMode.Create);
 					while (await cursor.ResponseStream.MoveNext(cancellationToken))
diff --git a/Engine/Source/Programs/Horde/Horde.Agent/Leases/LeaseManager.cs b/Engine/Source/Programs/Horde/Horde.Agent/Leases/LeaseManager.cs
--- a/Engine/Source/Programs/Horde/Horde.Agent/Leases/LeaseManager.cs
+++ b/Engine/Source/Programs/Horde/Horde.Agent/Leases/LeaseManager.cs
@@ -32,7 +32,7 @@ namespace Horde.Agent.Leases
 			/// <summary>
 			/// The worker lease state
 			/// </summary>
-			public Lease Lease { get; set; }
+			public RpcLease Lease { get; set; }
 
 			/// <summary>
 			/// Identifier for this lease
@@ -53,7 +53,7 @@ namespace Horde.Agent.Leases
 			/// Constructor
 			/// </summary>
 			/// <param name="lease">The worker lease state</param>
-			public LeaseInfo(Lease lease)
+			public LeaseInfo(RpcLease lease)
 			{
 				Lease = lease;
 				LeaseId = LeaseId.Parse(lease.Id);
@@ -99,7 +99,7 @@ namespace Horde.Agent.Leases
 		/// <summary>
 		/// Delegate for lease active events
 		/// </summary>
-		public delegate void LeaseActiveEvent(Lease lease);
+		public delegate void LeaseActiveEvent(RpcLease lease);
 
 		/// <summary>
 		/// Event triggered when a lease is accepted and set to active
@@ -124,7 +124,7 @@ namespace Horde.Agent.Leases
 		readonly LeaseLoggerFactory _leaseLoggerFactory;
 		readonly ILogger _logger;
 
-		AgentCapabilities? _capabilities;
+		RpcAgentCapabilities? _capabilities;
 
 		public LeaseManager(ISession session, CapabilitiesService capabilitiesService, StatusService statusService, IEnumerable<LeaseHandler> leaseHandlers, LeaseLoggerFactory leaseLoggerFactory, ILogger logger)
 		{
@@ -261,7 +261,7 @@ namespace Horde.Agent.Leases
 				}
 
 				// Build the next update request
-				UpdateSessionRequest updateSessionRequest = new UpdateSessionRequest();
+				RpcUpdateSessionRequest updateSessionRequest = new RpcUpdateSessionRequest();
 				updateSessionRequest.AgentId = _session.AgentId.ToString();
 				updateSessionRequest.SessionId = _session.SessionId.ToString();
 
@@ -270,7 +270,7 @@ namespace Horde.Agent.Leases
 				{
 					foreach (LeaseInfo leaseInfo in _activeLeases)
 					{
-						updateSessionRequest.Leases.Add(new Lease(leaseInfo.Lease));
+						updateSessionRequest.Leases.Add(new RpcLease(leaseInfo.Lease));
 					}
 					if (_sessionResult != null && _activeLeases.Count == 0)
 					{
@@ -304,7 +304,7 @@ namespace Horde.Agent.Leases
 				using (stopping ? (CancellationTokenRegistration?)null : stoppingToken.Register(() => _updateLeasesEvent.Set()))
 				{
 					// Update the state with the server
-					UpdateSessionResponse? updateSessionResponse = await UpdateSessionAsync(rpcClient, updateSessionRequest, waitTask);
+					RpcUpdateSessionResponse? updateSessionResponse = await UpdateSessionAsync(rpcClient, updateSessionRequest, waitTask);
 
 					lock (_lockObject)
 					{
@@ -325,7 +325,7 @@ namespace Horde.Agent.Leases
 							NumLeasesCompleted += numRemoved;
 
 							// Create any new leases and cancel any running leases
-							foreach (Lease serverLease in updateSessionResponse.Leases)
+							foreach (RpcLease serverLease in updateSessionResponse.Leases)
 							{
 								if (serverLease.State == RpcLeaseState.Cancelled)
 								{
@@ -413,9 +413,9 @@ namespace Horde.Agent.Leases
 		/// <param name="updateSessionRequest">The session update request</param>
 		/// <param name="waitTask">Task which can be used to jump out of the update early</param>
 		/// <returns>Response from the call</returns>
-		async Task<UpdateSessionResponse?> UpdateSessionAsync(HordeRpc.HordeRpcClient rpcClient, UpdateSessionRequest updateSessionRequest, Task waitTask)
+		async Task<RpcUpdateSessionResponse?> UpdateSessionAsync(HordeRpc.HordeRpcClient rpcClient, RpcUpdateSessionRequest updateSessionRequest, Task waitTask)
 		{
-			UpdateSessionResponse? updateSessionResponse = null;
+			RpcUpdateSessionResponse? updateSessionResponse = null;
 			try
 			{
 				updateSessionResponse = await UpdateSessionInternalAsync(rpcClient, updateSessionRequest, waitTask);
@@ -451,10 +451,10 @@ namespace Horde.Agent.Leases
 		/// <param name="request">The session update request</param>
 		/// <param name="waitTask">Task to use to terminate the wait</param>
 		/// <returns>The response object</returns>
-		async Task<UpdateSessionResponse?> UpdateSessionInternalAsync(HordeRpc.HordeRpcClient rpcClient, UpdateSessionRequest request, Task waitTask)
+		async Task<RpcUpdateSessionResponse?> UpdateSessionInternalAsync(HordeRpc.HordeRpcClient rpcClient, RpcUpdateSessionRequest request, Task waitTask)
 		{
 			DateTime deadline = DateTime.UtcNow + TimeSpan.FromMinutes(2.0);
-			using AsyncDuplexStreamingCall<UpdateSessionRequest, UpdateSessionResponse> call = rpcClient.UpdateSession(deadline: deadline);
+			using AsyncDuplexStreamingCall<RpcUpdateSessionRequest, RpcUpdateSessionResponse> call = rpcClient.UpdateSession(deadline: deadline);
 			_logger.LogDebug("Updating session {SessionId} (Status={Status})", request.SessionId, request.Status);
 
 			// Write the request to the server
@@ -473,7 +473,7 @@ namespace Horde.Agent.Leases
 			await call.RequestStream.CompleteAsync();
 
 			// Wait for a response or a new update to come in, then close the request stream
-			UpdateSessionResponse? response = null;
+			RpcUpdateSessionResponse? response = null;
 			while (await moveNextAsync)
 			{
 				response = call.ResponseStream.Current;
diff --git a/Engine/Source/Programs/Horde/Horde.Agent/Services/CapabilitiesService.cs b/Engine/Source/Programs/Horde/Horde.Agent/Services/CapabilitiesService.cs
--- a/Engine/Source/Programs/Horde/Horde.Agent/Services/CapabilitiesService.cs
+++ b/Engine/Source/Programs/Horde/Horde.Agent/Services/CapabilitiesService.cs
@@ -45,12 +45,12 @@ namespace Horde.Agent.Services
 		/// </summary>
 		/// <param name="workingDir">Working directory for the agent</param>
 		/// <returns>Worker object for advertising to the server</returns>
-		public async Task<AgentCapabilities> GetCapabilitiesAsync(DirectoryReference? workingDir)
+		public async Task<RpcAgentCapabilities> GetCapabilitiesAsync(DirectoryReference? workingDir)
 		{
 			_logger.LogInformation("Querying agent capabilities... (may take up to 30 seconds)");
 			Stopwatch timer = Stopwatch.StartNew();
 
-			Task<AgentCapabilities> task = GetCapabilitiesInternalAsync(workingDir);
+			Task<RpcAgentCapabilities> task = GetCapabilitiesInternalAsync(workingDir);
 			while (!task.IsCompleted)
 			{
 				Task delayTask = Task.Delay(TimeSpan.FromSeconds(30.0));
@@ -64,15 +64,15 @@ namespace Horde.Agent.Services
 			return await task;
 		}
 
-		async Task<AgentCapabilities> GetCapabilitiesInternalAsync(DirectoryReference? workingDir)
+		async Task<RpcAgentCapabilities> GetCapabilitiesInternalAsync(DirectoryReference? workingDir)
 		{
 			ILogger logger = _logger;
 
 			// Create the primary device
-			DeviceCapabilities primaryDevice = new DeviceCapabilities();
+			RpcDeviceCapabilities primaryDevice = new RpcDeviceCapabilities();
 			primaryDevice.Handle = "Primary";
 
-			List<DeviceCapabilities> otherDevices = new List<DeviceCapabilities>();
+			List<RpcDeviceCapabilities> otherDevices = new List<RpcDeviceCapabilities>();
 			if (RuntimeInformation.IsOSPlatform(OSPlatform.Windows))
 			{
 				primaryDevice.Properties.Add("Platform=Win64");
@@ -426,7 +426,7 @@ namespace Horde.Agent.Services
 			}
 
 			// Create the worker
-			AgentCapabilities agent = new();
+			RpcAgentCapabilities agent = new();
 			agent.Devices.Add(primaryDevice);
 			agent.Devices.AddRange(otherDevices);
 
@@ -471,7 +471,7 @@ namespace Horde.Agent.Services
 			}
 		}
 
-		static void AddCpuInfo(DeviceCapabilities primaryDevice, Dictionary<string, int> nameToCount, int numLogicalCores, int numPhysicalCores)
+		static void AddCpuInfo(RpcDeviceCapabilities primaryDevice, Dictionary<string, int> nameToCount, int numLogicalCores, int numPhysicalCores)
 		{
 			if (nameToCount.Count > 0)
 			{
diff --git a/Engine/Source/Programs/Horde/Horde.Agent/Services/SessionFactory.cs b/Engine/Source/Programs/Horde/Horde.Agent/Services/SessionFactory.cs
--- a/Engine/Source/Programs/Horde/Horde.Agent/Services/SessionFactory.cs
+++ b/Engine/Source/Programs/Horde/Horde.Agent/Services/SessionFactory.cs
@@ -154,7 +154,7 @@ namespace Horde.Agent.Services
 			logger.LogInformation("Server: {Server}", serverProfile.Url);
 
 			// Show the worker capabilities
-			AgentCapabilities capabilities = await capabilitiesService.GetCapabilitiesAsync(workingDir);
+			RpcAgentCapabilities capabilities = await capabilitiesService.GetCapabilitiesAsync(workingDir);
 			if (capabilities.Properties.Count > 0)
 			{
 				logger.LogInformation("Global:");
@@ -163,7 +163,7 @@ namespace Horde.Agent.Services
 					logger.LogInformation("  {AgentProperty}", property);
 				}
 			}
-			foreach (DeviceCapabilities device in capabilities.Devices)
+			foreach (RpcDeviceCapabilities device in capabilities.Devices)
 			{
 				logger.LogInformation("{DeviceName} Device:", device.Handle);
 				foreach (string property in device.Properties)
@@ -207,13 +207,13 @@ namespace Horde.Agent.Services
 			// Create the session
 			statusService.Set(AgentStatusMessage.ConnectingToServer);
 
-			CreateSessionResponse createSessionResponse;
+			RpcCreateSessionResponse createSessionResponse;
 			using (GrpcChannel channel = await grpcService.CreateGrpcChannelAsync(registrationInfo.Token, cancellationToken))
 			{
 				HordeRpc.HordeRpcClient rpcClient = new HordeRpc.HordeRpcClient(channel);
 
 				// Create the session information
-				CreateSessionRequest sessionRequest = new CreateSessionRequest();
+				RpcCreateSessionRequest sessionRequest = new RpcCreateSessionRequest();
 				sessionRequest.Id = registrationInfo.Id;
 				sessionRequest.Status = RpcAgentStatus.Ok;
 				sessionRequest.Capabilities = capabilities;
@@ -292,7 +292,7 @@ namespace Horde.Agent.Services
 			await FileReference.WriteAllBytesAsync(settingsFile, data, cancellationToken);
 		}
 
-		static async Task<AgentRegistration> RegisterAgentAsync(GrpcService grpcService, AgentSettings agentSettings, AgentCapabilities capabilities, ILogger logger, CancellationToken cancellationToken)
+		static async Task<AgentRegistration> RegisterAgentAsync(GrpcService grpcService, AgentSettings agentSettings, RpcAgentCapabilities capabilities, ILogger logger, CancellationToken cancellationToken)
 		{
 			ServerProfile serverProfile = agentSettings.GetCurrentServerProfile();
 			using GrpcChannel grpcChannel = await grpcService.CreateGrpcChannelAsync(serverProfile.Token, cancellationToken);
@@ -302,11 +302,11 @@ namespace Horde.Agent.Services
 				logger.LogInformation("Registering agent directly...");
 				HordeRpc.HordeRpcClient rpcClient = new HordeRpc.HordeRpcClient(grpcChannel);
 
-				CreateAgentRequest createAgentRequest = new CreateAgentRequest();
+				RpcCreateAgentRequest createAgentRequest = new RpcCreateAgentRequest();
 				createAgentRequest.Name = agentSettings.GetAgentName();
 				createAgentRequest.Ephemeral = agentSettings.Ephemeral;
 
-				CreateAgentResponse createAgentResponse = await rpcClient.CreateAgentAsync(createAgentRequest, null, null, cancellationToken);
+				RpcCreateAgentResponse createAgentResponse = await rpcClient.CreateAgentAsync(createAgentRequest, null, null, cancellationToken);
 				return new AgentRegistration(serverProfile.Url, createAgentResponse.Id, createAgentResponse.Token);
 			}
 
@@ -351,7 +351,7 @@ namespace Horde.Agent.Services
 			}
 		}
 
-		static string? GetProperty(AgentCapabilities capabilities, string name)
+		static string? GetProperty(RpcAgentCapabilities capabilities, string name)
 		{
 			foreach (string property in capabilities.Devices[0].Properties)
 			{
diff --git a/Engine/Source/Programs/Horde/Horde.Agent/Services/TelemetryService.cs b/Engine/Source/Programs/Horde/Horde.Agent/Services/TelemetryService.cs
--- a/Engine/Source/Programs/Horde/Horde.Agent/Services/TelemetryService.cs
+++ b/Engine/Source/Programs/Horde/Horde.Agent/Services/TelemetryService.cs
@@ -47,9 +47,9 @@ public class CpuMetrics
 	/// Convert to Protobuf-based event
 	/// </summary>
 	/// <returns></returns>
-	public AgentCpuMetricsEvent ToEvent()
+	public RpcAgentCpuMetricsEvent ToEvent()
 	{
-		return new AgentCpuMetricsEvent { User = User, System = System, Idle = Idle };
+		return new RpcAgentCpuMetricsEvent { User = User, System = System, Idle = Idle };
 	}
 }
 
@@ -88,9 +88,9 @@ public class MemoryMetrics
 	/// Convert to Protobuf-based event
 	/// </summary>
 	/// <returns></returns>
-	public AgentMemoryMetricsEvent ToEvent()
+	public RpcAgentMemoryMetricsEvent ToEvent()
 	{
-		return new AgentMemoryMetricsEvent { Total = Total, Free = Available, Used = Used, UsedPercentage = UsedPercentage };
+		return new RpcAgentMemoryMetricsEvent { Total = Total, Free = Available, Used = Used, UsedPercentage = UsedPercentage };
 	}
 }
 
@@ -211,7 +211,7 @@ class TelemetryService : BackgroundService
 	private readonly AgentSettings _agentSettings;
 	private readonly ILogger<TelemetryService> _logger;
 	private readonly TimeSpan _reportInterval;
-	private readonly AgentMetadataEvent _agentMetadataEvent;
+	private readonly RpcAgentMetadataEvent _agentMetadataEvent;
 	private readonly TimeSpan _agentMetadataReportInterval = TimeSpan.FromMinutes(2);
 	private ISystemMetrics? _systemMetrics;
 	private DateTime _lastTimeAgentMetadataSent = DateTime.UnixEpoch;
@@ -494,9 +494,9 @@ class TelemetryService : BackgroundService
 		{
 			_logger.LogDebug("Sending telemetry events to server...");
 
-			SendTelemetryEventsRequest request = new();
+			RpcSendTelemetryEventsRequest request = new();
 			Timestamp utcNow = Timestamp.FromDateTime(DateTime.UtcNow);
-			ExecutionMetadata em = new()
+			RpcExecutionMetadata em = new()
 			{
 				LeaseId = _jobHandler.CurrentLeaseId.ToString(),
 				JobId = _jobHandler.CurrentJobId,
@@ -504,26 +504,26 @@ class TelemetryService : BackgroundService
 			};
 
 			{
-				AgentCpuMetricsEvent cpuMetricsEvent = _systemMetrics.GetCpu().ToEvent();
+				RpcAgentCpuMetricsEvent cpuMetricsEvent = _systemMetrics.GetCpu().ToEvent();
 				cpuMetricsEvent.AgentId = _agentMetadataEvent.AgentId;
 				cpuMetricsEvent.Timestamp = utcNow;
 				cpuMetricsEvent.ExecutionMetadata = em;
-				request.Events.Add(new WrappedTelemetryEvent { Cpu = cpuMetricsEvent });
+				request.Events.Add(new RpcWrappedTelemetryEvent { Cpu = cpuMetricsEvent });
 			}
 
 			{
-				AgentMemoryMetricsEvent memMetricsEvent = _systemMetrics.GetMemory().ToEvent();
+				RpcAgentMemoryMetricsEvent memMetricsEvent = _systemMetrics.GetMemory().ToEvent();
 				memMetricsEvent.AgentId = _agentMetadataEvent.AgentId;
 				memMetricsEvent.Timestamp = utcNow;
 				memMetricsEvent.ExecutionMetadata = em;
-				request.Events.Add(new WrappedTelemetryEvent { Mem = memMetricsEvent });
+				request.Events.Add(new RpcWrappedTelemetryEvent { Mem = memMetricsEvent });
 			}
 
 			if (DateTime.UtcNow > _lastTimeAgentMetadataSent + _agentMetadataReportInterval)
 			{
 				// Report agent metadata every now and then as events are not guaranteed to be delivered.
 				// Re-sending ensures the metadata will eventually make it to the server.
-				request.Events.Add(new WrappedTelemetryEvent { AgentMetadata = _agentMetadataEvent });
+				request.Events.Add(new RpcWrappedTelemetryEvent { AgentMetadata = _agentMetadataEvent });
 				_lastTimeAgentMetadataSent = DateTime.UtcNow;
 			}
 
@@ -534,9 +534,9 @@ class TelemetryService : BackgroundService
 		return true;
 	}
 
-	AgentMetadataEvent GetAgentMetadataEvent()
+	RpcAgentMetadataEvent GetAgentMetadataEvent()
 	{
-		AgentMetadataEvent e = new()
+		RpcAgentMetadataEvent e = new()
 		{
 			Ip = null,
 			Hostname = _agentSettings.GetAgentName(),
diff --git a/Engine/Source/Programs/Horde/Horde.Agent/Utilities/HordeRpc.cs b/Engine/Source/Programs/Horde/Horde.Agent/Utilities/HordeRpc.cs
--- a/Engine/Source/Programs/Horde/Horde.Agent/Utilities/HordeRpc.cs
+++ b/Engine/Source/Programs/Horde/Horde.Agent/Utilities/HordeRpc.cs
@@ -15,24 +15,24 @@ using Google.Protobuf.Collections;
 #pragma warning disable CA1716
 namespace HordeCommon.Rpc
 {
-	partial class Property
+	partial class RpcProperty
 	{
-		public Property(string name, string value)
+		public RpcProperty(string name, string value)
 		{
 			Name = name;
 			Value = value;
 		}
 
-		public Property(KeyValuePair<string, string> pair)
+		public RpcProperty(KeyValuePair<string, string> pair)
 		{
 			Name = pair.Key;
 			Value = pair.Value;
 		}
 	}
 
-	partial class PropertyUpdate
+	partial class RpcPropertyUpdate
 	{
-		public PropertyUpdate(string name, string? value)
+		public RpcPropertyUpdate(string name, string? value)
 		{
 			Name = name;
 			Value = value;
@@ -41,14 +41,14 @@ namespace HordeCommon.Rpc
 
 	static class PropertyExtensions
 	{
-		public static string GetValue(this RepeatedField<Property> properties, string name)
+		public static string GetValue(this RepeatedField<RpcProperty> properties, string name)
 		{
 			return properties.First(x => x.Name == name).Value;
 		}
 
-		public static bool TryGetValue(this RepeatedField<Property> properties, string name, [MaybeNullWhen(false)] out string result)
+		public static bool TryGetValue(this RepeatedField<RpcProperty> properties, string name, [MaybeNullWhen(false)] out string result)
 		{
-			Property? property = properties.FirstOrDefault(x => x.Name == name);
+			RpcProperty? property = properties.FirstOrDefault(x => x.Name == name);
 			if (property == null)
 			{
 				result = null!;
@@ -62,34 +62,34 @@ namespace HordeCommon.Rpc
 		}
 	}
 
-	partial class GetStreamRequest
+	partial class RpcGetStreamRequest
 	{
-		public GetStreamRequest(StreamId streamId)
+		public RpcGetStreamRequest(StreamId streamId)
 		{
 			StreamId = streamId.ToString();
 		}
 	}
 
-	partial class UpdateStreamRequest
+	partial class RpcUpdateStreamRequest
 	{
-		public UpdateStreamRequest(StreamId streamId, Dictionary<string, string?> properties)
+		public RpcUpdateStreamRequest(StreamId streamId, Dictionary<string, string?> properties)
 		{
 			StreamId = streamId.ToString();
-			Properties.AddRange(properties.Select(x => new PropertyUpdate(x.Key, x.Value)));
+			Properties.AddRange(properties.Select(x => new RpcPropertyUpdate(x.Key, x.Value)));
 		}
 	}
 
-	partial class GetJobRequest
+	partial class RpcGetJobRequest
 	{
-		public GetJobRequest(JobId jobId)
+		public RpcGetJobRequest(JobId jobId)
 		{
 			JobId = jobId.ToString();
 		}
 	}
 
-	partial class BeginBatchRequest
+	partial class RpcBeginBatchRequest
 	{
-		public BeginBatchRequest(JobId jobId, JobStepBatchId batchId, LeaseId leaseId)
+		public RpcBeginBatchRequest(JobId jobId, JobStepBatchId batchId, LeaseId leaseId)
 		{
 			JobId = jobId.ToString();
 			BatchId = batchId.ToString();
@@ -97,9 +97,9 @@ namespace HordeCommon.Rpc
 		}
 	}
 
-	partial class FinishBatchRequest
+	partial class RpcFinishBatchRequest
 	{
-		public FinishBatchRequest(JobId jobId, JobStepBatchId batchId, LeaseId leaseId)
+		public RpcFinishBatchRequest(JobId jobId, JobStepBatchId batchId, LeaseId leaseId)
 		{
 			JobId = jobId.ToString();
 			BatchId = batchId.ToString();
@@ -107,9 +107,9 @@ namespace HordeCommon.Rpc
 		}
 	}
 
-	partial class BeginStepRequest
+	partial class RpcBeginStepRequest
 	{
-		public BeginStepRequest(JobId jobId, JobStepBatchId batchId, LeaseId leaseId)
+		public RpcBeginStepRequest(JobId jobId, JobStepBatchId batchId, LeaseId leaseId)
 		{
 			JobId = jobId.ToString();
 			BatchId = batchId.ToString();
@@ -117,9 +117,9 @@ namespace HordeCommon.Rpc
 		}
 	}
 
-	partial class UpdateStepRequest
+	partial class RpcUpdateStepRequest
 	{
-		public UpdateStepRequest(JobId jobId, JobStepBatchId batchId, JobStepId stepId, JobStepState state, JobStepOutcome outcome)
+		public RpcUpdateStepRequest(JobId jobId, JobStepBatchId batchId, JobStepId stepId, JobStepState state, JobStepOutcome outcome)
 		{
 			JobId = jobId.ToString();
 			BatchId = batchId.ToString();
@@ -129,9 +129,9 @@ namespace HordeCommon.Rpc
 		}
 	}
 
-	partial class GetStepRequest
+	partial class RpcGetStepRequest
 	{
-		public GetStepRequest(JobId jobId, JobStepBatchId batchId, JobStepId stepId)
+		public RpcGetStepRequest(JobId jobId, JobStepBatchId batchId, JobStepId stepId)
 		{
 			JobId = jobId.ToString();
 			BatchId = batchId.ToString();
@@ -139,9 +139,9 @@ namespace HordeCommon.Rpc
 		}
 	}
 
-	partial class GetStepResponse
+	partial class RpcGetStepResponse
 	{
-		public GetStepResponse(JobStepOutcome outcome, JobStepState state, bool abortRequested)
+		public RpcGetStepResponse(JobStepOutcome outcome, JobStepState state, bool abortRequested)
 		{
 			Outcome = (int)outcome;
 			State = (int)state;
@@ -149,9 +149,9 @@ namespace HordeCommon.Rpc
 		}
 	}
 
-	partial class CreateEventRequest
+	partial class RpcCreateEventRequest
 	{
-		public CreateEventRequest(EventSeverity severity, LogId logId, int lineIndex, int lineCount)
+		public RpcCreateEventRequest(RpcEventSeverity severity, LogId logId, int lineIndex, int lineCount)
 		{
 			Severity = severity;
 			LogId = logId.ToString();
@@ -160,17 +160,17 @@ namespace HordeCommon.Rpc
 		}
 	}
 
-	partial class CreateEventsRequest
+	partial class RpcCreateEventsRequest
 	{
-		public CreateEventsRequest(IEnumerable<CreateEventRequest> events)
+		public RpcCreateEventsRequest(IEnumerable<RpcCreateEventRequest> events)
 		{
 			Events.AddRange(events);
 		}
 	}
 
-	partial class WriteOutputRequest
+	partial class RpcWriteOutputRequest
 	{
-		public WriteOutputRequest(LogId logId, long offset, int lineIndex, ByteString data, bool flush)
+		public RpcWriteOutputRequest(LogId logId, long offset, int lineIndex, ByteString data, bool flush)
 		{
 			LogId = logId.ToString();
 			Offset = offset;
@@ -180,9 +180,9 @@ namespace HordeCommon.Rpc
 		}
 	}
 
-	partial class DownloadSoftwareRequest
+	partial class RpcDownloadSoftwareRequest
 	{
-		public DownloadSoftwareRequest(string version)
+		public RpcDownloadSoftwareRequest(string version)
 		{
 			Version = version;
 		}
@@ -191,7 +191,7 @@ namespace HordeCommon.Rpc
 
 namespace HordeCommon.Rpc.Messages.Telemetry
 {
-	partial class AgentMetadataEvent
+	partial class RpcAgentMetadataEvent
 	{
 		/// <summary>
 		/// Calculate an agent ID
diff --git a/Engine/Source/Programs/Horde/Horde.Agent/Utility/ArtifactUploader.cs b/Engine/Source/Programs/Horde/Horde.Agent/Utility/ArtifactUploader.cs
--- a/Engine/Source/Programs/Horde/Horde.Agent/Utility/ArtifactUploader.cs
+++ b/Engine/Source/Programs/Horde/Horde.Agent/Utility/ArtifactUploader.cs
@@ -115,10 +115,10 @@ namespace Horde.Agent.Utility
 			using (FileStream artifactStream = FileReference.Open(artifactFile, FileMode.Open, FileAccess.Read, FileShare.Read))
 			{
 				logger.LogInformation("Uploading artifact {ArtifactName} from {ArtifactFile} ({Size:n0} bytes)", artifactName, artifactFile, artifactStream.Length);
-				using (AsyncClientStreamingCall<UploadArtifactRequest, UploadArtifactResponse> cursor = client.UploadArtifact(null, null, cancellationToken))
+				using (AsyncClientStreamingCall<RpcUploadArtifactRequest, RpcUploadArtifactResponse> cursor = client.UploadArtifact(null, null, cancellationToken))
 				{
 					// Upload the metadata in the initial request
-					UploadArtifactMetadata metadata = new UploadArtifactMetadata();
+					RpcUploadArtifactMetadata metadata = new RpcUploadArtifactMetadata();
 					metadata.JobId = jobId.ToString();
 					metadata.BatchId = batchId.ToString();
 					metadata.StepId = stepId.ToString();
@@ -126,7 +126,7 @@ namespace Horde.Agent.Utility
 					metadata.MimeType = GetMimeType(artifactFile);
 					metadata.Length = artifactStream.Length;
 
-					UploadArtifactRequest initialRequest = new UploadArtifactRequest();
+					RpcUploadArtifactRequest initialRequest = new RpcUploadArtifactRequest();
 					initialRequest.Metadata = metadata;
 
 					await cursor.RequestStream.WriteAsync(initialRequest, cancellationToken);
@@ -141,7 +141,7 @@ namespace Horde.Agent.Utility
 							throw new InvalidDataException($"Unable to read data from {artifactFile} beyond offset {offset}; expected length to be {metadata.Length}");
 						}
 
-						UploadArtifactRequest request = new UploadArtifactRequest();
+						RpcUploadArtifactRequest request = new RpcUploadArtifactRequest();
 						request.Data = Google.Protobuf.ByteString.CopyFrom(buffer, 0, bytesRead);
 						await cursor.RequestStream.WriteAsync(request, cancellationToken);
 
@@ -152,7 +152,7 @@ namespace Horde.Agent.Utility
 					await cursor.RequestStream.CompleteAsync();
 
 					// Read the response
-					UploadArtifactResponse response = await cursor.ResponseAsync;
+					RpcUploadArtifactResponse response = await cursor.ResponseAsync;
 					return response.Id;
 				}
 			}
diff --git a/Engine/Source/Programs/Horde/Horde.Agent/Utility/JsonRpcLogSink.cs b/Engine/Source/Programs/Horde/Horde.Agent/Utility/JsonRpcLogSink.cs
--- a/Engine/Source/Programs/Horde/Horde.Agent/Utility/JsonRpcLogSink.cs
+++ b/Engine/Source/Programs/Horde/Horde.Agent/Utility/JsonRpcLogSink.cs
@@ -16,8 +16,8 @@ namespace Horde.Agent.Utility
 {
 	interface IJsonRpcLogSink : IAsyncDisposable
 	{
-		Task WriteEventsAsync(List<CreateEventRequest> events, CancellationToken cancellationToken);
-		Task WriteOutputAsync(WriteOutputRequest request, CancellationToken cancellationToken);
+		Task WriteEventsAsync(List<RpcCreateEventRequest> events, CancellationToken cancellationToken);
+		Task WriteOutputAsync(RpcWriteOutputRequest request, CancellationToken cancellationToken);
 		Task SetOutcomeAsync(JobStepOutcome outcome, CancellationToken cancellationToken);
 	}
 
@@ -175,7 +175,7 @@ namespace Horde.Agent.Utility
 			{
 				try
 				{
-					await _connection.InvokeAsync((JobRpc.JobRpcClient x) => x.UpdateStepAsync(new UpdateStepRequest(_jobId.Value, _jobBatchId.Value, _jobStepId.Value, JobStepState.Unspecified, outcome)), cancellationToken);
+					await _connection.InvokeAsync((JobRpc.JobRpcClient x) => x.UpdateStepAsync(new RpcUpdateStepRequest(_jobId.Value, _jobBatchId.Value, _jobStepId.Value, JobStepState.Unspecified, outcome)), cancellationToken);
 				}
 				catch (Exception ex)
 				{
@@ -185,13 +185,13 @@ namespace Horde.Agent.Utility
 		}
 
 		/// <inheritdoc/>
-		public async Task WriteEventsAsync(List<CreateEventRequest> events, CancellationToken cancellationToken)
+		public async Task WriteEventsAsync(List<RpcCreateEventRequest> events, CancellationToken cancellationToken)
 		{
-			await _connection.InvokeAsync((JobRpc.JobRpcClient x) => x.CreateEventsAsync(new CreateEventsRequest(events)), cancellationToken);
+			await _connection.InvokeAsync((JobRpc.JobRpcClient x) => x.CreateEventsAsync(new RpcCreateEventsRequest(events)), cancellationToken);
 		}
 
 		/// <inheritdoc/>
-		public async Task WriteOutputAsync(WriteOutputRequest request, CancellationToken cancellationToken)
+		public async Task WriteOutputAsync(RpcWriteOutputRequest request, CancellationToken cancellationToken)
 		{
 			_builder.WriteData(request.Data.Memory);
 			_bufferLength += request.Data.Length;
diff --git a/Engine/Source/Programs/Horde/Horde.Agent/Utility/RpcConnection.cs b/Engine/Source/Programs/Horde/Horde.Agent/Utility/RpcConnection.cs
--- a/Engine/Source/Programs/Horde/Horde.Agent/Utility/RpcConnection.cs
+++ b/Engine/Source/Programs/Horde/Horde.Agent/Utility/RpcConnection.cs
@@ -546,10 +546,10 @@ namespace Horde.Agent.Utility
 					while (!reconnectTaskSource.Task.IsCompleted && !_stoppingTaskSource.Task.IsCompleted)
 					{
 						// Client HAS TO BE the one to end the call, by closing the stream. The server must keep receiving requests from us until it does.
-						using (AsyncDuplexStreamingCall<QueryServerStateRequest, QueryServerStateResponse> call = client.QueryServerStateV2())
+						using (AsyncDuplexStreamingCall<RpcQueryServerStateRequest, RpcQueryServerStateResponse> call = client.QueryServerStateV2())
 						{
 							// Send the name of this agent
-							QueryServerStateRequest request = new QueryServerStateRequest();
+							RpcQueryServerStateRequest request = new RpcQueryServerStateRequest();
 							request.Name = Dns.GetHostName();
 							await call.RequestStream.WriteAsync(request);
 
@@ -562,7 +562,7 @@ namespace Horde.Agent.Utility
 							}
 
 							// Read the server response
-							QueryServerStateResponse response = call.ResponseStream.Current;
+							RpcQueryServerStateResponse response = call.ResponseStream.Current;
 							if (!response.Stopping)
 							{
 								// The first time we connect, log the server name and create the subconnection
diff --git a/Engine/Source/Programs/Horde/Horde.Agent/Utility/ServerLogger.cs b/Engine/Source/Programs/Horde/Horde.Agent/Utility/ServerLogger.cs
--- a/Engine/Source/Programs/Horde/Horde.Agent/Utility/ServerLogger.cs
+++ b/Engine/Source/Programs/Horde/Horde.Agent/Utility/ServerLogger.cs
@@ -161,7 +161,7 @@ namespace Horde.Agent.Utility
 
 			// Buffer for events read in a single iteration
 			JsonRpcLogWriter writer = new JsonRpcLogWriter();
-			List<CreateEventRequest> events = new List<CreateEventRequest>();
+			List<RpcCreateEventRequest> events = new List<RpcCreateEventRequest>();
 
 			// The current jobstep outcome
 			JobStepOutcome postedOutcome = JobStepOutcome.Success;
@@ -183,11 +183,11 @@ namespace Horde.Agent.Utility
 						{
 							if (jsonLogEvent.Level == LogLevel.Warning && ++numWarnings <= MaxWarnings)
 							{
-								AddEvent(jsonLogEvent.Data.Span, nextLineIndex, Math.Max(lineCount, jsonLogEvent.LineCount), EventSeverity.Warning, events);
+								AddEvent(jsonLogEvent.Data.Span, nextLineIndex, Math.Max(lineCount, jsonLogEvent.LineCount), RpcEventSeverity.Warning, events);
 							}
 							else if ((jsonLogEvent.Level == LogLevel.Error || jsonLogEvent.Level == LogLevel.Critical) && ++numErrors <= MaxErrors)
 							{
-								AddEvent(jsonLogEvent.Data.Span, nextLineIndex, Math.Max(lineCount, jsonLogEvent.LineCount), EventSeverity.Error, events);
+								AddEvent(jsonLogEvent.Data.Span, nextLineIndex, Math.Max(lineCount, jsonLogEvent.LineCount), RpcEventSeverity.Error, events);
 							}
 						}
 						nextLineIndex += lineCount;
@@ -211,7 +211,7 @@ namespace Horde.Agent.Utility
 					(ReadOnlyMemory<byte> packet, int packetLineCount) = writer.CreatePacket();
 					try
 					{
-						await _sink.WriteOutputAsync(new WriteOutputRequest(_logId, packetOffset, packetLineIndex, UnsafeByteOperations.UnsafeWrap(packet), false), CancellationToken.None);
+						await _sink.WriteOutputAsync(new RpcWriteOutputRequest(_logId, packetOffset, packetLineIndex, UnsafeByteOperations.UnsafeWrap(packet), false), CancellationToken.None);
 						packetOffset += packet.Length;
 						packetLineIndex += packetLineCount;
 					}
@@ -253,7 +253,7 @@ namespace Horde.Agent.Utility
 				{
 					try
 					{
-						await _sink.WriteOutputAsync(new WriteOutputRequest(_logId, packetOffset, packetLineIndex, ByteString.Empty, true), CancellationToken.None);
+						await _sink.WriteOutputAsync(new RpcWriteOutputRequest(_logId, packetOffset, packetLineIndex, ByteString.Empty, true), CancellationToken.None);
 					}
 					catch (Exception ex)
 					{
@@ -264,11 +264,11 @@ namespace Horde.Agent.Utility
 			}
 		}
 
-		void AddEvent(ReadOnlySpan<byte> span, int lineIndex, int lineCount, EventSeverity severity, List<CreateEventRequest> events)
+		void AddEvent(ReadOnlySpan<byte> span, int lineIndex, int lineCount, RpcEventSeverity severity, List<RpcCreateEventRequest> events)
 		{
 			try
 			{
-				events.Add(new CreateEventRequest(severity, _logId, lineIndex, lineCount));
+				events.Add(new RpcCreateEventRequest(severity, _logId, lineIndex, lineCount));
 			}
 			catch (Exception ex)
 			{
diff --git a/Engine/Source/Programs/Horde/Horde.Agent/Utility/TempStorage.cs b/Engine/Source/Programs/Horde/Horde.Agent/Utility/TempStorage.cs
--- a/Engine/Source/Programs/Horde/Horde.Agent/Utility/TempStorage.cs
+++ b/Engine/Source/Programs/Horde/Horde.Agent/Utility/TempStorage.cs
@@ -603,13 +603,13 @@ namespace Horde.Storage.Utility
 				ArtifactName artifactName = GetArtifactNameForNode(nodeName);
 				ArtifactType artifactType = ArtifactType.StepOutput;
 
-				GetJobArtifactRequest artifactRequest = new GetJobArtifactRequest();
+				RpcGetJobArtifactRequest artifactRequest = new RpcGetJobArtifactRequest();
 				artifactRequest.JobId = jobId.ToString();
 				artifactRequest.StepId = stepId.ToString();
 				artifactRequest.Name = artifactName.ToString();
 				artifactRequest.Type = artifactType.ToString();
 
-				GetJobArtifactResponse artifact = await jobRpc.Client.GetArtifactAsync(artifactRequest, cancellationToken: cancellationToken);
+				RpcGetJobArtifactResponse artifact = await jobRpc.Client.GetArtifactAsync(artifactRequest, cancellationToken: cancellationToken);
 
 				NamespaceId namespaceId = new NamespaceId(artifact.NamespaceId);
 				RefName refName = new RefName(artifact.RefName);
@@ -741,13 +741,13 @@ namespace Horde.Storage.Utility
 				ArtifactName artifactName = GetArtifactNameForNode(nodeName);
 				ArtifactType artifactType = ArtifactType.StepOutput;
 
-				GetJobArtifactRequest artifactRequest = new GetJobArtifactRequest();
+				RpcGetJobArtifactRequest artifactRequest = new RpcGetJobArtifactRequest();
 				artifactRequest.JobId = jobId.ToString();
 				artifactRequest.StepId = stepId.ToString();
 				artifactRequest.Name = artifactName.ToString();
 				artifactRequest.Type = artifactType.ToString();
 
-				GetJobArtifactResponse artifact = await jobRpc.Client.GetArtifactAsync(artifactRequest, cancellationToken: cancellationToken);
+				RpcGetJobArtifactResponse artifact = await jobRpc.Client.GetArtifactAsync(artifactRequest, cancellationToken: cancellationToken);
 				NamespaceId namespaceId = new NamespaceId(artifact.NamespaceId);
 				RefName refName = new RefName(artifact.RefName);
 
diff --git a/Engine/Source/Programs/Horde/Horde.Agent/Utility/WorkspaceInfo.cs b/Engine/Source/Programs/Horde/Horde.Agent/Utility/WorkspaceInfo.cs
--- a/Engine/Source/Programs/Horde/Horde.Agent/Utility/WorkspaceInfo.cs
+++ b/Engine/Source/Programs/Horde/Horde.Agent/Utility/WorkspaceInfo.cs
@@ -127,7 +127,7 @@ namespace Horde.Agent.Utility
 		/// <param name="logger">Logger output</param>
 		/// <param name="cancellationToken">Cancellation token</param>
 		/// <returns>New workspace info</returns>
-		public static async Task<WorkspaceInfo> SetupWorkspaceAsync(AgentWorkspace workspace, DirectoryReference rootDir, ManagedWorkspaceOptions options, ILogger logger, CancellationToken cancellationToken)
+		public static async Task<WorkspaceInfo> SetupWorkspaceAsync(RpcAgentWorkspace workspace, DirectoryReference rootDir, ManagedWorkspaceOptions options, ILogger logger, CancellationToken cancellationToken)
 		{
 			// Fill in the default credentials iff they are not set
 			string? serverAndPort = String.IsNullOrEmpty(workspace.ServerAndPort) ? null : workspace.ServerAndPort;
@@ -295,7 +295,7 @@ namespace Horde.Agent.Utility
 		/// </summary>
 		/// <param name="workspace">Workspace settings</param>
 		/// <returns></returns>
-		public static ManagedWorkspaceOptions GetMwOptions(AgentWorkspace workspace)
+		public static ManagedWorkspaceOptions GetMwOptions(RpcAgentWorkspace workspace)
 		{
 			const string NameKey = "name";
 			const string ManagedWorkspaceValue = "managedWorkspace";
diff --git a/Engine/Source/Programs/Horde/Horde.Server.Tests/Agents/AgentServiceTest.cs b/Engine/Source/Programs/Horde/Horde.Server.Tests/Agents/AgentServiceTest.cs
--- a/Engine/Source/Programs/Horde/Horde.Server.Tests/Agents/AgentServiceTest.cs
+++ b/Engine/Source/Programs/Horde/Horde.Server.Tests/Agents/AgentServiceTest.cs
@@ -119,7 +119,7 @@ public class AgentServiceTest : TestSetup
 		DateTime lastStatusChange = agent.LastStatusChange!.Value;
 		await Clock.AdvanceAsync(TimeSpan.FromMinutes(1));
 
-		agent = (await AgentService.UpdateSessionAsync(agent, agent.SessionId!.Value, status, null, null, new List<Lease>()))!;
+		agent = (await AgentService.UpdateSessionAsync(agent, agent.SessionId!.Value, status, null, null, new List<RpcLease>()))!;
 		if (expectTimestampUpdate)
 		{
 			AssertNotEqual(lastStatusChange, agent.LastStatusChange);
diff --git a/Engine/Source/Programs/Horde/Horde.Server.Tests/Issues/IssueServiceTests.cs b/Engine/Source/Programs/Horde/Horde.Server.Tests/Issues/IssueServiceTests.cs
--- a/Engine/Source/Programs/Horde/Horde.Server.Tests/Issues/IssueServiceTests.cs
+++ b/Engine/Source/Programs/Horde/Horde.Server.Tests/Issues/IssueServiceTests.cs
@@ -100,7 +100,7 @@ namespace Horde.Server.Tests.Issues
 					LogEvent ev = ParseEvent(line);
 					if (ev.LineIndex == 0)
 					{
-						EventSeverity severity = (level == LogLevel.Warning) ? EventSeverity.Warning : EventSeverity.Error;
+						RpcEventSeverity severity = (level == LogLevel.Warning) ? RpcEventSeverity.Warning : RpcEventSeverity.Error;
 						await _logFileService.CreateEventsAsync(new List<NewLogEventData> { new NewLogEventData { LogId = _logId, LineIndex = _lineIndex, LineCount = ev.LineCount, Severity = severity } }, CancellationToken.None);
 					}
 				}
@@ -316,11 +316,11 @@ namespace Horde.Server.Tests.Issues
 			buffer.GetSpan(1)[0] = (byte)'\n';
 			buffer.Advance(1);
 
-			EventSeverity severity = (logLevel == LogLevel.Error) ? EventSeverity.Error : (logLevel == LogLevel.Warning) ? EventSeverity.Warning : EventSeverity.Information;
+			RpcEventSeverity severity = (logLevel == LogLevel.Error) ? RpcEventSeverity.Error : (logLevel == LogLevel.Warning) ? RpcEventSeverity.Warning : RpcEventSeverity.Information;
 			await AddEventAsync(job, batchIdx, stepIdx, severity, buffer.WrittenMemory.ToArray());
 		}
 
-		async Task AddEventAsync(IJob job, int batchIdx, int stepIdx, EventSeverity severity, byte[] data)
+		async Task AddEventAsync(IJob job, int batchIdx, int stepIdx, RpcEventSeverity severity, byte[] data)
 		{
 			LogId logId = job.Batches[batchIdx].Steps[stepIdx].LogId!.Value;
 
diff --git a/Engine/Source/Programs/Horde/Horde.Server.Tests/Server/RpcServiceTest.cs b/Engine/Source/Programs/Horde/Horde.Server.Tests/Server/RpcServiceTest.cs
--- a/Engine/Source/Programs/Horde/Horde.Server.Tests/Server/RpcServiceTest.cs
+++ b/Engine/Source/Programs/Horde/Horde.Server.Tests/Server/RpcServiceTest.cs
@@ -32,7 +32,6 @@ using MongoDB.Bson;
 
 namespace Horde.Server.Tests.Server
 {
-	using AgentCapabilities = HordeCommon.Rpc.Messages.AgentCapabilities;
 	using ISession = Microsoft.AspNetCore.Http.ISession;
 
 	public class AppLifetimeStub : IHostApplicationLifetime
@@ -323,14 +322,14 @@ namespace Horde.Server.Tests.Server
 		[TestMethod]
 		public async Task CreateSessionTestAsync()
 		{
-			CreateSessionRequest req = new CreateSessionRequest();
+			RpcCreateSessionRequest req = new RpcCreateSessionRequest();
 			await Assert.ThrowsExceptionAsync<StructuredRpcException>(() => RpcService.CreateSession(req, _adminContext));
 
 			req.Id = new AgentId("MyName").ToString();
 			await Assert.ThrowsExceptionAsync<StructuredRpcException>(() => RpcService.CreateSession(req, _adminContext));
 
-			req.Capabilities = new AgentCapabilities();
-			CreateSessionResponse res = await RpcService.CreateSession(req, _adminContext);
+			req.Capabilities = new RpcAgentCapabilities();
+			RpcCreateSessionResponse res = await RpcService.CreateSession(req, _adminContext);
 
 			Assert.AreEqual("MYNAME", res.AgentId);
 			// TODO: Check Token, ExpiryTime, SessionId 
@@ -339,15 +338,15 @@ namespace Horde.Server.Tests.Server
 		[TestMethod]
 		public async Task AgentJoinsPoolThroughPropertiesAsync()
 		{
-			CreateSessionRequest req = new() { Id = new AgentId("bogusAgentName").ToString(), Capabilities = new AgentCapabilities() };
+			RpcCreateSessionRequest req = new() { Id = new AgentId("bogusAgentName").ToString(), Capabilities = new RpcAgentCapabilities() };
 			req.Capabilities.Properties.Add($"{KnownPropertyNames.RequestedPools}=fooPool,barPool");
-			CreateSessionResponse res = await RpcService.CreateSession(req, _adminContext);
+			RpcCreateSessionResponse res = await RpcService.CreateSession(req, _adminContext);
 
 			IAgent agent = (await AgentService.GetAgentAsync(new AgentId(res.AgentId)))!;
 			CollectionAssert.AreEquivalent(new List<PoolId> { new("fooPool"), new("barPool") }, agent.GetPools().ToList());
 
 			// Connect a second time, when the agent has already been created
-			req = new() { Id = new AgentId("bogusAgentName").ToString(), Capabilities = new AgentCapabilities() };
+			req = new() { Id = new AgentId("bogusAgentName").ToString(), Capabilities = new RpcAgentCapabilities() };
 			req.Capabilities.Properties.Add($"{KnownPropertyNames.RequestedPools}=bazPool");
 			res = await RpcService.CreateSession(req, _adminContext);
 
@@ -358,9 +357,9 @@ namespace Horde.Server.Tests.Server
 		[TestMethod]
 		public async Task PropertiesFromAgentCapabilitiesAsync()
 		{
-			CreateSessionRequest req = new() { Id = new AgentId("bogusAgentName").ToString(), Capabilities = new AgentCapabilities() };
+			RpcCreateSessionRequest req = new() { Id = new AgentId("bogusAgentName").ToString(), Capabilities = new RpcAgentCapabilities() };
 			req.Capabilities.Properties.Add("fooKey=barValue");
-			CreateSessionResponse res = await RpcService.CreateSession(req, _adminContext);
+			RpcCreateSessionResponse res = await RpcService.CreateSession(req, _adminContext);
 			IAgent agent = (await AgentService.GetAgentAsync(new AgentId(res.AgentId)))!;
 			Assert.IsTrue(agent.Properties.Contains("fooKey=barValue"));
 		}
@@ -368,9 +367,9 @@ namespace Horde.Server.Tests.Server
 		[TestMethod]
 		public async Task PropertiesFromDeviceCapabilitiesAsync()
 		{
-			CreateSessionRequest req = new() { Id = new AgentId("bogusAgentName").ToString(), Capabilities = new AgentCapabilities() };
-			req.Capabilities.Devices.Add(new DeviceCapabilities { Handle = "someHandle", Properties = { "foo=bar" } });
-			CreateSessionResponse res = await RpcService.CreateSession(req, _adminContext);
+			RpcCreateSessionRequest req = new() { Id = new AgentId("bogusAgentName").ToString(), Capabilities = new RpcAgentCapabilities() };
+			req.Capabilities.Devices.Add(new RpcDeviceCapabilities { Handle = "someHandle", Properties = { "foo=bar" } });
+			RpcCreateSessionResponse res = await RpcService.CreateSession(req, _adminContext);
 			IAgent agent = (await AgentService.GetAgentAsync(new AgentId(res.AgentId)))!;
 			Assert.IsTrue(agent.Properties.Contains("foo=bar"));
 		}
@@ -378,9 +377,9 @@ namespace Horde.Server.Tests.Server
 		[TestMethod]
 		public async Task KnownPropertiesAreSetAsResourcesAsync()
 		{
-			CreateSessionRequest req = new() { Id = new AgentId("bogusAgentName").ToString(), Capabilities = new AgentCapabilities() };
-			req.Capabilities.Devices.Add(new DeviceCapabilities { Handle = "someHandle", Properties = { $"{KnownPropertyNames.LogicalCores}=10" } });
-			CreateSessionResponse res = await RpcService.CreateSession(req, _adminContext);
+			RpcCreateSessionRequest req = new() { Id = new AgentId("bogusAgentName").ToString(), Capabilities = new RpcAgentCapabilities() };
+			req.Capabilities.Devices.Add(new RpcDeviceCapabilities { Handle = "someHandle", Properties = { $"{KnownPropertyNames.LogicalCores}=10" } });
+			RpcCreateSessionResponse res = await RpcService.CreateSession(req, _adminContext);
 			IAgent agent = (await AgentService.GetAgentAsync(new AgentId(res.AgentId)))!;
 			Assert.AreEqual(10, agent.Resources[KnownPropertyNames.LogicalCores]);
 		}
@@ -388,22 +387,22 @@ namespace Horde.Server.Tests.Server
 		[TestMethod]
 		public async Task UpdateSessionTestAsync()
 		{
-			CreateSessionRequest createReq = new CreateSessionRequest
+			RpcCreateSessionRequest createReq = new RpcCreateSessionRequest
 			{
 				Id = new AgentId("UpdateSessionTest1").ToString(),
-				Capabilities = new AgentCapabilities()
+				Capabilities = new RpcAgentCapabilities()
 			};
-			CreateSessionResponse createRes = await RpcService.CreateSession(createReq, _adminContext);
+			RpcCreateSessionResponse createRes = await RpcService.CreateSession(createReq, _adminContext);
 			string agentId = createRes.AgentId;
 			string sessionId = createRes.SessionId;
 
-			TestAsyncStreamReader<UpdateSessionRequest> requestStream =
-				new TestAsyncStreamReader<UpdateSessionRequest>(_adminContext);
-			TestServerStreamWriter<UpdateSessionResponse> responseStream =
-				new TestServerStreamWriter<UpdateSessionResponse>(_adminContext);
+			TestAsyncStreamReader<RpcUpdateSessionRequest> requestStream =
+				new TestAsyncStreamReader<RpcUpdateSessionRequest>(_adminContext);
+			TestServerStreamWriter<RpcUpdateSessionResponse> responseStream =
+				new TestServerStreamWriter<RpcUpdateSessionResponse>(_adminContext);
 			Task call = RpcService.UpdateSession(requestStream, responseStream, _adminContext);
 
-			requestStream.AddMessage(new UpdateSessionRequest { AgentId = "does-not-exist", SessionId = sessionId });
+			requestStream.AddMessage(new RpcUpdateSessionRequest { AgentId = "does-not-exist", SessionId = sessionId });
 			StructuredRpcException re = await Assert.ThrowsExceptionAsync<StructuredRpcException>(() => call);
 			Assert.AreEqual(StatusCode.NotFound, re.StatusCode);
 			Assert.IsTrue(re.Message.Contains("Invalid agent name", StringComparison.OrdinalIgnoreCase));
@@ -414,14 +413,14 @@ namespace Horde.Server.Tests.Server
 		{
 			RpcService._longPollTimeout = TimeSpan.FromMilliseconds(200);
 
-			TestAsyncStreamReader<QueryServerStateRequest> requestStream =
-				new TestAsyncStreamReader<QueryServerStateRequest>(_adminContext);
-			TestServerStreamWriter<QueryServerStateResponse> responseStream =
-				new TestServerStreamWriter<QueryServerStateResponse>(_adminContext);
+			TestAsyncStreamReader<RpcQueryServerStateRequest> requestStream =
+				new TestAsyncStreamReader<RpcQueryServerStateRequest>(_adminContext);
+			TestServerStreamWriter<RpcQueryServerStateResponse> responseStream =
+				new TestServerStreamWriter<RpcQueryServerStateResponse>(_adminContext);
 			Task call = RpcService.QueryServerState(requestStream, responseStream, _adminContext);
 
-			requestStream.AddMessage(new QueryServerStateRequest { Name = "bogusAgentName" });
-			QueryServerStateResponse? res = await responseStream.ReadNextAsync();
+			requestStream.AddMessage(new RpcQueryServerStateRequest { Name = "bogusAgentName" });
+			RpcQueryServerStateResponse? res = await responseStream.ReadNextAsync();
 			Assert.IsNotNull(res);
 
 			res = await responseStream.ReadNextAsync();
@@ -434,22 +433,22 @@ namespace Horde.Server.Tests.Server
 		[TestMethod]
 		public async Task FinishBatchTestAsync()
 		{
-			CreateSessionRequest createReq = new CreateSessionRequest
+			RpcCreateSessionRequest createReq = new RpcCreateSessionRequest
 			{
 				Id = new AgentId("UpdateSessionTest1").ToString(),
-				Capabilities = new AgentCapabilities()
+				Capabilities = new RpcAgentCapabilities()
 			};
-			CreateSessionResponse createRes = await RpcService.CreateSession(createReq, _adminContext);
+			RpcCreateSessionResponse createRes = await RpcService.CreateSession(createReq, _adminContext);
 			string agentId = createRes.AgentId;
 			string sessionId = createRes.SessionId;
 
-			TestAsyncStreamReader<UpdateSessionRequest> requestStream =
-				new TestAsyncStreamReader<UpdateSessionRequest>(_adminContext);
-			TestServerStreamWriter<UpdateSessionResponse> responseStream =
-				new TestServerStreamWriter<UpdateSessionResponse>(_adminContext);
+			TestAsyncStreamReader<RpcUpdateSessionRequest> requestStream =
+				new TestAsyncStreamReader<RpcUpdateSessionRequest>(_adminContext);
+			TestServerStreamWriter<RpcUpdateSessionResponse> responseStream =
+				new TestServerStreamWriter<RpcUpdateSessionResponse>(_adminContext);
 			Task call = RpcService.UpdateSession(requestStream, responseStream, _adminContext);
 
-			requestStream.AddMessage(new UpdateSessionRequest { AgentId = "does-not-exist", SessionId = sessionId });
+			requestStream.AddMessage(new RpcUpdateSessionRequest { AgentId = "does-not-exist", SessionId = sessionId });
 			StructuredRpcException re = await Assert.ThrowsExceptionAsync<StructuredRpcException>(() => call);
 			Assert.AreEqual(StatusCode.NotFound, re.StatusCode);
 			Assert.IsTrue(re.Message.Contains("Invalid agent name", StringComparison.OrdinalIgnoreCase));
@@ -470,7 +469,7 @@ namespace Horde.Server.Tests.Server
 			string[] data = { "foo", "bar", "baz", "qux" };
 			string dataStr = String.Join("", data);
 
-			UploadArtifactMetadata metadata = new UploadArtifactMetadata
+			RpcUploadArtifactMetadata metadata = new RpcUploadArtifactMetadata
 			{
 				JobId = fixture.Job1.Id.ToString(),
 				BatchId = fixture.Job1.Batches[0].Id.ToString(),
@@ -496,14 +495,14 @@ namespace Horde.Server.Tests.Server
 						RequestStream.Complete();
 						await Task.Delay(500);
 			*/
-			TestAsyncStreamReader<UploadArtifactRequest> requestStream = new TestAsyncStreamReader<UploadArtifactRequest>(context);
-			Task<UploadArtifactResponse> call = RpcService.UploadArtifact(requestStream, context);
-			requestStream.AddMessage(new UploadArtifactRequest { Metadata = metadata });
-			requestStream.AddMessage(new UploadArtifactRequest { Data = ByteString.CopyFromUtf8(data[0]) });
-			requestStream.AddMessage(new UploadArtifactRequest { Data = ByteString.CopyFromUtf8(data[1]) });
-			requestStream.AddMessage(new UploadArtifactRequest { Data = ByteString.CopyFromUtf8(data[2]) });
-			requestStream.AddMessage(new UploadArtifactRequest { Data = ByteString.CopyFromUtf8(data[3]) });
-			UploadArtifactResponse res = await call;
+			TestAsyncStreamReader<RpcUploadArtifactRequest> requestStream = new TestAsyncStreamReader<RpcUploadArtifactRequest>(context);
+			Task<RpcUploadArtifactResponse> call = RpcService.UploadArtifact(requestStream, context);
+			requestStream.AddMessage(new RpcUploadArtifactRequest { Metadata = metadata });
+			requestStream.AddMessage(new RpcUploadArtifactRequest { Data = ByteString.CopyFromUtf8(data[0]) });
+			requestStream.AddMessage(new RpcUploadArtifactRequest { Data = ByteString.CopyFromUtf8(data[1]) });
+			requestStream.AddMessage(new RpcUploadArtifactRequest { Data = ByteString.CopyFromUtf8(data[2]) });
+			requestStream.AddMessage(new RpcUploadArtifactRequest { Data = ByteString.CopyFromUtf8(data[3]) });
+			RpcUploadArtifactResponse res = await call;
 
 			IArtifactV1? artifact = await ArtifactCollection.GetArtifactAsync(ObjectId.Parse(res.Id));
 			Assert.IsNotNull(artifact);
diff --git a/Engine/Source/Programs/Horde/Horde.Server/Agents/AgentService.cs b/Engine/Source/Programs/Horde/Horde.Server/Agents/AgentService.cs
--- a/Engine/Source/Programs/Horde/Horde.Server/Agents/AgentService.cs
+++ b/Engine/Source/Programs/Horde/Horde.Server/Agents/AgentService.cs
@@ -26,6 +26,7 @@ using Horde.Server.Tasks;
 using Horde.Server.Utilities;
 using HordeCommon;
 using HordeCommon.Rpc;
+using HordeCommon.Rpc.Messages;
 using Microsoft.Extensions.Hosting;
 using Microsoft.Extensions.Logging;
 using StackExchange.Redis;
@@ -468,7 +469,7 @@ namespace Horde.Server.Agents
 		/// <param name="newLeases">Leases that the agent knows about</param>
 		/// <param name="cancellationToken"></param>
 		/// <returns>True if a lease was assigned, false otherwise</returns>
-		public async Task<IAgent?> WaitForLeaseAsync(IAgent? agent, IList<HordeCommon.Rpc.Messages.Lease> newLeases, CancellationToken cancellationToken = default)
+		public async Task<IAgent?> WaitForLeaseAsync(IAgent? agent, IList<HordeCommon.Rpc.Messages.RpcLease> newLeases, CancellationToken cancellationToken = default)
 		{
 			HashSet<string> knownLeases = new HashSet<string>(newLeases.Select(x => x.Id), StringComparer.OrdinalIgnoreCase);
 			while (agent != null && agent.Leases.All(x => knownLeases.Contains(x.Id.ToString())))
@@ -662,7 +663,7 @@ namespace Horde.Server.Agents
 		/// <summary>
 		/// 
 		/// </summary>
-		public async Task<IAgent?> UpdateSessionWithWaitAsync(IAgent inAgent, SessionId sessionId, AgentStatus status, IReadOnlyList<string>? properties, IReadOnlyDictionary<string, int>? resources, IList<HordeCommon.Rpc.Messages.Lease> newLeases, CancellationToken cancellationToken = default)
+		public async Task<IAgent?> UpdateSessionWithWaitAsync(IAgent inAgent, SessionId sessionId, AgentStatus status, IReadOnlyList<string>? properties, IReadOnlyDictionary<string, int>? resources, IList<HordeCommon.Rpc.Messages.RpcLease> newLeases, CancellationToken cancellationToken = default)
 		{
 			IAgent? agent = inAgent;
 
@@ -689,7 +690,7 @@ namespace Horde.Server.Agents
 		/// <param name="newLeases">New list of leases for this session</param>
 		/// <param name="cancellationToken">Cancellation token for the operation</param>
 		/// <returns>Updated agent state</returns>
-		public async Task<IAgent?> UpdateSessionAsync(IAgent inAgent, SessionId sessionId, AgentStatus status, IReadOnlyList<string>? properties, IReadOnlyDictionary<string, int>? resources, IList<HordeCommon.Rpc.Messages.Lease> newLeases, CancellationToken cancellationToken = default)
+		public async Task<IAgent?> UpdateSessionAsync(IAgent inAgent, SessionId sessionId, AgentStatus status, IReadOnlyList<string>? properties, IReadOnlyDictionary<string, int>? resources, IList<HordeCommon.Rpc.Messages.RpcLease> newLeases, CancellationToken cancellationToken = default)
 		{
 			DateTime utcNow = _clock.UtcNow;
 
@@ -731,13 +732,13 @@ namespace Horde.Server.Agents
 				List<AgentLease> leases = new List<AgentLease>(agent.Leases);
 
 				// Remove any completed leases from the agent
-				Dictionary<LeaseId, HordeCommon.Rpc.Messages.Lease> leaseIdToNewState = newLeases.ToDictionary(x => LeaseId.Parse(x.Id), x => x);
+				Dictionary<LeaseId, RpcLease> leaseIdToNewState = newLeases.ToDictionary(x => LeaseId.Parse(x.Id), x => x);
 				for (int idx = 0; idx < leases.Count; idx++)
 				{
 					AgentLease lease = leases[idx];
 					if (lease.State == LeaseState.Cancelled)
 					{
-						HordeCommon.Rpc.Messages.Lease? newLease;
+						RpcLease? newLease;
 						if (!leaseIdToNewState.TryGetValue(lease.Id, out newLease) || newLease.State == RpcLeaseState.Cancelled || newLease.State == RpcLeaseState.Completed)
 						{
 							await RemoveLeaseAsync(agent, lease, utcNow, LeaseOutcome.Cancelled, null, cancellationToken);
@@ -747,7 +748,7 @@ namespace Horde.Server.Agents
 					}
 					else
 					{
-						HordeCommon.Rpc.Messages.Lease? newLease;
+						RpcLease? newLease;
 						if (leaseIdToNewState.TryGetValue(lease.Id, out newLease) && (LeaseState)newLease.State != lease.State)
 						{
 							if (newLease.State == RpcLeaseState.Cancelled || newLease.State == RpcLeaseState.Completed)
diff --git a/Engine/Source/Programs/Horde/Horde.Server/Agents/Fleet/JobQueueStrategy.cs b/Engine/Source/Programs/Horde/Horde.Server/Agents/Fleet/JobQueueStrategy.cs
--- a/Engine/Source/Programs/Horde/Horde.Server/Agents/Fleet/JobQueueStrategy.cs
+++ b/Engine/Source/Programs/Horde/Horde.Server/Agents/Fleet/JobQueueStrategy.cs
@@ -9,7 +9,6 @@ using System.Threading;
 using System.Threading.Tasks;
 using EpicGames.Core;
 using EpicGames.Horde.Agents.Pools;
-using EpicGames.Horde.Jobs;
 using EpicGames.Horde.Streams;
 using Horde.Server.Agents.Pools;
 using Horde.Server.Jobs;
diff --git a/Engine/Source/Programs/Horde/Horde.Server/Agents/IAgent.cs b/Engine/Source/Programs/Horde/Horde.Server/Agents/IAgent.cs
--- a/Engine/Source/Programs/Horde/Horde.Server/Agents/IAgent.cs
+++ b/Engine/Source/Programs/Horde/Horde.Server/Agents/IAgent.cs
@@ -101,7 +101,7 @@ namespace Horde.Server.Agents
 		/// Constructor
 		/// </summary>
 		/// <param name="workspace">RPC message to construct from</param>
-		public AgentWorkspaceInfo(AgentWorkspace workspace)
+		public AgentWorkspaceInfo(RpcAgentWorkspace workspace)
 			: this(workspace.ConfiguredCluster, workspace.ConfiguredUserName, workspace.Identifier, workspace.Stream, (workspace.View.Count > 0) ? workspace.View.ToList() : null, workspace.Incremental, workspace.Method)
 		{
 		}
@@ -164,10 +164,10 @@ namespace Horde.Server.Agents
 		/// <param name="server">The Perforce server</param>
 		/// <param name="credentials">Credentials for the server</param>
 		/// <returns>The RPC message</returns>
-		public AgentWorkspace ToRpcMessage(IPerforceServer server, PerforceCredentials? credentials)
+		public RpcAgentWorkspace ToRpcMessage(IPerforceServer server, PerforceCredentials? credentials)
 		{
 			// Construct the message
-			AgentWorkspace result = new AgentWorkspace
+			RpcAgentWorkspace result = new RpcAgentWorkspace
 			{
 				ConfiguredCluster = Cluster,
 				ConfiguredUserName = UserName,
@@ -402,9 +402,9 @@ namespace Horde.Server.Agents
 		/// Converts this lease to an RPC message
 		/// </summary>
 		/// <returns>RPC message</returns>
-		public HordeCommon.Rpc.Messages.Lease ToRpcMessage()
+		public RpcLease ToRpcMessage()
 		{
-			HordeCommon.Rpc.Messages.Lease lease = new HordeCommon.Rpc.Messages.Lease();
+			RpcLease lease = new RpcLease();
 			lease.Id = Id.ToString();
 			lease.Payload = Google.Protobuf.WellKnownTypes.Any.Parser.ParseFrom(Payload);
 			lease.State = (RpcLeaseState)State;
@@ -965,14 +965,14 @@ namespace Horde.Server.Agents
 		/// <param name="workspaceMessages">List of messages</param>
 		/// <param name="cancellationToken">Cancellation token for the operation</param>
 		/// <returns>The RPC message</returns>
-		public static async Task<bool> TryAddWorkspaceMessageAsync(this IAgent agent, AgentWorkspaceInfo workspace, PerforceCluster cluster, PerforceLoadBalancer loadBalancer, IList<AgentWorkspace> workspaceMessages, CancellationToken cancellationToken)
+		public static async Task<bool> TryAddWorkspaceMessageAsync(this IAgent agent, AgentWorkspaceInfo workspace, PerforceCluster cluster, PerforceLoadBalancer loadBalancer, IList<RpcAgentWorkspace> workspaceMessages, CancellationToken cancellationToken)
 		{
 			// Find a matching server, trying to use a previously selected one if possible
 			string? baseServerAndPort;
 			string? serverAndPort;
 			bool partitioned;
 
-			AgentWorkspace? existingWorkspace = workspaceMessages.FirstOrDefault(x => x.ConfiguredCluster == workspace.Cluster);
+			RpcAgentWorkspace? existingWorkspace = workspaceMessages.FirstOrDefault(x => x.ConfiguredCluster == workspace.Cluster);
 			if (existingWorkspace != null)
 			{
 				baseServerAndPort = existingWorkspace.BaseServerAndPort;
@@ -1012,7 +1012,7 @@ namespace Horde.Server.Agents
 			}
 
 			// Construct the message
-			AgentWorkspace result = new AgentWorkspace
+			RpcAgentWorkspace result = new RpcAgentWorkspace
 			{
 				ConfiguredCluster = workspace.Cluster,
 				ConfiguredUserName = workspace.UserName,
diff --git a/Engine/Source/Programs/Horde/Horde.Server/Issues/IssueService.cs b/Engine/Source/Programs/Horde/Horde.Server/Issues/IssueService.cs
--- a/Engine/Source/Programs/Horde/Horde.Server/Issues/IssueService.cs
+++ b/Engine/Source/Programs/Horde/Horde.Server/Issues/IssueService.cs
@@ -158,10 +158,10 @@ namespace Horde.Server.Issues
 				EventData = logEventData;
 			}
 
-			static LogLevel GetLogLevelFromSeverity(EventSeverity severity) => severity switch
+			static LogLevel GetLogLevelFromSeverity(RpcEventSeverity severity) => severity switch
 			{
-				EventSeverity.Error => LogLevel.Error,
-				EventSeverity.Warning => LogLevel.Warning,
+				RpcEventSeverity.Error => LogLevel.Error,
+				RpcEventSeverity.Warning => LogLevel.Warning,
 				_ => LogLevel.Information
 			};
 		}
diff --git a/Engine/Source/Programs/Horde/Horde.Server/Jobs/Graphs/GraphCollection.cs b/Engine/Source/Programs/Horde/Horde.Server/Jobs/Graphs/GraphCollection.cs
--- a/Engine/Source/Programs/Horde/Horde.Server/Jobs/Graphs/GraphCollection.cs
+++ b/Engine/Source/Programs/Horde/Horde.Server/Jobs/Graphs/GraphCollection.cs
@@ -165,7 +165,7 @@ namespace Horde.Server.Jobs.Graphs
 			public string? UgsProject { get; set; }
 
 			[BsonIgnoreIfNull]
-			public LabelChange Change { get; set; }
+			public RpcLabelChange Change { get; set; }
 
 			public List<NodeRef> RequiredNodes { get; set; }
 			public List<NodeRef> IncludedNodes { get; set; }
@@ -179,7 +179,7 @@ namespace Horde.Server.Jobs.Graphs
 				IncludedNodes = new List<NodeRef>();
 			}
 
-			public Label(string? dashboardName, string? dashboardCategory, string? ugsName, string? ugsProject, LabelChange change, List<NodeRef> requiredNodes, List<NodeRef> includedNodes)
+			public Label(string? dashboardName, string? dashboardCategory, string? ugsName, string? ugsProject, RpcLabelChange change, List<NodeRef> requiredNodes, List<NodeRef> includedNodes)
 			{
 				DashboardName = dashboardName;
 				DashboardCategory = dashboardCategory;
diff --git a/Engine/Source/Programs/Horde/Horde.Server/Jobs/Graphs/IGraph.cs b/Engine/Source/Programs/Horde/Horde.Server/Jobs/Graphs/IGraph.cs
--- a/Engine/Source/Programs/Horde/Horde.Server/Jobs/Graphs/IGraph.cs
+++ b/Engine/Source/Programs/Horde/Horde.Server/Jobs/Graphs/IGraph.cs
@@ -232,7 +232,7 @@ namespace Horde.Server.Jobs.Graphs
 		/// <summary>
 		/// Which change to display the label on
 		/// </summary>
-		public LabelChange Change { get; }
+		public RpcLabelChange Change { get; }
 
 		/// <summary>
 		/// List of required nodes for the aggregate to be valid
@@ -645,7 +645,7 @@ namespace Horde.Server.Jobs.Graphs
 		/// <summary>
 		/// Which change the label applies to
 		/// </summary>
-		public LabelChange Change { get; set; }
+		public RpcLabelChange Change { get; set; }
 
 		/// <summary>
 		/// Nodes which must be part of the job for the aggregate to be valid
diff --git a/Engine/Source/Programs/Horde/Horde.Server/Jobs/IJob.cs b/Engine/Source/Programs/Horde/Horde.Server/Jobs/IJob.cs
--- a/Engine/Source/Programs/Horde/Horde.Server/Jobs/IJob.cs
+++ b/Engine/Source/Programs/Horde/Horde.Server/Jobs/IJob.cs
@@ -23,6 +23,7 @@ using Horde.Server.Jobs.Timing;
 using Horde.Server.Server;
 using Horde.Server.Streams;
 using Horde.Server.Ugs;
+using HordeCommon.Rpc;
 using MongoDB.Bson;
 using MongoDB.Driver;
 
@@ -1556,9 +1557,9 @@ namespace Horde.Server.Jobs
 		/// </summary>
 		/// <param name="job">The job document</param>
 		/// <returns></returns>
-		public static HordeCommon.Rpc.GetJobResponse ToRpcResponse(this IJob job)
+		public static RpcGetJobResponse ToRpcResponse(this IJob job)
 		{
-			HordeCommon.Rpc.GetJobResponse response = new HordeCommon.Rpc.GetJobResponse();
+			RpcGetJobResponse response = new RpcGetJobResponse();
 			response.StreamId = job.StreamId.ToString();
 			response.Change = job.Change;
 			response.CodeChange = job.CodeChange;
diff --git a/Engine/Source/Programs/Horde/Horde.Server/Jobs/JobMessages.cs b/Engine/Source/Programs/Horde/Horde.Server/Jobs/JobMessages.cs
--- a/Engine/Source/Programs/Horde/Horde.Server/Jobs/JobMessages.cs
+++ b/Engine/Source/Programs/Horde/Horde.Server/Jobs/JobMessages.cs
@@ -3,9 +3,7 @@
 using System;
 using System.Collections.Generic;
 using System.ComponentModel.DataAnnotations;
-using System.Linq;
 using System.Text.Json.Serialization;
-using EpicGames.Core;
 using EpicGames.Horde.Artifacts;
 using EpicGames.Horde.Commits;
 using EpicGames.Horde.Common;
diff --git a/Engine/Source/Programs/Horde/Horde.Server/Jobs/JobRpcService.cs b/Engine/Source/Programs/Horde/Horde.Server/Jobs/JobRpcService.cs
--- a/Engine/Source/Programs/Horde/Horde.Server/Jobs/JobRpcService.cs
+++ b/Engine/Source/Programs/Horde/Horde.Server/Jobs/JobRpcService.cs
@@ -40,12 +40,6 @@ using MongoDB.Bson.Serialization;
 
 namespace Horde.Server.Jobs
 {
-	using RpcGetJobResponse = HordeCommon.Rpc.GetJobResponse;
-	using RpcGetStepResponse = HordeCommon.Rpc.GetStepResponse;
-	using RpcGetStreamResponse = HordeCommon.Rpc.GetStreamResponse;
-	using RpcUpdateJobRequest = HordeCommon.Rpc.UpdateJobRequest;
-	using RpcUpdateStepRequest = HordeCommon.Rpc.UpdateStepRequest;
-
 	/// <summary>
 	/// Implements the Horde gRPC service for bots updating their status and dequeing work
 	/// </summary>
@@ -71,7 +65,7 @@ namespace Horde.Server.Jobs
 		}
 
 		/// <inheritdoc/>
-		public override async Task<Common.Rpc.CreateJobArtifactResponse> CreateArtifact(CreateJobArtifactRequest request, ServerCallContext context)
+		public override async Task<RpcCreateJobArtifactResponse> CreateArtifact(RpcCreateJobArtifactRequest request, ServerCallContext context)
 		{
 			(IJob job, _, IJobStep step) = await AuthorizeAsync(request.JobId, request.StepId, context);
 
@@ -101,11 +95,11 @@ namespace Horde.Server.Jobs
 			claims.Add(new AclClaimConfig(HordeClaimTypes.WriteNamespace, artifact.NamespaceId.ToString()));
 
 			string token = await _aclService.IssueBearerTokenAsync(claims, TimeSpan.FromHours(8.0), context.CancellationToken);
-			return new Common.Rpc.CreateJobArtifactResponse { Id = artifact.Id.ToString(), NamespaceId = artifact.NamespaceId.ToString(), RefName = artifact.RefName.ToString(), Token = token };
+			return new RpcCreateJobArtifactResponse { Id = artifact.Id.ToString(), NamespaceId = artifact.NamespaceId.ToString(), RefName = artifact.RefName.ToString(), Token = token };
 		}
 
 		/// <inheritdoc/>
-		public override async Task<Common.Rpc.CreateJobArtifactResponseV2> CreateArtifactV2(CreateJobArtifactRequestV2 request, ServerCallContext context)
+		public override async Task<RpcCreateJobArtifactResponseV2> CreateArtifactV2(RpcCreateJobArtifactRequestV2 request, ServerCallContext context)
 		{
 			(IJob job, _, IJobStep step) = await AuthorizeAsync(request.JobId, request.StepId, context);
 
@@ -135,7 +129,7 @@ namespace Horde.Server.Jobs
 
 			string token = await _aclService.IssueBearerTokenAsync(claims, TimeSpan.FromHours(8.0), context.CancellationToken);
 
-			Common.Rpc.CreateJobArtifactResponseV2 response = new Common.Rpc.CreateJobArtifactResponseV2();
+			RpcCreateJobArtifactResponseV2 response = new RpcCreateJobArtifactResponseV2();
 			response.Id = artifact.Id.ToString();
 			response.NamespaceId = artifact.NamespaceId.ToString();
 			response.RefName = artifact.RefName.ToString();
@@ -144,7 +138,7 @@ namespace Horde.Server.Jobs
 		}
 
 		/// <inheritdoc/>
-		public override async Task<Common.Rpc.GetJobArtifactResponse> GetArtifact(GetJobArtifactRequest request, ServerCallContext context)
+		public override async Task<RpcGetJobArtifactResponse> GetArtifact(RpcGetJobArtifactRequest request, ServerCallContext context)
 		{
 			(IJob job, _, _) = await AuthorizeAsync(request.JobId, request.StepId, context);
 
@@ -162,7 +156,7 @@ namespace Horde.Server.Jobs
 
 			string token = await _aclService.IssueBearerTokenAsync(claims, TimeSpan.FromHours(8.0), context.CancellationToken);
 
-			Common.Rpc.GetJobArtifactResponse response = new Common.Rpc.GetJobArtifactResponse();
+			RpcGetJobArtifactResponse response = new RpcGetJobArtifactResponse();
 			response.Id = artifact.Id.ToString();
 			response.NamespaceId = artifact.NamespaceId.ToString();
 			response.RefName = artifact.RefName.ToString();
@@ -205,43 +199,43 @@ namespace Horde.Server.Jobs
 		}
 
 		/// <inheritdoc/>
-		public override Task<RpcGetStreamResponse> GetStream(GetStreamRequest request, ServerCallContext context) => _jobRpcCommon.GetStreamAsync(request, context);
+		public override Task<RpcGetStreamResponse> GetStream(RpcGetStreamRequest request, ServerCallContext context) => _jobRpcCommon.GetStreamAsync(request, context);
 
 		/// <inheritdoc/>
-		public override Task<RpcGetJobResponse> GetJob(GetJobRequest request, ServerCallContext context) => _jobRpcCommon.GetJobAsync(request, context);
+		public override Task<RpcGetJobResponse> GetJob(RpcGetJobRequest request, ServerCallContext context) => _jobRpcCommon.GetJobAsync(request, context);
 
 		/// <inheritdoc/>
 		public override Task<Empty> UpdateJob(RpcUpdateJobRequest request, ServerCallContext context) => _jobRpcCommon.UpdateJobAsync(request, context);
 
 		/// <inheritdoc/>
-		public override Task<BeginBatchResponse> BeginBatch(BeginBatchRequest request, ServerCallContext context) => _jobRpcCommon.BeginBatchAsync(request, context);
+		public override Task<RpcBeginBatchResponse> BeginBatch(RpcBeginBatchRequest request, ServerCallContext context) => _jobRpcCommon.BeginBatchAsync(request, context);
 
 		/// <inheritdoc/>
-		public override Task<Empty> FinishBatch(FinishBatchRequest request, ServerCallContext context) => _jobRpcCommon.FinishBatchAsync(request, context);
+		public override Task<Empty> FinishBatch(RpcFinishBatchRequest request, ServerCallContext context) => _jobRpcCommon.FinishBatchAsync(request, context);
 
 		/// <inheritdoc/>
-		public override Task<BeginStepResponse> BeginStep(BeginStepRequest request, ServerCallContext context) => _jobRpcCommon.BeginStepAsync(request, context);
+		public override Task<RpcBeginStepResponse> BeginStep(RpcBeginStepRequest request, ServerCallContext context) => _jobRpcCommon.BeginStepAsync(request, context);
 
 		/// <inheritdoc/>
 		public override Task<Empty> UpdateStep(RpcUpdateStepRequest request, ServerCallContext context) => _jobRpcCommon.UpdateStepAsync(request, context);
 
 		/// <inheritdoc/>
-		public override Task<RpcGetStepResponse> GetStep(GetStepRequest request, ServerCallContext context) => _jobRpcCommon.GetStepAsync(request, context);
+		public override Task<RpcGetStepResponse> GetStep(RpcGetStepRequest request, ServerCallContext context) => _jobRpcCommon.GetStepAsync(request, context);
 
 		/// <inheritdoc/>
-		public override Task<UpdateGraphResponse> UpdateGraph(UpdateGraphRequest request, ServerCallContext context) => _jobRpcCommon.UpdateGraphAsync(request, context);
+		public override Task<RpcUpdateGraphResponse> UpdateGraph(RpcUpdateGraphRequest request, ServerCallContext context) => _jobRpcCommon.UpdateGraphAsync(request, context);
 
 		/// <inheritdoc/>
-		public override Task<Empty> CreateEvents(CreateEventsRequest request, ServerCallContext context) => _jobRpcCommon.CreateEventsAsync(request, context);
+		public override Task<Empty> CreateEvents(RpcCreateEventsRequest request, ServerCallContext context) => _jobRpcCommon.CreateEventsAsync(request, context);
 
 		/// <inheritdoc/>
-		public override Task<UploadArtifactResponse> UploadArtifact(IAsyncStreamReader<UploadArtifactRequest> reader, ServerCallContext context) => _jobRpcCommon.UploadArtifactAsync(reader, context);
+		public override Task<RpcUploadArtifactResponse> UploadArtifact(IAsyncStreamReader<RpcUploadArtifactRequest> reader, ServerCallContext context) => _jobRpcCommon.UploadArtifactAsync(reader, context);
 
 		/// <inheritdoc/>
-		public override Task<UploadTestDataResponse> UploadTestData(IAsyncStreamReader<UploadTestDataRequest> reader, ServerCallContext context) => _jobRpcCommon.UploadTestDataAsync(reader, context);
+		public override Task<RpcUploadTestDataResponse> UploadTestData(IAsyncStreamReader<RpcUploadTestDataRequest> reader, ServerCallContext context) => _jobRpcCommon.UploadTestDataAsync(reader, context);
 
 		/// <inheritdoc/>
-		public override Task<CreateReportResponse> CreateReport(CreateReportRequest request, ServerCallContext context) => _jobRpcCommon.CreateReportAsync(request, context);
+		public override Task<RpcCreateReportResponse> CreateReport(RpcCreateReportRequest request, ServerCallContext context) => _jobRpcCommon.CreateReportAsync(request, context);
 	}
 
 	/// <summary>
@@ -285,7 +279,7 @@ namespace Horde.Server.Jobs
 		/// <param name="request">Request arguments</param>
 		/// <param name="context">Context for the RPC call</param>
 		/// <returns>Information about the new agent</returns>
-		public Task<RpcGetStreamResponse> GetStreamAsync(GetStreamRequest request, ServerCallContext context)
+		public Task<RpcGetStreamResponse> GetStreamAsync(RpcGetStreamRequest request, ServerCallContext context)
 		{
 			StreamId streamIdValue = new StreamId(request.StreamId);
 
@@ -308,7 +302,7 @@ namespace Horde.Server.Jobs
 		/// <param name="request">Request arguments</param>
 		/// <param name="context">Context for the RPC call</param>
 		/// <returns>Information about the new agent</returns>
-		public async Task<RpcGetJobResponse> GetJobAsync(GetJobRequest request, ServerCallContext context)
+		public async Task<RpcGetJobResponse> GetJobAsync(RpcGetJobRequest request, ServerCallContext context)
 		{
 			JobId jobIdValue = JobId.Parse(request.JobId);
 
@@ -355,7 +349,7 @@ namespace Horde.Server.Jobs
 		/// <param name="request">Request arguments</param>
 		/// <param name="context">Context for the RPC call</param>
 		/// <returns>Information about the new agent</returns>
-		public async Task<BeginBatchResponse> BeginBatchAsync(BeginBatchRequest request, ServerCallContext context)
+		public async Task<RpcBeginBatchResponse> BeginBatchAsync(RpcBeginBatchRequest request, ServerCallContext context)
 		{
 			JobStepBatchId batchId = JobStepBatchId.Parse(request.BatchId);
 
@@ -380,7 +374,7 @@ namespace Horde.Server.Jobs
 			IGraph graph = await _jobService.GetGraphAsync(job);
 			AgentConfig agentConfig = streamConfig.AgentTypes[graph.Groups[batch.GroupIdx].AgentType];
 
-			BeginBatchResponse response = new BeginBatchResponse();
+			RpcBeginBatchResponse response = new RpcBeginBatchResponse();
 			response.LogId = batch.LogId.ToString();
 			response.AgentType = graph.Groups[batch.GroupIdx].AgentType;
 			response.StreamName = streamConfig.Name;
@@ -408,7 +402,7 @@ namespace Horde.Server.Jobs
 		/// <param name="request">Request arguments</param>
 		/// <param name="context">Context for the RPC call</param>
 		/// <returns>Information about the new agent</returns>
-		public async Task<Empty> FinishBatchAsync(FinishBatchRequest request, ServerCallContext context)
+		public async Task<Empty> FinishBatchAsync(RpcFinishBatchRequest request, ServerCallContext context)
 		{
 			IJob? job = await _jobService.GetJobAsync(JobId.Parse(request.JobId));
 			if (job == null)
@@ -431,12 +425,12 @@ namespace Horde.Server.Jobs
 		/// <param name="request">Request arguments</param>
 		/// <param name="context">Context for the RPC call</param>
 		/// <returns>Information about the new agent</returns>
-		public async Task<BeginStepResponse> BeginStepAsync(BeginStepRequest request, ServerCallContext context)
+		public async Task<RpcBeginStepResponse> BeginStepAsync(RpcBeginStepRequest request, ServerCallContext context)
 		{
 			Boxed<ILogFile?> log = new Boxed<ILogFile?>(null);
 			for (; ; )
 			{
-				BeginStepResponse? response = await TryBeginStepAsync(request, log, context);
+				RpcBeginStepResponse? response = await TryBeginStepAsync(request, log, context);
 				if (response != null)
 				{
 					return response;
@@ -444,7 +438,7 @@ namespace Horde.Server.Jobs
 			}
 		}
 
-		async Task<BeginStepResponse?> TryBeginStepAsync(BeginStepRequest request, Boxed<ILogFile?> log, ServerCallContext context)
+		async Task<RpcBeginStepResponse?> TryBeginStepAsync(RpcBeginStepRequest request, Boxed<ILogFile?> log, ServerCallContext context)
 		{
 			// Check the job exists and we can access it
 			IJob? job = await _jobService.GetJobAsync(JobId.Parse(request.JobId));
@@ -461,7 +455,7 @@ namespace Horde.Server.Jobs
 			IJobStepBatch batch = AuthorizeBatch(job, JobStepBatchId.Parse(request.BatchId), context);
 			if (batch.State != JobStepBatchState.Starting && batch.State != JobStepBatchState.Running)
 			{
-				return new BeginStepResponse { State = BeginStepResponse.Types.Result.Complete };
+				return new RpcBeginStepResponse { State = RpcBeginStepResponse.Types.Result.Complete };
 			}
 
 			// Figure out which step to execute next
@@ -476,7 +470,7 @@ namespace Horde.Server.Jobs
 					{
 						return null;
 					}
-					return new BeginStepResponse { State = BeginStepResponse.Types.Result.Complete };
+					return new RpcBeginStepResponse { State = RpcBeginStepResponse.Types.Result.Complete };
 				}
 
 				// Check if this step is ready to be executed
@@ -488,7 +482,7 @@ namespace Horde.Server.Jobs
 				if (step.State == JobStepState.Waiting)
 				{
 					_logger.LogDebug("Waiting for job {JobId}, batch {BatchId}, step {StepId}", job.Id, batch.Id, step.Id);
-					return new BeginStepResponse { State = BeginStepResponse.Types.Result.Waiting };
+					return new RpcBeginStepResponse { State = RpcBeginStepResponse.Types.Result.Waiting };
 				}
 			}
 
@@ -515,8 +509,8 @@ namespace Horde.Server.Jobs
 			IJob? newJob = await _jobService.TryUpdateStepAsync(job, batch.Id, step.Id, streamConfig, JobStepState.Running, JobStepOutcome.Unspecified, newLogId: log.Value.Id);
 			if (newJob != null)
 			{
-				BeginStepResponse response = new BeginStepResponse();
-				response.State = BeginStepResponse.Types.Result.Ready;
+				RpcBeginStepResponse response = new RpcBeginStepResponse();
+				response.State = RpcBeginStepResponse.Types.Result.Ready;
 				response.LogId = log.Value.Id.ToString();
 				response.StepId = step.Id.ToString();
 				response.Name = node.Name;
@@ -648,7 +642,7 @@ namespace Horde.Server.Jobs
 				{
 					if (node.OutputNames.Contains(artifact.OutputName))
 					{
-						CreateGraphArtifactRequest stepArtifact = new CreateGraphArtifactRequest { Name = artifact.Name.ToString(), Type = artifact.Type.ToString(), Description = artifact.Description, BasePath = artifact.BasePath, OutputName = artifact.OutputName };
+						RpcCreateGraphArtifactRequest stepArtifact = new RpcCreateGraphArtifactRequest { Name = artifact.Name.ToString(), Type = artifact.Type.ToString(), Description = artifact.Description, BasePath = artifact.BasePath, OutputName = artifact.OutputName };
 						stepArtifact.Keys.AddRange(artifact.Keys);
 						stepArtifact.Metadata.AddRange(artifact.Metadata);
 						response.Artifacts.Add(stepArtifact);
@@ -778,7 +772,7 @@ namespace Horde.Server.Jobs
 		/// <param name="request">Request arguments</param>
 		/// <param name="context">Context for the RPC call</param>
 		/// <returns>Information about the step</returns>
-		public async Task<RpcGetStepResponse> GetStepAsync(GetStepRequest request, ServerCallContext context)
+		public async Task<RpcGetStepResponse> GetStepAsync(RpcGetStepRequest request, ServerCallContext context)
 		{
 			IJob job = await GetJobAsync(JobId.Parse(request.JobId), context.CancellationToken);
 			IJobStepBatch batch = AuthorizeBatch(job, JobStepBatchId.Parse(request.BatchId), context);
@@ -798,7 +792,7 @@ namespace Horde.Server.Jobs
 		/// <param name="request">Request arguments</param>
 		/// <param name="context">Context for the RPC call</param>
 		/// <returns>Information about the new agent</returns>
-		public async Task<UpdateGraphResponse> UpdateGraphAsync(UpdateGraphRequest request, ServerCallContext context)
+		public async Task<RpcUpdateGraphResponse> UpdateGraphAsync(RpcUpdateGraphRequest request, ServerCallContext context)
 		{
 			JobId jobIdValue = JobId.Parse(request.JobId);
 			for (; ; )
@@ -818,10 +812,10 @@ namespace Horde.Server.Jobs
 
 				// Add all the new nodes
 				List<NewGroup> newGroups = new List<NewGroup> { new NewGroup(oldGraph, oldGraph.Groups[0]) };
-				foreach (CreateGroupRequest group in request.Groups)
+				foreach (RpcCreateGroupRequest group in request.Groups)
 				{
 					List<NewNode> newNodes = new List<NewNode>();
-					foreach (CreateNodeRequest node in group.Nodes)
+					foreach (RpcCreateNodeRequest node in group.Nodes)
 					{
 						NewNode newNode = new NewNode(node.Name, node.Inputs.ToList(), node.Outputs.ToList(), node.InputDependencies.ToList(), node.OrderDependencies.ToList(), (Priority)node.Priority, node.AllowRetry, node.RunEarly, node.Warnings, new Dictionary<string, string>(node.Credentials), new Dictionary<string, string>(node.Properties), new NodeAnnotations(node.Annotations));
 						newNodes.Add(newNode);
@@ -831,7 +825,7 @@ namespace Horde.Server.Jobs
 
 				// Add all the new aggregates
 				List<NewAggregate> newAggregates = new List<NewAggregate>();
-				foreach (CreateAggregateRequest aggregate in request.Aggregates)
+				foreach (RpcCreateAggregateRequest aggregate in request.Aggregates)
 				{
 					NewAggregate newAggregate = new NewAggregate(aggregate.Name, aggregate.Nodes.ToList());
 					newAggregates.Add(newAggregate);
@@ -839,7 +833,7 @@ namespace Horde.Server.Jobs
 
 				// Add all the new labels
 				List<NewLabel> newLabels = new List<NewLabel>();
-				foreach (CreateLabelRequest label in request.Labels)
+				foreach (RpcCreateLabelRequest label in request.Labels)
 				{
 					NewLabel newLabel = new NewLabel();
 					newLabel.DashboardName = String.IsNullOrEmpty(label.DashboardName) ? null : label.DashboardName;
@@ -854,7 +848,7 @@ namespace Horde.Server.Jobs
 
 				// Add all the new artifacts
 				List<NewGraphArtifact> newArtifacts = new List<NewGraphArtifact>();
-				foreach (CreateGraphArtifactRequest artifact in request.Artifacts)
+				foreach (RpcCreateGraphArtifactRequest artifact in request.Artifacts)
 				{
 					ArtifactName name = new ArtifactName(StringId.Sanitize(artifact.Name));
 					ArtifactType type = new ArtifactType(StringId.Sanitize(artifact.Type));
@@ -876,7 +870,7 @@ namespace Horde.Server.Jobs
 				if (newJob != null)
 				{
 					_logger.LogInformation("Updating graph for {JobId} from {OldGraphHash} to {NewGraphHash}", job.Id, oldGraph.Id, newJob.GraphHash);
-					return new UpdateGraphResponse();
+					return new RpcUpdateGraphResponse();
 				}
 			}
 		}
@@ -887,7 +881,7 @@ namespace Horde.Server.Jobs
 		/// <param name="request">Request arguments</param>
 		/// <param name="context">Context for the RPC call</param>
 		/// <returns>Information about the new agent</returns>
-		public async Task<Empty> CreateEventsAsync(CreateEventsRequest request, ServerCallContext context)
+		public async Task<Empty> CreateEventsAsync(RpcCreateEventsRequest request, ServerCallContext context)
 		{
 			if (!_globalConfig.Value.Authorize(LogAclAction.CreateEvent, context.GetHttpContext().User))
 			{
@@ -895,7 +889,7 @@ namespace Horde.Server.Jobs
 			}
 
 			List<NewLogEventData> newEvents = new List<NewLogEventData>();
-			foreach (CreateEventRequest createEvent in request.Events)
+			foreach (RpcCreateEventRequest createEvent in request.Events)
 			{
 				NewLogEventData newEvent = new NewLogEventData();
 				newEvent.LogId = LogId.Parse(createEvent.LogId);
@@ -914,7 +908,7 @@ namespace Horde.Server.Jobs
 		/// <param name="reader">Request arguments</param>
 		/// <param name="context">Context for the RPC call</param>
 		/// <returns>Information about the new agent</returns>
-		public async Task<UploadArtifactResponse> UploadArtifactAsync(IAsyncStreamReader<UploadArtifactRequest> reader, ServerCallContext context)
+		public async Task<RpcUploadArtifactResponse> UploadArtifactAsync(IAsyncStreamReader<RpcUploadArtifactRequest> reader, ServerCallContext context)
 		{
 			// Advance to the metadata object
 			if (!await reader.MoveNext())
@@ -923,7 +917,7 @@ namespace Horde.Server.Jobs
 			}
 
 			// Read the request object
-			UploadArtifactMetadata? metadata = reader.Current.Metadata;
+			RpcUploadArtifactMetadata? metadata = reader.Current.Metadata;
 			if (metadata == null)
 			{
 				throw new StructuredRpcException(StatusCode.DataLoss, "Expected metadata in first artifact request");
@@ -944,7 +938,7 @@ namespace Horde.Server.Jobs
 			{
 				IArtifactV1 artifact = await _artifactCollection.CreateArtifactAsync(job.Id, step.Id, metadata.Name, metadata.MimeType, inputStream, context.CancellationToken);
 
-				UploadArtifactResponse response = new UploadArtifactResponse();
+				RpcUploadArtifactResponse response = new RpcUploadArtifactResponse();
 				response.Id = artifact.Id.ToString();
 				return response;
 			}
@@ -956,7 +950,7 @@ namespace Horde.Server.Jobs
 		/// <param name="reader">Request arguments</param>
 		/// <param name="context">Context for the RPC call</param>
 		/// <returns>Information about the new agent</returns>
-		public async Task<UploadTestDataResponse> UploadTestDataAsync(IAsyncStreamReader<UploadTestDataRequest> reader, ServerCallContext context)
+		public async Task<RpcUploadTestDataResponse> UploadTestDataAsync(IAsyncStreamReader<RpcUploadTestDataRequest> reader, ServerCallContext context)
 		{
 			IJob? job = null;
 			IJobStep? jobStep = null;
@@ -965,7 +959,7 @@ namespace Horde.Server.Jobs
 
 			while (await reader.MoveNext())
 			{
-				UploadTestDataRequest request = reader.Current;
+				RpcUploadTestDataRequest request = reader.Current;
 
 				JobId jobId = JobId.Parse(request.JobId);
 				if (job == null)
@@ -1009,7 +1003,7 @@ namespace Horde.Server.Jobs
 				throw new StructuredRpcException(StatusCode.NotFound, "Unable to get job or step for test data upload");
 			}
 
-			return new UploadTestDataResponse();
+			return new RpcUploadTestDataResponse();
 		}
 
 		/// <summary>
@@ -1018,7 +1012,7 @@ namespace Horde.Server.Jobs
 		/// <param name="request"></param>
 		/// <param name="context"></param>
 		/// <returns></returns>
-		public async Task<CreateReportResponse> CreateReportAsync(CreateReportRequest request, ServerCallContext context)
+		public async Task<RpcCreateReportResponse> CreateReportAsync(RpcCreateReportRequest request, ServerCallContext context)
 		{
 			IJob job = await GetJobAsync(JobId.Parse(request.JobId), context.CancellationToken);
 			if (!_globalConfig.Value.TryGetStream(job.StreamId, out StreamConfig? streamConfig))
@@ -1029,7 +1023,7 @@ namespace Horde.Server.Jobs
 			IJobStepBatch batch = AuthorizeBatch(job, JobStepBatchId.Parse(request.BatchId), context);
 
 			Report newReport = new Report { Name = request.Name, Placement = (ReportPlacement)request.Placement, Content = request.Content };
-			if (request.Scope == ReportScope.Job)
+			if (request.Scope == RpcReportScope.Job)
 			{
 				_logger.LogDebug("Adding report to job {JobId}: {Name} -> {Content}", job.Id, request.Name, request.Content);
 				await _jobService.UpdateJobAsync(job, reports: new List<Report> { newReport });
@@ -1040,7 +1034,7 @@ namespace Horde.Server.Jobs
 				await _jobService.UpdateStepAsync(job, batch.Id, JobStepId.Parse(request.StepId), streamConfig, newReports: new List<Report> { newReport });
 			}
 
-			return new CreateReportResponse();
+			return new RpcCreateReportResponse();
 		}
 	}
 }
diff --git a/Engine/Source/Programs/Horde/Horde.Server/Jobs/JobTaskSource.cs b/Engine/Source/Programs/Horde/Horde.Server/Jobs/JobTaskSource.cs
--- a/Engine/Source/Programs/Horde/Horde.Server/Jobs/JobTaskSource.cs
+++ b/Engine/Source/Programs/Horde/Horde.Server/Jobs/JobTaskSource.cs
@@ -908,7 +908,7 @@ namespace Horde.Server.Jobs
 			task.StoragePrefix = storagePrefix;
 			task.Token = await _aclService.IssueBearerTokenAsync(claims, null, cancellationToken);
 
-			List<AgentWorkspace> workspaces = new();
+			List<RpcAgentWorkspace> workspaces = new();
 
 			PerforceCluster? cluster = globalConfig.FindPerforceCluster(workspace.Cluster);
 			if (cluster == null)
@@ -1045,7 +1045,7 @@ namespace Horde.Server.Jobs
 
 				// The changelist number to display the badge for
 				int change;
-				if (label.Change == LabelChange.Code)
+				if (label.Change == RpcLabelChange.Code)
 				{
 					change = job.CodeChange;
 				}
diff --git a/Engine/Source/Programs/Horde/Horde.Server/Logs/ILogEvent.cs b/Engine/Source/Programs/Horde/Horde.Server/Logs/ILogEvent.cs
--- a/Engine/Source/Programs/Horde/Horde.Server/Logs/ILogEvent.cs
+++ b/Engine/Source/Programs/Horde/Horde.Server/Logs/ILogEvent.cs
@@ -19,7 +19,7 @@ namespace Horde.Server.Logs
 		/// <summary>
 		/// The log level
 		/// </summary>
-		EventSeverity Severity { get; }
+		RpcEventSeverity Severity { get; }
 
 		/// <summary>
 		/// The type of event
@@ -50,7 +50,7 @@ namespace Horde.Server.Logs
 		/// <summary>
 		/// Severity of the event
 		/// </summary>
-		public EventSeverity Severity { get; }
+		public RpcEventSeverity Severity { get; }
 
 		/// <summary>
 		/// Index of the first line for this event
@@ -81,7 +81,7 @@ namespace Horde.Server.Logs
 		/// <summary>
 		/// Severity of the event
 		/// </summary>
-		public EventSeverity Severity { get; set; }
+		public RpcEventSeverity Severity { get; set; }
 
 		/// <summary>
 		/// Index of the first line for this event
diff --git a/Engine/Source/Programs/Horde/Horde.Server/Logs/LogEventCollection.cs b/Engine/Source/Programs/Horde/Horde.Server/Logs/LogEventCollection.cs
--- a/Engine/Source/Programs/Horde/Horde.Server/Logs/LogEventCollection.cs
+++ b/Engine/Source/Programs/Horde/Horde.Server/Logs/LogEventCollection.cs
@@ -46,17 +46,17 @@ namespace Horde.Server.Logs
 			LogId ILogEvent.LogId => Id.LogId;
 			int ILogEvent.LineIndex => Id.LineIndex;
 			int ILogEvent.LineCount => LineCount ?? 1;
-			EventSeverity ILogEvent.Severity => IsWarning ? EventSeverity.Warning : EventSeverity.Error;
+			RpcEventSeverity ILogEvent.Severity => IsWarning ? RpcEventSeverity.Warning : RpcEventSeverity.Error;
 
 			public LogEventDocument()
 			{
 				Id = new LogEventId();
 			}
 
-			public LogEventDocument(LogId logId, EventSeverity severity, int lineIndex, int lineCount, ObjectId? spanId)
+			public LogEventDocument(LogId logId, RpcEventSeverity severity, int lineIndex, int lineCount, ObjectId? spanId)
 			{
 				Id = new LogEventId { LogId = logId, LineIndex = lineIndex };
-				IsWarning = severity == EventSeverity.Warning;
+				IsWarning = severity == RpcEventSeverity.Warning;
 				LineCount = (lineCount > 1) ? (int?)lineCount : null;
 				SpanId = spanId;
 			}
@@ -71,7 +71,7 @@ namespace Horde.Server.Logs
 		{
 			public ObjectId Id { get; set; }
 			public DateTime Time { get; set; }
-			public EventSeverity Severity { get; set; }
+			public RpcEventSeverity Severity { get; set; }
 			public LogId LogId { get; set; }
 			public int LineIndex { get; set; }
 			public int LineCount { get; set; }
diff --git a/Engine/Source/Programs/Horde/Horde.Server/Logs/LogEventMessages.cs b/Engine/Source/Programs/Horde/Horde.Server/Logs/LogEventMessages.cs
--- a/Engine/Source/Programs/Horde/Horde.Server/Logs/LogEventMessages.cs
+++ b/Engine/Source/Programs/Horde/Horde.Server/Logs/LogEventMessages.cs
@@ -20,7 +20,7 @@ namespace Horde.Server.Logs
 		/// <summary>
 		/// Severity of this event
 		/// </summary>
-		public EventSeverity Severity { get; set; }
+		public RpcEventSeverity Severity { get; set; }
 
 		/// <summary>
 		/// Index of the first line for this event
diff --git a/Engine/Source/Programs/Horde/Horde.Server/Logs/LogFileService.cs b/Engine/Source/Programs/Horde/Horde.Server/Logs/LogFileService.cs
--- a/Engine/Source/Programs/Horde/Horde.Server/Logs/LogFileService.cs
+++ b/Engine/Source/Programs/Horde/Horde.Server/Logs/LogFileService.cs
@@ -691,7 +691,7 @@ namespace Horde.Server.Logs
 			public IReadOnlyList<JsonLogEvent> Lines { get; }
 
 			EventId? ILogEventData.EventId => (Lines.Count > 0) ? Lines[0].EventId : null;
-			EventSeverity ILogEventData.Severity => (Lines.Count == 0) ? EventSeverity.Information : (Lines[0].Level == LogLevel.Warning) ? EventSeverity.Warning : EventSeverity.Error;
+			RpcEventSeverity ILogEventData.Severity => (Lines.Count == 0) ? RpcEventSeverity.Information : (Lines[0].Level == LogLevel.Warning) ? RpcEventSeverity.Warning : RpcEventSeverity.Error;
 
 			public LogEventData(IReadOnlyList<JsonLogEvent> lines)
 			{
diff --git a/Engine/Source/Programs/Horde/Horde.Server/Notifications/Sinks/SlackNotificationSink.cs b/Engine/Source/Programs/Horde/Horde.Server/Notifications/Sinks/SlackNotificationSink.cs
--- a/Engine/Source/Programs/Horde/Horde.Server/Notifications/Sinks/SlackNotificationSink.cs
+++ b/Engine/Source/Programs/Horde/Horde.Server/Notifications/Sinks/SlackNotificationSink.cs
@@ -639,8 +639,8 @@ namespace Horde.Server.Notifications.Sinks
 			}
 			else
 			{
-				List<ILogEventData> errors = events.Where(x => x.Severity == EventSeverity.Error).ToList();
-				List<ILogEventData> warnings = events.Where(x => x.Severity == EventSeverity.Warning).ToList();
+				List<ILogEventData> errors = events.Where(x => x.Severity == RpcEventSeverity.Error).ToList();
+				List<ILogEventData> warnings = events.Where(x => x.Severity == RpcEventSeverity.Warning).ToList();
 				List<string> eventStrings = new List<string>();
 				if (errors.Any())
 				{
@@ -1018,9 +1018,9 @@ namespace Horde.Server.Notifications.Sinks
 					if (logFile != null)
 					{
 						events = await _logFileService.FindEventsAsync(logFile, span.Id, 0, 50, cancellationToken);
-						if (events.Any(x => x.Severity == EventSeverity.Error))
+						if (events.Any(x => x.Severity == RpcEventSeverity.Error))
 						{
-							events.RemoveAll(x => x.Severity == EventSeverity.Warning);
+							events.RemoveAll(x => x.Severity == RpcEventSeverity.Warning);
 						}
 
 						List<string> eventStrings = new List<string>();
@@ -1454,9 +1454,9 @@ namespace Horde.Server.Notifications.Sinks
 				if (logFile != null)
 				{
 					List<ILogEvent> events = await _logFileService.FindEventsAsync(logFile, lastSpan.Id, 0, 20, cancellationToken);
-					if (events.Any(x => x.Severity == EventSeverity.Error))
+					if (events.Any(x => x.Severity == RpcEventSeverity.Error))
 					{
-						events.RemoveAll(x => x.Severity == EventSeverity.Warning);
+						events.RemoveAll(x => x.Severity == RpcEventSeverity.Warning);
 					}
 
 					for (int idx = 0; idx < Math.Min(events.Count, 3); idx++)
diff --git a/Engine/Source/Programs/Horde/Horde.Server/Perforce/PerforceLoadBalancer.cs b/Engine/Source/Programs/Horde/Horde.Server/Perforce/PerforceLoadBalancer.cs
--- a/Engine/Source/Programs/Horde/Horde.Server/Perforce/PerforceLoadBalancer.cs
+++ b/Engine/Source/Programs/Horde/Horde.Server/Perforce/PerforceLoadBalancer.cs
@@ -546,7 +546,7 @@ namespace Horde.Server.Perforce
 				HashSet<(string, string)> servers = new HashSet<(string, string)>();
 				if (any.TryUnpack(out ConformTask conformTask))
 				{
-					foreach (AgentWorkspace workspace in conformTask.Workspaces)
+					foreach (RpcAgentWorkspace workspace in conformTask.Workspaces)
 					{
 						servers.Add((workspace.Cluster, workspace.ServerAndPort));
 					}
diff --git a/Engine/Source/Programs/Horde/Horde.Server/Projects/ProjectConfig.cs b/Engine/Source/Programs/Horde/Horde.Server/Projects/ProjectConfig.cs
--- a/Engine/Source/Programs/Horde/Horde.Server/Projects/ProjectConfig.cs
+++ b/Engine/Source/Programs/Horde/Horde.Server/Projects/ProjectConfig.cs
@@ -8,7 +8,6 @@ using System.Security.Claims;
 using System.Text.Json.Serialization;
 using EpicGames.Core;
 using EpicGames.Horde.Acls;
-using EpicGames.Horde.Jobs;
 using EpicGames.Horde.Projects;
 using EpicGames.Horde.Telemetry;
 using Horde.Server.Acls;
diff --git a/Engine/Source/Programs/Horde/Horde.Server/Protos/horde/common/io_hash.cs b/Engine/Source/Programs/Horde/Horde.Server/Protos/horde/common/io_hash.cs
--- a/Engine/Source/Programs/Horde/Horde.Server/Protos/horde/common/io_hash.cs
+++ b/Engine/Source/Programs/Horde/Horde.Server/Protos/horde/common/io_hash.cs
@@ -5,16 +5,16 @@ namespace HordeCommon.Wrapper
 	using NativeIoHash = EpicGames.Core.IoHash;
 
 #pragma warning disable CS1591 // Missing XML comment for publicly visible type or member
-	partial class IoHash
+	partial class RpcIoHash
 	{
-		public static implicit operator NativeIoHash(IoHash hash)
+		public static implicit operator NativeIoHash(RpcIoHash hash)
 		{
 			return new NativeIoHash(hash.Data.ToByteArray());
 		}
 
-		public static implicit operator IoHash(NativeIoHash hash)
+		public static implicit operator RpcIoHash(NativeIoHash hash)
 		{
-			IoHash result = new IoHash();
+			RpcIoHash result = new RpcIoHash();
 			result.Data = Google.Protobuf.ByteString.CopyFrom(hash.ToByteArray());
 			return result;
 		}
diff --git a/Engine/Source/Programs/Horde/Horde.Server/Protos/horde/common/namespace_id.cs b/Engine/Source/Programs/Horde/Horde.Server/Protos/horde/common/namespace_id.cs
--- a/Engine/Source/Programs/Horde/Horde.Server/Protos/horde/common/namespace_id.cs
+++ b/Engine/Source/Programs/Horde/Horde.Server/Protos/horde/common/namespace_id.cs
@@ -7,16 +7,16 @@ namespace HordeCommon.Wrapper
 	using NativeNamespaceId = EpicGames.Horde.Storage.NamespaceId;
 
 #pragma warning disable CS1591 // Missing XML comment for publicly visible type or member
-	partial class NamespaceId
+	partial class RpcNamespaceId
 	{
-		public static implicit operator NativeNamespaceId(NamespaceId wrapper)
+		public static implicit operator NativeNamespaceId(RpcNamespaceId wrapper)
 		{
 			return new NativeNamespaceId(new Utf8String(wrapper.Data));
 		}
 
-		public static implicit operator NamespaceId(NativeNamespaceId source)
+		public static implicit operator RpcNamespaceId(NativeNamespaceId source)
 		{
-			NamespaceId result = new NamespaceId();
+			RpcNamespaceId result = new RpcNamespaceId();
 			result.Data = source.ToString();
 			return result;
 		}
diff --git a/Engine/Source/Programs/Horde/Horde.Server/Server/RpcService.cs b/Engine/Source/Programs/Horde/Horde.Server/Server/RpcService.cs
--- a/Engine/Source/Programs/Horde/Horde.Server/Server/RpcService.cs
+++ b/Engine/Source/Programs/Horde/Horde.Server/Server/RpcService.cs
@@ -27,6 +27,7 @@ using Horde.Server.Telemetry;
 using Horde.Server.Tools;
 using Horde.Server.Utilities;
 using HordeCommon.Rpc;
+using HordeCommon.Rpc.Messages;
 using HordeCommon.Rpc.Messages.Telemetry;
 using Microsoft.AspNetCore.Authorization;
 using Microsoft.Extensions.Logging;
@@ -34,14 +35,6 @@ using Microsoft.Extensions.Options;
 
 namespace Horde.Server.Server
 {
-	using RpcAgentCapabilities = HordeCommon.Rpc.Messages.AgentCapabilities;
-	using RpcDeviceCapabilities = HordeCommon.Rpc.Messages.DeviceCapabilities;
-	using RpcGetJobResponse = HordeCommon.Rpc.GetJobResponse;
-	using RpcGetStepResponse = HordeCommon.Rpc.GetStepResponse;
-	using RpcGetStreamResponse = HordeCommon.Rpc.GetStreamResponse;
-	using RpcUpdateJobRequest = HordeCommon.Rpc.UpdateJobRequest;
-	using RpcUpdateStepRequest = HordeCommon.Rpc.UpdateStepRequest;
-
 	/// <summary>
 	/// Implements the Horde gRPC service for bots updating their status and dequeing work
 	/// </summary>
@@ -88,15 +81,15 @@ namespace Horde.Server.Server
 		/// <param name="writer">Response writer</param>
 		/// <param name="context">Context for the call</param>
 		/// <returns>Response object</returns>
-		public override async Task QueryServerState(IAsyncStreamReader<QueryServerStateRequest> reader, IServerStreamWriter<QueryServerStateResponse> writer, ServerCallContext context)
+		public override async Task QueryServerState(IAsyncStreamReader<RpcQueryServerStateRequest> reader, IServerStreamWriter<RpcQueryServerStateResponse> writer, ServerCallContext context)
 		{
 			if (await reader.MoveNext())
 			{
-				QueryServerStateRequest request = reader.Current;
+				RpcQueryServerStateRequest request = reader.Current;
 				_logger.LogInformation("Start server query for client {Name}", request.Name);
 
 				// Return the current response
-				QueryServerStateResponse response = new QueryServerStateResponse();
+				RpcQueryServerStateResponse response = new RpcQueryServerStateResponse();
 				response.Name = Dns.GetHostName();
 				await writer.WriteAsync(response);
 
@@ -131,17 +124,17 @@ namespace Horde.Server.Server
 		/// <param name="writer">Response writer</param>
 		/// <param name="context">Context for the call</param>
 		/// <returns>Response object</returns>
-		public override async Task QueryServerStateV2(IAsyncStreamReader<QueryServerStateRequest> reader, IServerStreamWriter<QueryServerStateResponse> writer, ServerCallContext context)
+		public override async Task QueryServerStateV2(IAsyncStreamReader<RpcQueryServerStateRequest> reader, IServerStreamWriter<RpcQueryServerStateResponse> writer, ServerCallContext context)
 		{
 			if (await reader.MoveNext())
 			{
-				QueryServerStateRequest request = reader.Current;
+				RpcQueryServerStateRequest request = reader.Current;
 				_logger.LogDebug("Start server query for client {Name}", request.Name);
 
 				try
 				{
 					// Return the current response
-					QueryServerStateResponse response = new QueryServerStateResponse();
+					RpcQueryServerStateResponse response = new RpcQueryServerStateResponse();
 					response.Name = Dns.GetHostName();
 					response.Stopping = _lifetimeService.IsStopping;
 					await writer.WriteAsync(response);
@@ -176,7 +169,7 @@ namespace Horde.Server.Server
 		/// <param name="request">The request parameters</param>
 		/// <param name="context">Context for the call</param>
 		/// <returns>Response object</returns>
-		public override async Task<UpdateAgentWorkspacesResponse> UpdateAgentWorkspaces(UpdateAgentWorkspacesRequest request, ServerCallContext context)
+		public override async Task<RpcUpdateAgentWorkspacesResponse> UpdateAgentWorkspaces(RpcUpdateAgentWorkspacesRequest request, ServerCallContext context)
 		{
 			for (; ; )
 			{
@@ -197,7 +190,7 @@ namespace Horde.Server.Server
 				// Update the workspaces
 				if (await _agentService.TryUpdateWorkspacesAsync(agent, newWorkspaces, pendingConform, context.CancellationToken))
 				{
-					UpdateAgentWorkspacesResponse response = new UpdateAgentWorkspacesResponse();
+					RpcUpdateAgentWorkspacesResponse response = new RpcUpdateAgentWorkspacesResponse();
 					if (pendingConform)
 					{
 						response.Retry = await _conformTaskSource.GetWorkspacesAsync(agent, response.PendingWorkspaces, context.CancellationToken);
@@ -256,7 +249,7 @@ namespace Horde.Server.Server
 		/// <param name="request">Request to create a new agent</param>
 		/// <param name="context">Context for the RPC call</param>
 		/// <returns>Information about the new agent</returns>
-		public override async Task<CreateAgentResponse> CreateAgent(CreateAgentRequest request, ServerCallContext context)
+		public override async Task<RpcCreateAgentResponse> CreateAgent(RpcCreateAgentRequest request, ServerCallContext context)
 		{
 			using IDisposable? scope = _logger.BeginScope("CreateAgent({AgentId})", request.Name.ToString());
 
@@ -270,7 +263,7 @@ namespace Horde.Server.Server
 			List<AclClaimConfig> claims = new List<AclClaimConfig>();
 			claims.Add(new AclClaimConfig(HordeClaimTypes.Agent, agent.Id.ToString()));
 
-			CreateAgentResponse response = new CreateAgentResponse();
+			RpcCreateAgentResponse response = new RpcCreateAgentResponse();
 			response.Id = agent.Id.ToString();
 			response.Token = await _aclService.IssueBearerTokenAsync(claims, null, context.CancellationToken);
 
@@ -283,7 +276,7 @@ namespace Horde.Server.Server
 		/// <param name="request">Request to create a new agent</param>
 		/// <param name="context">Context for the RPC call</param>
 		/// <returns>Information about the new agent</returns>
-		public override async Task<CreateSessionResponse> CreateSession(CreateSessionRequest request, ServerCallContext context)
+		public override async Task<RpcCreateSessionResponse> CreateSession(RpcCreateSessionRequest request, ServerCallContext context)
 		{
 			if (request.Capabilities == null)
 			{
@@ -333,7 +326,7 @@ namespace Horde.Server.Server
 			}
 
 			// Create the response
-			CreateSessionResponse response = new CreateSessionResponse();
+			RpcCreateSessionResponse response = new RpcCreateSessionResponse();
 			response.AgentId = agent.Id.ToString();
 			response.SessionId = agent.SessionId.ToString();
 			response.ExpiryTime = Timestamp.FromDateTime(agent.SessionExpiresAt!.Value);
@@ -348,17 +341,17 @@ namespace Horde.Server.Server
 		/// <param name="writer">Writer for response objects</param>
 		/// <param name="context">Context for the RPC call</param>
 		/// <returns>Information about the new agent</returns>
-		public override async Task UpdateSession(IAsyncStreamReader<UpdateSessionRequest> reader, IServerStreamWriter<UpdateSessionResponse> writer, ServerCallContext context)
+		public override async Task UpdateSession(IAsyncStreamReader<RpcUpdateSessionRequest> reader, IServerStreamWriter<RpcUpdateSessionResponse> writer, ServerCallContext context)
 		{
 			// Read the request object
 			Task<bool> nextRequestTask = reader.MoveNext();
 			if (await nextRequestTask)
 			{
-				UpdateSessionRequest request = reader.Current;
+				RpcUpdateSessionRequest request = reader.Current;
 				using IDisposable? scope = _logger.BeginScope("UpdateSession for agent {AgentId}, session {SessionId}", request.AgentId, request.SessionId);
 
 				_logger.LogDebug("Updating session for {AgentId}", request.AgentId);
-				foreach (HordeCommon.Rpc.Messages.Lease lease in request.Leases)
+				foreach (RpcLease lease in request.Leases)
 				{
 					_logger.LogDebug("Session {SessionId}, Lease {LeaseId} - State: {LeaseState}, Outcome: {LeaseOutcome}", request.SessionId, lease.Id, lease.State, lease.Outcome);
 				}
@@ -421,7 +414,7 @@ namespace Horde.Server.Server
 				// Create the new session info
 				if (!context.CancellationToken.IsCancellationRequested)
 				{
-					UpdateSessionResponse response = new UpdateSessionResponse();
+					RpcUpdateSessionResponse response = new RpcUpdateSessionResponse();
 					response.Leases.Add(agent.Leases.Select(x => x.ToRpcMessage()));
 					response.ExpiryTime = (agent.SessionExpiresAt == null) ? new Timestamp() : Timestamp.FromDateTime(agent.SessionExpiresAt.Value);
 					response.Status = (RpcAgentStatus)agent.Status;
@@ -444,34 +437,34 @@ namespace Horde.Server.Server
 		}
 
 		/// <inheritdoc/>
-		public override Task<RpcGetStreamResponse> GetStream(GetStreamRequest request, ServerCallContext context) => _jobRpcCommon.GetStreamAsync(request, context);
+		public override Task<RpcGetStreamResponse> GetStream(RpcGetStreamRequest request, ServerCallContext context) => _jobRpcCommon.GetStreamAsync(request, context);
 
 		/// <inheritdoc/>
-		public override Task<RpcGetJobResponse> GetJob(GetJobRequest request, ServerCallContext context) => _jobRpcCommon.GetJobAsync(request, context);
+		public override Task<RpcGetJobResponse> GetJob(RpcGetJobRequest request, ServerCallContext context) => _jobRpcCommon.GetJobAsync(request, context);
 
 		/// <inheritdoc/>
 		public override Task<Empty> UpdateJob(RpcUpdateJobRequest request, ServerCallContext context) => _jobRpcCommon.UpdateJobAsync(request, context);
 
 		/// <inheritdoc/>
-		public override Task<BeginBatchResponse> BeginBatch(BeginBatchRequest request, ServerCallContext context) => _jobRpcCommon.BeginBatchAsync(request, context);
+		public override Task<RpcBeginBatchResponse> BeginBatch(RpcBeginBatchRequest request, ServerCallContext context) => _jobRpcCommon.BeginBatchAsync(request, context);
 
 		/// <inheritdoc/>
-		public override Task<Empty> FinishBatch(FinishBatchRequest request, ServerCallContext context) => _jobRpcCommon.FinishBatchAsync(request, context);
+		public override Task<Empty> FinishBatch(RpcFinishBatchRequest request, ServerCallContext context) => _jobRpcCommon.FinishBatchAsync(request, context);
 
 		/// <inheritdoc/>
-		public override Task<BeginStepResponse> BeginStep(BeginStepRequest request, ServerCallContext context) => _jobRpcCommon.BeginStepAsync(request, context);
+		public override Task<RpcBeginStepResponse> BeginStep(RpcBeginStepRequest request, ServerCallContext context) => _jobRpcCommon.BeginStepAsync(request, context);
 
 		/// <inheritdoc/>
 		public override Task<Empty> UpdateStep(RpcUpdateStepRequest request, ServerCallContext context) => _jobRpcCommon.UpdateStepAsync(request, context);
 
 		/// <inheritdoc/>
-		public override Task<RpcGetStepResponse> GetStep(GetStepRequest request, ServerCallContext context) => _jobRpcCommon.GetStepAsync(request, context);
+		public override Task<RpcGetStepResponse> GetStep(RpcGetStepRequest request, ServerCallContext context) => _jobRpcCommon.GetStepAsync(request, context);
 
 		/// <inheritdoc/>
-		public override Task<UpdateGraphResponse> UpdateGraph(UpdateGraphRequest request, ServerCallContext context) => _jobRpcCommon.UpdateGraphAsync(request, context);
+		public override Task<RpcUpdateGraphResponse> UpdateGraph(RpcUpdateGraphRequest request, ServerCallContext context) => _jobRpcCommon.UpdateGraphAsync(request, context);
 
 		/// <inheritdoc/>
-		public override Task<Empty> CreateEvents(CreateEventsRequest request, ServerCallContext context) => _jobRpcCommon.CreateEventsAsync(request, context);
+		public override Task<Empty> CreateEvents(RpcCreateEventsRequest request, ServerCallContext context) => _jobRpcCommon.CreateEventsAsync(request, context);
 
 		/// <summary>
 		/// Downloads a new agent archive
@@ -480,7 +473,7 @@ namespace Horde.Server.Server
 		/// <param name="responseStream">Writer for the output data</param>
 		/// <param name="context">Context for the RPC call</param>
 		/// <returns>Information about the new agent</returns>
-		public override async Task DownloadSoftware(DownloadSoftwareRequest request, IServerStreamWriter<DownloadSoftwareResponse> responseStream, ServerCallContext context)
+		public override async Task DownloadSoftware(RpcDownloadSoftwareRequest request, IServerStreamWriter<RpcDownloadSoftwareResponse> responseStream, ServerCallContext context)
 		{
 			int colonIdx = request.Version.IndexOf(':', StringComparison.Ordinal);
 			ToolId toolId = new ToolId(request.Version.Substring(0, colonIdx));
@@ -516,7 +509,7 @@ namespace Horde.Server.Server
 						break;
 					}
 
-					DownloadSoftwareResponse response = new DownloadSoftwareResponse();
+					RpcDownloadSoftwareResponse response = new RpcDownloadSoftwareResponse();
 					response.Data = UnsafeByteOperations.UnsafeWrap(buffer.Memory.Slice(0, read));
 					await responseStream.WriteAsync(response);
 
@@ -532,7 +525,7 @@ namespace Horde.Server.Server
 		/// <param name="request">Request arguments</param>
 		/// <param name="context">Context for the RPC call</param>
 		/// <returns>An empty response</returns>
-		public override async Task<Empty> SendTelemetryEvents(SendTelemetryEventsRequest request, ServerCallContext context)
+		public override async Task<Empty> SendTelemetryEvents(RpcSendTelemetryEventsRequest request, ServerCallContext context)
 		{
 			ISession? session = null;
 
@@ -543,9 +536,9 @@ namespace Horde.Server.Server
 			}
 
 			TelemetryRecordMeta agentMeta = new TelemetryRecordMeta("HordeAgent", session?.Version ?? "(Unknown)", ServerApp.DeploymentEnvironment, sessionId?.ToString() ?? "(Unknown)");
-			foreach (WrappedTelemetryEvent wrappedEvent in request.Events)
+			foreach (RpcWrappedTelemetryEvent wrappedEvent in request.Events)
 			{
-				OneofDescriptor oneofDescriptor = WrappedTelemetryEvent.Descriptor.Oneofs[0];
+				OneofDescriptor oneofDescriptor = RpcWrappedTelemetryEvent.Descriptor.Oneofs[0];
 				FieldDescriptor caseDescriptor = oneofDescriptor.Accessor.GetCaseFieldDescriptor(wrappedEvent);
 
 				object wrappedValue = caseDescriptor.Accessor.GetValue(wrappedEvent);
@@ -556,12 +549,12 @@ namespace Horde.Server.Server
 		}
 
 		/// <inheritdoc/>
-		public override Task<UploadArtifactResponse> UploadArtifact(IAsyncStreamReader<UploadArtifactRequest> reader, ServerCallContext context) => _jobRpcCommon.UploadArtifactAsync(reader, context);
+		public override Task<RpcUploadArtifactResponse> UploadArtifact(IAsyncStreamReader<RpcUploadArtifactRequest> reader, ServerCallContext context) => _jobRpcCommon.UploadArtifactAsync(reader, context);
 
 		/// <inheritdoc/>
-		public override Task<UploadTestDataResponse> UploadTestData(IAsyncStreamReader<UploadTestDataRequest> reader, ServerCallContext context) => _jobRpcCommon.UploadTestDataAsync(reader, context);
+		public override Task<RpcUploadTestDataResponse> UploadTestData(IAsyncStreamReader<RpcUploadTestDataRequest> reader, ServerCallContext context) => _jobRpcCommon.UploadTestDataAsync(reader, context);
 
 		/// <inheritdoc/>
-		public override Task<CreateReportResponse> CreateReport(CreateReportRequest request, ServerCallContext context) => _jobRpcCommon.CreateReportAsync(request, context);
+		public override Task<RpcCreateReportResponse> CreateReport(RpcCreateReportRequest request, ServerCallContext context) => _jobRpcCommon.CreateReportAsync(request, context);
 	}
 }
diff --git a/Engine/Source/Programs/Horde/Horde.Server/Streams/IStream.cs b/Engine/Source/Programs/Horde/Horde.Server/Streams/IStream.cs
--- a/Engine/Source/Programs/Horde/Horde.Server/Streams/IStream.cs
+++ b/Engine/Source/Programs/Horde/Horde.Server/Streams/IStream.cs
@@ -7,6 +7,7 @@ using EpicGames.Horde.Jobs;
 using EpicGames.Horde.Jobs.Templates;
 using EpicGames.Horde.Streams;
 using EpicGames.Horde.Users;
+using HordeCommon.Rpc;
 
 namespace Horde.Server.Streams
 {
@@ -167,9 +168,9 @@ namespace Horde.Server.Streams
 		/// </summary>
 		/// <param name="streamConfig">The stream config object</param>
 		/// <returns>New response instance</returns>
-		public static HordeCommon.Rpc.GetStreamResponse ToRpcResponse(this StreamConfig streamConfig)
+		public static RpcGetStreamResponse ToRpcResponse(this StreamConfig streamConfig)
 		{
-			HordeCommon.Rpc.GetStreamResponse response = new HordeCommon.Rpc.GetStreamResponse();
+			RpcGetStreamResponse response = new RpcGetStreamResponse();
 			response.Name = streamConfig.Name;
 			response.AgentTypes.Add(streamConfig.AgentTypes.ToDictionary(x => x.Key, x => x.Value.ToRpcResponse()));
 			return response;
diff --git a/Engine/Source/Programs/Horde/Horde.Server/Streams/StreamConfig.cs b/Engine/Source/Programs/Horde/Horde.Server/Streams/StreamConfig.cs
--- a/Engine/Source/Programs/Horde/Horde.Server/Streams/StreamConfig.cs
+++ b/Engine/Source/Programs/Horde/Horde.Server/Streams/StreamConfig.cs
@@ -16,7 +16,6 @@ using EpicGames.Horde.Acls;
 using EpicGames.Horde.Agents.Pools;
 using EpicGames.Horde.Commits;
 using EpicGames.Horde.Issues;
-using EpicGames.Horde.Jobs;
 using EpicGames.Horde.Jobs.Templates;
 using EpicGames.Horde.Replicators;
 using EpicGames.Horde.Streams;
@@ -32,6 +31,7 @@ using Horde.Server.Projects;
 using Horde.Server.Replicators;
 using Horde.Server.Server;
 using Horde.Server.Utilities;
+using HordeCommon.Rpc;
 
 namespace Horde.Server.Streams
 {
@@ -627,9 +627,9 @@ namespace Horde.Server.Streams
 		/// Creates an API response object from this stream
 		/// </summary>
 		/// <returns>The response object</returns>
-		public HordeCommon.Rpc.GetAgentTypeResponse ToRpcResponse()
+		public RpcGetAgentTypeResponse ToRpcResponse()
 		{
-			HordeCommon.Rpc.GetAgentTypeResponse response = new HordeCommon.Rpc.GetAgentTypeResponse();
+			RpcGetAgentTypeResponse response = new RpcGetAgentTypeResponse();
 			if (TempStorageDir != null)
 			{
 				response.TempStorageDir = TempStorageDir;
diff --git a/Engine/Source/Programs/Horde/Horde.Server/Streams/StreamsController.cs b/Engine/Source/Programs/Horde/Horde.Server/Streams/StreamsController.cs
--- a/Engine/Source/Programs/Horde/Horde.Server/Streams/StreamsController.cs
+++ b/Engine/Source/Programs/Horde/Horde.Server/Streams/StreamsController.cs
@@ -6,7 +6,6 @@ using System.Linq;
 using System.Threading;
 using System.Threading.Tasks;
 using EpicGames.Horde.Commits;
-using EpicGames.Horde.Jobs;
 using EpicGames.Horde.Jobs.Templates;
 using EpicGames.Horde.Projects;
 using EpicGames.Horde.Streams;
diff --git a/Engine/Source/Programs/Horde/Horde.Server/Tasks/ConformTaskSource.cs b/Engine/Source/Programs/Horde/Horde.Server/Tasks/ConformTaskSource.cs
--- a/Engine/Source/Programs/Horde/Horde.Server/Tasks/ConformTaskSource.cs
+++ b/Engine/Source/Programs/Horde/Horde.Server/Tasks/ConformTaskSource.cs
@@ -213,7 +213,7 @@ namespace Horde.Server.Tasks
 		}
 
 		/// <inheritdoc/>
-		public async Task<bool> GetWorkspacesAsync(IAgent agent, IList<HordeCommon.Rpc.Messages.AgentWorkspace> workspaces, CancellationToken cancellationToken)
+		public async Task<bool> GetWorkspacesAsync(IAgent agent, IList<RpcAgentWorkspace> workspaces, CancellationToken cancellationToken)
 		{
 			GlobalConfig globalConfig = _globalConfig.CurrentValue;
 
@@ -244,7 +244,7 @@ namespace Horde.Server.Tasks
 		/// <param name="leaseId">The lease id</param>
 		/// <param name="cancellationToken">Cancellation token for the operation</param>
 		/// <returns>True if the resource was allocated, false otherwise</returns>
-		private async Task<bool> AllocateConformLeaseAsync(AgentId agentId, IEnumerable<AgentWorkspace> workspaces, LeaseId leaseId, CancellationToken cancellationToken)
+		private async Task<bool> AllocateConformLeaseAsync(AgentId agentId, IEnumerable<RpcAgentWorkspace> workspaces, LeaseId leaseId, CancellationToken cancellationToken)
 		{
 			GlobalConfig globalConfig = _globalConfig.CurrentValue;
 			for (; ; )
@@ -256,7 +256,7 @@ namespace Horde.Server.Tasks
 				}
 
 				HashSet<string> servers = new HashSet<string>(StringComparer.OrdinalIgnoreCase);
-				foreach (AgentWorkspace workspace in workspaces)
+				foreach (RpcAgentWorkspace workspace in workspaces)
 				{
 					if (servers.Add(workspace.ServerAndPort))
 					{
diff --git a/Engine/Source/Programs/Horde/Horde.Server/Telemetry/Sinks/ClickHouseTelemetrySink.cs b/Engine/Source/Programs/Horde/Horde.Server/Telemetry/Sinks/ClickHouseTelemetrySink.cs
--- a/Engine/Source/Programs/Horde/Horde.Server/Telemetry/Sinks/ClickHouseTelemetrySink.cs
+++ b/Engine/Source/Programs/Horde/Horde.Server/Telemetry/Sinks/ClickHouseTelemetrySink.cs
@@ -60,18 +60,18 @@ namespace Horde.Server.Telemetry.Sinks
 		/// <inheritdoc />
 		public async ValueTask FlushAsync(CancellationToken cancellationToken)
 		{
-			List<AgentMetadataEvent> agentMetadataEvents = new();
-			List<AgentCpuMetricsEvent> agentCpuEvents = new();
-			List<AgentMemoryMetricsEvent> agentMemEvents = new();
+			List<RpcAgentMetadataEvent> agentMetadataEvents = new();
+			List<RpcAgentCpuMetricsEvent> agentCpuEvents = new();
+			List<RpcAgentMemoryMetricsEvent> agentMemEvents = new();
 
 			int c = 0;
 			while (_queuedEvents.TryDequeue(out TelemetryEvent? evt))
 			{
 				switch (evt.Payload)
 				{
-					case AgentMetadataEvent agentMetadata: agentMetadataEvents.Add(agentMetadata); break;
-					case AgentCpuMetricsEvent agentCpu: agentCpuEvents.Add(agentCpu); break;
-					case AgentMemoryMetricsEvent agentMem: agentMemEvents.Add(agentMem); break;
+					case RpcAgentMetadataEvent agentMetadata: agentMetadataEvents.Add(agentMetadata); break;
+					case RpcAgentCpuMetricsEvent agentCpu: agentCpuEvents.Add(agentCpu); break;
+					case RpcAgentMemoryMetricsEvent agentMem: agentMemEvents.Add(agentMem); break;
 				}
 				c++;
 			}
@@ -116,14 +116,14 @@ namespace Horde.Server.Telemetry.Sinks
 		/// </summary>
 		/// <param name="events"></param>
 		/// <param name="sb"></param>
-		private static void WriteAgentCpuSql(List<AgentCpuMetricsEvent> events, StringBuilder sb)
+		private static void WriteAgentCpuSql(List<RpcAgentCpuMetricsEvent> events, StringBuilder sb)
 		{
 			if (events.Count > 0)
 			{
 				sb.Append("INSERT INTO agentCpu (time, agentId, leaseId, jobId, jobBatchId, user, system, idle) VALUES \n");
-				foreach (AgentCpuMetricsEvent e in events)
+				foreach (RpcAgentCpuMetricsEvent e in events)
 				{
-					ExecutionMetadata? em = e.ExecutionMetadata;
+					RpcExecutionMetadata? em = e.ExecutionMetadata;
 					sb.AppendFormat("({0}, {1}, '{2}', '{3}', {4}, {5:F3}, {6:F3}, {7:F3}),\n",
 						e.Timestamp?.Seconds, (ulong)e.AgentId,
 						WriteEscapedObjectId(em?.LeaseId ?? EmptyObjectId),
@@ -141,12 +141,12 @@ namespace Horde.Server.Telemetry.Sinks
 		/// </summary>
 		/// <param name="events"></param>
 		/// <param name="sb"></param>
-		private static void WriteAgentMetadataSql(List<AgentMetadataEvent> events, StringBuilder sb)
+		private static void WriteAgentMetadataSql(List<RpcAgentMetadataEvent> events, StringBuilder sb)
 		{
 			if (events.Count > 0)
 			{
 				sb.Append("INSERT INTO agentMetadata (id, ip, hostname, region, az, env, version, os, osVersion, arch, props) VALUES \n");
-				foreach (AgentMetadataEvent e in events)
+				foreach (RpcAgentMetadataEvent e in events)
 				{
 					sb.AppendFormat("({0}, '{1}', '{2}', '{3}', '{4}', '{5}', '{6}', '{7}', '{8}', '{9}', {10}),\n",
 						(ulong)e.AgentId, e.Ip, e.Hostname, e.Region, e.AvailabilityZone, e.Environment, e.AgentVersion,
diff --git a/Engine/Source/Programs/Horde/Horde.Server/Utilities/ArtifactChunkStream.cs b/Engine/Source/Programs/Horde/Horde.Server/Utilities/ArtifactChunkStream.cs
--- a/Engine/Source/Programs/Horde/Horde.Server/Utilities/ArtifactChunkStream.cs
+++ b/Engine/Source/Programs/Horde/Horde.Server/Utilities/ArtifactChunkStream.cs
@@ -17,7 +17,7 @@ namespace Horde.Server.Utilities
 		/// <summary>
 		/// The grpc client reader. Should probably be templatized
 		/// </summary>
-		readonly IAsyncStreamReader<UploadArtifactRequest> _reader;
+		readonly IAsyncStreamReader<RpcUploadArtifactRequest> _reader;
 
 		/// <summary>
 		/// Position within the stream
@@ -32,7 +32,7 @@ namespace Horde.Server.Utilities
 		/// <summary>
 		/// The current request being read from
 		/// </summary>
-		UploadArtifactRequest? _request;
+		RpcUploadArtifactRequest? _request;
 
 		/// <summary>
 		/// Position within the current request
@@ -44,7 +44,7 @@ namespace Horde.Server.Utilities
 		/// </summary>
 		/// <param name="reader">the grpc reader</param>
 		/// <param name="length">filesize reported by the client</param>
-		public ArtifactChunkStream(IAsyncStreamReader<UploadArtifactRequest> reader, long length)
+		public ArtifactChunkStream(IAsyncStreamReader<RpcUploadArtifactRequest> reader, long length)
 		{
 			_reader = reader;
 			_streamLength = length;
