Horde: Move commit messages into EpicGames.Horde.

#rnx

diff --git a/Engine/Source/Programs/Horde/Horde.Server/Streams/StreamMessages.cs b/Engine/Source/Programs/Horde/Horde.Server/Streams/StreamMessages.cs
--- a/Engine/Source/Programs/Horde/Horde.Server/Streams/StreamMessages.cs
+++ b/Engine/Source/Programs/Horde/Horde.Server/Streams/StreamMessages.cs
@@ -4,18 +4,15 @@ using System;
 using System.Collections.Generic;
 using System.ComponentModel.DataAnnotations;
 using EpicGames.Core;
-using EpicGames.Horde.Commits;
 using EpicGames.Horde.Jobs;
 using EpicGames.Horde.Jobs.Templates;
 using EpicGames.Horde.Projects;
 using EpicGames.Horde.Streams;
 using EpicGames.Horde.Users;
-using Horde.Server.Commits;
 using Horde.Server.Issues;
 using Horde.Server.Jobs.Schedules;
 using Horde.Server.Jobs.Templates;
 using Horde.Server.Projects;
-using Horde.Server.Users;
 
 namespace Horde.Server.Streams
 {
@@ -521,65 +518,4 @@ namespace Horde.Server.Streams
 			PausedByUserInfo = pausedByUserInfo;
 		}
 	}
-
-	/// <summary>
-	/// Information about a commit
-	/// </summary>
-	public class GetCommitResponse
-	{
-		/// <summary>
-		/// The source changelist number
-		/// </summary>
-		public int Number { get; set; }
-
-		/// <summary>
-		/// Name of the user that authored this change [DEPRECATED]
-		/// </summary>
-		public string Author { get; set; }
-
-		/// <summary>
-		/// Information about the user that authored this change
-		/// </summary>
-		public GetThinUserInfoResponse AuthorInfo { get; set; }
-
-		/// <summary>
-		/// The description text
-		/// </summary>
-		public string Description { get; set; }
-
-		/// <summary>
-		/// Tags for this commit
-		/// </summary>
-		public List<CommitTag>? Tags { get; set; }
-
-		/// <summary>
-		/// List of files that were modified, relative to the stream base
-		/// </summary>
-		public List<string>? Files { get; set; }
-
-		/// <summary>
-		/// Constructor
-		/// </summary>
-		/// <param name="commit">The commit to construct from</param>
-		/// <param name="author">Author of the change</param>
-		/// <param name="tags">Tags for the commit</param>
-		/// <param name="files">Files modified by the commit</param>
-		public GetCommitResponse(ICommit commit, IUser author, IReadOnlyList<CommitTag>? tags, IReadOnlyList<string>? files)
-		{
-			Number = commit.Number;
-			Author = author.Name;
-			AuthorInfo = author.ToThinApiResponse();
-			Description = commit.Description;
-
-			if (tags != null)
-			{
-				Tags = new List<CommitTag>(tags);
-			}
-
-			if (files != null)
-			{
-				Files = new List<string>(files);
-			}
-		}
-	}
 }
diff --git a/Engine/Source/Programs/Horde/Horde.Server/Streams/StreamsController.cs b/Engine/Source/Programs/Horde/Horde.Server/Streams/StreamsController.cs
--- a/Engine/Source/Programs/Horde/Horde.Server/Streams/StreamsController.cs
+++ b/Engine/Source/Programs/Horde/Horde.Server/Streams/StreamsController.cs
@@ -278,11 +278,25 @@ namespace Horde.Server.Streams
 			foreach (ICommit commit in commits)
 			{
 				IUser? author = await _userCollection.GetCachedUserAsync(commit.AuthorId, cancellationToken);
-				responses.Add(new GetCommitResponse(commit, author!, null, null));
+				responses.Add(CreateGetCommitResponse(commit, author!, null, null));
 			}
 			return responses.ConvertAll(x => PropertyFilter.Apply(x, filter));
 		}
 
+		static GetCommitResponse CreateGetCommitResponse(ICommit commit, IUser author, IReadOnlyList<CommitTag>? tags, IReadOnlyList<string>? files)
+		{
+			GetCommitResponse response = new GetCommitResponse(commit.Number, author.ToThinApiResponse(), commit.Description);
+			if (tags != null)
+			{
+				response.Tags = new List<CommitTag>(tags);
+			}
+			if (files != null)
+			{
+				response.Files = new List<string>(files);
+			}
+			return response;
+		}
+
 		/// <summary>
 		/// Gets a list of changes for a stream
 		/// </summary>
@@ -317,7 +331,7 @@ namespace Horde.Server.Streams
 			IReadOnlyList<CommitTag> tags = await changeDetails.GetTagsAsync(cancellationToken);
 			IReadOnlyList<string> files = await changeDetails.GetFilesAsync(maxFiles, cancellationToken);
 
-			return PropertyFilter.Apply(new GetCommitResponse(changeDetails, author!, tags, files), filter);
+			return PropertyFilter.Apply(CreateGetCommitResponse(changeDetails, author!, tags, files), filter);
 		}
 
 		/// <summary>
diff --git a/Engine/Source/Programs/Shared/EpicGames.Horde/Commits/CommitMessages.cs b/Engine/Source/Programs/Shared/EpicGames.Horde/Commits/CommitMessages.cs
new file mode 100644
--- /dev/null
+++ b/Engine/Source/Programs/Shared/EpicGames.Horde/Commits/CommitMessages.cs
@@ -0,0 +1,56 @@
+// Copyright Epic Games, Inc. All Rights Reserved.
+
+using System.Collections.Generic;
+using EpicGames.Horde.Users;
+
+#pragma warning disable CA2227
+
+namespace EpicGames.Horde.Commits
+{
+	/// <summary>
+	/// Information about a commit
+	/// </summary>
+	public class GetCommitResponse
+	{
+		/// <summary>
+		/// The source changelist number
+		/// </summary>
+		public int Number { get; set; }
+
+		/// <summary>
+		/// Name of the user that authored this change [DEPRECATED]
+		/// </summary>
+		public string Author { get; set; }
+
+		/// <summary>
+		/// Information about the user that authored this change
+		/// </summary>
+		public GetThinUserInfoResponse AuthorInfo { get; set; }
+
+		/// <summary>
+		/// The description text
+		/// </summary>
+		public string Description { get; set; }
+
+		/// <summary>
+		/// Tags for this commit
+		/// </summary>
+		public List<CommitTag>? Tags { get; set; }
+
+		/// <summary>
+		/// List of files that were modified, relative to the stream base
+		/// </summary>
+		public List<string>? Files { get; set; }
+
+		/// <summary>
+		/// Constructor
+		/// </summary>
+		public GetCommitResponse(int number, GetThinUserInfoResponse authorInfo, string description)
+		{
+			Number = number;
+			Author = authorInfo.Name;
+			AuthorInfo = authorInfo;
+			Description = description;
+		}
+	}
+}
