Horde: Move log event messages into EpicGames.Horde.

#rnx

diff --git a/Engine/Source/Programs/Horde/Horde.Agent/Horde.Agent.csproj b/Engine/Source/Programs/Horde/Horde.Agent/Horde.Agent.csproj
--- a/Engine/Source/Programs/Horde/Horde.Agent/Horde.Agent.csproj
+++ b/Engine/Source/Programs/Horde/Horde.Agent/Horde.Agent.csproj
@@ -155,7 +155,6 @@
     <Protobuf Include="..\Horde.Server\Protos\horde\log_rpc.proto" ProtoRoot="..\Horde.Server\Protos" GrpcServices="Client" Access="Public" />
     <Protobuf Include="..\Horde.Server\Protos\horde\log_rpc_messages.proto" ProtoRoot="..\Horde.Server\Protos" Access="Public" />
     <Protobuf Include="..\Horde.Server\Protos\horde\common\agent_status.proto" ProtoRoot="..\Horde.Server\Protos" Access="Public" />
-    <Protobuf Include="..\Horde.Server\Protos\horde\common\event_severity.proto" ProtoRoot="..\Horde.Server\Protos" Access="Public" />
     <Protobuf Include="..\Horde.Server\Protos\horde\common\io_hash.proto" ProtoRoot="..\Horde.Server\Protos" Access="Public" />
     <Protobuf Include="..\Horde.Server\Protos\horde\common\label_change.proto" ProtoRoot="..\Horde.Server\Protos" Access="Public" />
     <Protobuf Include="..\Horde.Server\Protos\horde\common\lease_outcome.proto" ProtoRoot="..\Horde.Server\Protos" Access="Public" />
diff --git a/Engine/Source/Programs/Horde/Horde.Agent/Utilities/HordeRpc.cs b/Engine/Source/Programs/Horde/Horde.Agent/Utilities/HordeRpc.cs
--- a/Engine/Source/Programs/Horde/Horde.Agent/Utilities/HordeRpc.cs
+++ b/Engine/Source/Programs/Horde/Horde.Agent/Utilities/HordeRpc.cs
@@ -151,9 +151,9 @@ namespace HordeCommon.Rpc
 
 	partial class RpcCreateEventRequest
 	{
-		public RpcCreateEventRequest(RpcEventSeverity severity, LogId logId, int lineIndex, int lineCount)
+		public RpcCreateEventRequest(LogEventSeverity severity, LogId logId, int lineIndex, int lineCount)
 		{
-			Severity = severity;
+			Severity = (int)severity;
 			LogId = logId.ToString();
 			LineIndex = lineIndex;
 			LineCount = lineCount;
diff --git a/Engine/Source/Programs/Horde/Horde.Agent/Utility/ServerLogger.cs b/Engine/Source/Programs/Horde/Horde.Agent/Utility/ServerLogger.cs
--- a/Engine/Source/Programs/Horde/Horde.Agent/Utility/ServerLogger.cs
+++ b/Engine/Source/Programs/Horde/Horde.Agent/Utility/ServerLogger.cs
@@ -6,7 +6,6 @@ using EpicGames.Core;
 using EpicGames.Horde.Jobs;
 using EpicGames.Horde.Logs;
 using Google.Protobuf;
-using HordeCommon;
 using HordeCommon.Rpc;
 using Microsoft.Extensions.Logging;
 
@@ -183,11 +182,11 @@ namespace Horde.Agent.Utility
 						{
 							if (jsonLogEvent.Level == LogLevel.Warning && ++numWarnings <= MaxWarnings)
 							{
-								AddEvent(jsonLogEvent.Data.Span, nextLineIndex, Math.Max(lineCount, jsonLogEvent.LineCount), RpcEventSeverity.Warning, events);
+								AddEvent(jsonLogEvent.Data.Span, nextLineIndex, Math.Max(lineCount, jsonLogEvent.LineCount), LogEventSeverity.Warning, events);
 							}
 							else if ((jsonLogEvent.Level == LogLevel.Error || jsonLogEvent.Level == LogLevel.Critical) && ++numErrors <= MaxErrors)
 							{
-								AddEvent(jsonLogEvent.Data.Span, nextLineIndex, Math.Max(lineCount, jsonLogEvent.LineCount), RpcEventSeverity.Error, events);
+								AddEvent(jsonLogEvent.Data.Span, nextLineIndex, Math.Max(lineCount, jsonLogEvent.LineCount), LogEventSeverity.Error, events);
 							}
 						}
 						nextLineIndex += lineCount;
@@ -264,7 +263,7 @@ namespace Horde.Agent.Utility
 			}
 		}
 
-		void AddEvent(ReadOnlySpan<byte> span, int lineIndex, int lineCount, RpcEventSeverity severity, List<RpcCreateEventRequest> events)
+		void AddEvent(ReadOnlySpan<byte> span, int lineIndex, int lineCount, LogEventSeverity severity, List<RpcCreateEventRequest> events)
 		{
 			try
 			{
diff --git a/Engine/Source/Programs/Horde/Horde.Server.Tests/Issues/IssueServiceTests.cs b/Engine/Source/Programs/Horde/Horde.Server.Tests/Issues/IssueServiceTests.cs
--- a/Engine/Source/Programs/Horde/Horde.Server.Tests/Issues/IssueServiceTests.cs
+++ b/Engine/Source/Programs/Horde/Horde.Server.Tests/Issues/IssueServiceTests.cs
@@ -30,7 +30,6 @@ using Horde.Server.Tests.Stubs.Services;
 using Horde.Server.Users;
 using HordeAgent.Horde.Agent.Parser;
 using HordeAgent.Horde.Agent.Utility;
-using HordeCommon;
 using Microsoft.Extensions.Logging;
 using Microsoft.Extensions.Logging.Abstractions;
 using Microsoft.VisualStudio.TestTools.UnitTesting;
@@ -100,7 +99,7 @@ namespace Horde.Server.Tests.Issues
 					LogEvent ev = ParseEvent(line);
 					if (ev.LineIndex == 0)
 					{
-						RpcEventSeverity severity = (level == LogLevel.Warning) ? RpcEventSeverity.Warning : RpcEventSeverity.Error;
+						LogEventSeverity severity = (level == LogLevel.Warning) ? LogEventSeverity.Warning : LogEventSeverity.Error;
 						await _logFileService.CreateEventsAsync(new List<NewLogEventData> { new NewLogEventData { LogId = _logId, LineIndex = _lineIndex, LineCount = ev.LineCount, Severity = severity } }, CancellationToken.None);
 					}
 				}
@@ -316,11 +315,11 @@ namespace Horde.Server.Tests.Issues
 			buffer.GetSpan(1)[0] = (byte)'\n';
 			buffer.Advance(1);
 
-			RpcEventSeverity severity = (logLevel == LogLevel.Error) ? RpcEventSeverity.Error : (logLevel == LogLevel.Warning) ? RpcEventSeverity.Warning : RpcEventSeverity.Information;
+			LogEventSeverity severity = (logLevel == LogLevel.Error) ? LogEventSeverity.Error : (logLevel == LogLevel.Warning) ? LogEventSeverity.Warning : LogEventSeverity.Information;
 			await AddEventAsync(job, batchIdx, stepIdx, severity, buffer.WrittenMemory.ToArray());
 		}
 
-		async Task AddEventAsync(IJob job, int batchIdx, int stepIdx, RpcEventSeverity severity, byte[] data)
+		async Task AddEventAsync(IJob job, int batchIdx, int stepIdx, LogEventSeverity severity, byte[] data)
 		{
 			LogId logId = job.Batches[batchIdx].Steps[stepIdx].LogId!.Value;
 
diff --git a/Engine/Source/Programs/Horde/Horde.Server/Horde.Server.csproj b/Engine/Source/Programs/Horde/Horde.Server/Horde.Server.csproj
--- a/Engine/Source/Programs/Horde/Horde.Server/Horde.Server.csproj
+++ b/Engine/Source/Programs/Horde/Horde.Server/Horde.Server.csproj
@@ -104,7 +104,6 @@
     <Protobuf Include="Protos\horde\log_rpc_messages.proto" ProtoRoot="Protos" />
     <Protobuf Include="Protos\horde\relay_rpc.proto" ProtoRoot="Protos" />
     <Protobuf Include="Protos\horde\common\agent_status.proto" ProtoRoot="Protos" />
-    <Protobuf Include="Protos\horde\common\event_severity.proto" ProtoRoot="Protos" />
     <Protobuf Include="Protos\horde\common\io_hash.proto" ProtoRoot="Protos" />
     <Protobuf Include="Protos\horde\common\label_change.proto" ProtoRoot="Protos" />
     <Protobuf Include="Protos\horde\common\lease_outcome.proto" ProtoRoot="Protos" />
diff --git a/Engine/Source/Programs/Horde/Horde.Server/Issues/IssueService.cs b/Engine/Source/Programs/Horde/Horde.Server/Issues/IssueService.cs
--- a/Engine/Source/Programs/Horde/Horde.Server/Issues/IssueService.cs
+++ b/Engine/Source/Programs/Horde/Horde.Server/Issues/IssueService.cs
@@ -13,6 +13,7 @@ using EpicGames.Core;
 using EpicGames.Horde.Issues;
 using EpicGames.Horde.Jobs;
 using EpicGames.Horde.Jobs.Templates;
+using EpicGames.Horde.Logs;
 using EpicGames.Horde.Streams;
 using EpicGames.Horde.Users;
 using Horde.Server.Commits;
@@ -158,10 +159,10 @@ namespace Horde.Server.Issues
 				EventData = logEventData;
 			}
 
-			static LogLevel GetLogLevelFromSeverity(RpcEventSeverity severity) => severity switch
+			static LogLevel GetLogLevelFromSeverity(LogEventSeverity severity) => severity switch
 			{
-				RpcEventSeverity.Error => LogLevel.Error,
-				RpcEventSeverity.Warning => LogLevel.Warning,
+				LogEventSeverity.Error => LogLevel.Error,
+				LogEventSeverity.Warning => LogLevel.Warning,
 				_ => LogLevel.Information
 			};
 		}
diff --git a/Engine/Source/Programs/Horde/Horde.Server/Issues/IssuesController.cs b/Engine/Source/Programs/Horde/Horde.Server/Issues/IssuesController.cs
--- a/Engine/Source/Programs/Horde/Horde.Server/Issues/IssuesController.cs
+++ b/Engine/Source/Programs/Horde/Horde.Server/Issues/IssuesController.cs
@@ -745,7 +745,7 @@ namespace Horde.Server.Issues
 				if (logFile != null && await _jobService.AuthorizeAsync(logFile.JobId, LogAclAction.ViewLog, User, _globalConfig.Value, cancellationToken))
 				{
 					ILogEventData data = await _logFileService.GetEventDataAsync(logFile, logEvent.LineIndex, logEvent.LineCount, cancellationToken);
-					GetLogEventResponse response = new GetLogEventResponse(logEvent, data, issueId);
+					GetLogEventResponse response = LogsController.CreateGetLogEventResponse(logEvent, data, issueId);
 					responses.Add(PropertyFilter.Apply(response, filter));
 				}
 			}
diff --git a/Engine/Source/Programs/Horde/Horde.Server/Jobs/JobRpcService.cs b/Engine/Source/Programs/Horde/Horde.Server/Jobs/JobRpcService.cs
--- a/Engine/Source/Programs/Horde/Horde.Server/Jobs/JobRpcService.cs
+++ b/Engine/Source/Programs/Horde/Horde.Server/Jobs/JobRpcService.cs
@@ -893,7 +893,7 @@ namespace Horde.Server.Jobs
 			{
 				NewLogEventData newEvent = new NewLogEventData();
 				newEvent.LogId = LogId.Parse(createEvent.LogId);
-				newEvent.Severity = createEvent.Severity;
+				newEvent.Severity = (LogEventSeverity)createEvent.Severity;
 				newEvent.LineIndex = createEvent.LineIndex;
 				newEvent.LineCount = createEvent.LineCount;
 				newEvents.Add(newEvent);
diff --git a/Engine/Source/Programs/Horde/Horde.Server/Logs/ILogEvent.cs b/Engine/Source/Programs/Horde/Horde.Server/Logs/ILogEvent.cs
--- a/Engine/Source/Programs/Horde/Horde.Server/Logs/ILogEvent.cs
+++ b/Engine/Source/Programs/Horde/Horde.Server/Logs/ILogEvent.cs
@@ -5,7 +5,6 @@ using System.Linq;
 using System.Text.Json;
 using EpicGames.Core;
 using EpicGames.Horde.Logs;
-using HordeCommon;
 using Microsoft.Extensions.Logging;
 using MongoDB.Bson;
 
@@ -19,7 +18,7 @@ namespace Horde.Server.Logs
 		/// <summary>
 		/// The log level
 		/// </summary>
-		RpcEventSeverity Severity { get; }
+		LogEventSeverity Severity { get; }
 
 		/// <summary>
 		/// The type of event
@@ -50,7 +49,7 @@ namespace Horde.Server.Logs
 		/// <summary>
 		/// Severity of the event
 		/// </summary>
-		public RpcEventSeverity Severity { get; }
+		public LogEventSeverity Severity { get; }
 
 		/// <summary>
 		/// Index of the first line for this event
@@ -81,7 +80,7 @@ namespace Horde.Server.Logs
 		/// <summary>
 		/// Severity of the event
 		/// </summary>
-		public RpcEventSeverity Severity { get; set; }
+		public LogEventSeverity Severity { get; set; }
 
 		/// <summary>
 		/// Index of the first line for this event
diff --git a/Engine/Source/Programs/Horde/Horde.Server/Logs/LogEventCollection.cs b/Engine/Source/Programs/Horde/Horde.Server/Logs/LogEventCollection.cs
--- a/Engine/Source/Programs/Horde/Horde.Server/Logs/LogEventCollection.cs
+++ b/Engine/Source/Programs/Horde/Horde.Server/Logs/LogEventCollection.cs
@@ -7,7 +7,6 @@ using System.Threading.Tasks;
 using EpicGames.Horde.Logs;
 using Horde.Server.Server;
 using Horde.Server.Utilities;
-using HordeCommon;
 using Microsoft.Extensions.Logging;
 using MongoDB.Bson;
 using MongoDB.Bson.Serialization.Attributes;
@@ -46,17 +45,17 @@ namespace Horde.Server.Logs
 			LogId ILogEvent.LogId => Id.LogId;
 			int ILogEvent.LineIndex => Id.LineIndex;
 			int ILogEvent.LineCount => LineCount ?? 1;
-			RpcEventSeverity ILogEvent.Severity => IsWarning ? RpcEventSeverity.Warning : RpcEventSeverity.Error;
+			LogEventSeverity ILogEvent.Severity => IsWarning ? LogEventSeverity.Warning : LogEventSeverity.Error;
 
 			public LogEventDocument()
 			{
 				Id = new LogEventId();
 			}
 
-			public LogEventDocument(LogId logId, RpcEventSeverity severity, int lineIndex, int lineCount, ObjectId? spanId)
+			public LogEventDocument(LogId logId, LogEventSeverity severity, int lineIndex, int lineCount, ObjectId? spanId)
 			{
 				Id = new LogEventId { LogId = logId, LineIndex = lineIndex };
-				IsWarning = severity == RpcEventSeverity.Warning;
+				IsWarning = severity == LogEventSeverity.Warning;
 				LineCount = (lineCount > 1) ? (int?)lineCount : null;
 				SpanId = spanId;
 			}
@@ -71,7 +70,7 @@ namespace Horde.Server.Logs
 		{
 			public ObjectId Id { get; set; }
 			public DateTime Time { get; set; }
-			public RpcEventSeverity Severity { get; set; }
+			public LogEventSeverity Severity { get; set; }
 			public LogId LogId { get; set; }
 			public int LineIndex { get; set; }
 			public int LineCount { get; set; }
diff --git a/Engine/Source/Programs/Horde/Horde.Server/Logs/LogEventMessages.cs b/Engine/Source/Programs/Horde/Horde.Server/Logs/LogEventMessages.cs
deleted file mode 100644
diff --git a/Engine/Source/Programs/Horde/Horde.Server/Logs/LogFileService.cs b/Engine/Source/Programs/Horde/Horde.Server/Logs/LogFileService.cs
--- a/Engine/Source/Programs/Horde/Horde.Server/Logs/LogFileService.cs
+++ b/Engine/Source/Programs/Horde/Horde.Server/Logs/LogFileService.cs
@@ -21,7 +21,6 @@ using Horde.Server.Acls;
 using Horde.Server.Logs.Data;
 using Horde.Server.Storage;
 using Horde.Server.Utilities;
-using HordeCommon;
 using Microsoft.Extensions.Caching.Memory;
 using Microsoft.Extensions.Logging;
 using Microsoft.Extensions.Options;
@@ -691,7 +690,7 @@ namespace Horde.Server.Logs
 			public IReadOnlyList<JsonLogEvent> Lines { get; }
 
 			EventId? ILogEventData.EventId => (Lines.Count > 0) ? Lines[0].EventId : null;
-			RpcEventSeverity ILogEventData.Severity => (Lines.Count == 0) ? RpcEventSeverity.Information : (Lines[0].Level == LogLevel.Warning) ? RpcEventSeverity.Warning : RpcEventSeverity.Error;
+			LogEventSeverity ILogEventData.Severity => (Lines.Count == 0) ? LogEventSeverity.Information : (Lines[0].Level == LogLevel.Warning) ? LogEventSeverity.Warning : LogEventSeverity.Error;
 
 			public LogEventData(IReadOnlyList<JsonLogEvent> lines)
 			{
diff --git a/Engine/Source/Programs/Horde/Horde.Server/Logs/LogsController.cs b/Engine/Source/Programs/Horde/Horde.Server/Logs/LogsController.cs
--- a/Engine/Source/Programs/Horde/Horde.Server/Logs/LogsController.cs
+++ b/Engine/Source/Programs/Horde/Horde.Server/Logs/LogsController.cs
@@ -6,6 +6,7 @@ using System.IO;
 using System.Linq;
 using System.Security.Claims;
 using System.Text;
+using System.Text.Json;
 using System.Threading;
 using System.Threading.Tasks;
 using EpicGames.Core;
@@ -343,11 +344,29 @@ namespace Horde.Server.Logs
 					spanIdToIssueId[logEvent.SpanId.Value] = issueId;
 				}
 
-				responses.Add(new GetLogEventResponse(logEvent, logEventData, issueId));
+				responses.Add(CreateGetLogEventResponse(logEvent, logEventData, issueId));
 			}
 			return responses;
 		}
 
+		/// <summary>
+		/// Create a log event response message
+		/// </summary>
+		/// <param name="logEvent">The event to construct from</param>
+		/// <param name="eventData">The event data</param>
+		/// <param name="issueId">The issue for this event</param>
+		internal static GetLogEventResponse CreateGetLogEventResponse(ILogEvent logEvent, ILogEventData eventData, int? issueId)
+		{
+			GetLogEventResponse response = new GetLogEventResponse();
+			response.Severity = logEvent.Severity;
+			response.LogId = logEvent.LogId;
+			response.LineIndex = logEvent.LineIndex;
+			response.LineCount = logEvent.LineCount;
+			response.IssueId = issueId;
+			response.Lines.AddRange(eventData.Lines.Select(x => JsonDocument.Parse(x.Data).RootElement));
+			return response;
+		}
+
 		/// <summary>
 		/// Determines if the user is authorized to perform an action on a particular template
 		/// </summary>
diff --git a/Engine/Source/Programs/Horde/Horde.Server/Notifications/Sinks/SlackNotificationSink.cs b/Engine/Source/Programs/Horde/Horde.Server/Notifications/Sinks/SlackNotificationSink.cs
--- a/Engine/Source/Programs/Horde/Horde.Server/Notifications/Sinks/SlackNotificationSink.cs
+++ b/Engine/Source/Programs/Horde/Horde.Server/Notifications/Sinks/SlackNotificationSink.cs
@@ -639,8 +639,8 @@ namespace Horde.Server.Notifications.Sinks
 			}
 			else
 			{
-				List<ILogEventData> errors = events.Where(x => x.Severity == RpcEventSeverity.Error).ToList();
-				List<ILogEventData> warnings = events.Where(x => x.Severity == RpcEventSeverity.Warning).ToList();
+				List<ILogEventData> errors = events.Where(x => x.Severity == LogEventSeverity.Error).ToList();
+				List<ILogEventData> warnings = events.Where(x => x.Severity == LogEventSeverity.Warning).ToList();
 				List<string> eventStrings = new List<string>();
 				if (errors.Any())
 				{
@@ -1018,9 +1018,9 @@ namespace Horde.Server.Notifications.Sinks
 					if (logFile != null)
 					{
 						events = await _logFileService.FindEventsAsync(logFile, span.Id, 0, 50, cancellationToken);
-						if (events.Any(x => x.Severity == RpcEventSeverity.Error))
+						if (events.Any(x => x.Severity == LogEventSeverity.Error))
 						{
-							events.RemoveAll(x => x.Severity == RpcEventSeverity.Warning);
+							events.RemoveAll(x => x.Severity == LogEventSeverity.Warning);
 						}
 
 						List<string> eventStrings = new List<string>();
@@ -1454,9 +1454,9 @@ namespace Horde.Server.Notifications.Sinks
 				if (logFile != null)
 				{
 					List<ILogEvent> events = await _logFileService.FindEventsAsync(logFile, lastSpan.Id, 0, 20, cancellationToken);
-					if (events.Any(x => x.Severity == RpcEventSeverity.Error))
+					if (events.Any(x => x.Severity == LogEventSeverity.Error))
 					{
-						events.RemoveAll(x => x.Severity == RpcEventSeverity.Warning);
+						events.RemoveAll(x => x.Severity == LogEventSeverity.Warning);
 					}
 
 					for (int idx = 0; idx < Math.Min(events.Count, 3); idx++)
diff --git a/Engine/Source/Programs/Shared/EpicGames.Horde/Logs/LogEventMessages.cs b/Engine/Source/Programs/Shared/EpicGames.Horde/Logs/LogEventMessages.cs
new file mode 100644
--- /dev/null
+++ b/Engine/Source/Programs/Shared/EpicGames.Horde/Logs/LogEventMessages.cs
@@ -0,0 +1,71 @@
+// Copyright Epic Games, Inc. All Rights Reserved.
+
+using System.Collections.Generic;
+using System.Text.Json;
+using EpicGames.Core;
+using EpicGames.Horde.Logs;
+
+namespace EpicGames.Horde.Logs
+{
+	/// <summary>
+	/// Severity of a log event
+	/// </summary>
+	public enum LogEventSeverity
+	{
+		/// <summary>
+		/// Severity is not specified
+		/// </summary>
+		Unspecified = 0,
+
+		/// <summary>
+		/// Information severity
+		/// </summary>
+		Information = 1,
+
+		/// <summary>
+		/// Warning severity
+		/// </summary>
+		Warning = 2,
+
+		/// <summary>
+		/// Error severity
+		/// </summary>
+		Error = 3,
+	}
+
+	/// <summary>
+	/// Information about an uploaded event
+	/// </summary>
+	public class GetLogEventResponse
+	{
+		/// <summary>
+		/// Unique id of the log containing this event
+		/// </summary>
+		public LogId LogId { get; set; }
+
+		/// <summary>
+		/// Severity of this event
+		/// </summary>
+		public LogEventSeverity Severity { get; set; }
+
+		/// <summary>
+		/// Index of the first line for this event
+		/// </summary>
+		public int LineIndex { get; set; }
+
+		/// <summary>
+		/// Number of lines in the event
+		/// </summary>
+		public int LineCount { get; set; }
+
+		/// <summary>
+		/// The issue id associated with this event
+		/// </summary>
+		public int? IssueId { get; set; }
+
+		/// <summary>
+		/// The structured message data for this event
+		/// </summary>
+		public List<JsonElement> Lines { get; set; } = new List<JsonElement>();
+	}
+}
diff --git a/Engine/Source/Programs/Shared/EpicGames.Horde/Logs/LogMessages.cs b/Engine/Source/Programs/Shared/EpicGames.Horde/Logs/LogMessages.cs
--- a/Engine/Source/Programs/Shared/EpicGames.Horde/Logs/LogMessages.cs
+++ b/Engine/Source/Programs/Shared/EpicGames.Horde/Logs/LogMessages.cs
@@ -2,6 +2,7 @@
 
 using System;
 using System.Collections.Generic;
+using System.Text.Json;
 using EpicGames.Horde.Agents.Leases;
 using EpicGames.Horde.Agents.Sessions;
 using EpicGames.Horde.Jobs;
