Horde: Updating schema documentation

#rnx

diff --git a/Engine/Source/Programs/Horde/Docs/Config/Schema/Streams.md b/Engine/Source/Programs/Horde/Docs/Config/Schema/Streams.md
--- a/Engine/Source/Programs/Horde/Docs/Config/Schema/Streams.md
+++ b/Engine/Source/Programs/Horde/Docs/Config/Schema/Streams.md
@@ -100,12 +100,14 @@ Name | Description
 
 ## JobStepOutcome (Enum)
 
+Outcome for a jobstep
+
 Name | Description
 ---- | -----------
-`Unspecified` | 
-`Failure` | 
-`Warnings` | 
-`Success` | 
+`Unspecified` | Outcome is not known
+`Failure` | Step failed
+`Warnings` | Step completed with warnings
+`Success` | Step succeeded
 
 ## CommitTagConfig
 
@@ -349,6 +351,8 @@ Name | Description
 
 ## Priority (Enum)
 
+Priority of a job or step
+
 Name | Description
 ---- | -----------
 `Unspecified` | Not specified
