Horde: Extend timeout before a Perforce server is considered degraded

Ticker for PerforceLoadBalancer is not running consistently and sometimes skips execution, causing gaps in updates which in turn triggers the degraded check when returning servers.
Likely related to the out of memory issue we see with worker tasks. Not a real fix, but 7.5 mins is enough to hide the problem for now.

diff --git a/Engine/Source/Programs/Horde/Horde.Server/Perforce/PerforceLoadBalancer.cs b/Engine/Source/Programs/Horde/Horde.Server/Perforce/PerforceLoadBalancer.cs
--- a/Engine/Source/Programs/Horde/Horde.Server/Perforce/PerforceLoadBalancer.cs
+++ b/Engine/Source/Programs/Horde/Horde.Server/Perforce/PerforceLoadBalancer.cs
@@ -204,7 +204,7 @@ namespace Horde.Server.Perforce
 		{
 			PerforceServerList serverList = await _serverListSingleton.GetAsync(cancellationToken);
 
-			DateTime minLastUpdateTime = DateTime.UtcNow - TimeSpan.FromMinutes(2.5);
+			DateTime minLastUpdateTime = DateTime.UtcNow - TimeSpan.FromMinutes(7.5);
 			foreach (PerforceServerEntry server in serverList.Servers)
 			{
 				if (server.Status == PerforceServerStatus.Healthy && server.LastUpdateTime != null && server.LastUpdateTime.Value < minLastUpdateTime)
@@ -460,6 +460,7 @@ namespace Horde.Server.Perforce
 
 			// Now update the health of each entry in parallel
 			List<Task> tasks = [];
+			span.SetAttribute("numServers", list.Servers.Count);
 			foreach (PerforceServerEntry entry in list.Servers)
 			{
 				tasks.Add(Task.Run(() => UpdateHealthAsync(entry, cancellationToken), cancellationToken));
