Horde: Add debug endpoint for capturing a dotMemory snapshot

diff --git a/Engine/Source/Programs/Horde/Horde.Server/Server/DebugController.cs b/Engine/Source/Programs/Horde/Horde.Server/Server/DebugController.cs
--- a/Engine/Source/Programs/Horde/Horde.Server/Server/DebugController.cs
+++ b/Engine/Source/Programs/Horde/Horde.Server/Server/DebugController.cs
@@ -5,6 +5,7 @@
 
 using System;
 using System.Collections.Generic;
+using System.Diagnostics;
 using System.IO;
 using System.Linq;
 using System.Net;
@@ -617,21 +618,25 @@ namespace Horde.Server.Server
 		}
 
 		/// <summary>
-		/// Starts the profiler session
+		/// Start a CPU profiler session using dotTrace
+		/// Only one profiling session can run at a time.
 		/// </summary>
-		/// <returns>Text message</returns>
+		/// <returns>Status description</returns>
 		[HttpGet]
-		[Route("/api/v1/debug/profiler/start")]
-		public async Task<ActionResult> StartProfilerAsync()
+		[Route("/api/v1/debug/profiler/cpu/start")]
+		public async Task<ActionResult> StartCpuProfilerAsync()
 		{
 			if (!_globalConfig.Value.Authorize(ServerAclAction.Debug, User))
 			{
 				return Forbid(ServerAclAction.Debug);
 			}
 
+			// Downloads dotTrace executable if not available
+			Stopwatch sw = Stopwatch.StartNew();
 			await DotTrace.EnsurePrerequisiteAsync();
+			_logger.LogInformation("dotTrace prerequisites step finished in {SetupTimeMs} ms", sw.ElapsedMilliseconds);
 
-			string snapshotDir = Path.Join(Path.GetTempPath(), "horde-profiler-snapshots");
+			string snapshotDir = Path.Join(Path.GetTempPath(), "horde-cpu-profiler-snapshots");
 			if (!Directory.Exists(snapshotDir))
 			{
 				Directory.CreateDirectory(snapshotDir);
@@ -642,15 +647,15 @@ namespace Horde.Server.Server
 			DotTrace.Attach(config);
 			DotTrace.StartCollectingData();
 
-			return new ContentResult { ContentType = "text/plain", StatusCode = (int)HttpStatusCode.OK, Content = "Profiling session started. Using dir " + snapshotDir };
+			return new ContentResult { ContentType = "text/plain", StatusCode = (int)HttpStatusCode.OK, Content = "CPU profiling session started. Using dir " + snapshotDir };
 		}
 
 		/// <summary>
-		/// Stops the profiler session
+		/// Stops a CPU profiler session
 		/// </summary>
 		/// <returns>Text message</returns>
 		[HttpGet]
-		[Route("/api/v1/debug/profiler/stop")]
+		[Route("/api/v1/debug/profiler/cpu/stop")]
 		public ActionResult StopProfiler()
 		{
 			if (!_globalConfig.Value.Authorize(ServerAclAction.Debug, User))
@@ -660,15 +665,15 @@ namespace Horde.Server.Server
 
 			DotTrace.SaveData();
 			DotTrace.Detach();
-			return new ContentResult { ContentType = "text/plain", StatusCode = (int)HttpStatusCode.OK, Content = "Profiling session stopped" };
+			return new ContentResult { ContentType = "text/plain", StatusCode = (int)HttpStatusCode.OK, Content = "CPU profiling session stopped" };
 		}
 
 		/// <summary>
-		/// Downloads the captured profiling snapshots
+		/// Downloads the captured CPU profiling snapshots
 		/// </summary>
 		/// <returns>A .zip file containing the profiling snapshots</returns>
 		[HttpGet]
-		[Route("/api/v1/debug/profiler/download")]
+		[Route("/api/v1/debug/profiler/cpu/download")]
 		public ActionResult DownloadProfilingData()
 		{
 			if (!_globalConfig.Value.Authorize(ServerAclAction.Debug, User))
@@ -684,6 +689,44 @@ namespace Horde.Server.Server
 
 			return PhysicalFile(snapshotZipFile, "application/zip", Path.GetFileName(snapshotZipFile));
 		}
+		
+		/// <summary>
+		/// Take a memory snapshot using dotTrace
+		/// </summary>
+		/// <returns>A .dmw file containing the memory snapshot</returns>
+		[HttpGet]
+		[Route("/api/v1/debug/profiler/mem/snapshot")]
+		public async Task<ActionResult> TakeMemorySnapshotAsync()
+		{
+			if (!_globalConfig.Value.Authorize(ServerAclAction.Debug, User))
+			{
+				return Forbid(ServerAclAction.Debug);
+			}
+
+			// Downloads dotMemory executable if not available
+			Stopwatch sw = Stopwatch.StartNew();
+			await DotMemory.EnsurePrerequisiteAsync();
+			_logger.LogInformation("dotMemory prerequisites step finished in {SetupTimeMs} ms", sw.ElapsedMilliseconds);
+
+			string snapshotDir = Path.Join(Path.GetTempPath(), "horde-mem-profiler-snapshots");
+			if (!Directory.Exists(snapshotDir))
+			{
+				Directory.CreateDirectory(snapshotDir);
+			}
+
+			sw.Restart();
+			DotMemory.Config config = new();
+			config.SaveToDir(snapshotDir);
+			string workspaceFilePath = DotMemory.GetSnapshotOnce(config);
+			_logger.LogInformation("dotMemory snapshot captured in {CaptureTimeMs} ms", sw.ElapsedMilliseconds);
+			
+			if (!System.IO.File.Exists(workspaceFilePath))
+			{
+				return NotFound("The generated workspace file was not found");
+			}
+
+			return PhysicalFile(workspaceFilePath, "application/octet-stream", Path.GetFileName(workspaceFilePath));
+		}
 
 		/// <summary>
 		/// Throws an exception to debug error handling
