Horde: Fix documentation link using incorrect slashes.

#rnx

diff --git a/Engine/Source/Programs/Horde/Docs/Deployment/Server.md b/Engine/Source/Programs/Horde/Docs/Deployment/Server.md
--- a/Engine/Source/Programs/Horde/Docs/Deployment/Server.md
+++ b/Engine/Source/Programs/Horde/Docs/Deployment/Server.md
@@ -244,7 +244,7 @@ Horde's internal account system, for example), and it's important that these URL
 By default, this URL is derived from the server's reported DNS name, but this can be overridden through the
 `ServerUrl` property.
 
-To set up a self-signed certificate for testing see [Tutorials > Self Signed Certs](..\Tutorials\SelfSignedCerts.md).
+To set up a self-signed certificate for testing see [Tutorials > Self Signed Certs](../Tutorials/SelfSignedCerts.md).
 
 ### Monitoring
 
