Horde: Remove docs section instructing users to add a [StudioTelemetry.Config] section.

diff --git a/Engine/Source/Programs/Horde/Docs/Tutorials/Analytics.md b/Engine/Source/Programs/Horde/Docs/Tutorials/Analytics.md
--- a/Engine/Source/Programs/Horde/Docs/Tutorials/Analytics.md
+++ b/Engine/Source/Programs/Horde/Docs/Tutorials/Analytics.md
@@ -25,12 +25,6 @@ make sure it is enabled. It should be enabled by default.
 the following lines:
 
     ```ini
-    [StudioTelemetry.Config] 
-    SendTelemetry=true
-    SendUserData=false
-    SendHardwareData=true
-    SendOSData=true
-
     [StudioTelemetry.Provider.HordeAnalytics]
     Name=HordeAnalytics
     ProviderModule=AnalyticsET
