Horde: Add logging to identify stalling lease assignments

diff --git a/Engine/Source/Programs/Horde/Horde.Server/Agents/AgentService.cs b/Engine/Source/Programs/Horde/Horde.Server/Agents/AgentService.cs
--- a/Engine/Source/Programs/Horde/Horde.Server/Agents/AgentService.cs
+++ b/Engine/Source/Programs/Horde/Horde.Server/Agents/AgentService.cs
@@ -515,6 +515,7 @@ namespace Horde.Server.Agents
 					// If no task source is valid, just add a delay
 					if (tasks.Count == 0)
 					{
+						_logger.LogInformation("No task source valid for agent {AgentId}. Waiting {WaitTimeMs} ms", agent.Id, maxWaitTime.TotalMilliseconds);
 						await AsyncUtils.DelayNoThrow(maxWaitTime, cancellationToken);
 						break;
 					}
diff --git a/Engine/Source/Programs/Horde/Horde.Server/Jobs/JobTaskSource.cs b/Engine/Source/Programs/Horde/Horde.Server/Jobs/JobTaskSource.cs
--- a/Engine/Source/Programs/Horde/Horde.Server/Jobs/JobTaskSource.cs
+++ b/Engine/Source/Programs/Horde/Horde.Server/Jobs/JobTaskSource.cs
@@ -770,13 +770,16 @@ namespace Horde.Server.Jobs
 			IJob job = item._job;
 			IJobStepBatch batch = item.Batch;
 			IAgent agent = waiter.Agent;
-			_logger.LogDebug("Assigning job {JobId}, batch {BatchId} to waiter (agent {AgentId})", job.Id, batch.Id, agent.Id);
-
-			// Generate a new unique id for the lease
-			LeaseId leaseId = new LeaseId(BinaryIdUtils.CreateNew());
-
-			// The next time to try assigning to another agent
-			DateTime backOffTime = DateTime.UtcNow + TimeSpan.FromMinutes(1.0);
+			LeaseId leaseId = new (BinaryIdUtils.CreateNew()); // Generate a new unique id for the lease
+			
+			using IDisposable logScope = _logger
+				.WithProperty("JobId", job.Id.ToString())
+				.WithProperty("BatchId", batch.Id.ToString())
+				.WithProperty("AgentId", agent.Id.ToString())
+				.WithProperty("LeaseId", leaseId.ToString())
+				.BeginScope();
+			
+			_logger.LogInformation("Assigning job to waiter");
 
 			// Allocate a log ID but hold off creating the actual log file until the lease has been accepted
 			LogId logId = LogIdUtils.GenerateNewId();
@@ -823,20 +826,20 @@ namespace Horde.Server.Jobs
 					AgentLease lease = new AgentLease(leaseId, null, leaseName.ToString(), job.StreamId, item._poolId, logId, LeaseState.Pending, null, true, payload);
 					if (waiter.LeaseSource.TrySetResult(lease))
 					{
-						_logger.LogInformation("Assigned lease {LeaseId} to agent {AgentId}", leaseId, agent.Id);
+						_logger.LogInformation("Assigned lease to agent");
 						await _logFileService.CreateLogFileAsync(job.Id, leaseId, agent.SessionId, LogType.Json, logId, cancellationToken);
 						return lease;
 					}
 				}
 
 				// Cancel the lease
-				_logger.LogInformation("Unable to assign lease {LeaseId} to agent {AgentId}, cancelling", leaseId, agent.Id);
+				_logger.LogInformation("Unable to assign lease, cancelling");
 				await CancelLeaseAsync(waiter.Agent, job.Id, batch.Id, cancellationToken);
 			}
 			else
 			{
 				// Unable to assign job
-				_logger.LogInformation("Failed to assign job {JobId}, batch {BatchId} to agent {AgentId}. Refreshing queue entries.", job.Id, batch.Id, agent.Id);
+				_logger.LogInformation("Failed to assign job/batch to agent. Refreshing queue entries");
 
 				// Get the new copy of the job
 				newJob = await _jobs.GetAsync(job.Id, cancellationToken);
