Horde: Add a "My Checkouts" at top of shared device pivot
#rnx

diff --git a/Engine/Source/Programs/Horde/HordeDashboard/src/components/DeviceView.tsx b/Engine/Source/Programs/Horde/HordeDashboard/src/components/DeviceView.tsx
--- a/Engine/Source/Programs/Horde/HordeDashboard/src/components/DeviceView.tsx
+++ b/Engine/Source/Programs/Horde/HordeDashboard/src/components/DeviceView.tsx
@@ -21,6 +21,8 @@ const handler = new DeviceHandler();
 
 type DeviceItem = {
    device: GetDeviceResponse;
+   // Whether user has this device checked out
+   userCheckout?: boolean;
    status?: string;
 }
 
@@ -681,10 +683,29 @@ const DevicePanel: React.FC = observer(() => {
    const devices = GetSortedDeviceList(filteredDevices, !automationTab ? "status" : undefined);
    const newGroups: IGroup[] = [];
 
+   let myDevices: DeviceItem[] = [];
+   if (pivotState.key === "pivot-key-shared") {
+      myDevices = devices.filter(d => d.device.checkedOutByUserId === dashboard.userId);
+      devices.unshift(...myDevices.map(d => { return { ...d, userCheckout: true } }));
+   }
+
    let curPlatform: string | undefined;
    devices.forEach((item, index) => {
       const d = item.device;
-      if (d.platformId !== curPlatform) {
+      if (item.userCheckout) {
+         if (!newGroups.length) {
+            const key = `user_checkout`;
+            newGroups.push({
+               key: key,
+               name: "My Checkouts",
+               startIndex: 0,
+               count: myDevices.length,
+               level: 0,
+               isCollapsed: groups.find(g => g.key === key)?.isCollapsed
+            });
+         }
+      }
+      else if (d.platformId !== curPlatform) {
          curPlatform = d.platformId;
          const key = `group_${d.platformId}`;
          newGroups.push({
