Horde: Use strongly typed identifiers for artifact log messages in agent.

diff --git a/Engine/Source/Programs/Horde/Horde.Agent/Execution/JobExecutor.cs b/Engine/Source/Programs/Horde/Horde.Agent/Execution/JobExecutor.cs
--- a/Engine/Source/Programs/Horde/Horde.Agent/Execution/JobExecutor.cs
+++ b/Engine/Source/Programs/Horde/Horde.Agent/Execution/JobExecutor.cs
@@ -555,7 +555,8 @@ namespace Horde.Agent.Execution
 					artifactRequest.Type = artifactType.ToString();
 
 					RpcCreateJobArtifactResponseV2 artifact = await jobRpc.Client.CreateArtifactV2Async(artifactRequest, cancellationToken: cancellationToken);
-					logger.LogInformation("Creating output artifact {ArtifactId} '{ArtifactName}' ({ArtifactType}) with ref {RefName} in namespace {NamespaceId}", artifact.Id, artifactName, artifactType, artifact.RefName, artifact.NamespaceId);
+					ArtifactId artifactId = ArtifactId.Parse(artifact.Id);
+					logger.LogInformation("Creating output artifact {ArtifactId} '{ArtifactName}' ({ArtifactType}) with ref {RefName} in namespace {NamespaceId}", artifactId, artifactName, artifactType, artifact.RefName, artifact.NamespaceId);
 
 					// Write the data
 					using IStorageClient storage = CreateStorageClient(new NamespaceId(artifact.NamespaceId), artifact.Token);
@@ -903,7 +904,8 @@ namespace Horde.Agent.Execution
 				artifactRequest.Type = type.ToString();
 
 				RpcCreateJobArtifactResponseV2 artifact = await jobRpc.Client.CreateArtifactV2Async(artifactRequest, cancellationToken: cancellationToken);
-				Logger.LogInformation("Created artifact {ArtifactId} '{ArtifactName}' ({ArtifactType}) with ref {RefName} ({Link})", artifact.Id, name, type, artifact.RefName, $"{Session.ServerUrl}/api/v1/storage/{artifact.NamespaceId}/refs/{artifact.RefName}");
+				ArtifactId artifactId = ArtifactId.Parse(artifact.Id);
+				Logger.LogInformation("Created artifact {ArtifactId} '{ArtifactName}' ({ArtifactType}) with ref {RefName} ({Link})", artifactId, name, type, artifact.RefName, $"{Session.ServerUrl}/api/v1/storage/{artifact.NamespaceId}/refs/{artifact.RefName}");
 
 				using IStorageClient storage = CreateStorageClient(new NamespaceId(artifact.NamespaceId), artifact.Token);
 
@@ -918,13 +920,13 @@ namespace Horde.Agent.Execution
 					}
 					catch (Exception ex)
 					{
-						Logger.LogInformation(ex, "Error uploading files for artifact {ArtifactId}", artifact.Id);
+						Logger.LogInformation(ex, "Error uploading files for artifact {ArtifactId}", artifactId);
 						throw;
 					}
 				}
 				await storage.WriteRefAsync(new RefName(artifact.RefName), rootRef, cancellationToken: cancellationToken);
 
-				Logger.LogInformation("Uploaded artifact {ArtifactId}", artifact.Id);
+				Logger.LogInformation("Uploaded artifact {ArtifactId}", artifactId);
 			}
 			catch (Exception ex)
 			{
@@ -1103,7 +1105,8 @@ namespace Horde.Agent.Execution
 				artifactRequest.Type = ArtifactType.StepOutput.ToString();
 
 				RpcCreateJobArtifactResponseV2 artifact = await jobRpc.Client.CreateArtifactV2Async(artifactRequest, cancellationToken: cancellationToken);
-				logger.LogInformation("Created artifact {ArtifactId} '{ArtifactName}' ({ArtifactType}) with ref {RefName} ({RefUrl})", artifact.Id, artifactRequest.Name, ArtifactType.StepOutput, artifact.RefName, $"{Session.ServerUrl.ToString().TrimEnd('/')}/api/v1/storage/{artifact.NamespaceId}/refs/{artifact.RefName}");
+				ArtifactId artifactId = ArtifactId.Parse(artifact.Id);
+				logger.LogInformation("Created artifact {ArtifactId} '{ArtifactName}' ({ArtifactType}) with ref {RefName} ({RefUrl})", artifactId, artifactRequest.Name, ArtifactType.StepOutput, artifact.RefName, $"{Session.ServerUrl.ToString().TrimEnd('/')}/api/v1/storage/{artifact.NamespaceId}/refs/{artifact.RefName}");
 
 				using IStorageClient storage = CreateStorageClient(new NamespaceId(artifact.NamespaceId), artifact.Token);
 
@@ -1181,7 +1184,8 @@ namespace Horde.Agent.Execution
 					artifactRequest.Metadata.AddRange(graphArtifact.Metadata);
 
 					RpcCreateJobArtifactResponseV2 artifact = await jobRpc.Client.CreateArtifactV2Async(artifactRequest, cancellationToken: cancellationToken);
-					logger.LogInformation("Created artifact {ArtifactId} '{ArtifactName}' ({ArtifactType}) with ref {RefName} ({RefUrl})", artifact.Id, artifactRequest.Name, ArtifactType.StepOutput, artifact.RefName, $"{Session.ServerUrl.ToString().TrimEnd('/')}/api/v1/storage/{artifact.NamespaceId}/refs/{artifact.RefName}");
+					ArtifactId artifactId = ArtifactId.Parse(artifact.Id);
+					logger.LogInformation("Created artifact {ArtifactId} '{ArtifactName}' ({ArtifactType}) with ref {RefName} ({RefUrl})", artifactId, artifactRequest.Name, ArtifactType.StepOutput, artifact.RefName, $"{Session.ServerUrl.ToString().TrimEnd('/')}/api/v1/storage/{artifact.NamespaceId}/refs/{artifact.RefName}");
 
 					using IStorageClient storage = CreateStorageClient(new NamespaceId(artifact.NamespaceId), artifact.Token);
 
