Horde: Fix PerforceConnectionId types not being deserialized correctly.

#rnx

diff --git a/Engine/Source/Programs/Horde/Horde.Server/ServerSettings.cs b/Engine/Source/Programs/Horde/Horde.Server/ServerSettings.cs
--- a/Engine/Source/Programs/Horde/Horde.Server/ServerSettings.cs
+++ b/Engine/Source/Programs/Horde/Horde.Server/ServerSettings.cs
@@ -2,6 +2,7 @@
 
 using System;
 using System.Collections.Generic;
+using System.ComponentModel;
 using System.Diagnostics.CodeAnalysis;
 using System.Linq;
 using System.Net;
@@ -957,6 +958,7 @@ namespace Horde.Server
 	/// <param name="Id">Id to construct from</param>
 	[JsonSchemaString]
 	[StringIdConverter(typeof(PerforceConnectionIdConverter))]
+	[TypeConverter(typeof(StringIdTypeConverter<PerforceConnectionId, PerforceConnectionIdConverter>))]
 	public record struct PerforceConnectionId(StringId Id)
 	{
 		/// <summary>
