Horde: Dispose temporary Perforce connections used for health checks

diff --git a/Engine/Source/Programs/Horde/Horde.Server/Perforce/PerforceLoadBalancer.cs b/Engine/Source/Programs/Horde/Horde.Server/Perforce/PerforceLoadBalancer.cs
--- a/Engine/Source/Programs/Horde/Horde.Server/Perforce/PerforceLoadBalancer.cs
+++ b/Engine/Source/Programs/Horde/Horde.Server/Perforce/PerforceLoadBalancer.cs
@@ -739,7 +739,7 @@ namespace Horde.Server.Perforce
 
 			try
 			{
-				IPerforceConnection connection = await PerforceConnection.CreateAsync(settings, logger);
+				using IPerforceConnection connection = await PerforceConnection.CreateAsync(settings, logger);
 				await connection.GetInfoAsync(InfoOptions.ShortOutput, cancellationToken);
 				return new ServerHealth(PerforceServerStatus.Healthy, "Server responded to \"p4 info\"");
 			}
