Horde: Make log output directory configurable in agent

diff --git a/Engine/Source/Programs/Horde/Horde.Agent/AgentApp.cs b/Engine/Source/Programs/Horde/Horde.Agent/AgentApp.cs
--- a/Engine/Source/Programs/Horde/Horde.Agent/AgentApp.cs
+++ b/Engine/Source/Programs/Horde/Horde.Agent/AgentApp.cs
@@ -286,7 +286,7 @@ namespace Horde.Agent
 				.Build();
 		}
 
-		static AgentSettings BindSettings(IConfiguration configuration)
+		internal static AgentSettings BindSettings(IConfiguration configuration)
 		{
 			AgentSettings settings = new AgentSettings();
 			configuration.GetSection(AgentSettings.SectionName).Bind(settings);
diff --git a/Engine/Source/Programs/Horde/Horde.Agent/AgentSettings.cs b/Engine/Source/Programs/Horde/Horde.Agent/AgentSettings.cs
--- a/Engine/Source/Programs/Horde/Horde.Agent/AgentSettings.cs
+++ b/Engine/Source/Programs/Horde/Horde.Agent/AgentSettings.cs
@@ -216,10 +216,15 @@ namespace Horde.Agent
 		public PerforceExecutorSettings PerforceExecutor { get; set; } = new PerforceExecutorSettings();
 
 		/// <summary>
-		/// Working directory
+		/// Working directory for leases and jobs (i.e where files from Perforce will be checked out) 
 		/// </summary>
 		public DirectoryReference WorkingDir { get; set; } = DirectoryReference.Combine(AgentApp.DataDir, "Sandbox");
 
+		/// <summary>
+		/// Directory where agent and lease logs are written
+		/// </summary>
+		public DirectoryReference LogsDir { get; set; } = AgentApp.DataDir;
+
 		/// <summary>
 		/// Whether to mount the specified list of network shares
 		/// </summary>
diff --git a/Engine/Source/Programs/Horde/Horde.Agent/Leases/LeaseLoggerFactory.cs b/Engine/Source/Programs/Horde/Horde.Agent/Leases/LeaseLoggerFactory.cs
--- a/Engine/Source/Programs/Horde/Horde.Agent/Leases/LeaseLoggerFactory.cs
+++ b/Engine/Source/Programs/Horde/Horde.Agent/Leases/LeaseLoggerFactory.cs
@@ -25,7 +25,7 @@ namespace Horde.Agent.Leases
 		public LeaseLoggerFactory(IOptions<AgentSettings> settings, ILogger<LeaseLoggerFactory> logger)
 		{
 			_settings = settings.Value;
-			_logDir = DirectoryReference.Combine(settings.Value.WorkingDir, "Leases");
+			_logDir = DirectoryReference.Combine(settings.Value.LogsDir, "Leases");
 			_logger = logger;
 			_backgroundTask = BackgroundTask.StartNew(BackgroundCleanupAsync);
 		}
diff --git a/Engine/Source/Programs/Horde/Horde.Agent/Logging.cs b/Engine/Source/Programs/Horde/Horde.Agent/Logging.cs
--- a/Engine/Source/Programs/Horde/Horde.Agent/Logging.cs
+++ b/Engine/Source/Programs/Horde/Horde.Agent/Logging.cs
@@ -50,7 +50,8 @@ namespace Horde.Agent
 
 		static Serilog.ILogger CreateSerilogLogger(IConfiguration configuration)
 		{
-			DirectoryReference.CreateDirectory(AgentApp.DataDir);
+			AgentSettings settings = AgentApp.BindSettings(configuration);
+			DirectoryReference.CreateDirectory(settings.LogsDir);
 
 			ConsoleTheme theme;
 			if (RuntimeInformation.IsOSPlatform(OSPlatform.Windows) && Environment.OSVersion.Version < new Version(10, 0))
@@ -64,8 +65,8 @@ namespace Horde.Agent
 
 			return new LoggerConfiguration()
 				.WriteTo.Console(outputTemplate: "[{Timestamp:HH:mm:ss} {Level:w3}] {Indent}{Message:l}{NewLine}{Exception}", theme: theme)
-				.WriteTo.File(FileReference.Combine(AgentApp.DataDir, "Log-.txt").FullName, fileSizeLimitBytes: 50 * 1024 * 1024, rollingInterval: RollingInterval.Day, rollOnFileSizeLimit: true, retainedFileCountLimit: 10)
-				.WriteTo.File(new JsonFormatter(renderMessage: true), FileReference.Combine(AgentApp.DataDir, "Log-.json").FullName, fileSizeLimitBytes: 50 * 1024 * 1024, rollingInterval: RollingInterval.Day, rollOnFileSizeLimit: true, retainedFileCountLimit: 10)
+				.WriteTo.File(FileReference.Combine(settings.LogsDir, "Log-.txt").FullName, fileSizeLimitBytes: 50 * 1024 * 1024, rollingInterval: RollingInterval.Day, rollOnFileSizeLimit: true, retainedFileCountLimit: 10)
+				.WriteTo.File(new JsonFormatter(renderMessage: true), FileReference.Combine(settings.LogsDir, "Log-.json").FullName, fileSizeLimitBytes: 50 * 1024 * 1024, rollingInterval: RollingInterval.Day, rollOnFileSizeLimit: true, retainedFileCountLimit: 10)
 				.ReadFrom.Configuration(configuration)
 				.MinimumLevel.ControlledBy(LogLevelSwitch)
 				.Enrich.FromLogContext()
