Horde: Do not allow creating artifacts with an invalid type field.

#rnx

diff --git a/Engine/Source/Programs/Horde/Horde.Server/Artifacts/ArtifactCollection.cs b/Engine/Source/Programs/Horde/Horde.Server/Artifacts/ArtifactCollection.cs
--- a/Engine/Source/Programs/Horde/Horde.Server/Artifacts/ArtifactCollection.cs
+++ b/Engine/Source/Programs/Horde/Horde.Server/Artifacts/ArtifactCollection.cs
@@ -124,6 +124,15 @@ namespace Horde.Server.Artifacts
 		/// <inheritdoc/>
 		public async Task<IArtifact> AddAsync(ArtifactName name, ArtifactType type, string? description, StreamId streamId, int change, IEnumerable<string> keys, IEnumerable<string> metadata, AclScopeName scopeName, CancellationToken cancellationToken)
 		{
+			if (name.Id.IsEmpty)
+			{
+				throw new ArgumentException($"Artifact name cannot be empty", nameof(name));
+			}
+			if (type.Id.IsEmpty)
+			{
+				throw new ArgumentException($"Artifact type for '{name}' is not valid", nameof(type));
+			}
+
 			ArtifactId id = new ArtifactId(BinaryIdUtils.CreateNew());
 
 			NamespaceId namespaceId = Namespace.Artifacts;
