Horde: Send Slack notifications whenever config changes are applied.

diff --git a/Engine/Source/Programs/Horde/Horde.Server.Tests/Notifications/NotificationServiceTests.cs b/Engine/Source/Programs/Horde/Horde.Server.Tests/Notifications/NotificationServiceTests.cs
--- a/Engine/Source/Programs/Horde/Horde.Server.Tests/Notifications/NotificationServiceTests.cs
+++ b/Engine/Source/Programs/Horde/Horde.Server.Tests/Notifications/NotificationServiceTests.cs
@@ -8,6 +8,7 @@ using EpicGames.Horde.Jobs;
 using EpicGames.Horde.Streams;
 using Horde.Server.Agents;
 using Horde.Server.Agents.Pools;
+using Horde.Server.Configuration;
 using Horde.Server.Devices;
 using Horde.Server.Issues;
 using Horde.Server.Jobs;
@@ -40,7 +41,7 @@ namespace Horde.Server.Tests.Notifications
 		public Task NotifyJobStepCompleteAsync(IUser user, IJob job, IJobStepBatch batch, IJobStep step, INode node, List<ILogEventData> jobStepEventData, CancellationToken cancellationToken) { throw new NotImplementedException(); }
 		public Task NotifyLabelCompleteAsync(IUser user, IJob job, ILabel label, int labelIdx, LabelOutcome outcome, List<(string, JobStepOutcome, Uri)> stepData, CancellationToken cancellationToken) { throw new NotImplementedException(); }
 		public Task NotifyIssueUpdatedAsync(IIssue issue, CancellationToken cancellationToken) { throw new NotImplementedException(); }
-		public Task NotifyConfigUpdateAsync(Exception? ex, CancellationToken cancellationToken) => Task.CompletedTask;
+		public Task NotifyConfigUpdateAsync(ConfigUpdateInfo info, CancellationToken cancellationToken) => Task.CompletedTask;
 		public Task NotifyConfigUpdateFailureAsync(string errorMessage, string fileName, int? change = null, IUser? author = null, string? description = null, CancellationToken cancellationToken = default) { throw new NotImplementedException(); }
 		public Task NotifyDeviceServiceAsync(string message, IDevice? device = null, IDevicePool? pool = null, StreamConfig? stream = null, IJob? job = null, IJobStep? step = null, INode? node = null, IUser? user = null, CancellationToken cancellationToken = default) { throw new NotImplementedException(); }
 		public Task SendAgentReportAsync(AgentReport report, CancellationToken cancellationToken) => throw new NotImplementedException();
diff --git a/Engine/Source/Programs/Horde/Horde.Server/Configuration/ConfigService.cs b/Engine/Source/Programs/Horde/Horde.Server/Configuration/ConfigService.cs
--- a/Engine/Source/Programs/Horde/Horde.Server/Configuration/ConfigService.cs
+++ b/Engine/Source/Programs/Horde/Horde.Server/Configuration/ConfigService.cs
@@ -27,6 +27,11 @@ using StackExchange.Redis;
 
 namespace Horde.Server.Configuration
 {
+	/// <summary>
+	/// Information about the updated config
+	/// </summary>
+	public record class ConfigUpdateInfo(Dictionary<Uri, string> Sources, Exception? Exception);
+
 	/// <summary>
 	/// Service which processes runtime configuration data.
 	/// </summary>
@@ -132,7 +137,7 @@ namespace Horde.Server.Configuration
 		/// <summary>
 		/// Event for notifications that the config has been updated
 		/// </summary>
-		public event Action<Exception?>? OnConfigUpdate;
+		public event Action<ConfigUpdateInfo>? OnConfigUpdate;
 
 		/// <summary>
 		/// Constructor
@@ -448,14 +453,14 @@ namespace Horde.Server.Configuration
 			{
 				ConfigSnapshot snapshot = await CreateSnapshotAsync(cancellationToken);
 				await _health.UpdateAsync(HealthStatus.Healthy);
-				OnConfigUpdate?.Invoke(null);
+				OnConfigUpdate?.Invoke(new ConfigUpdateInfo(snapshot.Dependencies, null));
 				return snapshot;
 			}
 			catch (Exception ex)
 			{
 				_logger.LogError(ex, "Exception while updating config: {Message}", ex.Message);
 				await _health.UpdateAsync(HealthStatus.Unhealthy, ex.Message);
-				OnConfigUpdate?.Invoke(ex);
+				OnConfigUpdate?.Invoke(new ConfigUpdateInfo(new Dictionary<Uri, string>(), ex));
 				return null;
 			}
 		}
diff --git a/Engine/Source/Programs/Horde/Horde.Server/Notifications/INotificationService.cs b/Engine/Source/Programs/Horde/Horde.Server/Notifications/INotificationService.cs
--- a/Engine/Source/Programs/Horde/Horde.Server/Notifications/INotificationService.cs
+++ b/Engine/Source/Programs/Horde/Horde.Server/Notifications/INotificationService.cs
@@ -8,6 +8,7 @@ using System.Threading;
 using System.Threading.Tasks;
 using EpicGames.Horde.Jobs;
 using Horde.Server.Agents;
+using Horde.Server.Configuration;
 using Horde.Server.Devices;
 using Horde.Server.Issues;
 using Horde.Server.Jobs;
@@ -138,8 +139,8 @@ namespace Horde.Server.Notifications
 		/// <summary>
 		/// Notify slack channel about a configuration update
 		/// </summary>
-		/// <param name="ex">Exception thrown during the update. Null if the update completed successfully.</param>
-		void NotifyConfigUpdate(Exception? ex);
+		/// <param name="info">Information about the config update</param>
+		void NotifyConfigUpdate(ConfigUpdateInfo info);
 
 		/// <summary>
 		/// Notify slack channel about a stream update failure
diff --git a/Engine/Source/Programs/Horde/Horde.Server/Notifications/INotificationSink.cs b/Engine/Source/Programs/Horde/Horde.Server/Notifications/INotificationSink.cs
--- a/Engine/Source/Programs/Horde/Horde.Server/Notifications/INotificationSink.cs
+++ b/Engine/Source/Programs/Horde/Horde.Server/Notifications/INotificationSink.cs
@@ -6,6 +6,7 @@ using System.Threading;
 using System.Threading.Tasks;
 using EpicGames.Horde.Jobs;
 using Horde.Server.Agents;
+using Horde.Server.Configuration;
 using Horde.Server.Devices;
 using Horde.Server.Issues;
 using Horde.Server.Jobs;
@@ -99,9 +100,9 @@ namespace Horde.Server.Notifications
 		/// <summary>
 		/// Notification that the configuration state has changed
 		/// </summary>
-		/// <param name="ex">Exception during update. Null if the update completed successfully.</param>
+		/// <param name="info">Information about the updated config</param>
 		/// <param name="cancellationToken">Cancellation token for the operation</param>
-		Task NotifyConfigUpdateAsync(Exception? ex, CancellationToken cancellationToken);
+		Task NotifyConfigUpdateAsync(ConfigUpdateInfo info, CancellationToken cancellationToken);
 
 		/// <summary>
 		/// Notification that a stream has failed to update
diff --git a/Engine/Source/Programs/Horde/Horde.Server/Notifications/NotificationService.cs b/Engine/Source/Programs/Horde/Horde.Server/Notifications/NotificationService.cs
--- a/Engine/Source/Programs/Horde/Horde.Server/Notifications/NotificationService.cs
+++ b/Engine/Source/Programs/Horde/Horde.Server/Notifications/NotificationService.cs
@@ -286,10 +286,10 @@ namespace Horde.Server.Notifications
 		}
 
 		/// <inheritdoc/>
-		public void NotifyConfigUpdate(Exception? ex)
+		public void NotifyConfigUpdate(ConfigUpdateInfo info)
 		{
-			_logger.LogInformation(ex, "Configuration updated ({Result})", (ex == null) ? "success" : "failure");
-			EnqueueTasks((sink, ctx) => sink.NotifyConfigUpdateAsync(ex, ctx));
+			_logger.LogInformation(info.Exception, "Configuration updated ({Result})", (info.Exception == null) ? "success" : "failure");
+			EnqueueTasks((sink, ctx) => sink.NotifyConfigUpdateAsync(info, ctx));
 		}
 
 		/// <inheritdoc/>
diff --git a/Engine/Source/Programs/Horde/Horde.Server/Notifications/Sinks/SlackNotificationSink.cs b/Engine/Source/Programs/Horde/Horde.Server/Notifications/Sinks/SlackNotificationSink.cs
--- a/Engine/Source/Programs/Horde/Horde.Server/Notifications/Sinks/SlackNotificationSink.cs
+++ b/Engine/Source/Programs/Horde/Horde.Server/Notifications/Sinks/SlackNotificationSink.cs
@@ -2092,7 +2092,7 @@ namespace Horde.Server.Notifications.Sinks
 		#region Stream updates
 
 		/// <inheritdoc/>
-		public async Task NotifyConfigUpdateAsync(Exception? ex, CancellationToken cancellationToken)
+		public async Task NotifyConfigUpdateAsync(ConfigUpdateInfo info, CancellationToken cancellationToken)
 		{
 			if (String.IsNullOrEmpty(_settings.ConfigNotificationChannel))
 			{
@@ -2101,6 +2101,7 @@ namespace Horde.Server.Notifications.Sinks
 
 			const string EventId = "config-update";
 
+			Exception? ex = info.Exception;
 			if (ex != null)
 			{
 				_logger.LogInformation(ex, "Sending config update failure notification: {Message}", ex.Message);
@@ -2156,13 +2157,43 @@ namespace Horde.Server.Notifications.Sinks
 			}
 			else
 			{
-				if (await DeleteMessageStateAsync(_settings.ConfigNotificationChannel, EventId, cancellationToken))
+				Dictionary<string, int> perforceClusterToChange = new Dictionary<string, int>();
+				foreach (Uri uri in info.Sources.Keys)
 				{
-					SlackMessage message = new SlackMessage();
-					message.AddSection($"*Config Update Succeeded*");
-					await SendMessageAsync(_settings.ConfigNotificationChannel, message, cancellationToken);
-					await DeleteMessageStateAsync(_settings.ConfigNotificationChannel, EventId, cancellationToken);
+					if (uri.Scheme == PerforceConfigSource.Scheme)
+					{
+						string path = uri.LocalPath;
+
+						int atIdx = path.IndexOf('@');
+						if (atIdx != -1 && Int32.TryParse(path.Substring(atIdx + 1), out int change))
+						{
+							int existingChange;
+							if (!perforceClusterToChange.TryGetValue(uri.Host, out existingChange) || change > existingChange)
+							{
+								perforceClusterToChange[uri.Host] = change;
+							}
+						}
+					}
 				}
+
+				// Can only send updates after failures by checking the result of this delete
+				await DeleteMessageStateAsync(_settings.ConfigNotificationChannel, EventId, cancellationToken);
+
+				SlackMessage message = new SlackMessage();
+				message.AddSection($"*Config Update Succeeded*");
+
+				if(perforceClusterToChange.Count > 0)
+				{
+					List<string> lines = new List<string>();
+					foreach ((string clusterName, int change) in perforceClusterToChange)
+					{
+						lines.Add($"Perforce changelist ({clusterName}): {change}");
+					}
+					message.AddSection(String.Join("\n", lines));
+				}
+
+				await SendMessageAsync(_settings.ConfigNotificationChannel, message, cancellationToken);
+				await DeleteMessageStateAsync(_settings.ConfigNotificationChannel, EventId, cancellationToken);
 			}
 		}
 
