Horde: Use standard definitions for log value types in PerforceLogger.

#rnx

diff --git a/Engine/Source/Programs/Horde/Horde.Agent/Utility/PerforceLogger.cs b/Engine/Source/Programs/Horde/Horde.Agent/Utility/PerforceLogger.cs
--- a/Engine/Source/Programs/Horde/Horde.Agent/Utility/PerforceLogger.cs
+++ b/Engine/Source/Programs/Horde/Horde.Agent/Utility/PerforceLogger.cs
@@ -145,8 +145,8 @@ namespace Horde.Agent.Utility
 			}
 		}
 
-		static readonly Utf8String s_sourceFileType = new Utf8String("SourceFile");
-		static readonly Utf8String s_assetType = new Utf8String("Asset");
+		static readonly Utf8String s_sourceFileType = LogValueType.SourceFile;
+		static readonly Utf8String s_assetType = LogValueType.Asset;
 
 		static readonly Utf8String s_file = new Utf8String("file");
 
