Horde: Allow specifying a unique id for each parameter in a job template. If a parameter is not specified, a unique identifier will be generated from the widget label. Parameters can be used to start jobs (by specifying the Parameters dictionary in the CreateJobRequest object), and will be returned as part of the job response.

diff --git a/Engine/Source/Programs/Horde/Horde.Server/Jobs/IJob.cs b/Engine/Source/Programs/Horde/Horde.Server/Jobs/IJob.cs
--- a/Engine/Source/Programs/Horde/Horde.Server/Jobs/IJob.cs
+++ b/Engine/Source/Programs/Horde/Horde.Server/Jobs/IJob.cs
@@ -673,6 +673,11 @@ namespace Horde.Server.Jobs
 		/// </summary>
 		public IReadOnlyList<IJobStepBatch> Batches { get; }
 
+		/// <summary>
+		/// Parameters for the job
+		/// </summary>
+		public IReadOnlyDictionary<ParameterId, string> Parameters { get; }
+
 		/// <summary>
 		/// Optional user-defined properties for this job
 		/// </summary>
diff --git a/Engine/Source/Programs/Horde/Horde.Server/Jobs/IJobCollection.cs b/Engine/Source/Programs/Horde/Horde.Server/Jobs/IJobCollection.cs
--- a/Engine/Source/Programs/Horde/Horde.Server/Jobs/IJobCollection.cs
+++ b/Engine/Source/Programs/Horde/Horde.Server/Jobs/IJobCollection.cs
@@ -77,6 +77,9 @@ namespace Horde.Server.Jobs
 		/// <inheritdoc cref="IJob.NotificationChannelFilter"/>
 		public string? NotificationChannelFilter { get; set; }
 
+		/// <inheritdoc cref="IJob.Parameters"/>
+		public Dictionary<ParameterId, string> Parameters { get; } = new Dictionary<ParameterId, string>();
+
 		/// <inheritdoc cref="IJob.Arguments"/>
 		public List<string> Arguments { get; } = new List<string>();
 
diff --git a/Engine/Source/Programs/Horde/Horde.Server/Jobs/JobCollection.cs b/Engine/Source/Programs/Horde/Horde.Server/Jobs/JobCollection.cs
--- a/Engine/Source/Programs/Horde/Horde.Server/Jobs/JobCollection.cs
+++ b/Engine/Source/Programs/Horde/Horde.Server/Jobs/JobCollection.cs
@@ -291,6 +291,11 @@ namespace Horde.Server.Jobs
 			public int SchedulePriority { get; set; }
 			public List<JobStepBatchDocument> Batches { get; set; } = new List<JobStepBatchDocument>();
 			public List<Report>? Reports { get; set; }
+
+			[BsonDictionaryOptions(Representation = DictionaryRepresentation.ArrayOfDocuments)]
+			public Dictionary<ParameterId, string> Parameters { get; set; } = new Dictionary<ParameterId, string>();
+			IReadOnlyDictionary<ParameterId, string> IJob.Parameters => Parameters;
+
 			public List<string> Arguments { get; set; } = new List<string>();
 
 			[BsonDictionaryOptions(DictionaryRepresentation.ArrayOfDocuments)]
@@ -362,6 +367,7 @@ namespace Horde.Server.Jobs
 				ShowUgsAlerts = options.ShowUgsAlerts;
 				NotificationChannel = options.NotificationChannel;
 				NotificationChannelFilter = options.NotificationChannelFilter;
+				Parameters = new Dictionary<ParameterId, string>(options.Parameters);
 				Arguments.AddRange(options.Arguments);
 
 				foreach (KeyValuePair<string, string> pair in options.Environment)
diff --git a/Engine/Source/Programs/Horde/Horde.Server/Jobs/JobsController.cs b/Engine/Source/Programs/Horde/Horde.Server/Jobs/JobsController.cs
--- a/Engine/Source/Programs/Horde/Horde.Server/Jobs/JobsController.cs
+++ b/Engine/Source/Programs/Horde/Horde.Server/Jobs/JobsController.cs
@@ -203,6 +203,18 @@ namespace Horde.Server.Jobs
 			options.Arguments.AddRange(arguments);
 			options.JobOptions ??= create.JobOptions;
 
+			if (create.Parameters != null)
+			{
+				foreach ((ParameterId parameter, string value) in create.Parameters)
+				{
+					options.Parameters[parameter] = value;
+				}
+				foreach (Parameter parameter in template.Parameters)
+				{
+					parameter.GetArguments(options.Parameters, false, options.Arguments);
+				}
+			}
+
 			foreach ((string key, string value) in environment)
 			{
 				options.Environment[key] = value;
@@ -571,6 +583,7 @@ namespace Horde.Server.Jobs
 			response.AutoSubmitChange = job.AutoSubmitChange;
 			response.AutoSubmitMessage = job.AutoSubmitMessage;
 			response.Reports = job.Reports?.ConvertAll(x => CreateGetReportResponse(x));
+			response.Parameters = job.Parameters.ToDictionary();
 			response.Arguments = job.Arguments.ToList();
 			response.UpdateTime = new DateTimeOffset(job.UpdateTimeUtc);
 			response.UseArtifactsV2 = job.JobOptions?.UseNewTempStorage ?? true;
diff --git a/Engine/Source/Programs/Horde/Horde.Server/Jobs/Templates/ITemplate.cs b/Engine/Source/Programs/Horde/Horde.Server/Jobs/Templates/ITemplate.cs
--- a/Engine/Source/Programs/Horde/Horde.Server/Jobs/Templates/ITemplate.cs
+++ b/Engine/Source/Programs/Horde/Horde.Server/Jobs/Templates/ITemplate.cs
@@ -3,7 +3,10 @@
 using System;
 using System.Collections.Generic;
 using EpicGames.Core;
+using EpicGames.Horde;
 using EpicGames.Horde.Jobs;
+using EpicGames.Horde.Jobs.Templates;
+using Google.Protobuf.WellKnownTypes;
 using MongoDB.Bson.Serialization.Attributes;
 
 namespace Horde.Server.Jobs.Templates
@@ -15,6 +18,14 @@ namespace Horde.Server.Jobs.Templates
 	[BsonKnownTypes(typeof(GroupParameter), typeof(TextParameter), typeof(ListParameter), typeof(BoolParameter))]
 	public abstract class Parameter
 	{
+		/// <summary>
+		/// Gets the arguments for a job given a set of parameters
+		/// </summary>
+		/// <param name="parameters">Map of parameter id to value</param>
+		/// <param name="scheduledBuild">Whether this is a scheduled build</param>
+		/// <param name="arguments">Receives command line arguments for the job</param>
+		public abstract void GetArguments(Dictionary<ParameterId, string> parameters, bool scheduledBuild, List<string> arguments);
+
 		/// <summary>
 		/// Gets the default arguments for this parameter and its children
 		/// </summary>
@@ -71,6 +82,15 @@ namespace Horde.Server.Jobs.Templates
 			Children = children;
 		}
 
+		/// <inheritdoc/>
+		public override void GetArguments(Dictionary<ParameterId, string> parameters, bool scheduledBuild, List<string> arguments)
+		{
+			foreach (Parameter child in Children)
+			{
+				child.GetArguments(parameters, scheduledBuild, arguments);
+			}
+		}
+
 		/// <inheritdoc/>
 		public override void GetDefaultArguments(List<string> defaultArguments, bool scheduledBuild)
 		{
@@ -95,6 +115,11 @@ namespace Horde.Server.Jobs.Templates
 	/// </summary>
 	public sealed class TextParameter : Parameter
 	{
+		/// <summary>
+		/// Identifier for this parameter
+		/// </summary>
+		public ParameterId Id { get; set; }
+
 		/// <summary>
 		/// Label to display next to this parameter. Should default to the parameter name.
 		/// </summary>
@@ -150,6 +175,7 @@ namespace Horde.Server.Jobs.Templates
 		/// <summary>
 		/// Constructor
 		/// </summary>
+		/// <param name="id">Identifier for this parameter</param>
 		/// <param name="label">Parameter to pass this value to the BuildGraph script with</param>
 		/// <param name="argument">Label to show next to the parameter</param>
 		/// <param name="defaultValue">Default value for this argument</param>
@@ -158,8 +184,9 @@ namespace Horde.Server.Jobs.Templates
 		/// <param name="validation">Regex used to validate entries</param>
 		/// <param name="description">Message displayed for invalid values</param>
 		/// <param name="toolTip">Tool tip text to display</param>
-		public TextParameter(string label, string argument, string defaultValue, string? scheduleOverride, string? hint, string? validation, string? description, string? toolTip)
+		public TextParameter(ParameterId id, string label, string argument, string defaultValue, string? scheduleOverride, string? hint, string? validation, string? description, string? toolTip)
 		{
+			Id = id;
 			Label = label;
 			Argument = argument;
 			Default = defaultValue;
@@ -170,6 +197,19 @@ namespace Horde.Server.Jobs.Templates
 			ToolTip = toolTip;
 		}
 
+		/// <inheritdoc/>
+		public override void GetArguments(Dictionary<ParameterId, string> parameters, bool scheduledBuild, List<string> arguments)
+		{
+			if (parameters.TryGetValue(Id, out string? value))
+			{
+				arguments.Add(Argument + value);
+			}
+			else
+			{
+				arguments.Add(Argument + (scheduledBuild ? (ScheduleOverride ?? Default) : Default));
+			}
+		}
+
 		/// <inheritdoc/>
 		public override void GetDefaultArguments(List<string> defaultArguments, bool scheduledBuild)
 		{
@@ -182,7 +222,7 @@ namespace Horde.Server.Jobs.Templates
 		/// <returns>Serializable parameter data</returns>
 		public override ParameterData ToData()
 		{
-			return new TextParameterData(Label, Argument, Default, ScheduleOverride, Hint, Validation, Description, ToolTip);
+			return new TextParameterData(Id, Label, Argument, Default, ScheduleOverride, Hint, Validation, Description, ToolTip);
 		}
 	}
 
@@ -191,6 +231,11 @@ namespace Horde.Server.Jobs.Templates
 	/// </summary>
 	public class ListParameterItem
 	{
+		/// <summary>
+		/// Identifier for this parameter
+		/// </summary>
+		public ParameterId Id { get; set; }
+
 		/// <summary>
 		/// Group to display this entry in
 		/// </summary>
@@ -248,6 +293,7 @@ namespace Horde.Server.Jobs.Templates
 		/// <summary>
 		/// Constructor
 		/// </summary>
+		/// <param name="id">Identifier for this parameter</param>
 		/// <param name="group">The group to put this parameter in</param>
 		/// <param name="text">Text to display for this option</param>
 		/// <param name="argumentIfEnabled">Argument to add if this option is enabled</param>
@@ -256,8 +302,9 @@ namespace Horde.Server.Jobs.Templates
 		/// <param name="argumentsIfDisabled">Arguments to add if this option is disabled </param>
 		/// <param name="defaultValue">Whether this item is selected by default</param>
 		/// <param name="scheduleOverride">Override for this value in scheduled builds</param>
-		public ListParameterItem(string? group, string text, string? argumentIfEnabled, List<string>? argumentsIfEnabled, string? argumentIfDisabled, List<string>? argumentsIfDisabled, bool defaultValue, bool? scheduleOverride)
+		public ListParameterItem(ParameterId id, string? group, string text, string? argumentIfEnabled, List<string>? argumentsIfEnabled, string? argumentIfDisabled, List<string>? argumentsIfDisabled, bool defaultValue, bool? scheduleOverride)
 		{
+			Id = id;
 			Group = group;
 			Text = text;
 			ArgumentIfEnabled = argumentIfEnabled;
@@ -274,7 +321,7 @@ namespace Horde.Server.Jobs.Templates
 		/// <returns>Serializable parameter data</returns>
 		public ListParameterItemData ToData()
 		{
-			return new ListParameterItemData(Group, Text, ArgumentIfEnabled, ArgumentsIfEnabled, ArgumentIfDisabled, ArgumentsIfDisabled, Default, ScheduleOverride);
+			return new ListParameterItemData(Id, Group, Text, ArgumentIfEnabled, ArgumentsIfEnabled, ArgumentIfDisabled, ArgumentsIfDisabled, Default, ScheduleOverride);
 		}
 	}
 
@@ -327,33 +374,48 @@ namespace Horde.Server.Jobs.Templates
 			ToolTip = toolTip;
 		}
 
+		/// <inheritdoc/>
+		public override void GetArguments(Dictionary<ParameterId, string> parameters, bool scheduledBuild, List<string> arguments)
+		{
+			foreach (ListParameterItem item in Items)
+			{
+				bool value = BoolParameter.GetValue(item.Id, parameters) ?? (scheduledBuild ? (item.ScheduleOverride ?? item.Default) : item.Default);
+				GetCommandLineArgumentsForItem(item, value, arguments);
+			}
+		}
+
 		/// <inheritdoc/>
 		public override void GetDefaultArguments(List<string> defaultArguments, bool scheduledBuild)
 		{
 			foreach (ListParameterItem item in Items)
 			{
 				bool value = scheduledBuild ? (item.ScheduleOverride ?? item.Default) : item.Default;
-				if (value)
+				GetCommandLineArgumentsForItem(item, value, defaultArguments);
+			}
+		}
+
+		static void GetCommandLineArgumentsForItem(ListParameterItem item, bool value, List<string> arguments)
+		{
+			if (value)
+			{
+				if (item.ArgumentIfEnabled != null)
 				{
-					if (item.ArgumentIfEnabled != null)
-					{
-						defaultArguments.Add(item.ArgumentIfEnabled);
-					}
-					if (item.ArgumentsIfEnabled != null)
-					{
-						defaultArguments.AddRange(item.ArgumentsIfEnabled);
-					}
+					arguments.Add(item.ArgumentIfEnabled);
 				}
-				else
+				if (item.ArgumentsIfEnabled != null)
 				{
-					if (item.ArgumentIfDisabled != null)
-					{
-						defaultArguments.Add(item.ArgumentIfDisabled);
-					}
-					if (item.ArgumentsIfDisabled != null)
-					{
-						defaultArguments.AddRange(item.ArgumentsIfDisabled);
-					}
+					arguments.AddRange(item.ArgumentsIfEnabled);
+				}
+			}
+			else
+			{
+				if (item.ArgumentIfDisabled != null)
+				{
+					arguments.Add(item.ArgumentIfDisabled);
+				}
+				if (item.ArgumentsIfDisabled != null)
+				{
+					arguments.AddRange(item.ArgumentsIfDisabled);
 				}
 			}
 		}
@@ -373,6 +435,11 @@ namespace Horde.Server.Jobs.Templates
 	/// </summary>
 	public class BoolParameter : Parameter
 	{
+		/// <summary>
+		/// Identifier for this parameter
+		/// </summary>
+		public ParameterId Id { get; set; }
+
 		/// <summary>
 		/// Label to display next to this parameter.
 		/// </summary>
@@ -430,6 +497,7 @@ namespace Horde.Server.Jobs.Templates
 		/// <summary>
 		/// Constructor
 		/// </summary>
+		/// <param name="id">Identifier for this parameter</param>
 		/// <param name="label">Label to display for this parameter</param>
 		/// <param name="argumentIfEnabled">Argument to add if this parameter is enabled</param>
 		/// <param name="argumentsIfEnabled">Arguments to add if this parameter is enabled</param>
@@ -438,8 +506,9 @@ namespace Horde.Server.Jobs.Templates
 		/// <param name="defaultValue">Default value for this argument</param>
 		/// <param name="scheduleOverride">Override for this argument in scheduled builds</param>
 		/// <param name="toolTip">Tool tip text to display</param>
-		public BoolParameter(string label, string? argumentIfEnabled, List<string>? argumentsIfEnabled, string? argumentIfDisabled, List<string>? argumentsIfDisabled, bool defaultValue, bool? scheduleOverride, string? toolTip)
+		public BoolParameter(ParameterId id, string label, string? argumentIfEnabled, List<string>? argumentsIfEnabled, string? argumentIfDisabled, List<string>? argumentsIfDisabled, bool defaultValue, bool? scheduleOverride, string? toolTip)
 		{
+			Id = id;
 			Label = label;
 			ArgumentIfEnabled = argumentIfEnabled;
 			ArgumentsIfEnabled = argumentsIfEnabled;
@@ -450,30 +519,63 @@ namespace Horde.Server.Jobs.Templates
 			ToolTip = toolTip;
 		}
 
+		/// <summary>
+		/// Parses a bool parameter value from a dictionary
+		/// </summary>
+		/// <param name="parameterId">The parameter id</param>
+		/// <param name="parameters">Map from parameter id to value</param>
+		public static bool? GetValue(ParameterId parameterId, Dictionary<ParameterId, string> parameters)
+		{
+			if (parameters.TryGetValue(parameterId, out string? stringValue))
+			{
+				if (stringValue.Equals("0", StringComparison.Ordinal) || stringValue.Equals("false", StringComparison.OrdinalIgnoreCase))
+				{
+					return false;
+				}
+				if (stringValue.Equals("1", StringComparison.Ordinal) || stringValue.Equals("true", StringComparison.OrdinalIgnoreCase))
+				{
+					return true;
+				}
+			}
+			return null;
+		}
+
+		/// <inheritdoc/>
+		public override void GetArguments(Dictionary<ParameterId, string> parameters, bool scheduledBuild, List<string> arguments)
+		{
+			bool value = GetValue(Id, parameters) ?? (scheduledBuild ? (ScheduleOverride ?? Default) : Default);
+			GetCommandLineArgumentsForItem(value, arguments);
+		}
+
 		/// <inheritdoc/>
 		public override void GetDefaultArguments(List<string> defaultArguments, bool scheduledBuild)
 		{
 			bool value = scheduledBuild ? (ScheduleOverride ?? Default) : Default;
+			GetCommandLineArgumentsForItem(value, defaultArguments);
+		}
+
+		void GetCommandLineArgumentsForItem(bool value, List<string> arguments)
+		{
 			if (value)
 			{
 				if (!String.IsNullOrEmpty(ArgumentIfEnabled))
 				{
-					defaultArguments.Add(ArgumentIfEnabled);
+					arguments.Add(ArgumentIfEnabled);
 				}
 				if (ArgumentsIfEnabled != null)
 				{
-					defaultArguments.AddRange(ArgumentsIfEnabled);
+					arguments.AddRange(ArgumentsIfEnabled);
 				}
 			}
 			else
 			{
 				if (!String.IsNullOrEmpty(ArgumentIfDisabled))
 				{
-					defaultArguments.Add(ArgumentIfDisabled);
+					arguments.Add(ArgumentIfDisabled);
 				}
 				if (ArgumentsIfDisabled != null)
 				{
-					defaultArguments.AddRange(ArgumentsIfDisabled);
+					arguments.AddRange(ArgumentsIfDisabled);
 				}
 			}
 		}
@@ -483,7 +585,7 @@ namespace Horde.Server.Jobs.Templates
 		/// <returns><see cref="BoolParameterData"/> instance</returns>
 		public override ParameterData ToData()
 		{
-			return new BoolParameterData(Label, ArgumentIfEnabled, ArgumentsIfEnabled, ArgumentIfDisabled, ArgumentsIfDisabled, Default, ScheduleOverride, ToolTip);
+			return new BoolParameterData(Id, Label, ArgumentIfEnabled, ArgumentsIfEnabled, ArgumentIfDisabled, ArgumentsIfDisabled, Default, ScheduleOverride, ToolTip);
 		}
 	}
 
diff --git a/Engine/Source/Programs/Horde/Horde.Server/Jobs/Templates/TemplateConfig.cs b/Engine/Source/Programs/Horde/Horde.Server/Jobs/Templates/TemplateConfig.cs
--- a/Engine/Source/Programs/Horde/Horde.Server/Jobs/Templates/TemplateConfig.cs
+++ b/Engine/Source/Programs/Horde/Horde.Server/Jobs/Templates/TemplateConfig.cs
@@ -7,6 +7,7 @@ using System.IO;
 using System.Text.Json.Serialization;
 using EpicGames.Core;
 using EpicGames.Horde.Jobs;
+using EpicGames.Horde.Jobs.Templates;
 using Horde.Server.Configuration;
 
 namespace Horde.Server.Jobs.Templates
@@ -102,13 +103,58 @@ namespace Horde.Server.Jobs.Templates
 		/// <summary>
 		/// Callback after a parameter has been read.
 		/// </summary>
-		public abstract void PostLoad();
+		public abstract void PostLoad(HashSet<ParameterId> parameterIds);
 
 		/// <summary>
 		/// Convert to a parameter object
 		/// </summary>
 		/// <returns><see cref="Parameter"/> object</returns>
 		public abstract Parameter ToModel();
+
+		/// <summary>
+		/// Adds the given parameter id to the set of known ids, or create a new one
+		/// </summary>
+		/// <param name="explicitParameterId"></param>
+		/// <param name="name"></param>
+		/// <param name="parameterIds"></param>
+		/// <returns></returns>
+		protected ParameterId AddParameterId(ParameterId explicitParameterId, string name, HashSet<ParameterId> parameterIds)
+		{
+			if (explicitParameterId.Id.IsEmpty)
+			{
+				return CreateUniqueParameterId(name, parameterIds);
+			}
+			else if (parameterIds.Add(explicitParameterId))
+			{
+				return explicitParameterId;
+			}
+			else
+			{
+				throw new InvalidOperationException($"Parameter '{explicitParameterId}' has been defined twice");
+			}
+		}
+
+		/// <summary>
+		/// Creates a unique parameter id based on the given name
+		/// </summary>
+		/// <param name="name">Default name for the parameter</param>
+		/// <param name="parameterIds">The existing parameter ids</param>
+		protected ParameterId CreateUniqueParameterId(string name, HashSet<ParameterId> parameterIds)
+		{
+			string baseName = name;
+			if (String.IsNullOrEmpty(baseName))
+			{
+				baseName = "unnamed";
+			}
+
+			ParameterId parameterId = ParameterId.Sanitize(baseName);
+			for (int idx = 2; !parameterIds.Add(parameterId); idx++)
+			{
+				parameterId = ParameterId.Sanitize($"{baseName}-{idx}");
+			}
+
+			return parameterId;
+		}
 	}
 
 	/// <summary>
@@ -171,7 +217,7 @@ namespace Horde.Server.Jobs.Templates
 		}
 
 		/// <inheritdoc/>
-		public override void PostLoad()
+		public override void PostLoad(HashSet<ParameterId> parameterIds)
 		{
 		}
 
@@ -188,6 +234,11 @@ namespace Horde.Server.Jobs.Templates
 	[JsonDiscriminator("Text")]
 	public class TextParameterData : ParameterData
 	{
+		/// <summary>
+		/// Identifier for this parameter
+		/// </summary>
+		public ParameterId Id { get; set; }
+
 		/// <summary>
 		/// Name of the parameter associated with this parameter.
 		/// </summary>
@@ -241,6 +292,7 @@ namespace Horde.Server.Jobs.Templates
 		/// <summary>
 		/// Constructor
 		/// </summary>
+		/// <param name="id">Identifier for this parameter</param>
 		/// <param name="label">Label to show next to the parameter</param>
 		/// <param name="argument">Argument to pass this value with</param>
 		/// <param name="defaultValue">Default value for this parameter</param>
@@ -249,8 +301,9 @@ namespace Horde.Server.Jobs.Templates
 		/// <param name="validation">Regex used to validate entries</param>
 		/// <param name="validationError">Message displayed to explain validation issues</param>
 		/// <param name="toolTip">Tool tip text to display</param>
-		public TextParameterData(string label, string argument, string defaultValue, string? scheduleOverride, string? hint, string? validation, string? validationError, string? toolTip)
+		public TextParameterData(ParameterId id, string label, string argument, string defaultValue, string? scheduleOverride, string? hint, string? validation, string? validationError, string? toolTip)
 		{
+			Id = id;
 			Label = label;
 			Argument = argument;
 			Default = defaultValue;
@@ -262,14 +315,15 @@ namespace Horde.Server.Jobs.Templates
 		}
 
 		/// <inheritdoc/>
-		public override void PostLoad()
+		public override void PostLoad(HashSet<ParameterId> parameterIds)
 		{
+			Id = AddParameterId(Id, Label, parameterIds);
 		}
 
 		/// <inheritdoc/>
 		public override Parameter ToModel()
 		{
-			return new TextParameter(Label, Argument, Default, ScheduleOverride, Hint, Validation, ValidationError, ToolTip);
+			return new TextParameter(Id, Label, Argument, Default, ScheduleOverride, Hint, Validation, ValidationError, ToolTip);
 		}
 	}
 
@@ -299,6 +353,11 @@ namespace Horde.Server.Jobs.Templates
 	/// </summary>
 	public class ListParameterItemData
 	{
+		/// <summary>
+		/// Identifier for this parameter
+		/// </summary>
+		public ParameterId Id { get; set; }
+
 		/// <summary>
 		/// Optional group heading to display this entry under, if the picker style supports it.
 		/// </summary>
@@ -350,6 +409,7 @@ namespace Horde.Server.Jobs.Templates
 		/// <summary>
 		/// Constructor
 		/// </summary>
+		/// <param name="id">Identifier for this parameter</param>
 		/// <param name="group">The group to put this parameter in</param>
 		/// <param name="text">Text to display for this option</param>
 		/// <param name="argumentIfEnabled">Argument to pass for this item if it's enabled</param>
@@ -358,8 +418,9 @@ namespace Horde.Server.Jobs.Templates
 		/// <param name="argumentsIfDisabled">Argument to pass for this item if it's enabled</param>
 		/// <param name="defaultValue">Whether this item is selected by default</param>
 		/// <param name="scheduleOverride">Overridden value for this item for scheduled builds</param>
-		public ListParameterItemData(string? group, string text, string? argumentIfEnabled, List<string>? argumentsIfEnabled, string? argumentIfDisabled, List<string>? argumentsIfDisabled, bool defaultValue, bool? scheduleOverride)
+		public ListParameterItemData(ParameterId id, string? group, string text, string? argumentIfEnabled, List<string>? argumentsIfEnabled, string? argumentIfDisabled, List<string>? argumentsIfDisabled, bool defaultValue, bool? scheduleOverride)
 		{
+			Id = id;
 			Group = group;
 			Text = text;
 			ArgumentIfEnabled = argumentIfEnabled;
@@ -376,7 +437,7 @@ namespace Horde.Server.Jobs.Templates
 		/// <returns>New <see cref="ListParameterItem"/> object</returns>
 		public ListParameterItem ToModel()
 		{
-			return new ListParameterItem(Group, Text, ArgumentIfEnabled, ArgumentsIfEnabled, ArgumentIfDisabled, ArgumentsIfDisabled, Default, ScheduleOverride);
+			return new ListParameterItem(Id, Group, Text, ArgumentIfEnabled, ArgumentsIfEnabled, ArgumentIfDisabled, ArgumentsIfDisabled, Default, ScheduleOverride);
 		}
 	}
 
@@ -431,10 +492,12 @@ namespace Horde.Server.Jobs.Templates
 		}
 
 		/// <inheritdoc/>
-		public override void PostLoad()
+		public override void PostLoad(HashSet<ParameterId> parameterIds)
 		{
 			foreach (ListParameterItemData item in Items)
 			{
+				item.Id = AddParameterId(item.Id, $"{Label}.{item.Text}", parameterIds);
+
 				if (!String.IsNullOrEmpty(item.ArgumentIfEnabled) && item.ArgumentsIfEnabled != null && item.ArgumentsIfEnabled.Count > 0)
 				{
 					throw new InvalidDataException("Cannot specify both 'ArgumentIfEnabled' and 'ArgumentsIfEnabled'");
@@ -459,6 +522,11 @@ namespace Horde.Server.Jobs.Templates
 	[JsonDiscriminator("Bool")]
 	public class BoolParameterData : ParameterData
 	{
+		/// <summary>
+		/// Identifier for this parameter
+		/// </summary>
+		public ParameterId Id { get; set; }
+
 		/// <summary>
 		/// Name of the parameter associated with this parameter.
 		/// </summary>
@@ -510,6 +578,7 @@ namespace Horde.Server.Jobs.Templates
 		/// <summary>
 		/// Constructor
 		/// </summary>
+		/// <param name="id">Identifier for this parameter</param>
 		/// <param name="label">Label to show next to this parameter</param>
 		/// <param name="argumentIfEnabled">Argument to add if this parameter is enabled</param>
 		/// <param name="argumentsIfEnabled">Arguments to add if this parameter is enabled</param>
@@ -518,8 +587,9 @@ namespace Horde.Server.Jobs.Templates
 		/// <param name="defaultValue">Whether this option is enabled by default</param>
 		/// <param name="scheduleOverride">Override for scheduled builds</param>
 		/// <param name="toolTip">The tool tip text to display</param>
-		public BoolParameterData(string label, string? argumentIfEnabled, List<string>? argumentsIfEnabled, string? argumentIfDisabled, List<string>? argumentsIfDisabled, bool defaultValue, bool? scheduleOverride, string? toolTip)
+		public BoolParameterData(ParameterId id, string label, string? argumentIfEnabled, List<string>? argumentsIfEnabled, string? argumentIfDisabled, List<string>? argumentsIfDisabled, bool defaultValue, bool? scheduleOverride, string? toolTip)
 		{
+			Id = id;
 			Label = label;
 			ArgumentIfEnabled = argumentIfEnabled;
 			ArgumentsIfEnabled = argumentsIfEnabled;
@@ -531,8 +601,10 @@ namespace Horde.Server.Jobs.Templates
 		}
 
 		/// <inheritdoc/>
-		public override void PostLoad()
+		public override void PostLoad(HashSet<ParameterId> parameterIds)
 		{
+			Id = AddParameterId(Id, Label, parameterIds);
+
 			if (!String.IsNullOrEmpty(ArgumentIfEnabled) && ArgumentsIfEnabled != null && ArgumentsIfEnabled.Count > 0)
 			{
 				throw new InvalidDataException("Cannot specify both 'ArgumentIfEnabled' and 'ArgumentsIfEnabled'");
@@ -546,7 +618,7 @@ namespace Horde.Server.Jobs.Templates
 		/// <inheritdoc/>
 		public override Parameter ToModel()
 		{
-			return new BoolParameter(Label, ArgumentIfEnabled, ArgumentsIfEnabled, ArgumentIfDisabled, ArgumentsIfDisabled, Default, ScheduleOverride, ToolTip);
+			return new BoolParameter(Id, Label, ArgumentIfEnabled, ArgumentsIfEnabled, ArgumentIfDisabled, ArgumentsIfDisabled, Default, ScheduleOverride, ToolTip);
 		}
 	}
 }
\ No newline at end of file
diff --git a/Engine/Source/Programs/Horde/Horde.Server/Streams/StreamConfig.cs b/Engine/Source/Programs/Horde/Horde.Server/Streams/StreamConfig.cs
--- a/Engine/Source/Programs/Horde/Horde.Server/Streams/StreamConfig.cs
+++ b/Engine/Source/Programs/Horde/Horde.Server/Streams/StreamConfig.cs
@@ -280,9 +280,10 @@ namespace Horde.Server.Streams
 
 			foreach (TemplateRefConfig template in Templates)
 			{
+				HashSet<ParameterId> parameterIds = new HashSet<ParameterId>();
 				foreach (ParameterData parameter in template.Parameters)
 				{
-					parameter.PostLoad();
+					parameter.PostLoad(parameterIds);
 				}
 			}
 
diff --git a/Engine/Source/Programs/Shared/EpicGames.Horde/Jobs/JobMessages.cs b/Engine/Source/Programs/Shared/EpicGames.Horde/Jobs/JobMessages.cs
--- a/Engine/Source/Programs/Shared/EpicGames.Horde/Jobs/JobMessages.cs
+++ b/Engine/Source/Programs/Shared/EpicGames.Horde/Jobs/JobMessages.cs
@@ -213,6 +213,11 @@ namespace EpicGames.Horde.Jobs
 		/// </summary>
 		public bool? UpdateIssues { get; set; }
 
+		/// <summary>
+		/// Values for the template parameters
+		/// </summary>
+		public Dictionary<ParameterId, string>? Parameters { get; set; }
+
 		/// <summary>
 		/// Arguments for the job
 		/// </summary>
@@ -478,6 +483,11 @@ namespace EpicGames.Horde.Jobs
 		/// <summary>
 		/// Parameters for the job
 		/// </summary>
+		public Dictionary<ParameterId, string> Parameters { get; set; } = new Dictionary<ParameterId, string>();
+
+		/// <summary>
+		/// Command line arguments for the job
+		/// </summary>
 		public List<string> Arguments { get; set; } = new List<string>();
 
 		/// <summary>
diff --git a/Engine/Source/Programs/Shared/EpicGames.Horde/Jobs/Templates/ParameterId.cs b/Engine/Source/Programs/Shared/EpicGames.Horde/Jobs/Templates/ParameterId.cs
new file mode 100644
--- /dev/null
+++ b/Engine/Source/Programs/Shared/EpicGames.Horde/Jobs/Templates/ParameterId.cs
@@ -0,0 +1,46 @@
+// Copyright Epic Games, Inc. All Rights Reserved.
+
+using System;
+using System.Collections.Generic;
+using System.Linq;
+using System.Text;
+using System.Threading.Tasks;
+using EpicGames.Core;
+
+namespace EpicGames.Horde.Jobs.Templates
+{
+	/// <summary>
+	/// Identifier for a template parameter
+	/// </summary>
+	/// <param name="Id">Id to construct from</param>
+	[JsonSchemaString]
+	[StringIdConverter(typeof(ParameterIdConverter))]
+	public record struct ParameterId(StringId Id)
+	{
+		/// <summary>
+		/// Constructor
+		/// </summary>
+		public ParameterId(string id) : this(new StringId(id))
+		{
+		}
+
+		/// <inheritdoc cref="StringId.Sanitize(string)"/>
+		public static ParameterId Sanitize(string name)
+			=> new ParameterId(StringId.Sanitize(name));
+
+		/// <inheritdoc/>
+		public override string ToString() => Id.ToString();
+	}
+
+	/// <summary>
+	/// Converter to and from <see cref="StringId"/> instances.
+	/// </summary>
+	class ParameterIdConverter : StringIdConverter<ParameterId>
+	{
+		/// <inheritdoc/>
+		public override ParameterId FromStringId(StringId id) => new ParameterId(id);
+
+		/// <inheritdoc/>
+		public override StringId ToStringId(ParameterId value) => value.Id;
+	}
+}
diff --git a/Engine/Source/Programs/Shared/EpicGames.Horde/StringId.cs b/Engine/Source/Programs/Shared/EpicGames.Horde/StringId.cs
--- a/Engine/Source/Programs/Shared/EpicGames.Horde/StringId.cs
+++ b/Engine/Source/Programs/Shared/EpicGames.Horde/StringId.cs
@@ -99,7 +99,7 @@ namespace EpicGames.Horde
 				{
 					result.Append((char)('a' + (character - 'A')));
 				}
-				else if (IsValidCharacter(character))
+				else if (IsValidCharacter(character) && (result.Length > 0 || character != '.'))
 				{
 					result.Append(character);
 				}
@@ -108,7 +108,7 @@ namespace EpicGames.Horde
 					result.Append('-');
 				}
 			}
-			while (result.Length > 0 && result[^1] == '-')
+			while (result.Length > 0 && (result[^1] == '-' || result[^1] == '.'))
 			{
 				result.Remove(result.Length - 1, 1);
 			}
diff --git a/Engine/Source/Programs/Shared/EpicGames.Horde/StringIdConverter.cs b/Engine/Source/Programs/Shared/EpicGames.Horde/StringIdConverter.cs
--- a/Engine/Source/Programs/Shared/EpicGames.Horde/StringIdConverter.cs
+++ b/Engine/Source/Programs/Shared/EpicGames.Horde/StringIdConverter.cs
@@ -124,10 +124,20 @@ namespace EpicGames.Horde
 		readonly TConverter _converter = new TConverter();
 
 		/// <inheritdoc/>
-		public override TValue Read(ref Utf8JsonReader reader, Type typeToConvert, JsonSerializerOptions options) => _converter.FromStringId(new StringId(new Utf8String(reader.GetUtf8String().ToArray())));
+		public override TValue Read(ref Utf8JsonReader reader, Type typeToConvert, JsonSerializerOptions options)
+			=> _converter.FromStringId(new StringId(new Utf8String(reader.GetUtf8String().ToArray())));
 
 		/// <inheritdoc/>
-		public override void Write(Utf8JsonWriter writer, TValue value, JsonSerializerOptions options) => writer.WriteStringValue(_converter.ToStringId(value).Span);
+		public override void Write(Utf8JsonWriter writer, TValue value, JsonSerializerOptions options) 
+			=> writer.WriteStringValue(_converter.ToStringId(value).Span);
+
+		/// <inheritdoc/>
+		public override TValue ReadAsPropertyName(ref Utf8JsonReader reader, Type typeToConvert, JsonSerializerOptions options)
+			=> Read(ref reader, typeToConvert, options);
+
+		/// <inheritdoc/>
+		public override void WriteAsPropertyName(Utf8JsonWriter writer, TValue value, JsonSerializerOptions options)
+			=> writer.WritePropertyName(_converter.ToStringId(value).Span);
 	}
 
 	/// <summary>
