Horde: Ensure blobs uploaded to the tools endpoint have the tool id as a prefix to the blob locator.

diff --git a/Engine/Source/Programs/Horde/Horde.Server/Tools/ToolsController.cs b/Engine/Source/Programs/Horde/Horde.Server/Tools/ToolsController.cs
--- a/Engine/Source/Programs/Horde/Horde.Server/Tools/ToolsController.cs
+++ b/Engine/Source/Programs/Horde/Horde.Server/Tools/ToolsController.cs
@@ -62,6 +62,15 @@ namespace Horde.Server.Tools
 				return Forbid(ToolAclAction.UploadTool, id);
 			}
 
+			if (String.IsNullOrEmpty(request.Prefix))
+			{
+				request.Prefix = id.ToString();
+			}
+			else
+			{
+				request.Prefix = $"{id}/{request.Prefix}";
+			}
+
 			IStorageBackend storageBackend = _toolCollection.CreateStorageBackend(tool);
 			return await StorageController.WriteBlobAsync(storageBackend, request, cancellationToken);
 		}
@@ -382,7 +391,7 @@ namespace Horde.Server.Tools
 		/// <returns>Information about all the artifacts</returns>
 		[HttpGet]
 		[Route("/api/v1/tools/{id}/blobs/{*locator}")]
-		public async Task<ActionResult<object>> ReadToolBlobAsync(ToolId id, BlobLocator locator, CancellationToken cancellationToken = default)
+		public async Task<ActionResult> ReadToolBlobAsync(ToolId id, BlobLocator locator, CancellationToken cancellationToken = default)
 		{
 			ITool? tool = await _toolCollection.GetAsync(id, _globalConfig.Value, cancellationToken);
 			if (tool == null)
@@ -400,7 +409,7 @@ namespace Horde.Server.Tools
 			}
 
 			IStorageBackend storageBackend = _toolCollection.CreateStorageBackend(tool);
-			return StorageController.ReadBlobInternalAsync(storageBackend, locator, Request.Headers, cancellationToken);
+			return await StorageController.ReadBlobInternalAsync(storageBackend, locator, Request.Headers, cancellationToken);
 		}
 
 		bool AuthorizeDownload(ToolConfig toolConfig)
