Horde: Fix some static analysis warnings.

#rnx

diff --git a/Engine/Source/Programs/Horde/Horde.Server/Jobs/Templates/ITemplate.cs b/Engine/Source/Programs/Horde/Horde.Server/Jobs/Templates/ITemplate.cs
--- a/Engine/Source/Programs/Horde/Horde.Server/Jobs/Templates/ITemplate.cs
+++ b/Engine/Source/Programs/Horde/Horde.Server/Jobs/Templates/ITemplate.cs
@@ -3,10 +3,8 @@
 using System;
 using System.Collections.Generic;
 using EpicGames.Core;
-using EpicGames.Horde;
 using EpicGames.Horde.Jobs;
 using EpicGames.Horde.Jobs.Templates;
-using Google.Protobuf.WellKnownTypes;
 using MongoDB.Bson.Serialization.Attributes;
 
 namespace Horde.Server.Jobs.Templates
diff --git a/Engine/Source/Programs/Horde/Horde.Server/Notifications/Sinks/SlackNotificationSink.cs b/Engine/Source/Programs/Horde/Horde.Server/Notifications/Sinks/SlackNotificationSink.cs
--- a/Engine/Source/Programs/Horde/Horde.Server/Notifications/Sinks/SlackNotificationSink.cs
+++ b/Engine/Source/Programs/Horde/Horde.Server/Notifications/Sinks/SlackNotificationSink.cs
@@ -2164,7 +2164,7 @@ namespace Horde.Server.Notifications.Sinks
 					{
 						string path = uri.LocalPath;
 
-						int atIdx = path.IndexOf('@');
+						int atIdx = path.IndexOf('@', StringComparison.Ordinal);
 						if (atIdx != -1 && Int32.TryParse(path.Substring(atIdx + 1), out int change))
 						{
 							int existingChange;
diff --git a/Engine/Source/Programs/Shared/EpicGames.Horde/Jobs/Templates/ParameterId.cs b/Engine/Source/Programs/Shared/EpicGames.Horde/Jobs/Templates/ParameterId.cs
--- a/Engine/Source/Programs/Shared/EpicGames.Horde/Jobs/Templates/ParameterId.cs
+++ b/Engine/Source/Programs/Shared/EpicGames.Horde/Jobs/Templates/ParameterId.cs
@@ -1,10 +1,5 @@
 // Copyright Epic Games, Inc. All Rights Reserved.
 
-using System;
-using System.Collections.Generic;
-using System.Linq;
-using System.Text;
-using System.Threading.Tasks;
 using EpicGames.Core;
 
 namespace EpicGames.Horde.Jobs.Templates
