Horde: Updating schema documentation

#rnx

diff --git a/Engine/Source/Programs/Horde/Docs/Config/Schema/Streams.md b/Engine/Source/Programs/Horde/Docs/Config/Schema/Streams.md
--- a/Engine/Source/Programs/Horde/Docs/Config/Schema/Streams.md
+++ b/Engine/Source/Programs/Horde/Docs/Config/Schema/Streams.md
@@ -390,6 +390,7 @@ Free-form text entry parameter
 Name | Description
 ---- | -----------
 `type` | Text<br>Type discriminator
+`id` | `string`<br>Identifier for this parameter
 `label` | `string`<br>Name of the parameter associated with this parameter.
 `argument` | `string`<br>Argument to pass to the executor
 `default` | `string`<br>Default value for this argument
@@ -427,6 +428,7 @@ Possible option for a list parameter
 
 Name | Description
 ---- | -----------
+`id` | `string`<br>Identifier for this parameter
 `group` | `string`<br>Optional group heading to display this entry under, if the picker style supports it.
 `text` | `string`<br>Name of the parameter associated with this list.
 `argumentIfEnabled` | `string`<br>Argument to pass with this parameter.
@@ -443,6 +445,7 @@ Allows the user to toggle an option on or off
 Name | Description
 ---- | -----------
 `type` | Bool<br>Type discriminator
+`id` | `string`<br>Identifier for this parameter
 `label` | `string`<br>Name of the parameter associated with this parameter.
 `argumentIfEnabled` | `string`<br>Argument to add if this parameter is enabled
 `argumentsIfEnabled` | `string[]`<br>Argument to add if this parameter is enabled
diff --git a/Engine/Source/Programs/Horde/Docs/Deployment/AgentSettings.md b/Engine/Source/Programs/Horde/Docs/Deployment/AgentSettings.md
--- a/Engine/Source/Programs/Horde/Docs/Deployment/AgentSettings.md
+++ b/Engine/Source/Programs/Horde/Docs/Deployment/AgentSettings.md
@@ -14,7 +14,8 @@ Name | Description
 `executor` | `string`<br>The executor to use for jobs
 `localExecutor` | [`LocalExecutorSettings`](#localexecutorsettings)<br>Settings for the local executor
 `perforceExecutor` | [`PerforceExecutorSettings`](#perforceexecutorsettings)<br>Settings for the perforce executor
-`workingDir` | [`DirectoryReference`](#directoryreference)<br>Working directory
+`workingDir` | [`DirectoryReference`](#directoryreference)<br>Working directory for leases and jobs (i.e where files from Perforce will be checked out)
+`logsDir` | [`DirectoryReference`](#directoryreference)<br>Directory where agent and lease logs are written
 `shareMountingEnabled` | `boolean`<br>Whether to mount the specified list of network shares
 `shares` | [`MountNetworkShare`](#mountnetworkshare)`[]`<br>List of network shares to mount
 `processNamesToTerminate` | `string[]`<br>List of process names to terminate after a job
