Horde: Add extra trace attributes to workspace materializer

diff --git a/Engine/Source/Programs/Horde/Horde.Agent/Execution/ManagedWorkspaceMaterializer.cs b/Engine/Source/Programs/Horde/Horde.Agent/Execution/ManagedWorkspaceMaterializer.cs
--- a/Engine/Source/Programs/Horde/Horde.Agent/Execution/ManagedWorkspaceMaterializer.cs
+++ b/Engine/Source/Programs/Horde/Horde.Agent/Execution/ManagedWorkspaceMaterializer.cs
@@ -143,6 +143,8 @@ public sealed class ManagedWorkspaceMaterializer : IWorkspaceMaterializer
 		scope.Span.SetTag("Cluster", _agentWorkspace.Cluster);
 		scope.Span.SetTag("Incremental", _agentWorkspace.Incremental);
 		scope.Span.SetTag("Method", _agentWorkspace.Method);
+		scope.Span.SetTag("Stream", _agentWorkspace.Stream);
+		scope.Span.SetTag("Partitioned", _agentWorkspace.Partitioned);
 		scope.Span.SetTag("UseCacheFile", _useCacheFile);
 		return scope;
 	}
