Horde: Update release notes

#rnx

diff --git a/Engine/Source/Programs/Horde/Docs/ReleaseNotes.md b/Engine/Source/Programs/Horde/Docs/ReleaseNotes.md
--- a/Engine/Source/Programs/Horde/Docs/ReleaseNotes.md
+++ b/Engine/Source/Programs/Horde/Docs/ReleaseNotes.md
@@ -1,5 +1,17 @@
 # Release Notes
 
+## 2024-04-12
+
+* Ensure blobs uploaded to the tools endpoint have the tool id as a prefix to the blob locator (32907981)
+* Allow specifying a unique id for each parameter in a job template (32906466)
+* Send Slack notifications whenever config changes are applied (32897388)
+* Make log output directory configurable in agent (32886152)
+* Use strongly typed identifiers for artifact log messages in agent. (32861971)
+* Fix memory leak and incorrect usage of native P4API library (32798937, 32794622)
+* Add debug endpoint for capturing a dotMemory snapshot (32790300)
+* Refactor API message names, project locations and reduce gRPC dependency (32785690, 32783652)
+* Allow clients to upload the list of references for each blob, removing the need for the server to do it later. (32768299)
+
 ## 2024-04-05
 
 * Handle socket shutdown errors gracefully rather than throwing exceptions. (32756972)
