Horde: Increase the max number of form values allowed, to fix issue uploading bundles with a large number of references.

diff --git a/Engine/Source/Programs/Horde/Horde.Server/Startup.cs b/Engine/Source/Programs/Horde/Horde.Server/Startup.cs
--- a/Engine/Source/Programs/Horde/Horde.Server/Startup.cs
+++ b/Engine/Source/Programs/Horde/Horde.Server/Startup.cs
@@ -101,6 +101,7 @@ using Microsoft.AspNetCore.DataProtection;
 using Microsoft.AspNetCore.Hosting;
 using Microsoft.AspNetCore.Hosting.Server.Features;
 using Microsoft.AspNetCore.Http;
+using Microsoft.AspNetCore.Http.Features;
 using Microsoft.AspNetCore.HttpOverrides;
 using Microsoft.AspNetCore.Mvc;
 using Microsoft.AspNetCore.Mvc.Controllers;
@@ -429,6 +430,12 @@ namespace Horde.Server
 			// IOptionsMonitor pattern for live updating of configuration settings
 			services.Configure<ServerSettings>(x => BindServerSettings(Configuration, x));
 
+			// We may upload a large number of references with posts to storage endpoints. Increase the max number of form values to allow this (the default is 1024).
+			services.Configure<FormOptions>(options =>
+			{
+				options.ValueCountLimit = 100_000;
+			});
+
 			// Bind the settings again for local variable access in this method
 			ServerSettings settings = new();
 			BindServerSettings(Configuration, settings);
