Horde: Fix static analysis warning.

diff --git a/Engine/Source/Programs/Horde/Horde.Server/Startup.cs b/Engine/Source/Programs/Horde/Horde.Server/Startup.cs
--- a/Engine/Source/Programs/Horde/Horde.Server/Startup.cs
+++ b/Engine/Source/Programs/Horde/Horde.Server/Startup.cs
@@ -128,6 +128,8 @@ using Swashbuckle.AspNetCore.SwaggerGen;
 using Status = Grpc.Core.Status;
 using StatusCode = Grpc.Core.StatusCode;
 
+#pragma warning disable CA1505 // 'ConfigureServices' has a maintainability index of '9'. Rewrite or refactor the code to increase its maintainability index (MI) above '9'.
+
 namespace Horde.Server
 {
 	using ContentHash = EpicGames.Core.ContentHash;
