Horde: Fix issues related to string ids used for job templates exceeding the maximum allowed length.

- Max length of a string id is now increased to 128 characters.
- Identifiers generated from template parameter names are now capped at MaxLength - 10, to allow space for a number to be added afterwards.
- Sanitization of strings will cap the maximum length of parameters at MaxLength.
- Existing string ids that are deserialized from BSON are assumed to be valid and will not throw an exception on deserialization, in order to handle any data that already exists in the DB.

diff --git a/Engine/Source/Programs/Horde/Horde.Server/Jobs/Templates/TemplateConfig.cs b/Engine/Source/Programs/Horde/Horde.Server/Jobs/Templates/TemplateConfig.cs
--- a/Engine/Source/Programs/Horde/Horde.Server/Jobs/Templates/TemplateConfig.cs
+++ b/Engine/Source/Programs/Horde/Horde.Server/Jobs/Templates/TemplateConfig.cs
@@ -6,6 +6,7 @@ using System.ComponentModel.DataAnnotations;
 using System.IO;
 using System.Text.Json.Serialization;
 using EpicGames.Core;
+using EpicGames.Horde;
 using EpicGames.Horde.Jobs;
 using EpicGames.Horde.Jobs.Templates;
 using Horde.Server.Configuration;
@@ -147,6 +148,12 @@ namespace Horde.Server.Jobs.Templates
 				baseName = "unnamed";
 			}
 
+			const int MaxBaseLength = StringId.MaxLength - 10;
+			if (baseName.Length > MaxBaseLength)
+			{
+				baseName = baseName.Substring(0, MaxBaseLength);
+			}
+
 			ParameterId parameterId = ParameterId.Sanitize(baseName);
 			for (int idx = 2; !parameterIds.Add(parameterId); idx++)
 			{
diff --git a/Engine/Source/Programs/Horde/Horde.Server/Utilities/StringIdConverter.cs b/Engine/Source/Programs/Horde/Horde.Server/Utilities/StringIdConverter.cs
--- a/Engine/Source/Programs/Horde/Horde.Server/Utilities/StringIdConverter.cs
+++ b/Engine/Source/Programs/Horde/Horde.Server/Utilities/StringIdConverter.cs
@@ -2,6 +2,7 @@
 
 using System;
 using System.Reflection;
+using EpicGames.Core;
 using EpicGames.Horde;
 using MongoDB.Bson.Serialization;
 using MongoDB.Bson.Serialization.Serializers;
@@ -17,7 +18,7 @@ namespace Horde.Server.Utilities
 		readonly TConverter _converter = new TConverter();
 
 		/// <inheritdoc/>
-		public override TValue Deserialize(BsonDeserializationContext context, BsonDeserializationArgs args) => _converter.FromStringId(new StringId(context.Reader.ReadString()));
+		public override TValue Deserialize(BsonDeserializationContext context, BsonDeserializationArgs args) => _converter.FromStringId(new StringId(new Utf8String(context.Reader.ReadString()), StringId.Validate.None));
 
 		/// <inheritdoc/>
 		public override void Serialize(BsonSerializationContext context, BsonSerializationArgs args, TValue value) => context.Writer.WriteString(_converter.ToStringId(value).ToString());
diff --git a/Engine/Source/Programs/Shared/EpicGames.Horde/StringId.cs b/Engine/Source/Programs/Shared/EpicGames.Horde/StringId.cs
--- a/Engine/Source/Programs/Shared/EpicGames.Horde/StringId.cs
+++ b/Engine/Source/Programs/Shared/EpicGames.Horde/StringId.cs
@@ -92,7 +92,7 @@ namespace EpicGames.Horde
 		public static StringId Sanitize(string text)
 		{
 			StringBuilder result = new StringBuilder();
-			for (int idx = 0; idx < text.Length; idx++)
+			for (int idx = 0; idx < text.Length && result.Length < MaxLength; idx++)
 			{
 				char character = (char)text[idx];
 				if (character >= 'A' && character <= 'Z')
