Horde: Upgrade OpenTelemetry to latest version

diff --git a/Engine/Source/Programs/Horde/Horde.Server/Horde.Server.csproj b/Engine/Source/Programs/Horde/Horde.Server/Horde.Server.csproj
--- a/Engine/Source/Programs/Horde/Horde.Server/Horde.Server.csproj
+++ b/Engine/Source/Programs/Horde/Horde.Server/Horde.Server.csproj
@@ -58,14 +58,14 @@
       <IncludeAssets>runtime; build; native; contentfiles; analyzers; buildtransitive</IncludeAssets>
     </PackageReference>
     <PackageReference Include="MongoDB.Driver" Version="2.22.0" />
-    <PackageReference Include="OpenTelemetry" Version="1.6.0" />
-    <PackageReference Include="OpenTelemetry.Exporter.Console" Version="1.6.0" />
-    <PackageReference Include="OpenTelemetry.Exporter.OpenTelemetryProtocol" Version="1.6.0" />
-    <PackageReference Include="OpenTelemetry.Extensions.Hosting" Version="1.6.0" />
-    <PackageReference Include="OpenTelemetry.Instrumentation.AspNetCore" Version="1.0.0-rc9.14" />
-    <PackageReference Include="OpenTelemetry.Instrumentation.GrpcNetClient" Version="1.0.0-rc9.14" />
-    <PackageReference Include="OpenTelemetry.Instrumentation.Http" Version="1.0.0-rc9.14" />
-    <PackageReference Include="OpenTelemetry.Instrumentation.StackExchangeRedis" Version="1.0.0-rc9.8" />
+    <PackageReference Include="OpenTelemetry" Version="1.8.0" />
+    <PackageReference Include="OpenTelemetry.Exporter.Console" Version="1.8.0" />
+    <PackageReference Include="OpenTelemetry.Exporter.OpenTelemetryProtocol" Version="1.8.0" />
+    <PackageReference Include="OpenTelemetry.Extensions.Hosting" Version="1.8.0" />
+    <PackageReference Include="OpenTelemetry.Instrumentation.AspNetCore" Version="1.8.1" />
+    <PackageReference Include="OpenTelemetry.Instrumentation.GrpcNetClient" Version="1.8.0-beta.1" />
+    <PackageReference Include="OpenTelemetry.Instrumentation.Http" Version="1.8.1" />
+    <PackageReference Include="OpenTelemetry.Instrumentation.StackExchangeRedis" Version="1.0.0-rc9.14" />
     <PackageReference Include="Polly" Version="8.2.0" />
     <PackageReference Include="Polly.Extensions.Http" Version="3.0.0" />
     <PackageReference Include="Serilog.AspNetCore" Version="8.0.0" />
diff --git a/Engine/Source/Programs/Horde/Horde.Server/Utilities/OpenTelemetry.cs b/Engine/Source/Programs/Horde/Horde.Server/Utilities/OpenTelemetry.cs
--- a/Engine/Source/Programs/Horde/Horde.Server/Utilities/OpenTelemetry.cs
+++ b/Engine/Source/Programs/Horde/Horde.Server/Utilities/OpenTelemetry.cs
@@ -141,7 +141,7 @@ public static class OpenTelemetryHelper
 		{
 			builder.AddOtlpExporter(name, exporter =>
 			{
-				exporter.Endpoint = exporterSettings.Endpoint;
+				exporter.Endpoint = exporterSettings.Endpoint!;
 				exporter.Protocol = Enum.TryParse(exporterSettings.Protocol, true, out OtlpExportProtocol protocol) ? protocol : OtlpExportProtocol.Grpc;
 			});
 		}
@@ -162,7 +162,7 @@ public static class OpenTelemetryHelper
 		{
 			builder.AddOtlpExporter(name, exporter =>
 			{
-				exporter.Endpoint = exporterSettings.Endpoint;
+				exporter.Endpoint = exporterSettings.Endpoint!;
 				exporter.Protocol = Enum.TryParse(exporterSettings.Protocol, true, out OtlpExportProtocol protocol) ? protocol : OtlpExportProtocol.Grpc;
 			});
 		}
