Horde: Add test for downloading tool with ACLs

diff --git a/Engine/Source/Programs/Horde/Horde.Server.Tests/Tools/ToolTests.cs b/Engine/Source/Programs/Horde/Horde.Server.Tests/Tools/ToolTests.cs
--- a/Engine/Source/Programs/Horde/Horde.Server.Tests/Tools/ToolTests.cs
+++ b/Engine/Source/Programs/Horde/Horde.Server.Tests/Tools/ToolTests.cs
@@ -132,5 +132,19 @@ namespace Horde.Server.Tests.Tools
 				Assert.IsTrue(fileData.AsSpan().SequenceEqual(outputFileData.AsSpan()));
 			}
 		}
+
+		public static async Task<byte[]> CreateZipFileDataAsync(string fileName, string fileData)
+		{
+			using MemoryStream stream = new ();
+			using (ZipArchive archive = new (stream, ZipArchiveMode.Create))
+			{
+				ZipArchiveEntry entry = archive.CreateEntry(fileName);
+				await using (Stream entryStream = entry.Open())
+				{
+					await entryStream.WriteAsync(Encoding.UTF8.GetBytes(fileData));
+				}
+			}
+			return stream.ToArray();
+		}
 	}
 }
\ No newline at end of file
diff --git a/Engine/Source/Programs/Horde/Horde.Server.Tests/Tools/ToolsControllerTests.cs b/Engine/Source/Programs/Horde/Horde.Server.Tests/Tools/ToolsControllerTests.cs
new file mode 100644
--- /dev/null
+++ b/Engine/Source/Programs/Horde/Horde.Server.Tests/Tools/ToolsControllerTests.cs
@@ -0,0 +1,66 @@
+﻿// Copyright Epic Games, Inc. All Rights Reserved.
+
+using System.Collections.Generic;
+using System.IO;
+using System.Net;
+using System.Net.Http;
+using System.Net.Http.Headers;
+using System.Threading.Tasks;
+using EpicGames.Horde.Server;
+using EpicGames.Horde.Tools;
+using Horde.Server.Acls;
+using Horde.Server.Configuration;
+using Horde.Server.Server;
+using Horde.Server.ServiceAccounts;
+using Horde.Server.Storage;
+using Horde.Server.Tools;
+using Horde.Server.Users;
+using Horde.Server.Utilities;
+using Microsoft.Extensions.DependencyInjection;
+using Microsoft.VisualStudio.TestTools.UnitTesting;
+
+namespace Horde.Server.Tests.Tools;
+
+[TestClass]
+public class ToolsControllerTests
+{
+	[TestMethod]
+	public async Task DownloadToolWithAclAsync()
+	{
+		AclEntryConfig aclEntryConfig = new (HordeClaims.AgentRoleClaim, [ToolAclAction.DownloadTool]);
+		ToolId toolId = new ("foo");
+		GlobalConfig globalConfig = new();
+		
+		globalConfig.Storage.Backends.Clear();
+		globalConfig.Storage.Backends.Add(new BackendConfig { Id = new BackendId("tools-backend"), Type = StorageBackendType.Memory });
+		globalConfig.Storage.Namespaces.Clear();
+		globalConfig.Storage.Namespaces.Add(new NamespaceConfig { Id = Namespace.Tools, Backend = new BackendId("tools-backend") });
+		globalConfig.Tools.Add(new ToolConfig(toolId) { Name = "Foo", Description = "This is foo", Acl = new AclConfig() { Entries = [aclEntryConfig] }, Public = false });
+		
+		ServerSettings serverSettings = new() { AuthMethod = AuthMethod.Horde };
+		globalConfig.PostLoad(serverSettings);
+		Dictionary<string, string> settings = new() { { "Horde:AuthMethod", AuthMethod.Horde.ToString() } };
+		await using FakeHordeWebApp app = new (settings);
+		
+		ConfigService configService = app.ServiceProvider.GetRequiredService<ConfigService>();
+		IToolCollection tools = app.ServiceProvider.GetRequiredService<IToolCollection>();
+		IServiceAccountCollection serviceAccounts = app.ServiceProvider.GetRequiredService<IServiceAccountCollection>();
+
+		configService.OverrideConfig(globalConfig);
+
+		List<IUserClaim> claims = [new UserClaim("http://epicgames.com/ue/horde/role", "agent")];
+		(IServiceAccount _, string token) = await serviceAccounts.CreateAsync(new CreateServiceAccountOptions("myDesc", claims));
+		
+		// Create tool and deployment
+		using MemoryStream ms = new(await ToolTests.CreateZipFileDataAsync("foo.txt", "foo content"));
+		ITool? tool = await tools.GetAsync(toolId, globalConfig);
+		Assert.IsNotNull(tool);
+		tool = await tools.CreateDeploymentAsync(tool, new ToolDeploymentConfig() { Version = "1" }, ms, globalConfig);
+
+		HttpClient client = app.CreateHttpClient();
+		using HttpRequestMessage req = new (HttpMethod.Get, $"/api/v1/tools/{toolId.Id}?action=download");
+		req.Headers.Authorization = new AuthenticationHeaderValue("ServiceAccount", token);
+		HttpResponseMessage res = await client.SendAsync(req);
+		Assert.AreEqual(HttpStatusCode.OK, res.StatusCode);
+	}
+}
