Horde: Replace CRA (Create React App) which is no longer maintained with Vite
#rnx
#siubmittool bypass

diff --git a/Engine/Source/Programs/Horde/HordeDashboard/Dockerfile b/Engine/Source/Programs/Horde/HordeDashboard/Dockerfile
--- a/Engine/Source/Programs/Horde/HordeDashboard/Dockerfile
+++ b/Engine/Source/Programs/Horde/HordeDashboard/Dockerfile
@@ -1,14 +1,14 @@
 FROM node:20.11.0 as build-env
 ARG VersionInfo
 ARG DashboardConfig
-ENV REACT_APP_VERSION_INFO=${VersionInfo}
-ENV REACT_APP_DASHBOARD_CONFIG=${DashboardConfig}
+ENV VITE_VERSION_INFO=${VersionInfo}
+ENV VITE_DASHBOARD_CONFIG=${DashboardConfig}
 WORKDIR /app
 RUN wget https://s3.amazonaws.com/rds-downloads/rds-combined-ca-bundle.pem
 COPY Source/Programs/Horde/HordeDashboard ./Dashboard
-RUN cd ./Dashboard && yarn install
-RUN cd ./Dashboard && yarn run build
-COPY Source/Programs/Horde/HordeDashboard/documentation ./Dashboard/build/documentation
+RUN cd ./Dashboard && npm install --legacy-peer-deps
+RUN cd ./Dashboard && npm run build
+COPY Source/Programs/Horde/HordeDashboard/documentation ./Dashboard/dist/documentation
 
 FROM nginx:1.17.10 as runtime-env
 ## Remove default nginx index page and conf
@@ -16,7 +16,7 @@ RUN rm -rf /usr/share/nginx/html/*
 RUN rm /etc/nginx/conf.d/default.conf 
 WORKDIR /app
 COPY Source/Programs/Horde/HordeDashboard/nginx/nginx.conf /etc/nginx/conf.d
-COPY --from=build-env /app/Dashboard/build /usr/share/nginx/html
+COPY --from=build-env /app/Dashboard/dist /usr/share/nginx/html
 COPY --from=build-env /app/rds-combined-ca-bundle.pem .
 EXPOSE 80
 CMD ["nginx", "-g", "daemon off;"]
diff --git a/Engine/Source/Programs/Horde/HordeDashboard/index.html b/Engine/Source/Programs/Horde/HordeDashboard/index.html
new file mode 100644
--- /dev/null
+++ b/Engine/Source/Programs/Horde/HordeDashboard/index.html
@@ -0,0 +1,13 @@
+<!doctype html>
+<html lang="en">
+  <head>
+    <meta charset="UTF-8" />
+    <link rel="icon" type="image/svg+xml" href="/favicon.svg" />
+    <meta name="viewport" content="width=device-width, initial-scale=1.0" />
+    <title>Horde</title>
+  </head>
+  <body>
+    <div id="root"></div>
+    <script type="module" src="/src/main.tsx"></script>
+  </body>
+</html>
diff --git a/Engine/Source/Programs/Horde/HordeDashboard/installer/Dockerfile b/Engine/Source/Programs/Horde/HordeDashboard/installer/Dockerfile
--- a/Engine/Source/Programs/Horde/HordeDashboard/installer/Dockerfile
+++ b/Engine/Source/Programs/Horde/HordeDashboard/installer/Dockerfile
@@ -1,9 +1,9 @@
 FROM node:20.11.0 as build-env
 ARG VersionInfo
-ENV REACT_APP_VERSION_INFO=${VersionInfo}
+ENV VITE_VERSION_INFO=${VersionInfo}
 WORKDIR /app
 
 COPY Programs/Horde/HordeDashboard ./Dashboard
-RUN cd ./Dashboard && yarn install
-RUN cd ./Dashboard && yarn run build
+RUN cd ./Dashboard && npm install --legacy-peer-deps
+RUN cd ./Dashboard && npm run build
 
diff --git a/Engine/Source/Programs/Horde/HordeDashboard/package.json b/Engine/Source/Programs/Horde/HordeDashboard/package.json
--- a/Engine/Source/Programs/Horde/HordeDashboard/package.json
+++ b/Engine/Source/Programs/Horde/HordeDashboard/package.json
@@ -1,10 +1,18 @@
 {
   "name": "horde-dashboard",
-  "version": "0.2.6",
   "private": true,
-  "main": "index.js",
-  "proxy": "https://localhost:5001",
+  "version": "0.0.1",
+  "type": "module",  
+  "scripts": {
+    "dev": "vite",
+    "build": "tsc && vite build",    
+    "lint": "eslint . --ext ts,tsx --report-unused-disable-directives --max-warnings 0",
+    "preview": "vite preview",
+    "license-check": "license-checker-rseidelsohn --production --excludePackages horde-dashboard --depth 0 --limitAttributes licenses --csv --customPath ./license-format.json --onlyAllow \"Apache-2.0;MIT;ISC\""
+  },
   "dependencies": {
+    "react": "^18.2.0",
+    "react-dom": "^18.2.0",
     "@blueprintjs/icons": "^3.13.0",
     "@datadog/browser-logs": "^4.30.0",
     "@fluentui/react": "^8.106.2",
@@ -13,56 +21,36 @@
     "d3": "^7.6.1",
     "idb": "^7.1.1",
     "markdown-to-jsx": "^7.1.8",
-    "mobx": "^6.7.0",
-    "mobx-react-lite": "^3.4.0",
+    "mobx": "^6.12.2",
+    "mobx-react-lite": "^4.0.7",
     "moment-timezone": "^0.5.40",
-    "react": "^18.2.0",
     "react-custom-scrollbars": "^4.2.1",
-    "react-dom": "^18.2.0",
     "react-flame-graph": "^1.4.0",
     "react-highlighter": "^0.4.3",
     "react-router-dom": "6.6.2",
-    "react-router-hash-link": "^2.4.3",    
+    "react-router-hash-link": "^2.4.3",
     "react-sparklines": "^1.7.0",
     "react-text-marquee": "^1.0.1"
   },
-  "scripts": {
-    "start": "rimraf ./build&&set NODE_OPTIONS=\"--max-http-header-size=80000\"&&react-scripts start",
-    "build": "react-scripts build",
-    "test": "react-scripts test",
-    "eject": "react-scripts eject",
-    "license-check": "license-checker-rseidelsohn --production --excludePackages horde-dashboard --depth 0 --limitAttributes licenses --csv --customPath ./license-format.json --onlyAllow \"Apache-2.0;MIT;ISC\"",
-    "lint-full": "eslint ./src/**/*.{ts,tsx}",
-    "line-full-fix": "eslint --fix ./src/**/*.{ts,tsx}"
-  },
-  "browserslist": {
-    "production": [
-      ">0.2%",
-      "not dead",
-      "not op_mini all"
-    ],
-    "development": [
-      "last 1 chrome version",
-      "last 1 firefox version",
-      "last 1 safari version"
-    ]
-  },  
-  "resolutions": {
-    "strip-ansi": "6.0.1"    
-  },  
   "devDependencies": {
+    "@types/react": "^18.2.66",
+    "@types/react-dom": "^18.2.22",
+    "@typescript-eslint/eslint-plugin": "^7.2.0",
+    "@typescript-eslint/parser": "^7.2.0",
+    "@vitejs/plugin-react": "^4.2.1",
+    "eslint": "^8.57.0",
+    "eslint-plugin-react-hooks": "^4.6.0",
+    "eslint-plugin-react-refresh": "^0.4.6",
+    "typescript": "^5.2.2",
+    "vite": "^5.2.0",
     "@types/crypto-js": "^4.1.1",
     "@types/d3": "^7.4.0",
     "@types/markdown-to-jsx": "7.0.1",
-    "@types/react": "^18.0.26",
     "@types/react-custom-scrollbars": "^4.0.10",
-    "@types/react-dom": "^18.0.10",
     "@types/react-highlighter": "^0.3.5",
     "@types/react-router-dom": "^5.3.3",
     "@types/react-sparklines": "^1.7.5",
     "license-checker-rseidelsohn": "^4.3.0",
-    "rimraf": "^3.0.0",
-    "typescript": "^4.7.4",
-    "react-scripts": "^5.0.1"
+    "rimraf": "^3.0.0"
   }
-}
+}
\ No newline at end of file
diff --git a/Engine/Source/Programs/Horde/HordeDashboard/public/index.html b/Engine/Source/Programs/Horde/HordeDashboard/public/index.html
deleted file mode 100644
diff --git a/Engine/Source/Programs/Horde/HordeDashboard/public/manifest.json b/Engine/Source/Programs/Horde/HordeDashboard/public/manifest.json
deleted file mode 100644
diff --git a/Engine/Source/Programs/Horde/HordeDashboard/src/App.tsx b/Engine/Source/Programs/Horde/HordeDashboard/src/App.tsx
--- a/Engine/Source/Programs/Horde/HordeDashboard/src/App.tsx
+++ b/Engine/Source/Programs/Horde/HordeDashboard/src/App.tsx
@@ -51,9 +51,7 @@ const RouteError: React.FC = () => {
 
 const Main: React.FC = () => {
 
-   const [init, setInit] = useState(false);
-   const [pluginsLoaded, setPluginsLoaded] = useState(false);
-   
+   const [init, setInit] = useState(false);   
 
    if (window.location.pathname === "/login") {
       return <HordeLoginView/>
@@ -71,7 +69,7 @@ const Main: React.FC = () => {
 
       backend.init().then(() => {
 
-         backend.getCurrentUser().then(user => {
+         backend.getCurrentUser().then(() => {
 
             setInit(true);
             return null;
@@ -108,12 +106,7 @@ const Main: React.FC = () => {
       </ThemeProvider>);
    }
 
-   if (!pluginsLoaded) {
-      hordePlugins.loadPlugins(config.plugins).finally(() => {
-         setPluginsLoaded(true);
-      })
-      return <ThemeProvider applyTo='body' theme={dashboard.darktheme ? darkTheme : lightTheme} />;
-   }
+   hordePlugins.loadPlugins(config.plugins);
 
    if (!router) {
 
diff --git a/Engine/Source/Programs/Horde/HordeDashboard/src/Plugins/AutomatedTestSession/index.tsx b/Engine/Source/Programs/Horde/HordeDashboard/src/Plugins/AutomatedTestSession/index.tsx
--- a/Engine/Source/Programs/Horde/HordeDashboard/src/Plugins/AutomatedTestSession/index.tsx
+++ b/Engine/Source/Programs/Horde/HordeDashboard/src/Plugins/AutomatedTestSession/index.tsx
@@ -1,6 +1,6 @@
 // Copyright Epic Games, Inc. All Rights Reserved.
 
-import { Plugin, PluginMount } from "..";
+import { Plugin, PluginMount } from "../pluginTypes";
 import { TestSessionView, TestSessionViewAll, BuildHealthTestSessionView, routePath } from "./components/AutomatedTestSession";
 
 class AutomatedTestSessionPlugin implements Plugin {
diff --git a/Engine/Source/Programs/Horde/HordeDashboard/src/Plugins/AutomatedTestSession/models/AutomatedTestSession.ts b/Engine/Source/Programs/Horde/HordeDashboard/src/Plugins/AutomatedTestSession/models/AutomatedTestSession.ts
--- a/Engine/Source/Programs/Horde/HordeDashboard/src/Plugins/AutomatedTestSession/models/AutomatedTestSession.ts
+++ b/Engine/Source/Programs/Horde/HordeDashboard/src/Plugins/AutomatedTestSession/models/AutomatedTestSession.ts
@@ -1152,7 +1152,7 @@ export class TestSessionCollection implements Loader {
         while (!result.done) {
             const sessionInfo = result.value.TestSessionInfo;
             if (sessionInfo) {
-                tests.push(...Object.values(result.value.TestSessionInfo.Tests));
+                tests.push(...Object.values(result.value.TestSessionInfo.Tests) as TestResult[]);
             }
             result = iterator.next();
         }
@@ -1234,7 +1234,7 @@ export class TestSessionCollection implements Loader {
         const iterator = this._datahandler.iterItemsByChange();
         let change = iterator.next();
         while (!change.done) {
-            change.value.forEach((item) => {
+            (change.value as TestDataWrapper[]).forEach((item) => {
                 const session = item.getDataHandler() as TestSessionWrapper;
                 if (session.Name === sessionName) {
                     if (!this._includePreflight && session.PreFlightChange) {
diff --git a/Engine/Source/Programs/Horde/HordeDashboard/src/Plugins/SimpleTestReport/index.tsx b/Engine/Source/Programs/Horde/HordeDashboard/src/Plugins/SimpleTestReport/index.tsx
--- a/Engine/Source/Programs/Horde/HordeDashboard/src/Plugins/SimpleTestReport/index.tsx
+++ b/Engine/Source/Programs/Horde/HordeDashboard/src/Plugins/SimpleTestReport/index.tsx
@@ -1,6 +1,6 @@
 // Copyright Epic Games, Inc. All Rights Reserved.
 
-import { Plugin, PluginMount } from "..";
+import { Plugin, PluginMount } from "../pluginTypes";
 import { TestReportView, ReportLinkItemPane } from "./components/SimpleReport";
 
 class SimpleTestReportPlugin implements Plugin {
diff --git a/Engine/Source/Programs/Horde/HordeDashboard/src/Plugins/UnrealAutomatedTests/index.tsx b/Engine/Source/Programs/Horde/HordeDashboard/src/Plugins/UnrealAutomatedTests/index.tsx
--- a/Engine/Source/Programs/Horde/HordeDashboard/src/Plugins/UnrealAutomatedTests/index.tsx
+++ b/Engine/Source/Programs/Horde/HordeDashboard/src/Plugins/UnrealAutomatedTests/index.tsx
@@ -1,6 +1,6 @@
 // Copyright Epic Games, Inc. All Rights Reserved.
 
-import { Plugin, PluginMount } from "..";
+import { Plugin, PluginMount } from "../pluginTypes";
 import { TestPassSummaryView } from "./components/UnrealAutomatedTests";
 
 class UnrealAutomatedTestsPlugin implements Plugin {
diff --git a/Engine/Source/Programs/Horde/HordeDashboard/src/Plugins/example/index.tsx b/Engine/Source/Programs/Horde/HordeDashboard/src/Plugins/example/index.tsx
--- a/Engine/Source/Programs/Horde/HordeDashboard/src/Plugins/example/index.tsx
+++ b/Engine/Source/Programs/Horde/HordeDashboard/src/Plugins/example/index.tsx
@@ -1,5 +1,5 @@
 // Copyright Epic Games, Inc. All Rights Reserved.
-import { Plugin, PluginMount } from "..";
+import { Plugin, PluginMount } from "../pluginTypes";
 import { ExampleBar } from "./routes/ExampleBar";
 import { JobDetailPluginExample } from "./components/JobDetailPluginExample";
 
diff --git a/Engine/Source/Programs/Horde/HordeDashboard/src/Plugins/index.tsx b/Engine/Source/Programs/Horde/HordeDashboard/src/Plugins/index.tsx
--- a/Engine/Source/Programs/Horde/HordeDashboard/src/Plugins/index.tsx
+++ b/Engine/Source/Programs/Horde/HordeDashboard/src/Plugins/index.tsx
@@ -1,70 +1,50 @@
 // Copyright Epic Games, Inc. All Rights Reserved.
 
-import React from 'react';
-
-// render components with top level routes
-export type PluginRoute = {
-    path: string;
-    component: React.FC;
-}
-
-// render components with specific mount points
-export enum PluginMount {
-    JobDetailPanel,
-    TestReportPanel,
-    TestReportLink,
-    BuildHealthSummary,
-}
-
-export type PluginComponent = {
-    mount: PluginMount;       
-    component: React.FC<any>;
-    // optional unique identifier for the component
-    id?:string;
-}
-
-export interface Plugin {
-    name: string;
-    routes?: PluginRoute[];
-    components?: PluginComponent[];
-}
+import { Plugin, PluginMount, PluginRoute, PluginComponent } from "./pluginTypes";
+import SimpleTestReportPlugin from "./SimpleTestReport"
+import AutomatedTestSessionPlugin from "./AutomatedTestSession"
+import UnrealAutomatedTestsPlugin from "./UnrealAutomatedTests"
 
 export class Plugins {
 
-    get routes(): PluginRoute[] {
-        return this.plugins.map(p => p.routes ?? []).flat();
-    }
+   get routes(): PluginRoute[] {
+      return this.plugins.map(p => p.routes ?? []).flat();
+   }
 
-    getComponents(mount: PluginMount, id?:string): PluginComponent[] {
-        return this.plugins.map(p => {
-            return p.components?.filter(c => c.mount === mount && c.id === id) ?? [];
-        }).flat();
-    }
+   getComponents(mount: PluginMount, id?: string): PluginComponent[] {
+      return this.plugins.map(p => {
+         return p.components?.filter(c => c.mount === mount && c.id === id) ?? [];
+      }).flat();
+   }
 
-    loadPlugins(pluginList: string[] | undefined): Promise<void> {
+   loadPlugins(pluginList: string[] | undefined): void {
 
-        console.log("loading plugins...");
+      console.log("loading plugins...");
 
-        return new Promise<void>((resolve) => {
+      this.plugins.push(...[AutomatedTestSessionPlugin, SimpleTestReportPlugin, UnrealAutomatedTestsPlugin]);
 
-            if (!pluginList || !pluginList.length || this.plugins.length) {
-                return resolve();
-            }
+      /*
+      return new Promise<void>((resolve) => {
 
-            Promise.all(pluginList.map(async (plugin: any) => {
-                await import(`./${plugin}`).then((m: any) => {
-                    this.plugins.push(m.default);
-                    console.log(`loaded plugin: ${plugin}`);
-                }).catch((reason: any) => {
-                    console.error(`unable to load plugin: ${plugin}, ${reason}`);
-                });
-            })).finally(() => {
-                resolve();
-            });
-        });
-    }
+         if (!pluginList || !pluginList.length || this.plugins.length) {
+            return resolve();
+         }
 
-    private plugins: Plugin[] = [];
+         Promise.all(pluginList.map(async (plugin: any) => {
+            await import(`./${plugin}/index.js`).then((m: any) => {
+               this.plugins.push(m.default);
+               console.log(`loaded plugin: ${plugin}`);
+            }).catch((reason: any) => {
+               console.error(`unable to load plugin: ${plugin}, ${reason}`);
+            });
+         })).finally(() => {
+            resolve();
+         });
+      });
+      */
+   }
+
+   private plugins: Plugin[] = [];
 
 }
 
diff --git a/Engine/Source/Programs/Horde/HordeDashboard/src/Plugins/pluginTypes.ts b/Engine/Source/Programs/Horde/HordeDashboard/src/Plugins/pluginTypes.ts
new file mode 100644
--- /dev/null
+++ b/Engine/Source/Programs/Horde/HordeDashboard/src/Plugins/pluginTypes.ts
@@ -0,0 +1,28 @@
+
+// render components with top level routes
+export type PluginRoute = {
+    path: string;
+    component: React.FC;
+ }
+ 
+ // render components with specific mount points
+ export enum PluginMount {
+    JobDetailPanel,
+    TestReportPanel,
+    TestReportLink,
+    BuildHealthSummary,
+ }
+ 
+ export type PluginComponent = {
+    mount: PluginMount;
+    component: React.FC<any>;
+    // optional unique identifier for the component
+    id?: string;
+ }
+ 
+ export interface Plugin {
+    name: string;
+    routes?: PluginRoute[];
+    components?: PluginComponent[];
+ }
+ 
\ No newline at end of file
diff --git a/Engine/Source/Programs/Horde/HordeDashboard/src/backend/Backend.ts b/Engine/Source/Programs/Horde/HordeDashboard/src/backend/Backend.ts
--- a/Engine/Source/Programs/Horde/HordeDashboard/src/backend/Backend.ts
+++ b/Engine/Source/Programs/Horde/HordeDashboard/src/backend/Backend.ts
@@ -1928,7 +1928,7 @@ export class Backend {
     // server url if on a separate origin from dashboard, as with local development and a debug token
     get serverUrl(): string {
 
-        const url = process.env.REACT_APP_HORDE_BACKEND;
+        const url = import.meta.env.VITE_HORDE_BACKEND;
 
         if (!url) {
             return "";
@@ -1939,7 +1939,7 @@ export class Backend {
 
     get debugToken(): string {
 
-        const token = process.env.REACT_APP_HORDE_DEBUG_TOKEN;
+        const token = import.meta.env.VITE_HORDE_DEBUG_TOKEN;
 
         if (!token) {
             return "";
diff --git a/Engine/Source/Programs/Horde/HordeDashboard/src/backend/Config.ts b/Engine/Source/Programs/Horde/HordeDashboard/src/backend/Config.ts
--- a/Engine/Source/Programs/Horde/HordeDashboard/src/backend/Config.ts
+++ b/Engine/Source/Programs/Horde/HordeDashboard/src/backend/Config.ts
@@ -33,12 +33,12 @@ export function getSiteConfig() {
 	let isProd = false;
 	
 	try {
-		// note: this must be exactly `process.env.REACT_APP_DASHBOARD_CONFIG`
+		// note: this must be exactly `import.meta.env.VITE_DASHBOARD_CONFIG`
 		// as webpack does a simple find and replace to the value in the .env
-		// so process?.env?.REACT_APP_DASHBOARD_CONFIG for example is invalid
-		isProd = process.env.REACT_APP_DASHBOARD_CONFIG !== "Development";
+		// so import?.meta?.env>.VITE_DASHBOARD_CONFIG for example is invalid
+		isProd = import.meta.env.VITE_DASHBOARD_CONFIG !== "Development";
 	 } catch (reason) {
-		console.log("Process env error on REACT_APP_DASHBOARD_CONFIG:", reason);
+		console.log("Process env error on VITE_DASHBOARD_CONFIG:", reason);
 	 }
           
     return isProd ? prod : dev;
diff --git a/Engine/Source/Programs/Horde/HordeDashboard/src/backend/Dashboard.ts b/Engine/Source/Programs/Horde/HordeDashboard/src/backend/Dashboard.ts
--- a/Engine/Source/Programs/Horde/HordeDashboard/src/backend/Dashboard.ts
+++ b/Engine/Source/Programs/Horde/HordeDashboard/src/backend/Dashboard.ts
@@ -152,7 +152,7 @@ export class Dashboard {
     get preview(): boolean {
 
         try {
-            return !!(process.env.REACT_APP_HORDE_DEBUG_PREVIEW) || window?.location?.hostname?.indexOf("preview") !== -1;
+            return !!(import.meta.env.VITE_HORDE_DEBUG_PREVIEW) || window?.location?.hostname?.indexOf("preview") !== -1;
         } catch (reason) {
             console.error(reason);
         }
diff --git a/Engine/Source/Programs/Horde/HordeDashboard/src/components/AutomationSuiteDetails.tsx b/Engine/Source/Programs/Horde/HordeDashboard/src/components/AutomationSuiteDetails.tsx
--- a/Engine/Source/Programs/Horde/HordeDashboard/src/components/AutomationSuiteDetails.tsx
+++ b/Engine/Source/Programs/Horde/HordeDashboard/src/components/AutomationSuiteDetails.tsx
@@ -12,6 +12,9 @@ import { StatusBar, StatusBarStack } from "./AutomationCommon";
 import { AutomationSuiteTest } from "./AutomationSuiteTest";
 import { getHordeStyling } from "../styles/Styles";
 
+// Handle bad "@types/d3" types, fix if addressed upstream
+const _d3 = d3 as any;
+
 type TestId = string;
 type SelectionType = d3.Selection<SVGGElement, unknown, null, undefined>;
 type DivSelectionType = d3.Selection<HTMLDivElement, unknown, null, undefined>;
@@ -545,8 +548,8 @@ class SuiteGraphRenderer {
       this.hasRendered = true;
       this.forceRender = false;
 
-      const X = d3.map(this.testData, (t) => t);
-      const xDomain = d3.extent(handler.changeDates.values(), d => d.getTime() / 1000);
+      const X = _d3.map(this.testData, (t) => t);
+      const xDomain = d3.extent(handler.changeDates.values() as any, (d:any) => d.getTime() / 1000);
 
       let Y: TestOutcome[] = [];
       const scolors = dashboard.getStatusColors();
@@ -558,14 +561,14 @@ class SuiteGraphRenderer {
          "Skipped": scolors.get(StatusColor.Skipped)!
       };
 
-      Y = d3.map(this.testData, (d) => d.warningCount ? TestOutcome.Warning : d.outcome);
+      Y = _d3.map(this.testData, (d) => d.warningCount ? TestOutcome.Warning : d.outcome);
 
       const width = suiteGraphWidth - 32; // important: the 32 offset is to provide more mouse room to roll down latest CL's
       const height = this.headerOnly ? 20 : 16;
 
       const xRange = [this.margin.left, width - this.margin.right];
 
-      const xScale = d3.scaleTime(xDomain as any, xRange);
+      const xScale = _d3.scaleTime(xDomain, xRange);
 
       const I = d3.range(X.length);
 
@@ -701,8 +704,8 @@ class SuiteGraphRenderer {
 
       const handleMouseMove = (event: any) => {
 
-         const mouseX = d3.pointer(event)[0];
-         const mouseY = d3.pointer(event)[1];
+         const mouseX = _d3.pointer(event)[0];
+         const mouseY = _d3.pointer(event)[1];
 
          const closest = closestRef(mouseX, mouseY);
          if (closest) {
@@ -730,8 +733,8 @@ class SuiteGraphRenderer {
 
       const handleMouseClick = (event: any) => {
 
-         const mouseX = d3.pointer(event)[0];
-         const mouseY = d3.pointer(event)[1];
+         const mouseX = _d3.pointer(event)[0];
+         const mouseY = _d3.pointer(event)[1];
 
          const closest = closestRef(mouseX, mouseY);
          if (closest && this.testId) {
diff --git a/Engine/Source/Programs/Horde/HordeDashboard/src/components/AutomationView.tsx b/Engine/Source/Programs/Horde/HordeDashboard/src/components/AutomationView.tsx
--- a/Engine/Source/Programs/Horde/HordeDashboard/src/components/AutomationView.tsx
+++ b/Engine/Source/Programs/Horde/HordeDashboard/src/components/AutomationView.tsx
@@ -20,6 +20,8 @@ import ErrorBoundary from './ErrorBoundary';
 import { TopNav } from './TopNav';
 import { getHordeStyling } from '../styles/Styles';
 
+// Handle bad "@types/d3" types, fix if addressed upstream
+const _d3 = d3 as any;
 
 type SelectionType = d3.Selection<SVGGElement, unknown, null, undefined>;
 type DivSelectionType = d3.Selection<HTMLDivElement, unknown, null, undefined>;
@@ -810,11 +812,11 @@ class AutomationGraph {
          "Skipped": scolors.get(StatusColor.Skipped)!
       };
 
-      const X = d3.map(refs, (r) => handler.changeDates.get(r.buildChangeList)!.getTime() / 1000);
-      const Y = d3.map(refs, (r) => this.metaIds.indexOf(r.metaId));
+      const X = _d3.map(refs, (r) => handler.changeDates.get(r.buildChangeList)!.getTime() / 1000);
+      const Y = _d3.map(refs, (r) => this.metaIds.indexOf(r.metaId));
       let Z: (TestOutcome | "Unspecified" | "Warning" | "Success" | "Failure")[] = [];
       if (this.suite) {
-         Z = d3.map(refs, (r) => {
+         Z = _d3.map(refs, (r) => {
 
             if (!r.suiteErrorCount && !r.suiteWarningCount && !r.suiteSkipCount) {
                // If there are skipped, will be a shape with warning/error/success included??
@@ -832,12 +834,12 @@ class AutomationGraph {
 
          });
       } else {
-         Z = d3.map(refs, (r) => r.outcome ?? "Unspecified");
+         Z = _d3.map(refs, (r) => r.outcome ?? "Unspecified");
       }
 
-      const xDomain = d3.extent(handler.changeDates.values(), d => d.getTime() / 1000);
+      const xDomain = d3.extent(handler.changeDates.values() as any, (d:any) => d.getTime() / 1000);
       let yDomain: any = Y;
-      yDomain = new d3.InternSet(yDomain);
+      yDomain = new _d3.InternSet(yDomain);
 
       const I = d3.range(X.length);
 
@@ -847,8 +849,8 @@ class AutomationGraph {
       const xRange = [this.margin.left, width - this.margin.right];
       let yRange = [this.margin.top, height - this.margin.bottom];
 
-      const xScale = d3.scaleTime(xDomain as any, xRange);
-      const yScale = d3.scalePoint(yDomain, yRange).round(true).padding(yPadding);
+      const xScale = _d3.scaleTime(xDomain, xRange);
+      const yScale = _d3.scalePoint(yDomain, yRange).round(true).padding(yPadding);
 
       const svg = d3.select(container)
          .append("svg")
@@ -859,12 +861,13 @@ class AutomationGraph {
 
       const g = svg.append("g")
          .selectAll()
-         .data(d3.group(I, i => Y[i]))
+         .data(_d3.group(I, i => Y[i]))
          .join("g")
          .attr("transform", ([y]) => `translate(0,${(yScale(y) as any) + 16})`);
 
       g.append("line")
-         .attr("stroke", ([y]) => {
+         .attr("stroke", (y: any) => {
+            // Test Vite
             const status = this.metaStatus.get(this.metaIds[y])!;
             if (status === "Failure") {
                return colors[this.metaStatus.get(this.metaIds[y])!]
@@ -895,7 +898,7 @@ class AutomationGraph {
          .style("font-family", "Horde Open Sans Regular")
          .style("font-size", 10)
          .attr("dy", "0.15em") // center stream name
-         .attr("x", ([, I]) => 0)
+         .attr("x", () => 0)
          .attr("fill", ([y]) => {
             const status = this.metaStatus.get(this.metaIds[y])!;
             if (!dashboard.darktheme && status === "Failure") {
@@ -1001,8 +1004,8 @@ class AutomationGraph {
 
       const handleMouseMove = (event: any) => {
 
-         const mouseX = d3.pointer(event)[0];
-         const mouseY = d3.pointer(event)[1];
+         const mouseX = _d3.pointer(event)[0];
+         const mouseY = _d3.pointer(event)[1];
 
          const closest = closestData(mouseX, mouseY);
          if (closest) {
@@ -1033,8 +1036,8 @@ class AutomationGraph {
 
       const handleMouseClick = (event: any) => {
 
-         const mouseX = d3.pointer(event)[0];
-         const mouseY = d3.pointer(event)[1];
+         const mouseX = _d3.pointer(event)[0];
+         const mouseY = _d3.pointer(event)[1];
 
          const closest = closestData(mouseX, mouseY);
          if (closest) {
diff --git a/Engine/Source/Programs/Horde/HordeDashboard/src/components/DevicePoolTelemetry.tsx b/Engine/Source/Programs/Horde/HordeDashboard/src/components/DevicePoolTelemetry.tsx
--- a/Engine/Source/Programs/Horde/HordeDashboard/src/components/DevicePoolTelemetry.tsx
+++ b/Engine/Source/Programs/Horde/HordeDashboard/src/components/DevicePoolTelemetry.tsx
@@ -14,6 +14,9 @@ import { getHordeStyling } from "../styles/Styles";
 
 type PoolId = string;
 
+// Handle bad "@types/d3" types, fix if addressed upstream
+const _d3 = d3 as any;
+
 type SelectionType = d3.Selection<SVGGElement, unknown, null, undefined>;
 type ScalarTime = d3.ScaleTime<number, number>
 type DivSelectionType = d3.Selection<HTMLDivElement, unknown, null, undefined>;
@@ -669,9 +672,9 @@ class PoolTelemetryGraph {
          .append("path")
          .attr("class", function (d) { return "telemetryArea " + (data.streamSelectors[d.key] ?? d.key); })
          .style("shapeRendering", "geometricPrecision")
-         .style("fill", function (d) { return colors[d.key] ?? schemaColor(d.key) })
+         .style("fill", function (d) { return colors[d.key] ?? schemaColor(d.key) as any })
          .attr("d", area as any)
-         .on("mouseover", (e, d) => { highlight(e, e.target.classList[1], true) })
+         .on("mouseover", (e:any) => { highlight(e, e.target.classList[1], true) })
          .on("mouseleave", (e, d) => { noHighlight(e, d) });
 
       // legend
@@ -687,8 +690,8 @@ class PoolTelemetryGraph {
          .attr("y", function (d, i) { return posy + 10 + i * (size + 5 + 4) })
          .attr("width", size)
          .attr("height", size)
-         .style("fill", function (d: any) { return colors[d] ?? schemaColor(d) })
-         .on("mouseover", highlight)
+         .style("fill", function (d: any) { return colors[d] ?? schemaColor(d) as any })
+         .on("mouseover", highlight as any)
          .on("mouseleave", noHighlight)
 
       svg.selectAll("legendlabels")
@@ -704,7 +707,7 @@ class PoolTelemetryGraph {
          .style("alignment-baseline", "middle")
          .style("font-family", "Horde Open Sans Regular")
          .style("font-size", "13px")
-         .on("mouseover", highlight)
+         .on("mouseover", highlight as any)
          .on("mouseleave", noHighlight)
 
       // tooltip
@@ -721,7 +724,7 @@ class PoolTelemetryGraph {
       //.style("pointer-events", "none")
 
       svg.on("mousemove", (event) => this.handleMouseMove(event))
-      svg.on("wheel", (event) => {
+      svg.on("wheel", (event:any) => {
          if (event.wheelDelta < 0) {
             updateChart(undefined);
          }
@@ -805,7 +808,7 @@ class PoolTelemetryGraph {
          return data.streamNamesReverse[streamName] === sid;
       });
 
-      const closest = this.closestData(d3.pointer(event)[0]);
+      const closest = this.closestData(_d3.pointer(event)[0]);
       if (!closest) {
          this.updateTooltip(false);
          return undefined;
@@ -902,7 +905,7 @@ class PoolTelemetryGraph {
 
       text += "</div>";
 
-      this.updateTooltip(true, d3.pointer(event)[0], d3.pointer(event)[1], text);
+      this.updateTooltip(true, _d3.pointer(event)[0], _d3.pointer(event)[1], text);
 
    }
 
diff --git a/Engine/Source/Programs/Horde/HordeDashboard/src/components/NewBuild.tsx b/Engine/Source/Programs/Horde/HordeDashboard/src/components/NewBuild.tsx
--- a/Engine/Source/Programs/Horde/HordeDashboard/src/components/NewBuild.tsx
+++ b/Engine/Source/Programs/Horde/HordeDashboard/src/components/NewBuild.tsx
@@ -1516,7 +1516,7 @@ export const NewBuild: React.FC<{ streamId: string; show: boolean; onClose: (new
 
       const submit = async () => {
 
-         if (!!process.env.REACT_APP_HORDE_DEBUG_NEW_JOB) {
+         if (!!import.meta.env.VITE_HORDE_DEBUG_NEW_JOB) {
 
             console.log("Debug Job Submit");
             console.log(JSON.stringify(data));
diff --git a/Engine/Source/Programs/Horde/HordeDashboard/src/components/TestReportPanel.tsx b/Engine/Source/Programs/Horde/HordeDashboard/src/components/TestReportPanel.tsx
--- a/Engine/Source/Programs/Horde/HordeDashboard/src/components/TestReportPanel.tsx
+++ b/Engine/Source/Programs/Horde/HordeDashboard/src/components/TestReportPanel.tsx
@@ -5,7 +5,8 @@ import { mergeStyleSets } from '@fluentui/react/lib/Styling';
 import { Stack, Text } from '@fluentui/react';
 import { Link } from 'react-router-dom';
 import { TestData } from '../backend/Api';
-import hordePlugins, { PluginMount } from  '../Plugins';
+import { PluginMount } from  '../Plugins/pluginTypes';
+import hordePlugins from  '../Plugins';
 import { ComponentMount } from './TestReportView';
 import { getHordeStyling } from '../styles/Styles';
 import { getHordeTheme } from '../styles/theme';
diff --git a/Engine/Source/Programs/Horde/HordeDashboard/src/components/TestReportView.tsx b/Engine/Source/Programs/Horde/HordeDashboard/src/components/TestReportView.tsx
--- a/Engine/Source/Programs/Horde/HordeDashboard/src/components/TestReportView.tsx
+++ b/Engine/Source/Programs/Horde/HordeDashboard/src/components/TestReportView.tsx
@@ -4,7 +4,8 @@ import React, { useEffect, useRef, useState } from 'react';
 import { TopNav } from './TopNav';
 import { Breadcrumbs, BreadcrumbItem } from './Breadcrumbs';
 import { Stack, Spinner, SpinnerSize, Text } from '@fluentui/react';
-import hordePlugins, { PluginMount } from  '../Plugins';
+import { PluginMount } from  '../Plugins/pluginTypes';
+import hordePlugins from  '../Plugins';
 import { projectStore } from '../backend/ProjectStore';
 import { useParams } from 'react-router-dom';
 import { TestDataCollection, TestDataWrapper } from '../backend/TestDataHandler'
diff --git a/Engine/Source/Programs/Horde/HordeDashboard/src/components/VersionModal.tsx b/Engine/Source/Programs/Horde/HordeDashboard/src/components/VersionModal.tsx
--- a/Engine/Source/Programs/Horde/HordeDashboard/src/components/VersionModal.tsx
+++ b/Engine/Source/Programs/Horde/HordeDashboard/src/components/VersionModal.tsx
@@ -42,10 +42,10 @@ export const VersionModal: React.FC<{ show: boolean, onClose: () => void }> = ({
    let dashboardVersion: string | undefined;
 
    try {
-      // note: this must be exactly `process.env.REACT_APP_VERSION_INFO`
+      // note: this must be exactly `import.meta.env.VITE_VERSION_INFO`
       // as webpack does a simple find and replace to the value in the .env
-      // so process?.env?.REACT_APP_VERSION_INFO for example is invalid
-      dashboardVersion = process.env.REACT_APP_VERSION_INFO;
+      // so import?.meta?.env?.VITE_VERSION_INFO for example is invalid
+      dashboardVersion = import.meta.env.VITE_VERSION_INFO;
    } catch (reason) {
       console.log("Process env error:", reason);
    }
diff --git a/Engine/Source/Programs/Horde/HordeDashboard/src/components/bisection/BisectionList.tsx b/Engine/Source/Programs/Horde/HordeDashboard/src/components/bisection/BisectionList.tsx
--- a/Engine/Source/Programs/Horde/HordeDashboard/src/components/bisection/BisectionList.tsx
+++ b/Engine/Source/Programs/Horde/HordeDashboard/src/components/bisection/BisectionList.tsx
@@ -17,6 +17,9 @@ import { getHordeStyling } from '../../styles/Styles';
 
 type SelectionType = d3.Selection<SVGGElement, unknown, null, undefined>;
 
+// Handle bad "@types/d3" types, fix if addressed upstream
+const _d3 = d3 as any;
+
 class Tooltip {
 
    constructor() {
@@ -269,7 +272,7 @@ class BisectionRenderer {
          changeSet.clear();
       }
 
-      const X = d3.map(Array.from(changeSet).sort((a, b) => a - b), (t) => t);
+      const X = _d3.map(Array.from(changeSet).sort((a, b) => a - b), (t) => t);
       const I = d3.range(X.length);
 
       let svg = this.svg;
@@ -361,10 +364,10 @@ class BisectionRenderer {
             return;
          }
 
-         let mouseX = d3.pointer(event)[0];
-         let mouseY = d3.pointer(event)[1];
+         let mouseX = _d3.pointer(event)[0];
+         let mouseY = _d3.pointer(event)[1];
          const change = closestCL(mouseX, mouseY);
-         this.tooltip.update(this.bisection, change, Math.floor(d3.pointer(event, container)[0]), Math.floor(d3.pointer(event, container)[1]));
+         this.tooltip.update(this.bisection, change, Math.floor(_d3.pointer(event, container)[0]), Math.floor(_d3.pointer(event, container)[1]));
       }
 
       const handleMouseLeave = (event: any) => {
diff --git a/Engine/Source/Programs/Horde/HordeDashboard/src/components/docs/DocView.tsx b/Engine/Source/Programs/Horde/HordeDashboard/src/components/docs/DocView.tsx
--- a/Engine/Source/Programs/Horde/HordeDashboard/src/components/docs/DocView.tsx
+++ b/Engine/Source/Programs/Horde/HordeDashboard/src/components/docs/DocView.tsx
@@ -192,7 +192,7 @@ const DocPanel: React.FC<{ docName: string }> = observer(({ docName }) => {
       // This timeout is horrible, though scrollIntoView is inaccurate until the rendering has "settled"
       // I tried quite a few approaches to this, and had to move on, this works
       setTimeout(() => {         
-         if (anchor === window.location.hash?.split("?")[0]?.slice(1) ?? "") {
+         if (anchor === (window.location.hash?.split("?")[0]?.slice(1) ?? "")) {
             const element = document.getElementById(anchor);
             element?.scrollIntoView();   
          }
diff --git a/Engine/Source/Programs/Horde/HordeDashboard/src/components/jobDetailsV2/JobDetailStepHistoryV2.tsx b/Engine/Source/Programs/Horde/HordeDashboard/src/components/jobDetailsV2/JobDetailStepHistoryV2.tsx
--- a/Engine/Source/Programs/Horde/HordeDashboard/src/components/jobDetailsV2/JobDetailStepHistoryV2.tsx
+++ b/Engine/Source/Programs/Horde/HordeDashboard/src/components/jobDetailsV2/JobDetailStepHistoryV2.tsx
@@ -212,7 +212,7 @@ class StepHistoryDataView extends JobDataView {
 
 type SelectionType = d3.Selection<SVGGElement, unknown, null, undefined>;
 type Zoom = d3.ZoomBehavior<Element, unknown>;
-type Scalar = d3.ScaleLinear<number, number, never>;
+type Scalar = d3.ScaleLinear<number, number>;
 
 enum ZoomType {
    Current,
@@ -399,7 +399,7 @@ class StepHistoryRenderer {
       const arrowPoints = [[0, 0], [0, 10], [10, 5]] as any;
       svg.append("marker")
          .attr("id", "cmarker")
-         .attr('viewBox', [0, 0, 10, 10])
+         .attr('viewBox', [0, 0, 10, 10] as any)
          .attr("refX", 5)
          .attr("refY", 5)
          .attr("markerWidth", 5)
@@ -524,8 +524,11 @@ class StepHistoryRenderer {
             return;
          }
 
-         let mouseX = d3.pointer(event)[0];
-         let mouseY = d3.pointer(event)[1];
+         // bad typings
+         const _d3 = d3 as any;
+
+         let mouseX = _d3.pointer(event)[0];
+         let mouseY = _d3.pointer(event)[1];
 
          const ref = closestData(mouseX, mouseY);
 
@@ -548,11 +551,11 @@ class StepHistoryRenderer {
          if (ref) {
             this.showToolTip()
          }
-         dataView.tooltip.update(ref, d3.pointer(event, container)[0], mouseY, ref?.change);
+         dataView.tooltip.update(ref, _d3.pointer(event, container)[0], mouseY, ref?.change);
       }      
 
       // events
-      svg.on("wheel", (event) => { event.preventDefault(); })
+      svg.on("wheel", (event:any) => { event.preventDefault(); })
 
       svg.on("mousemove", (event) => { this.showToolTip(true); handleMouseMove(event);  });
       svg.on("mouseleave", (event) => { if (!dataView.tooltip.frozen) dataView.tooltip.update(undefined);})
diff --git a/Engine/Source/Programs/Horde/HordeDashboard/src/components/jobDetailsV2/JobDetailStepTrendsV2.tsx b/Engine/Source/Programs/Horde/HordeDashboard/src/components/jobDetailsV2/JobDetailStepTrendsV2.tsx
--- a/Engine/Source/Programs/Horde/HordeDashboard/src/components/jobDetailsV2/JobDetailStepTrendsV2.tsx
+++ b/Engine/Source/Programs/Horde/HordeDashboard/src/components/jobDetailsV2/JobDetailStepTrendsV2.tsx
@@ -20,6 +20,9 @@ import { getHordeStyling } from "../../styles/Styles";
 
 const sideRail: ISideRailLink = { text: "Trends", url: "rail_step_trends" };
 
+// Handle bad "@types/d3" types, fix if addressed upstream
+const _d3 = d3 as any;
+
 class Tooltip {
 
    constructor() {
@@ -221,7 +224,7 @@ class StepTrendsDataView extends JobDataView {
 
 type SelectionType = d3.Selection<SVGGElement, unknown, null, undefined>;
 type Zoom = d3.ZoomBehavior<Element, unknown>;
-type Scalar = d3.ScaleLinear<number, number, never>;
+type Scalar = d3.ScaleLinear<number, number>;
 
 class StepTrendsRenderer {
 
@@ -327,7 +330,7 @@ class StepTrendsRenderer {
       const arrowPoints = [[0, 0], [0, 10], [10, 5]] as any;
       svg.append("marker")
          .attr("id", "cmarker")
-         .attr('viewBox', [0, 0, 10, 10])
+         .attr('viewBox', [0, 0, 10, 10] as any)
          .attr("refX", 5)
          .attr("refY", 5)
          .attr("markerWidth", 5)
@@ -514,8 +517,8 @@ class StepTrendsRenderer {
             return;
          }
 
-         let mouseX = d3.pointer(event)[0];
-         let mouseY = d3.pointer(event)[1];
+         let mouseX = _d3.pointer(event)[0];
+         let mouseY = _d3.pointer(event)[1];
 
          const ref = closestData(mouseX, mouseY);
 
@@ -538,11 +541,11 @@ class StepTrendsRenderer {
          if (ref) {
             this.showToolTip()
          }
-         dataView.tooltip.update(ref, d3.pointer(event, container)[0], mouseY, ref?.change);
+         dataView.tooltip.update(ref, _d3.pointer(event, container)[0], mouseY, ref?.change);
       }
 
       // events
-      svg.on("wheel", (event) => { event.preventDefault(); })
+      svg.on("wheel", (event:any) => { event.preventDefault(); })
 
       svg.on("mousemove", (event) => { this.showToolTip(true); handleMouseMove(event); });
       svg.on("mouseleave", (event) => { if (!dataView.tooltip.frozen) dataView.tooltip.update(undefined); })
diff --git a/Engine/Source/Programs/Horde/HordeDashboard/src/components/jobDetailsV2/JobDetailTimeline.tsx b/Engine/Source/Programs/Horde/HordeDashboard/src/components/jobDetailsV2/JobDetailTimeline.tsx
--- a/Engine/Source/Programs/Horde/HordeDashboard/src/components/jobDetailsV2/JobDetailTimeline.tsx
+++ b/Engine/Source/Programs/Horde/HordeDashboard/src/components/jobDetailsV2/JobDetailTimeline.tsx
@@ -17,6 +17,9 @@ import { getHordeStyling } from "../../styles/Styles";
 
 type SelectionType = d3.Selection<SVGGElement, unknown, null, undefined>;
 
+// There are some bad @type/d3 mappings
+const _d3 = d3 as any;
+
 type Batch = GetBatchResponse & {
    utcStart: Date;
    utcFinish: Date;
@@ -257,7 +260,7 @@ class TimelineDataView extends JobDataView {
             lanes.push({ batches: [batch] });
          } else {
             const lane = bestLane(batch, clanes);
-            console.assert(lane);
+            console.assert(!!lane);
             lane!.batches.push(batch);
          }
 
@@ -451,11 +454,11 @@ class TimelineRenderer {
       };
 
 
-      const X = d3.map(spans, (s) => s);
-      const Y = d3.map(spans, (s) => s.lane);
+      const X = (d3 as any).map(spans, (s) => s);
+      const Y = (d3 as any).map(spans, (s) => s.lane);
 
       let yDomain: any = Y;
-      yDomain = new d3.InternSet(yDomain);
+      yDomain = new (d3 as any).InternSet(yDomain);
 
       const I = d3.range(X.length);
 
@@ -468,7 +471,7 @@ class TimelineRenderer {
          .domain([dataView.minTime!, dataView.maxTime!].map(d => d.getTime() / 1000))
          .range([margin.left, width - margin.right])
 
-      const yScale = d3.scalePoint(yDomain, yRange).round(true).padding(yPadding);
+      const yScale = (d3 as any).scalePoint(yDomain, yRange).round(true).padding(yPadding);
 
       let svg = this.svg;
 
@@ -487,7 +490,7 @@ class TimelineRenderer {
 
       const g = svg.append("g")
          .selectAll()
-         .data(d3.group(I, i => Y[i]))
+         .data((d3 as any).group(I, i => Y[i]))
          .join("g")
 
       const waitColor = dashboard.darktheme ? "#5B6367" : "#D3D2D1";
@@ -652,8 +655,8 @@ class TimelineRenderer {
             return;
          }
 
-         let mouseX = d3.pointer(event)[0];
-         let mouseY = d3.pointer(event)[1];
+         let mouseX = (d3 as any).pointer(event)[0];
+         let mouseY = (d3 as any).pointer(event)[1];
 
          const span = closestData(mouseX, mouseY);
 
@@ -683,7 +686,7 @@ class TimelineRenderer {
                .attr("stroke-width", 0)
          }
 
-         dataView.tooltip.update(span, d3.pointer(event, container)[0], d3.pointer(event, container)[1]);
+         dataView.tooltip.update(span, (d3 as any).pointer(event, container)[0], (d3 as any).pointer(event, container)[1]);
 
       }
 
@@ -713,7 +716,7 @@ class TimelineRenderer {
 
       svg.on("mousemove", (event) => handleMouseMove(event));
       svg.on("mouseleave", (event) => { handleMouseLeave(event); })
-      svg.on("wheel", (event) => { event.preventDefault(); })
+      svg.on("wheel", (event) => { (event as any).preventDefault(); })
 
    }
 
diff --git a/Engine/Source/Programs/Horde/HordeDashboard/src/components/jobDetailsV2/JobDetailsViewCommon.tsx b/Engine/Source/Programs/Horde/HordeDashboard/src/components/jobDetailsV2/JobDetailsViewCommon.tsx
--- a/Engine/Source/Programs/Horde/HordeDashboard/src/components/jobDetailsV2/JobDetailsViewCommon.tsx
+++ b/Engine/Source/Programs/Horde/HordeDashboard/src/components/jobDetailsV2/JobDetailsViewCommon.tsx
@@ -454,7 +454,7 @@ export class JobDetailsV2 extends PollBase {
 
       const labels = this.labels;
 
-      if (!labels || idx >= labels.length) {
+      if (!labels || (idx as number)>= labels.length) {
          return undefined;
       }
 
diff --git a/Engine/Source/Programs/Horde/HordeDashboard/src/components/jobDetailsV2/StepDetailTestPanel.tsx b/Engine/Source/Programs/Horde/HordeDashboard/src/components/jobDetailsV2/StepDetailTestPanel.tsx
--- a/Engine/Source/Programs/Horde/HordeDashboard/src/components/jobDetailsV2/StepDetailTestPanel.tsx
+++ b/Engine/Source/Programs/Horde/HordeDashboard/src/components/jobDetailsV2/StepDetailTestPanel.tsx
@@ -8,7 +8,8 @@ import { Link } from 'react-router-dom';
 import backend from '../../backend';
 import { TestData } from '../../backend/Api';
 import { ISideRailLink } from '../../base/components/SideRail';
-import hordePlugins, { PluginMount } from '../../Plugins';
+import { PluginMount } from  '../../Plugins/pluginTypes';
+import hordePlugins from  '../../Plugins';
 import { ComponentMount } from '../TestReportView';
 import { JobDataView, JobDetailsV2 } from './JobDetailsViewCommon';
 import { getHordeStyling } from '../../styles/Styles';
diff --git a/Engine/Source/Programs/Horde/HordeDashboard/src/components/telemetry/TelemetryLineGraph.ts b/Engine/Source/Programs/Horde/HordeDashboard/src/components/telemetry/TelemetryLineGraph.ts
--- a/Engine/Source/Programs/Horde/HordeDashboard/src/components/telemetry/TelemetryLineGraph.ts
+++ b/Engine/Source/Programs/Horde/HordeDashboard/src/components/telemetry/TelemetryLineGraph.ts
@@ -5,9 +5,12 @@ import dashboard from "../../backend/Dashboard";
 import { displayTimeZone, msecToElapsed } from "../../base/utilities/timeUtils";
 import { graphColors } from "./TelemetryData";
 
+// Handle bad "@types/d3" types, fix if addressed upstream
+const _d3 = d3 as any;
+
 type SelectionType = d3.Selection<SVGGElement, unknown, null, undefined>;
 type Zoom = d3.ZoomBehavior<Element, unknown>;
-type Scalar = d3.ScaleLinear<number, number, never>;
+type Scalar = d3.ScaleLinear<number, number>;
 
 type DataPoint = [number, number, string, number, number];
 
@@ -75,7 +78,7 @@ export class TelemetryLineRenderer {
          .attr("height", height);
 
       const points = allMetrics.map((m, index) => [x(m.time.getTime() / 1000), y(m.value), m.key, legend.indexOf(m.key) % graphColors.length, index]) as DataPoint[];
-      const groups = d3.rollup(points, v => Object.assign(v, { z: v[0][2] }), d => d[2]);
+      const groups = _d3.rollup(points, v => Object.assign(v, { z: v[0][2] }), d => d[2]);
       const gvalues = Array.from(groups.values());
 
       const line = d3.line().curve(d3.curveMonotoneX);
@@ -223,8 +226,8 @@ export class TelemetryLineRenderer {
 
       const handleMouseMove = (event: any) => {
 
-         const mouseX = d3.pointer(event)[0];
-         const mouseY = d3.pointer(event)[1];
+         const mouseX = _d3.pointer(event)[0];
+         const mouseY = _d3.pointer(event)[1];
 
          // find closest point on x axis for each data key
          const closestX = new Map<string, DataPoint>();
diff --git a/Engine/Source/Programs/Horde/HordeDashboard/src/index.tsx b/Engine/Source/Programs/Horde/HordeDashboard/src/index.tsx
deleted file mode 100644
diff --git a/Engine/Source/Programs/Horde/HordeDashboard/src/main.tsx b/Engine/Source/Programs/Horde/HordeDashboard/src/main.tsx
new file mode 100644
--- /dev/null
+++ b/Engine/Source/Programs/Horde/HordeDashboard/src/main.tsx
@@ -0,0 +1,14 @@
+import React from 'react'
+import ReactDOM from 'react-dom/client'
+import App from './App.tsx'
+import './index.css'
+import { configure } from 'mobx';
+
+// configure mobx
+configure({
+   reactionRequiresObservable: true,
+   enforceActions: "observed"
+});
+
+ReactDOM.createRoot(document.getElementById('root')!).render(<App />);
+
diff --git a/Engine/Source/Programs/Horde/HordeDashboard/src/react-app-env.d.ts b/Engine/Source/Programs/Horde/HordeDashboard/src/react-app-env.d.ts
deleted file mode 100644
diff --git a/Engine/Source/Programs/Horde/HordeDashboard/src/serviceWorker.ts b/Engine/Source/Programs/Horde/HordeDashboard/src/serviceWorker.ts
deleted file mode 100644
diff --git a/Engine/Source/Programs/Horde/HordeDashboard/src/setupTests.ts b/Engine/Source/Programs/Horde/HordeDashboard/src/setupTests.ts
deleted file mode 100644
diff --git a/Engine/Source/Programs/Horde/HordeDashboard/src/types/react-oidc-context.d.ts b/Engine/Source/Programs/Horde/HordeDashboard/src/types/react-oidc-context.d.ts
deleted file mode 100644
diff --git a/Engine/Source/Programs/Horde/HordeDashboard/src/vite-env.d.ts b/Engine/Source/Programs/Horde/HordeDashboard/src/vite-env.d.ts
new file mode 100644
--- /dev/null
+++ b/Engine/Source/Programs/Horde/HordeDashboard/src/vite-env.d.ts
@@ -0,0 +1 @@
+/// <reference types="vite/client" />
diff --git a/Engine/Source/Programs/Horde/HordeDashboard/tsconfig.json b/Engine/Source/Programs/Horde/HordeDashboard/tsconfig.json
--- a/Engine/Source/Programs/Horde/HordeDashboard/tsconfig.json
+++ b/Engine/Source/Programs/Horde/HordeDashboard/tsconfig.json
@@ -1,28 +1,32 @@
 {
   "compilerOptions": {
-    "target": "es2015",
+    "target": "ES2020",
+    "useDefineForClassFields": true,
     "lib": [
       "dom",
       "dom.iterable",
       "esnext"
     ],
-    "allowJs": true,
+    "module": "ESNext",
     "skipLibCheck": true,
-    "esModuleInterop": true,
-    "allowSyntheticDefaultImports": true,
-    "strict": true,
-    "forceConsistentCasingInFileNames": true,
-    "module": "esnext",
-    "moduleResolution": "node",
+    "experimentalDecorators": true,
+
+    /* Bundler mode */
+    "moduleResolution": "bundler",
+    "allowImportingTsExtensions": true,
     "resolveJsonModule": true,
     "isolatedModules": true,
     "noEmit": true,
-    "jsx": "react-jsx",
-    "experimentalDecorators": true,
+    "jsx": "react-jsx",    
+
+    /* Linting */
+    "strict": false,
     "noFallthroughCasesInSwitch": true
+    /*
+    "noUnusedLocals": true,
+    "noUnusedParameters": true,    
+    */
   },
-  "include": [
-    "src",
-    "types"
-  ]
+  "include": ["src"],
+  "references": [{ "path": "./tsconfig.node.json" }]
 }
diff --git a/Engine/Source/Programs/Horde/HordeDashboard/tsconfig.node.json b/Engine/Source/Programs/Horde/HordeDashboard/tsconfig.node.json
new file mode 100644
--- /dev/null
+++ b/Engine/Source/Programs/Horde/HordeDashboard/tsconfig.node.json
@@ -0,0 +1,11 @@
+{
+  "compilerOptions": {
+    "composite": true,
+    "skipLibCheck": true,
+    "module": "ESNext",
+    "moduleResolution": "bundler",
+    "allowSyntheticDefaultImports": true,
+    "strict": true
+  },
+  "include": ["vite.config.ts"]
+}
diff --git a/Engine/Source/Programs/Horde/HordeDashboard/vite.config.ts b/Engine/Source/Programs/Horde/HordeDashboard/vite.config.ts
new file mode 100644
--- /dev/null
+++ b/Engine/Source/Programs/Horde/HordeDashboard/vite.config.ts
@@ -0,0 +1,51 @@
+import { defineConfig } from 'vite'
+import react from '@vitejs/plugin-react'
+
+const proxyTarget = "http://127.0.0.1:13340"
+const debug = false;
+
+// https://vitejs.dev/config/
+export default defineConfig({
+  plugins: [react({
+    babel: {
+      parserOpts: {
+        plugins: ['decorators-legacy'],
+      },
+    },
+  })],
+  build: {
+    chunkSizeWarningLimit: 8192
+  },
+  server: {
+    proxy: {
+      '/api': {
+        target: proxyTarget,
+        changeOrigin: true
+      },
+    }
+  },
+  preview: {
+    proxy: {
+      '/api': {
+        target: proxyTarget,
+        changeOrigin: true,
+        secure: false,
+        ws: true,
+        configure: (proxy, _options) => {
+          proxy.on('error', (err, _req, _res) => {
+            if (debug)
+              console.log('proxy error', err);
+          });
+          proxy.on('proxyReq', (proxyReq, req, _res) => {
+            if (debug)
+              console.log('Sending Request to the Target:', req.method, req.url);
+          });
+          proxy.on('proxyRes', (proxyRes, req, _res) => {
+            if (debug)
+              console.log('Received Response from the Target:', proxyRes.statusCode, req.url);
+          });
+        },
+      }
+    }
+  }
+})
