Horde: Report time since last user activity in each HTTP request

This helps identify how much load potentially unused pages generate on the Horde server

#rnx

diff --git a/Engine/Source/Programs/Horde/HordeDashboard/src/backend/Fetch.ts b/Engine/Source/Programs/Horde/HordeDashboard/src/backend/Fetch.ts
--- a/Engine/Source/Programs/Horde/HordeDashboard/src/backend/Fetch.ts
+++ b/Engine/Source/Programs/Horde/HordeDashboard/src/backend/Fetch.ts
@@ -2,6 +2,7 @@
 
 import { ErrorHandler, ErrorInfo } from "../components/ErrorHandler";
 import { GetDashboardChallengeResponse } from "./Api";
+import userInactivity from "./UserInactivity";
 //import { setDatadogUser } from './Datadog';
 
 export enum ChallengeStatus {
@@ -395,6 +396,10 @@ export class Fetch {
         if (this.withCredentials) {
             headers["Authorization"] = this.authorization!;
         }
+        
+        // Report how idle the user has been on current page in each HTTP request
+        // This helps identify how much load potentially unused pages generate on the Horde server
+        headers["X-Horde-LastUserActivity"] = userInactivity.getSecondsSinceLastActivity().toString();
 
         return headers;
     }
diff --git a/Engine/Source/Programs/Horde/HordeDashboard/src/backend/UserInactivity.ts b/Engine/Source/Programs/Horde/HordeDashboard/src/backend/UserInactivity.ts
new file mode 100644
--- /dev/null
+++ b/Engine/Source/Programs/Horde/HordeDashboard/src/backend/UserInactivity.ts
@@ -0,0 +1,36 @@
+/**
+ * Tracks a timestamp for when the user last performed a UI activity
+ */
+class UserInactivity {
+    private lastActivityTime: number;
+    private readonly eventHandler: () => void;
+    
+    constructor() {
+        this.lastActivityTime = Date.now();
+        this.eventHandler = () => this.resetLastActivityTime();
+    }
+
+    start() {
+        // Not listening for 'mousemove' as it's fired a lot and can affect performance
+        window.addEventListener('mouseenter', this.eventHandler);
+        window.addEventListener('mouseleave', this.eventHandler);
+        window.addEventListener('touchstart', this.eventHandler);
+        window.addEventListener('keydown', this.eventHandler);
+        window.addEventListener('scroll', this.eventHandler);
+        window.addEventListener('click', this.eventHandler);
+        window.addEventListener('pointerdown', this.eventHandler);
+    }
+
+    public getSecondsSinceLastActivity(): number {
+        return Math.floor((Date.now() - this.lastActivityTime) / 1000);
+    }
+
+    private resetLastActivityTime() {
+        this.lastActivityTime = Date.now();
+    }
+}
+
+const userInactivity = new UserInactivity()
+userInactivity.start();
+
+export default userInactivity;
\ No newline at end of file
