Horde: Get agent version from compile-time generated string

Encode version in a source file as inferring version from assemblies is unreliable on platform-specific binaries.

#rnx

diff --git a/Engine/Source/Programs/Horde/Horde.Agent/AgentApp.cs b/Engine/Source/Programs/Horde/Horde.Agent/AgentApp.cs
--- a/Engine/Source/Programs/Horde/Horde.Agent/AgentApp.cs
+++ b/Engine/Source/Programs/Horde/Horde.Agent/AgentApp.cs
@@ -90,9 +90,9 @@ namespace Horde.Agent
 #pragma warning restore IL3000 // Avoid accessing Assembly file path when publishing as a single file		
 
 		/// <summary>
-		/// The current application version
+		/// The current application version from compile-time generated version
 		/// </summary>
-		public static string Version { get; } = GetVersion();
+		public static string Version { get; } = VersionInfo.Version;
 
 		/// <summary>
 		/// Default settings for json serialization
@@ -334,39 +334,6 @@ namespace Horde.Agent
 			GlobalTracer.Register(openTracer);
 		}
 
-		/// <summary>
-		/// Gets the version of the current assembly
-		/// </summary>
-		/// <returns></returns>
-		[SuppressMessage("SingleFile", "IL3000:Avoid accessing Assembly file path when publishing as a single file", Justification = "Has fallback handling")]
-		static string GetVersion()
-		{
-			try
-			{
-				string? assemblyPath = Assembly.GetExecutingAssembly().Location;
-				if (String.IsNullOrEmpty(assemblyPath))
-				{
-					// It's possible the current assembly is packaged as self-contained, try resolving path another way
-					assemblyPath = Process.GetCurrentProcess().MainModule?.FileName;
-				}
-
-				if (assemblyPath != null)
-				{
-					string? version = FileVersionInfo.GetVersionInfo(assemblyPath).ProductVersion;
-					if (version != null)
-					{
-						return version;
-					}
-				}
-			}
-			catch
-			{
-				// Ignore
-			}
-
-			return "unknown";
-		}
-
 		/// <summary>
 		/// Gets the application directory
 		/// </summary>
diff --git a/Engine/Source/Programs/Horde/Horde.Agent/Horde.Agent.csproj b/Engine/Source/Programs/Horde/Horde.Agent/Horde.Agent.csproj
--- a/Engine/Source/Programs/Horde/Horde.Agent/Horde.Agent.csproj
+++ b/Engine/Source/Programs/Horde/Horde.Agent/Horde.Agent.csproj
@@ -136,6 +136,7 @@
     <Compile Include="../../../../Restricted/*/Source/Programs/Horde/Horde.Agent/**/*.cs">
       <Link>$([System.Text.RegularExpressions.Regex]::Replace(%(FullPath), '^.+?[\\/]Horde.Agent[\\/]', ''))</Link>
     </Compile>
+    <Compile Include="$(IntermediateOutputPath)VersionInfo.cs" />
   </ItemGroup>
 
   <ItemGroup>
@@ -205,5 +206,17 @@
     <MakeDir Directories="$(FullSymbolsDir)" />
     <Copy SourceFiles="@(Files)" DestinationFolder="$(FullSymbolsDir)" />
   </Target>
+
+  <!-- Encode version in a source file as inferring version from assemblies is unreliable on platform-specific binaries -->
+  <Target Name="PrepareVersionInfo" BeforeTargets="BeforeBuild">
+    <PropertyGroup>
+      <TemplateFile>Utility/VersionInfo.cs.template</TemplateFile>
+      <OutputFile>$(IntermediateOutputPath)/VersionInfo.cs</OutputFile>
+    </PropertyGroup>
+    <ReadLinesFromFile File="$(TemplateFile)">
+      <Output TaskParameter="Lines" ItemName="Lines"/>
+    </ReadLinesFromFile>
+    <WriteLinesToFile File="$(OutputFile)" Lines="@(Lines->Replace('$VERSION$', '$(InformationalVersion)'))" Overwrite="true" Encoding="UTF-8"/>
+  </Target>
   
 </Project>
diff --git a/Engine/Source/Programs/Horde/Horde.Agent/Utility/VersionInfo.cs.template b/Engine/Source/Programs/Horde/Horde.Agent/Utility/VersionInfo.cs.template
new file mode 100644
--- /dev/null
+++ b/Engine/Source/Programs/Horde/Horde.Agent/Utility/VersionInfo.cs.template
@@ -0,0 +1,24 @@
+// Copyright Epic Games, Inc. All Rights Reserved.
+
+namespace Horde.Agent.Utility;
+
+/// <summary>
+/// Encodes version by replacing the version pattern below
+/// Platform-specific binaries built for Linux and macOS do not encode this information like Windows binaries do
+/// </summary>
+
+public static class VersionInfo
+{
+	/// <summary>
+	/// Version (the pattern is replaced prior to compilation)
+	/// </summary>
+	private const string VersionConstant = "$VERSION$";
+
+	/// <summary>
+	/// Version of this Horde Agent
+	/// </summary>
+	public static string Version
+	{
+		get => String.IsNullOrEmpty(VersionConstant) ? "unknown" : VersionConstant;
+	}
+}
\ No newline at end of file
