Horde: Do not allow multiple agent services to run at once.

diff --git a/Engine/Source/Programs/Horde/Horde.Agent/Commands/Service/ServiceRun.cs b/Engine/Source/Programs/Horde/Horde.Agent/Commands/Service/ServiceRun.cs
--- a/Engine/Source/Programs/Horde/Horde.Agent/Commands/Service/ServiceRun.cs
+++ b/Engine/Source/Programs/Horde/Horde.Agent/Commands/Service/ServiceRun.cs
@@ -42,8 +42,15 @@ namespace Horde.Agent.Commands.Service
 			}
 			else
 			{
-				Console.WriteLine($"Unable to parse log level: {LogLevelStr}");
-				return 0;
+				logger.LogError("Unable to parse log level: {Level}", LogLevelStr);
+				return 1;
+			}
+
+			using SingleInstanceMutex mutex = new SingleInstanceMutex();
+			if (!await mutex.HasMutexAsync())
+			{
+				logger.LogError("Another instance of the Horde Agent is already running.");
+				return 1;
 			}
 
 			IHostBuilder hostBuilder = Host.CreateDefaultBuilder();
@@ -93,5 +100,51 @@ namespace Horde.Agent.Commands.Service
 
 			return 0;
 		}
+
+		class SingleInstanceMutex : IDisposable
+		{
+			Thread _backgroundThread;
+			readonly ManualResetEventSlim _quitEvent = new ManualResetEventSlim();
+			readonly TaskCompletionSource<bool> _waitEvent = new TaskCompletionSource<bool>();
+
+			public SingleInstanceMutex()
+			{
+				_backgroundThread = new Thread(WaitForMutex) { IsBackground = true };
+				_backgroundThread.Start();
+			}
+
+			public Task<bool> HasMutexAsync()
+				=> _waitEvent.Task;
+
+			public void Dispose()
+			{
+				if (_backgroundThread != null)
+				{
+					_quitEvent.Set();
+					_quitEvent.Dispose();
+
+					_backgroundThread.Join();
+					_backgroundThread = null!;
+				}
+			}
+
+			void WaitForMutex()
+			{
+				using Mutex mutex = new Mutex(false, "Horde.Agent.{C1F30772-CDD3-41E9-A6CE-42356DE7DEE3}");
+				try
+				{
+					if (!mutex.WaitOne(0))
+					{
+						_waitEvent.SetResult(false);
+						return;
+					}
+				}
+				catch (AbandonedMutexException)
+				{
+				}
+				_waitEvent.SetResult(true);
+				_quitEvent.Wait();
+			}
+		}
 	}
 }
