PR #11760: Add missing -WorkingDir= in ExecStart for linux service template

#rnx

diff --git a/Engine/Source/Programs/Horde/Docs/Deployment/Agent.md b/Engine/Source/Programs/Horde/Docs/Deployment/Agent.md
--- a/Engine/Source/Programs/Horde/Docs/Deployment/Agent.md
+++ b/Engine/Source/Programs/Horde/Docs/Deployment/Agent.md
@@ -172,7 +172,7 @@ Create a service descriptor file in `/etc/systemd/system/horde-agent.service` (s
     Description=Horde Agent
 
     [Service]
-    ExecStart=dotnet {{ HORDE_PATH }} {{ HORDE_WORKING_DIRECTORY }}
+    ExecStart=dotnet {{ HORDE_PATH }} -WorkingDir={{ HORDE_WORKING_DIRECTORY }}
     WorkingDirectory={{ HORDE_WORKING_DIRECTORY }}
 
     Restart=always
