PR #11760: Fix copy source location during Stage Horde Installer Dashboard.

#rnx

diff --git a/Engine/Source/Programs/Horde/BuildHorde.xml b/Engine/Source/Programs/Horde/BuildHorde.xml
--- a/Engine/Source/Programs/Horde/BuildHorde.xml
+++ b/Engine/Source/Programs/Horde/BuildHorde.xml
@@ -187,7 +187,7 @@
 			<Docker Arguments="create --name horde-installer-container hordedashboard:installer"/>
 			
 			<!-- Copy the dashboard build to the server -->
-			<Docker Arguments="cp horde-installer-container:/app/Dashboard/build/. &quot;$(StagingDir)/Dashboard&quot;"/>
+			<Docker Arguments="cp horde-installer-container:/app/Dashboard/dist/. &quot;$(StagingDir)/Dashboard&quot;"/>
 
 			<!-- Remove the container now that files are copied -->
 			<Docker Arguments="rm horde-installer-container"/>
