Horde: Add debug endpoint for job timings

Returns a table for time taken by batch setup, work and teardown. Grouped by stream and job template.

#rnx

diff --git a/Engine/Source/Programs/Horde/Horde.Server/Server/DebugController.cs b/Engine/Source/Programs/Horde/Horde.Server/Server/DebugController.cs
--- a/Engine/Source/Programs/Horde/Horde.Server/Server/DebugController.cs
+++ b/Engine/Source/Programs/Horde/Horde.Server/Server/DebugController.cs
@@ -23,6 +23,7 @@ using Horde.Common.Rpc;
 using Horde.Server.Agents.Relay;
 using Horde.Server.Configuration;
 using Horde.Server.Jobs;
+using Horde.Server.Jobs.Graphs;
 using Horde.Server.Logs;
 using Horde.Server.Projects;
 using Horde.Server.Streams;
@@ -599,6 +600,187 @@ namespace Horde.Server.Server
 			return new ContentResult { ContentType = "text/csv", StatusCode = (int)HttpStatusCode.OK, Content = sb.ToString() };
 		}
 
+		record JobTiming(
+			string StreamId,
+			string TemplateId,
+			string Name,
+			IReadOnlySet<string> StepNames,
+			TimeSpan BatchSetupDuration,
+			TimeSpan BatchWorkDuration,
+			TimeSpan BatchTeardownDuration)
+		{
+			public JobTiming Merge(JobTiming jt)
+			{
+				if (StreamId != jt.StreamId || TemplateId != jt.TemplateId)
+				{
+					throw new ArgumentException("StreamId or TemplateId do not match");
+				}
+
+				HashSet<string> newNames = [..StepNames.Union(jt.StepNames)];
+				return new JobTiming(StreamId, TemplateId, Name, newNames, 
+					BatchSetupDuration + jt.BatchSetupDuration,
+					BatchWorkDuration + jt.BatchWorkDuration,
+					BatchTeardownDuration + jt.BatchTeardownDuration);
+			}
+		}
+		
+		/// <summary>
+		/// Display a table listing each template with job timings (setup, work and teardown durations)
+		/// </summary>
+		/// <returns>Async task</returns>
+		[HttpGet]
+		[Route("/api/v1/debug/job-timings")]
+		public async Task<ActionResult> GetJobTimingsAsync(
+			[FromQuery] DateTimeOffset? minCreateTime = null,
+			[FromQuery] DateTimeOffset? maxCreateTime = null,
+			[FromQuery] bool onlySetupBuild = false,
+			[FromQuery] string? format = "html")
+		{
+			if (!_globalConfig.Value.Authorize(ServerAclAction.Debug, User))
+			{
+				return Forbid(ServerAclAction.Debug);
+			}
+
+			minCreateTime ??= DateTimeOffset.UtcNow.Subtract(TimeSpan.FromDays(1));
+			IReadOnlyList<IJob> jobs = await _jobService.FindJobsAsync(minCreateTime: minCreateTime, maxCreateTime: maxCreateTime);
+			List<JobTiming> allJobTimings = await CalculateJobTimingsAsync(jobs);
+
+			if (onlySetupBuild)
+			{
+				allJobTimings = allJobTimings.Where(x => x.StepNames.SetEquals(["Setup Build"])).ToList();
+			}
+			
+			IReadOnlyList<JobTiming> jobTimingsByTemplate = GroupJobTimings(allJobTimings);
+			IEnumerable<JobTiming> sortedJobTimings = jobTimingsByTemplate.OrderBy(x => x.BatchSetupDuration).Reverse();
+
+			if (format == "csv")
+			{
+				return GetJobTimingsAsCsv(sortedJobTimings);
+			}
+
+			StringBuilder sb = new();
+
+			sb.AppendLine("<style>");
+			sb.AppendLine("body { font-family: 'Helvetica Neue', Helvetica, Arial, sans-serif; }");
+			sb.AppendLine("table { border-collapse: collapse; width: 100%; font-size: 12px; }");
+			sb.AppendLine("th, td { border: 1px solid black; text-align: left; padding: 8px; }");
+			sb.AppendLine("th { background-color: #f2f2f2; }");
+			sb.AppendLine("</style>");
+
+			sb.AppendLine("<h1>Job timings by template</h1>");
+			sb.AppendLine("<p>Durations specified in seconds.</p>");
+			sb.AppendLine("<table>");
+			sb.AppendLine("<thead><tr>");
+			sb.Append("<th>Stream</th>");
+			sb.Append("<th>Template</th>");
+			sb.Append("<th>Name</th>");
+			sb.Append("<th>Batch Setup</th>");
+			sb.Append("<th>Batch Work</th>");
+			sb.Append("<th>Batch Teardown</th>");
+			sb.Append("<th>Steps</th>");
+			sb.AppendLine("</tr></thead>");
+
+			foreach (JobTiming jt in sortedJobTimings)
+			{
+				sb.AppendLine("<tr>");
+				sb.Append($"<td>{jt.StreamId}</td>");
+				sb.Append($"<td>{jt.TemplateId}</td>");
+				sb.Append($"<td>{jt.Name}</td>");
+				sb.Append($"<td>{(int)jt.BatchSetupDuration.TotalSeconds}</td>");
+				sb.Append($"<td>{(int)jt.BatchWorkDuration.TotalSeconds}</td>");
+				sb.Append($"<td>{(int)jt.BatchTeardownDuration.TotalSeconds}</td>");
+				sb.Append($"<td>{String.Join(", ", jt.StepNames.Order())}</td>");
+				sb.AppendLine("</tr>");
+			}
+
+			sb.AppendLine("</table>");
+			return new ContentResult { ContentType = "text/html", StatusCode = (int)HttpStatusCode.OK, Content = sb.ToString() };
+		}
+		
+		private ActionResult GetJobTimingsAsCsv(IEnumerable<JobTiming> jobTimings)
+		{
+			StringBuilder sb = new();
+			sb.AppendJoin('\t', ["Stream", "Template", "Name", "Batch Setup", "Batch Work", "Batch Teardown", "Steps"]).AppendLine();
+
+			foreach (JobTiming jt in jobTimings)
+			{
+				sb.Append($"{jt.StreamId}\t");
+				sb.Append($"{jt.TemplateId}\t");
+				sb.Append($"{jt.Name}\t");
+				sb.Append($"{(int)jt.BatchSetupDuration.TotalSeconds}\t");
+				sb.Append($"{(int)jt.BatchWorkDuration.TotalSeconds}\t");
+				sb.Append($"{(int)jt.BatchTeardownDuration.TotalSeconds}\t");
+				sb.AppendJoin(',', jt.StepNames.Order());
+				sb.AppendLine();
+			}
+
+			return new ContentResult { ContentType = "text/csv", StatusCode = (int)HttpStatusCode.OK, Content = sb.ToString() };
+		}
+
+		private async Task<List<JobTiming>> CalculateJobTimingsAsync(IEnumerable<IJob> jobs)
+		{
+			List<JobTiming> timings = [];
+			foreach (IJob job in jobs)
+			{
+				IGraph graph = await _jobService.GetGraphAsync(job);
+				
+				foreach (IJobStepBatch batch in job.Batches)
+				{
+					if (batch.State != JobStepBatchState.Complete || batch.StartTimeUtc == null || batch.FinishTimeUtc == null)
+					{
+						continue;
+					}
+					DateTime firstStepStartTime = DateTime.MaxValue;
+					DateTime lastStepFinishTime = DateTime.MinValue;
+
+					HashSet<string> stepNames = [];
+					foreach (IJobStep step in batch.Steps)
+					{
+						if (step.StartTimeUtc == null || step.FinishTimeUtc == null)
+						{
+							continue;
+						}
+						
+						INode node = graph.GetNode(new NodeRef(batch.GroupIdx, step.NodeIdx));
+						firstStepStartTime = step.StartTimeUtc.Value < firstStepStartTime ? step.StartTimeUtc.Value : firstStepStartTime;
+						lastStepFinishTime = step.FinishTimeUtc.Value > lastStepFinishTime ? step.FinishTimeUtc.Value : lastStepFinishTime;
+						stepNames.Add(node.Name);
+					}
+
+					if (firstStepStartTime == DateTime.MaxValue || lastStepFinishTime == DateTime.MinValue)
+					{
+						continue;
+					}
+
+					TimeSpan batchSetupDuration = firstStepStartTime - batch.StartTimeUtc.Value;
+					TimeSpan batchWorkDuration = lastStepFinishTime - firstStepStartTime;
+					TimeSpan batchTeardownDuration = batch.FinishTimeUtc.Value - lastStepFinishTime;
+					timings.Add(new JobTiming(job.StreamId.ToString(), job.TemplateId.ToString(), job.Name, stepNames, batchSetupDuration, batchWorkDuration, batchTeardownDuration));
+				}
+			}
+
+			return timings;
+		}
+
+		private static IReadOnlyList<JobTiming> GroupJobTimings(IEnumerable<JobTiming> jobTimings)
+		{
+			Dictionary<string, JobTiming> groupedTimings = new();
+			foreach (JobTiming timing in jobTimings)
+			{
+				string key = $"{timing.StreamId}-{timing.TemplateId}";
+				if (!groupedTimings.TryGetValue(key, out JobTiming? groupTiming))
+				{
+					groupedTimings[key] = timing;
+				}
+				else
+				{
+					groupedTimings[key] = groupTiming.Merge(timing);
+				}
+			}
+
+			return groupedTimings.Values.ToList();
+		}
+
 		/// <summary>
 		/// Populate the database with test data
 		/// </summary>
