Horde: Fix warning

#rnx

diff --git a/Engine/Source/Programs/Horde/Horde.Server/Server/DebugController.cs b/Engine/Source/Programs/Horde/Horde.Server/Server/DebugController.cs
--- a/Engine/Source/Programs/Horde/Horde.Server/Server/DebugController.cs
+++ b/Engine/Source/Programs/Horde/Horde.Server/Server/DebugController.cs
@@ -697,7 +697,7 @@ namespace Horde.Server.Server
 			return new ContentResult { ContentType = "text/html", StatusCode = (int)HttpStatusCode.OK, Content = sb.ToString() };
 		}
 		
-		private ActionResult GetJobTimingsAsCsv(IEnumerable<JobTiming> jobTimings)
+		private static ActionResult GetJobTimingsAsCsv(IEnumerable<JobTiming> jobTimings)
 		{
 			StringBuilder sb = new();
 			sb.AppendJoin('\t', ["Stream", "Template", "Name", "Batch Setup", "Batch Work", "Batch Teardown", "Steps"]).AppendLine();
