Horde: Include the changelist number and users that modified the  files in config update notifications.

diff --git a/Engine/Source/Programs/Horde/Horde.Server/Configuration/ConfigService.cs b/Engine/Source/Programs/Horde/Horde.Server/Configuration/ConfigService.cs
--- a/Engine/Source/Programs/Horde/Horde.Server/Configuration/ConfigService.cs
+++ b/Engine/Source/Programs/Horde/Horde.Server/Configuration/ConfigService.cs
@@ -10,6 +10,7 @@ using System.Text.Json;
 using System.Threading;
 using System.Threading.Tasks;
 using EpicGames.Core;
+using EpicGames.Horde.Users;
 using EpicGames.Redis;
 using Horde.Server.Projects;
 using Horde.Server.Server;
@@ -30,7 +31,7 @@ namespace Horde.Server.Configuration
 	/// <summary>
 	/// Information about the updated config
 	/// </summary>
-	public record class ConfigUpdateInfo(Dictionary<Uri, string> Sources, Exception? Exception);
+	public record class ConfigUpdateInfo(List<string> Status, HashSet<UserId> Authors, Exception? Exception);
 
 	/// <summary>
 	/// Service which processes runtime configuration data.
@@ -217,7 +218,7 @@ namespace Horde.Server.Configuration
 
 			public void Add(OverrideConfigFile file) => _overrides.Add(file.Uri, file);
 
-			public async Task<IConfigFile[]> GetAsync(Uri[] uris, CancellationToken cancellationToken)
+			public async Task<IConfigFile[]> GetFilesAsync(Uri[] uris, CancellationToken cancellationToken)
 			{
 				IConfigFile[] files = new IConfigFile[uris.Length];
 				for (int idx = 0; idx < uris.Length; idx++)
@@ -229,12 +230,16 @@ namespace Horde.Server.Configuration
 					}
 					else
 					{
-						file = (await _inner.GetAsync(new[] { uris[idx] }, cancellationToken))[0];
+						file = (await _inner.GetFilesAsync(new[] { uris[idx] }, cancellationToken))[0];
 					}
 					files[idx] = file;
 				}
 				return files;
 			}
+
+			/// <inheritdoc/>
+			public Task GetUpdateInfoAsync(IReadOnlyDictionary<Uri, string> files, IReadOnlyDictionary<Uri, string>? prevFiles, ConfigUpdateInfo updateInfo, CancellationToken cancellationToken)
+				=> Task.CompletedTask;
 		}
 
 		/// <summary>
@@ -393,7 +398,7 @@ namespace Horde.Server.Configuration
 			{
 				if (snapshot == null || await IsOutOfDateAsync(snapshot, cancellationToken))
 				{
-					snapshot = await CreateSnapshotGuardedAsync(cancellationToken);
+					snapshot = await CreateSnapshotGuardedAsync(snapshot?.Dependencies, cancellationToken);
 					if (snapshot != null)
 					{
 						try
@@ -445,26 +450,37 @@ namespace Horde.Server.Configuration
 		/// <summary>
 		/// Create a new snapshot object, catching any exceptions and sending notifications
 		/// </summary>
+		/// <param name="prevFiles">Previous snapshot data</param>
 		/// <param name="cancellationToken">Cancellation token for the operation</param>
 		/// <returns>New config snapshot</returns>
-		async Task<ConfigSnapshot?> CreateSnapshotGuardedAsync(CancellationToken cancellationToken)
+		async Task<ConfigSnapshot?> CreateSnapshotGuardedAsync(IReadOnlyDictionary<Uri, string>? prevFiles, CancellationToken cancellationToken)
 		{
 			try
 			{
 				ConfigSnapshot snapshot = await CreateSnapshotAsync(cancellationToken);
 				await _health.UpdateAsync(HealthStatus.Healthy);
-				OnConfigUpdate?.Invoke(new ConfigUpdateInfo(snapshot.Dependencies, null));
+				await NotifyConfigUpdateAsync(snapshot.Dependencies, prevFiles, null, cancellationToken);
 				return snapshot;
 			}
-			catch (Exception ex)
+			catch (ConfigException ex)
 			{
 				_logger.LogError(ex, "Exception while updating config: {Message}", ex.Message);
 				await _health.UpdateAsync(HealthStatus.Unhealthy, ex.Message);
-				OnConfigUpdate?.Invoke(new ConfigUpdateInfo(new Dictionary<Uri, string>(), ex));
+				await NotifyConfigUpdateAsync(ex.GetContext().Files.ToDictionary(x => x.Key, x => x.Value.Revision), prevFiles, ex, cancellationToken);
 				return null;
 			}
 		}
 
+		async Task NotifyConfigUpdateAsync(IReadOnlyDictionary<Uri, string> files, IReadOnlyDictionary<Uri, string>? prevFiles, Exception? exception, CancellationToken cancellationToken)
+		{
+			ConfigUpdateInfo info = new ConfigUpdateInfo(new List<string>(), new HashSet<UserId>(), exception);
+			foreach (IConfigSource source in _sources.Values)
+			{
+				await source.GetUpdateInfoAsync(files, prevFiles, info, cancellationToken);
+			}
+			OnConfigUpdate?.Invoke(info);
+		}
+
 		internal Uri GetGlobalConfigUri()
 		{
 			return GetGlobalConfigUri(_serverSettings.ConfigPath, ServerApp.ConfigDir.FullName);
@@ -631,7 +647,7 @@ namespace Horde.Server.Configuration
 						return true;
 					}
 
-					IConfigFile[] files = await source.GetAsync(pairs.ConvertAll(x => x.Key), cancellationToken);
+					IConfigFile[] files = await source.GetFilesAsync(pairs.ConvertAll(x => x.Key), cancellationToken);
 					for (int idx = 0; idx < pairs.Length; idx++)
 					{
 						if (!files[idx].Revision.Equals(pairs[idx].Value, StringComparison.Ordinal))
diff --git a/Engine/Source/Programs/Horde/Horde.Server/Configuration/ConfigSource.cs b/Engine/Source/Programs/Horde/Horde.Server/Configuration/ConfigSource.cs
--- a/Engine/Source/Programs/Horde/Horde.Server/Configuration/ConfigSource.cs
+++ b/Engine/Source/Programs/Horde/Horde.Server/Configuration/ConfigSource.cs
@@ -8,7 +8,9 @@ using System.Linq;
 using System.Threading;
 using System.Threading.Tasks;
 using EpicGames.Core;
+using EpicGames.Horde.Users;
 using EpicGames.Perforce;
+using Horde.Server.Accounts;
 using Horde.Server.Users;
 using Microsoft.Extensions.Caching.Memory;
 using Microsoft.Extensions.Logging;
@@ -37,7 +39,16 @@ namespace Horde.Server.Configuration
 		/// <param name="uris">Locations of the config files to query</param>
 		/// <param name="cancellationToken">Cancellation token for the operation</param>
 		/// <returns>Config file data</returns>
-		Task<IConfigFile[]> GetAsync(Uri[] uris, CancellationToken cancellationToken);
+		Task<IConfigFile[]> GetFilesAsync(Uri[] uris, CancellationToken cancellationToken);
+
+		/// <summary>
+		/// Gets summary infomration for sending notifications
+		/// </summary>
+		/// <param name="files">New files for the configuration</param>
+		/// <param name="prevFiles">Previous set of files for the configuration</param>
+		/// <param name="updateInfo">Information about the update</param>
+		/// <param name="cancellationToken">Cancellation token for the operation</param>
+		Task GetUpdateInfoAsync(IReadOnlyDictionary<Uri, string> files, IReadOnlyDictionary<Uri, string>? prevFiles, ConfigUpdateInfo updateInfo, CancellationToken cancellationToken);
 	}
 
 	/// <summary>
@@ -54,7 +65,7 @@ namespace Horde.Server.Configuration
 		/// <returns>Config file data</returns>
 		public static async Task<IConfigFile> GetAsync(this IConfigSource source, Uri uri, CancellationToken cancellationToken)
 		{
-			IConfigFile[] result = await source.GetAsync(new[] { uri }, cancellationToken);
+			IConfigFile[] result = await source.GetFilesAsync(new[] { uri }, cancellationToken);
 			return result[0];
 		}
 	}
@@ -105,7 +116,7 @@ namespace Horde.Server.Configuration
 		}
 
 		/// <inheritdoc/>
-		public Task<IConfigFile[]> GetAsync(Uri[] uris, CancellationToken cancellationToken)
+		public Task<IConfigFile[]> GetFilesAsync(Uri[] uris, CancellationToken cancellationToken)
 		{
 			IConfigFile[] result = new IConfigFile[uris.Length];
 			for (int idx = 0; idx < uris.Length; idx++)
@@ -119,6 +130,10 @@ namespace Horde.Server.Configuration
 			}
 			return Task.FromResult(result);
 		}
+
+		/// <inheritdoc/>
+		public Task GetUpdateInfoAsync(IReadOnlyDictionary<Uri, string> files, IReadOnlyDictionary<Uri, string>? prevFiles, ConfigUpdateInfo updateInfo, CancellationToken cancellationToken)
+			=> Task.CompletedTask;
 	}
 
 	/// <summary>
@@ -177,7 +192,7 @@ namespace Horde.Server.Configuration
 		}
 
 		/// <inheritdoc/>
-		public async Task<IConfigFile[]> GetAsync(Uri[] uris, CancellationToken cancellationToken)
+		public async Task<IConfigFile[]> GetFilesAsync(Uri[] uris, CancellationToken cancellationToken)
 		{
 			IConfigFile[] files = new IConfigFile[uris.Length];
 			for (int idx = 0; idx < uris.Length; idx++)
@@ -218,6 +233,10 @@ namespace Horde.Server.Configuration
 			}
 			return files;
 		}
+
+		/// <inheritdoc/>
+		public Task GetUpdateInfoAsync(IReadOnlyDictionary<Uri, string> files, IReadOnlyDictionary<Uri, string>? prevFiles, ConfigUpdateInfo updateInfo, CancellationToken cancellationToken)
+			=> Task.CompletedTask;
 	}
 
 	/// <summary>
@@ -274,7 +293,7 @@ namespace Horde.Server.Configuration
 		}
 
 		/// <inheritdoc/>
-		public async Task<IConfigFile[]> GetAsync(Uri[] uris, CancellationToken cancellationToken)
+		public async Task<IConfigFile[]> GetFilesAsync(Uri[] uris, CancellationToken cancellationToken)
 		{
 			Dictionary<Uri, IConfigFile> results = new Dictionary<Uri, IConfigFile>();
 			foreach (IGrouping<string, Uri> group in uris.GroupBy(x => x.Host))
@@ -373,5 +392,44 @@ namespace Horde.Server.Configuration
 			IPerforceConnection connection = await PerforceConnection.CreateAsync(connectionSettings.ToPerforceSettings(), _logger);
 			return connection;
 		}
+
+		/// <inheritdoc/>
+		public async Task GetUpdateInfoAsync(IReadOnlyDictionary<Uri, string> files, IReadOnlyDictionary<Uri, string>? prevFiles, ConfigUpdateInfo updateInfo, CancellationToken cancellationToken)
+		{
+			foreach (IGrouping<string, KeyValuePair<Uri, string>> group in files.Where(x => x.Key.Scheme == PerforceConfigSource.Scheme).GroupBy(x => x.Key.Host))
+			{
+				// Figure out the most recent changelist
+				int change = group.Select(x => Int32.Parse(x.Value)).Max();
+				updateInfo.Status.Add($"Perforce changelist ({group.Key}): {change}");
+
+				// Find all the users that have committed changes
+				if (prevFiles != null)
+				{
+					using IPerforceConnection perforce = await ConnectAsync(group.Key, cancellationToken);
+					foreach ((Uri uri, string revision) in group)
+					{
+						string? prevRevision;
+						if (prevFiles.TryGetValue(uri, out prevRevision))
+						{
+							string fileSpec = $"{uri.LocalPath}@{prevRevision},{revision}";
+							await FindAuthorsAsync(perforce, fileSpec, updateInfo.Authors, cancellationToken);
+						}
+					}
+				}
+			}
+		}
+
+		async Task FindAuthorsAsync(IPerforceConnection perforce, string fileSpec, HashSet<UserId> authors, CancellationToken cancellationToken)
+		{
+			List<ChangesRecord> records = await perforce.GetChangesAsync(ChangesOptions.None, -1, ChangeStatus.Submitted, fileSpec, cancellationToken);
+			foreach (ChangesRecord record in records)
+			{
+				IUser? user = await _userCollection.FindUserByLoginAsync(record.User, cancellationToken);
+				if (user != null)
+				{
+					authors.Add(user.Id);
+				}
+			}
+		}
 	}
 }
diff --git a/Engine/Source/Programs/Horde/Horde.Server/Notifications/Sinks/SlackNotificationSink.cs b/Engine/Source/Programs/Horde/Horde.Server/Notifications/Sinks/SlackNotificationSink.cs
--- a/Engine/Source/Programs/Horde/Horde.Server/Notifications/Sinks/SlackNotificationSink.cs
+++ b/Engine/Source/Programs/Horde/Horde.Server/Notifications/Sinks/SlackNotificationSink.cs
@@ -2157,39 +2157,25 @@ namespace Horde.Server.Notifications.Sinks
 			}
 			else
 			{
-				Dictionary<string, int> perforceClusterToChange = new Dictionary<string, int>();
-				foreach (Uri uri in info.Sources.Keys)
-				{
-					if (uri.Scheme == PerforceConfigSource.Scheme)
-					{
-						string path = uri.LocalPath;
-
-						int atIdx = path.IndexOf('@', StringComparison.Ordinal);
-						if (atIdx != -1 && Int32.TryParse(path.Substring(atIdx + 1), out int change))
-						{
-							int existingChange;
-							if (!perforceClusterToChange.TryGetValue(uri.Host, out existingChange) || change > existingChange)
-							{
-								perforceClusterToChange[uri.Host] = change;
-							}
-						}
-					}
-				}
-
 				// Can only send updates after failures by checking the result of this delete
 				await DeleteMessageStateAsync(_settings.ConfigNotificationChannel, EventId, cancellationToken);
 
 				SlackMessage message = new SlackMessage();
 				message.AddSection($"*Config Update Succeeded*");
 
-				if(perforceClusterToChange.Count > 0)
+				if (info.Status.Count > 0)
+				{
+					message.AddSection(String.Join("\n", info.Status));
+				}
+				if (info.Authors.Count > 0)
 				{
-					List<string> lines = new List<string>();
-					foreach ((string clusterName, int change) in perforceClusterToChange)
+					List<string> mentions = new List<string>();
+					foreach (UserId userId in info.Authors)
 					{
-						lines.Add($"Perforce changelist ({clusterName}): {change}");
+						string mention = await FormatMentionAsync(userId, DefaultAllowMentions, cancellationToken);
+						mentions.Add(mention);
 					}
-					message.AddSection(String.Join("\n", lines));
+					message.AddSection($"cc {String.Join(", ", mentions)}");
 				}
 
 				await SendMessageAsync(_settings.ConfigNotificationChannel, message, cancellationToken);
