Horde: Fix revision range for config updates.

#rnx

diff --git a/Engine/Source/Programs/Horde/Horde.Server/Configuration/ConfigSource.cs b/Engine/Source/Programs/Horde/Horde.Server/Configuration/ConfigSource.cs
--- a/Engine/Source/Programs/Horde/Horde.Server/Configuration/ConfigSource.cs
+++ b/Engine/Source/Programs/Horde/Horde.Server/Configuration/ConfigSource.cs
@@ -10,7 +10,6 @@ using System.Threading.Tasks;
 using EpicGames.Core;
 using EpicGames.Horde.Users;
 using EpicGames.Perforce;
-using Horde.Server.Accounts;
 using Horde.Server.Users;
 using Microsoft.Extensions.Caching.Memory;
 using Microsoft.Extensions.Logging;
@@ -409,9 +408,10 @@ namespace Horde.Server.Configuration
 					foreach ((Uri uri, string revision) in group)
 					{
 						string? prevRevision;
-						if (prevFiles.TryGetValue(uri, out prevRevision))
+						if (prevFiles.TryGetValue(uri, out prevRevision) && prevRevision != revision)
 						{
-							string fileSpec = $"{uri.LocalPath}@{prevRevision},{revision}";
+							int prevRevisionNum = Int32.Parse(prevRevision);
+							string fileSpec = $"{uri.LocalPath}@{prevRevisionNum + 1},{revision}";
 							await FindAuthorsAsync(perforce, fileSpec, updateInfo.Authors, cancellationToken);
 						}
 					}
