Horde: Send DMs to individual users rather than the notification channel when config updates succeed.

diff --git a/Engine/Source/Programs/Horde/Horde.Server/Notifications/Sinks/SlackNotificationSink.cs b/Engine/Source/Programs/Horde/Horde.Server/Notifications/Sinks/SlackNotificationSink.cs
--- a/Engine/Source/Programs/Horde/Horde.Server/Notifications/Sinks/SlackNotificationSink.cs
+++ b/Engine/Source/Programs/Horde/Horde.Server/Notifications/Sinks/SlackNotificationSink.cs
@@ -2157,29 +2157,32 @@ namespace Horde.Server.Notifications.Sinks
 			}
 			else
 			{
-				// Can only send updates after failures by checking the result of this delete
-				await DeleteMessageStateAsync(_settings.ConfigNotificationChannel, EventId, cancellationToken);
-
 				SlackMessage message = new SlackMessage();
 				message.AddSection($"*Config Update Succeeded*");
-
 				if (info.Status.Count > 0)
 				{
 					message.AddSection(String.Join("\n", info.Status));
 				}
-				if (info.Authors.Count > 0)
+
+				// Send a notification to the channel
+				if (await DeleteMessageStateAsync(_settings.ConfigNotificationChannel, EventId, cancellationToken))
 				{
-					List<string> mentions = new List<string>();
-					foreach (UserId userId in info.Authors)
+					await SendMessageAsync(_settings.ConfigNotificationChannel, message, cancellationToken);
+				}
+
+				// Notify individual authors that their updates have succeeded.
+				foreach (UserId authorId in info.Authors)
+				{
+					IUser? user = await _userCollection.GetUserAsync(authorId, cancellationToken);
+					if (user != null)
 					{
-						string mention = await FormatMentionAsync(userId, DefaultAllowMentions, cancellationToken);
-						mentions.Add(mention);
+						string? slackUserId = await GetSlackUserIdAsync(user, cancellationToken);
+						if (slackUserId != null)
+						{
+							await SendMessageAsync(slackUserId, message, cancellationToken);
+						}
 					}
-					message.AddSection($"cc {String.Join(", ", mentions)}");
 				}
-
-				await SendMessageAsync(_settings.ConfigNotificationChannel, message, cancellationToken);
-				await DeleteMessageStateAsync(_settings.ConfigNotificationChannel, EventId, cancellationToken);
 			}
 		}
 
