Horde: Fix agent registration not being invalidated when server is reinstalled. Server returns "unauthenticated" (401) response, not "forbidden" (403). Forbidden is only returned when the server deliberately invalidates an agent.

diff --git a/Engine/Source/Programs/Horde/Horde.Agent/Services/SessionFactory.cs b/Engine/Source/Programs/Horde/Horde.Agent/Services/SessionFactory.cs
--- a/Engine/Source/Programs/Horde/Horde.Agent/Services/SessionFactory.cs
+++ b/Engine/Source/Programs/Horde/Horde.Agent/Services/SessionFactory.cs
@@ -225,7 +225,7 @@ namespace Horde.Agent.Services
 					createSessionResponse = await rpcClient.CreateSessionAsync(sessionRequest, null, null, cancellationToken);
 					logger.LogInformation("Created session. AgentName={AgentName} SessionId={SessionId}", currentSettings.GetAgentName(), createSessionResponse.SessionId);
 				}
-				catch (RpcException ex) when (ex.StatusCode == StatusCode.PermissionDenied)
+				catch (RpcException ex) when (ex.StatusCode == StatusCode.PermissionDenied || ex.StatusCode == StatusCode.Unauthenticated)
 				{
 					Uri serverUrl = grpcService.ServerProfile.Url;
 					if (registrationList.Entries.RemoveAll(x => x.Server == serverUrl) > 0)
