Horde: Add dashboard's time since last user activity to HTTP log

diff --git a/Engine/Source/Programs/Horde/Horde.Server/Startup.cs b/Engine/Source/Programs/Horde/Horde.Server/Startup.cs
--- a/Engine/Source/Programs/Horde/Horde.Server/Startup.cs
+++ b/Engine/Source/Programs/Horde/Horde.Server/Startup.cs
@@ -114,6 +114,7 @@ using Microsoft.Extensions.Hosting;
 using Microsoft.Extensions.Logging;
 using Microsoft.Extensions.Logging.Abstractions;
 using Microsoft.Extensions.Options;
+using Microsoft.Extensions.Primitives;
 using Microsoft.OpenApi.Models;
 using MongoDB.Bson;
 using MongoDB.Bson.Serialization;
@@ -1289,6 +1290,16 @@ namespace Horde.Server
 				options.EnrichDiagnosticContext = (diagnosticContext, httpContext) =>
 				{
 					diagnosticContext.Set("RemoteIP", httpContext?.Connection?.RemoteIpAddress);
+					
+					// Header sent by the dashboard to indicate how long a user has been inactive for a particular browser page (in seconds)
+					if (httpContext?.Request.Headers.TryGetValue("X-Horde-LastUserActivity", out StringValues values) is true)
+					{
+						string? value = values.FirstOrDefault();
+						if (!String.IsNullOrEmpty(value) && Int32.TryParse(value, out int lastUserActivity))
+						{
+							diagnosticContext.Set("LastUserActivity", lastUserActivity);
+						}
+					}
 				};
 			});
 
