Horde: Update release notes.

#rnx

diff --git a/Engine/Source/Programs/Horde/Docs/ReleaseNotes.md b/Engine/Source/Programs/Horde/Docs/ReleaseNotes.md
--- a/Engine/Source/Programs/Horde/Docs/ReleaseNotes.md
+++ b/Engine/Source/Programs/Horde/Docs/ReleaseNotes.md
@@ -1,5 +1,15 @@
 # Release Notes
 
+## 2024-04-18
+
+* Add dashboard's time since last user activity to HTTP log (33068679)
+* Fix agent registration not being invalidated when server is reinstalled. Server returns "unauthenticated" (401) response, not "forbidden" (403). Forbidden is only returned when the server deliberately invalidates an agent. (33060431)
+* Send DMs to individual users rather than the notification channel when config updates succeed. (33060054)
+* Prevent hypens before a period in a sanitized string id. (33044611)
+* EpicGames.Perforce: Fix parsing of records with multiple arrays. (33038480)
+* Do not allow multiple agent services to run at once. (33037301)
+* Handle zero byte reads due to socket shutdown when reading messages into a compute buffer. (33034396)
+
 ## 2024-04-12
 
 * Ensure blobs uploaded to the tools endpoint have the tool id as a prefix to the blob locator (32907981)
