Horde: Disable the default URLs configured by the ASP.NET docker image to fix warnings during server startup.

diff --git a/Engine/Source/Programs/Horde/Horde.Server/Dockerfile b/Engine/Source/Programs/Horde/Horde.Server/Dockerfile
--- a/Engine/Source/Programs/Horde/Horde.Server/Dockerfile
+++ b/Engine/Source/Programs/Horde/Horde.Server/Dockerfile
@@ -77,5 +77,6 @@ COPY --from=build-env /tmp/empty /tmp/dotcover-report/*.zip /
 # Create the runtime image
 FROM mcr.microsoft.com/dotnet/aspnet:8.0 AS runtime
 WORKDIR /app
+ENV ASPNETCORE_URLS=""
 COPY --from=build-env /app/out .
 ENTRYPOINT [ "dotnet", "HordeServer.dll" ]
