Horde: Downgrade obsolete endpoint logging to debug verbosity.

#rnx

diff --git a/Engine/Source/Programs/Horde/Horde.Server/Startup.cs b/Engine/Source/Programs/Horde/Horde.Server/Startup.cs
--- a/Engine/Source/Programs/Horde/Horde.Server/Startup.cs
+++ b/Engine/Source/Programs/Horde/Horde.Server/Startup.cs
@@ -316,7 +316,7 @@ namespace Horde.Server
 					|| actionDescriptor?.MethodInfo?.DeclaringType?.GetCustomAttribute<ObsoleteAttribute>() != null)
 				{
 					ILogger? logger = context.HttpContext.RequestServices.GetService<ILogger<ObsoleteLoggingFilter>>();
-					logger?.LogWarning("Using obsolete endpoint: {RequestPath} (Source: {RemoteIp})", context.HttpContext.Request.Path, context.HttpContext.Connection.RemoteIpAddress);
+					logger?.LogDebug("Using obsolete endpoint: {RequestPath} (Source: {RemoteIp})", context.HttpContext.Request.Path, context.HttpContext.Connection.RemoteIpAddress);
 				}
 			}
 		}
