Clear the default ASPNETCORE_HTTP_PORTS environment varialbe to prevent warnings about port overrides on server startup. We manage HTTP port configuration through Horde settings, and don't want the NET runtime image defaults.

diff --git a/Engine/Source/Programs/Horde/Horde.Server/Dockerfile b/Engine/Source/Programs/Horde/Horde.Server/Dockerfile
--- a/Engine/Source/Programs/Horde/Horde.Server/Dockerfile
+++ b/Engine/Source/Programs/Horde/Horde.Server/Dockerfile
@@ -78,5 +78,6 @@ COPY --from=build-env /tmp/empty /tmp/dotcover-report/*.zip /
 FROM mcr.microsoft.com/dotnet/aspnet:8.0 AS runtime
 WORKDIR /app
 ENV ASPNETCORE_URLS=""
+ENV ASPNETCORE_HTTP_PORTS=""
 COPY --from=build-env /app/out .
 ENTRYPOINT [ "dotnet", "HordeServer.dll" ]
