Horde: Add ReleaseNotes tool to regular solution. Remove a few hard-coded constants.

diff --git a/Engine/Source/Programs/Horde/Tools/ReleaseNotes/Program.cs b/Engine/Source/Programs/Horde/Tools/ReleaseNotes/Program.cs
new file mode 100644
--- /dev/null
+++ b/Engine/Source/Programs/Horde/Tools/ReleaseNotes/Program.cs
@@ -0,0 +1,161 @@
+// Copyright Epic Games, Inc. All Rights Reserved.
+
+using EpicGames.Core;
+using EpicGames.Horde;
+using EpicGames.Horde.Server;
+using EpicGames.Perforce;
+using Microsoft.Extensions.DependencyInjection;
+using Microsoft.Extensions.Logging;
+using System.Text;
+using System.Text.RegularExpressions;
+
+namespace CreateReleaseNotes
+{
+	internal class Program
+	{
+		static readonly string[] s_paths =
+		{
+			$"Engine/Source/Programs/Horde/...",
+			$"Engine/Source/Programs/Shared/...",
+		};
+
+		record class ChangeInfo(int Number, string Description, string? JiraTicket);
+
+		class Options
+		{
+			[CommandLine("-Server=", Required = true)]
+			public Uri Server { get; set; } = null!;
+
+			[CommandLine("-MinChange=")]
+			public int? MinChange { get; set; }
+		}
+
+		static async Task<int> Main(string[] args)
+		{
+			DefaultConsoleLogger defaultLogger = new DefaultConsoleLogger();
+
+			CommandLineArguments arguments = new CommandLineArguments(args);
+			Options options = arguments.ApplyTo<Options>(defaultLogger);
+			arguments.CheckAllArgumentsUsed(defaultLogger);
+
+			ServiceCollection serviceCollection = new ServiceCollection();
+			serviceCollection.AddLogging(builder =>
+			{
+				builder.AddEpicDefault();
+				builder.AddFilter("System.Net.Http.HttpClient", LogLevel.Warning);
+			});
+			serviceCollection.AddHorde(x => x.ServerUrl = options.Server);
+
+			await using ServiceProvider serviceProvider = serviceCollection.BuildServiceProvider();
+			ILogger logger = serviceProvider.GetRequiredService<ILogger<Program>>();
+
+			IHordeClient horde = serviceProvider.GetRequiredService<IHordeClient>();
+			HordeHttpClient hordeHttpClient = horde.CreateHttpClient();
+
+			int minChange = options.MinChange ?? 0;
+			if (minChange == 0)
+			{
+				logger.LogInformation("Querying Horde for currently deployed version...");
+				GetServerInfoResponse info = await hordeHttpClient.GetServerInfoAsync();
+
+				Match match = Regex.Match(info.ServerVersion, @"-(\d+)$");
+				if (!match.Success)
+				{
+					logger.LogError("Unexpected server version format: {Version}", info.ServerVersion);
+					return 1;
+				}
+
+				logger.LogInformation("Current version: {Version}", info.ServerVersion);
+				minChange = int.Parse(match.Groups[1].Value) + 1;
+			}
+
+			logger.LogInformation("Finding changes after CL {MinChange}...", minChange);
+			using IPerforceConnection perforce = await PerforceConnection.CreateAsync(serviceProvider.GetRequiredService<ILogger<PerforceConnection>>());
+
+			List<ChangeInfo> parsedChanges = new List<ChangeInfo>();
+
+			InfoRecord perforceInfo = await perforce.GetInfoAsync(InfoOptions.None);
+
+			List<string> paths = s_paths.Select(x => $"//{perforceInfo.ClientName}/{x}").ToList();
+
+			List<ChangesRecord> changes = await perforce.GetChangesAsync(ChangesOptions.LongOutput, null, minChange, -1, ChangeStatus.Submitted, null, paths);
+			HashSet<int> changeNumbers = new HashSet<int>();
+
+			logger.LogInformation("");
+			foreach (ChangesRecord change in changes)
+			{
+				if (changeNumbers.Add(change.Number))
+				{
+					string description = change.Description;
+					if (!Regex.IsMatch(description, @"^\s*#rnx\s*$", RegexOptions.Multiline))
+					{
+						Match jiraMatch = Regex.Match(description, @"^\s*#jira ([a-zA-Z]+-[0-9]+.*)$");
+						string? jiraTicket = jiraMatch.Success ? jiraMatch.Groups[1].Value : null;
+
+						description = Regex.Replace(description, @"^[a-zA-Z]+:\s*", "");
+						description = Regex.Replace(description, @"^\s*#.*$", "", RegexOptions.Multiline);
+						description = Regex.Replace(description, @"\n[^\n]*\n", "\n");
+						description = description.Trim();
+
+						parsedChanges.Add(new ChangeInfo(change.Number, description, jiraTicket));
+
+						string logDescription = description.Replace("\r", "", StringComparison.Ordinal).Replace("\n", " ", StringComparison.Ordinal);
+
+						const int MaxDescriptionLength = 80;
+						if (logDescription.Length > MaxDescriptionLength)
+						{
+							logDescription = logDescription.Substring(0, MaxDescriptionLength);
+						}
+
+						logger.LogInformation("{Change}: {Description}", change.Number, logDescription);
+					}
+				}
+			}
+
+			if (parsedChanges.Count == 0)
+			{
+				logger.LogError("No changes to deploy.");
+				return 1;
+			}
+
+			DateTime now = DateTime.Now;
+
+			List<string> lines = new List<string>();
+			lines.Add($"## {now.Year}-{now.Month:00}-{now.Day:00}");
+			lines.Add("");
+			foreach (ChangeInfo parsedChange in parsedChanges.OrderByDescending(x => x.Number))
+			{
+				lines.Add($"* {parsedChange.Description.Replace("\n", "\n  ", StringComparison.Ordinal)} ({parsedChange.Number})");
+			}
+			lines.Add("");
+
+			string ReleaseNotesFile = $"//{perforceInfo.ClientName}/Engine/Source/Programs/Horde/Docs/ReleaseNotes.md";
+			await perforce.TryRevertAsync(-1, null, RevertOptions.None, ReleaseNotesFile);
+			await perforce.EditAsync(-1, null, EditOptions.None, ReleaseNotesFile);
+
+			WhereRecord? where = await perforce.WhereAsync(ReleaseNotesFile).FirstOrDefaultAsync();
+			if (where == null || String.IsNullOrEmpty(where.Path))
+			{
+				logger.LogError("Unable to get local path for file {File}", ReleaseNotesFile);
+				return 1;
+			}
+
+			FileReference file = new FileReference(where.Path);
+
+			List<string> fileLines = new List<string>(await FileReference.ReadAllLinesAsync(file));
+
+			int insertIdx = 0;
+			while (insertIdx < fileLines.Count && Regex.IsMatch(fileLines[insertIdx], @"^(# .*|\s*)$"))
+			{
+				insertIdx++;
+			}
+
+			fileLines.InsertRange(insertIdx, lines);
+
+			logger.LogInformation("");
+			logger.LogInformation("Writing {File}", file);
+			await FileReference.WriteAllLinesAsync(file, fileLines);
+			return 0;
+		}
+	}
+}
diff --git a/Engine/Source/Programs/Horde/Tools/ReleaseNotes/ReleaseNotes.csproj b/Engine/Source/Programs/Horde/Tools/ReleaseNotes/ReleaseNotes.csproj
new file mode 100644
--- /dev/null
+++ b/Engine/Source/Programs/Horde/Tools/ReleaseNotes/ReleaseNotes.csproj
@@ -0,0 +1,22 @@
+<Project Sdk="Microsoft.NET.Sdk">
+
+  <PropertyGroup>
+    <OutputType>Exe</OutputType>
+    <TargetFramework>net8.0</TargetFramework>
+    <ImplicitUsings>enable</ImplicitUsings>
+    <Nullable>enable</Nullable>
+    <PublishSingleFile>true</PublishSingleFile>
+  </PropertyGroup>
+
+  <ItemGroup>
+    <PackageReference Include="Microsoft.Extensions.Logging.Console" Version="8.0.0" />
+  </ItemGroup>
+
+  <ItemGroup>
+    <ProjectReference Include="..\..\..\Shared\EpicGames.Core\EpicGames.Core.csproj" />
+    <ProjectReference Include="..\..\..\Shared\EpicGames.Horde\EpicGames.Horde.csproj" />
+    <ProjectReference Include="..\..\..\Shared\EpicGames.IoHash\EpicGames.IoHash.csproj" />
+    <ProjectReference Include="..\..\..\Shared\EpicGames.Perforce\EpicGames.Perforce.csproj" />
+  </ItemGroup>
+
+</Project>
