Horde: Recognize tags in Perforce changelists of the form '#horde 123' as indication that a change fixes a particular issue. The specific tag can be configured in the globals.json file to support multiple deployments with the same Perforce server.

diff --git a/Engine/Source/Programs/Horde/Horde.Server.Tests/Issues/IssueTagTests.cs b/Engine/Source/Programs/Horde/Horde.Server.Tests/Issues/IssueTagTests.cs
new file mode 100644
--- /dev/null
+++ b/Engine/Source/Programs/Horde/Horde.Server.Tests/Issues/IssueTagTests.cs
@@ -0,0 +1,26 @@
+// Copyright Epic Games, Inc. All Rights Reserved.
+
+using System;
+using System.Collections.Generic;
+using System.Linq;
+using System.Text;
+using System.Threading.Tasks;
+using Horde.Server.Issues;
+using Microsoft.VisualStudio.TestTools.UnitTesting;
+
+namespace Horde.Server.Tests.Issues
+{
+	[TestClass]
+	public class IssueTagTests
+	{
+		[TestMethod]
+		public void IssueTag()
+		{
+			Assert.IsTrue(IssueTagService.ParseTags("#horde", "hello\n#horde 123 ").SequenceEqual(new[] { 123 }));
+			Assert.IsTrue(IssueTagService.ParseTags("#horde", "hello\n#horde 123 456").SequenceEqual(new[] { 123, 456 }));
+			Assert.IsTrue(IssueTagService.ParseTags("#horde", "hello\n #horde 123 ").SequenceEqual(new[] { 123 }));
+			Assert.IsTrue(IssueTagService.ParseTags("#horde", "hello\n #horde 123 \n#horde 456").SequenceEqual(new[] { 123, 456 }));
+			Assert.IsTrue(IssueTagService.ParseTags("#horde", "hello\n #horde 123 \n #ROBOMERGE-SOURCE foo").SequenceEqual(Array.Empty<int>()));
+		}
+	}
+}
diff --git a/Engine/Source/Programs/Horde/Horde.Server/Issues/IssueTagService.cs b/Engine/Source/Programs/Horde/Horde.Server/Issues/IssueTagService.cs
new file mode 100644
--- /dev/null
+++ b/Engine/Source/Programs/Horde/Horde.Server/Issues/IssueTagService.cs
@@ -0,0 +1,137 @@
+// Copyright Epic Games, Inc. All Rights Reserved.
+
+using System;
+using System.Collections.Generic;
+using System.Text.RegularExpressions;
+using System.Threading;
+using System.Threading.Tasks;
+using EpicGames.Core;
+using EpicGames.Horde.Streams;
+using Horde.Server.Commits;
+using Horde.Server.Perforce;
+using Horde.Server.Server;
+using Horde.Server.Streams;
+using Horde.Server.Utilities;
+using HordeCommon;
+using Microsoft.Extensions.Hosting;
+using Microsoft.Extensions.Logging;
+using Microsoft.Extensions.Options;
+using MongoDB.Bson.Serialization.Attributes;
+using MongoDB.Bson.Serialization.Options;
+
+namespace Horde.Server.Issues
+{
+	/// <summary>
+	/// Polls revision control for fix changelist numbers in commits using the syntax '#horde 1234'
+	/// </summary>
+	public sealed class IssueTagService : IHostedService, IAsyncDisposable
+	{
+		[SingletonDocument("issue-tags")]
+		class State : SingletonBase
+		{
+			[BsonElement("streams"), BsonDictionaryOptions(DictionaryRepresentation.ArrayOfDocuments)]
+			public Dictionary<StreamId, int> Streams { get; set; } = new Dictionary<StreamId, int>();
+		}
+
+		readonly ISingletonDocument<State> _state;
+		readonly ICommitService _commitService;
+		readonly IIssueCollection _issueCollection;
+		readonly IOptionsMonitor<GlobalConfig> _globalConfig;
+		readonly ITicker _ticker;
+		readonly ILogger _logger;
+
+		/// <summary>
+		/// Constructor
+		/// </summary>
+		public IssueTagService(MongoService mongoService, ICommitService commitService, IIssueCollection issueCollection, IClock clock, IOptionsMonitor<GlobalConfig> globalConfig, ILogger<IssueTagService> logger)
+		{
+			_state = new SingletonDocument<State>(mongoService);
+			_commitService = commitService;
+			_issueCollection = issueCollection;
+			_globalConfig = globalConfig;
+			_ticker = clock.AddSharedTicker<IssueTagService>(TimeSpan.FromMinutes(1.0), TickAsync, logger);
+			_logger = logger;
+		}
+
+		/// <inheritdoc/>
+		public ValueTask DisposeAsync() => _ticker.DisposeAsync();
+
+		/// <inheritdoc/>
+		public Task StartAsync(CancellationToken cancellationToken) => _ticker.StartAsync();
+
+		/// <inheritdoc/>
+		public Task StopAsync(CancellationToken cancellationToken) => _ticker.StopAsync();
+
+		async ValueTask TickAsync(CancellationToken cancellationToken)
+		{
+			using CancellationTokenSource cancellationSource = CancellationTokenSource.CreateLinkedTokenSource(cancellationToken);
+			using IDisposable? listener = _globalConfig.OnChange((_, _) => cancellationSource.Cancel());
+
+			State initialState = await _state.GetAsync(cancellationToken);
+
+			GlobalConfig globalConfig = _globalConfig.CurrentValue;
+			try
+			{
+				await Parallel.ForEachAsync(globalConfig.Streams, cancellationSource.Token, async (stream, ctx) => await TickStreamAsync(stream, initialState, ctx));
+			}
+			catch (Exception ex)
+			{
+				_logger.LogError(ex, "Exception while scanning for #horde tags: {Message}", ex.Message);
+			}
+		}
+
+		async ValueTask TickStreamAsync(StreamConfig streamConfig, State initialState, CancellationToken cancellationToken)
+		{
+			ICommitCollection commits = _commitService.GetCollection(streamConfig);
+
+			int minChange;
+			if (!initialState.Streams.TryGetValue(streamConfig.Id, out minChange))
+			{
+				minChange = await commits.GetLatestNumberAsync(cancellationToken);
+			}
+
+			await foreach (ICommit commit in commits.SubscribeAsync(minChange + 1, null, cancellationToken))
+			{
+				foreach (int issueId in ParseTags(_globalConfig.CurrentValue.IssueFixedTag, commit.Description))
+				{
+					for (; ; )
+					{
+						IIssue? issue = await _issueCollection.GetIssueAsync(issueId, cancellationToken);
+						if (issue == null)
+						{
+							_logger.LogInformation("Commit {Change} by {Author} has invalid issue id {IssueId}", commit.Number, commit.AuthorId, issueId);
+							break;
+						}
+
+						issue = await _issueCollection.TryUpdateIssueAsync(issue, commit.AuthorId, newFixChange: commit.Number, newResolvedById: commit.AuthorId, cancellationToken: cancellationToken);
+						if (issue != null)
+						{
+							_logger.LogInformation("Commit {Change} by {Author} fixes issue id {IssueId}", commit.Number, commit.AuthorId, issueId);
+							break;
+						}
+					}
+				}
+				await _state.UpdateAsync(x => x.Streams[streamConfig.Id] = commit.Number, cancellationToken);
+			}
+		}
+
+		internal static IEnumerable<int> ParseTags(string issueFixedTag, string description)
+		{
+			if (!Regex.IsMatch(description, @"^\s*#ROBOMERGE-SOURCE", RegexOptions.Multiline))
+			{
+				foreach (Match match in Regex.Matches(description, $"^\\s*{issueFixedTag}\\s+(.*)$", RegexOptions.Multiline))
+				{
+					string[] issues = match.Groups[1].Value.Split(new[] { ' ', '\t' }, StringSplitOptions.RemoveEmptyEntries);
+					foreach (string issue in issues)
+					{
+						if (Int32.TryParse(issue, System.Globalization.NumberStyles.None, null, out int issueId))
+						{
+							yield return issueId;
+						}
+					}
+				}
+			}
+		}
+	}
+}
+
diff --git a/Engine/Source/Programs/Horde/Horde.Server/Server/GlobalConfig.cs b/Engine/Source/Programs/Horde/Horde.Server/Server/GlobalConfig.cs
--- a/Engine/Source/Programs/Horde/Horde.Server/Server/GlobalConfig.cs
+++ b/Engine/Source/Programs/Horde/Horde.Server/Server/GlobalConfig.cs
@@ -261,6 +261,11 @@ namespace Horde.Server.Server
 		/// </summary>
 		public AclConfig Acl { get; set; } = new AclConfig();
 
+		/// <summary>
+		/// Commit tag to use for marking issues as fixed
+		/// </summary>
+		public string IssueFixedTag { get; set; } = "#horde";
+
 		/// <summary>
 		/// Accessor for the ACL scope lookup
 		/// </summary>
diff --git a/Engine/Source/Programs/Horde/Horde.Server/Startup.cs b/Engine/Source/Programs/Horde/Horde.Server/Startup.cs
--- a/Engine/Source/Programs/Horde/Horde.Server/Startup.cs
+++ b/Engine/Source/Programs/Horde/Horde.Server/Startup.cs
@@ -863,6 +863,7 @@ namespace Horde.Server
 				services.AddHostedService(provider => provider.GetRequiredService<ConsistencyService>());
 				services.AddHostedService(provider => provider.GetRequiredService<IssueService>());
 				services.AddHostedService<IssueReportService>();
+				services.AddHostedService<IssueTagService>();
 				services.AddHostedService<MetricService>();
 				services.AddHostedService(provider => provider.GetRequiredService<PerforceLoadBalancer>());
 				services.AddHostedService<PoolUpdateService>();
