Horde: Fix static analysis warnings.

#rnx

diff --git a/Engine/Source/Programs/Horde/Horde.Server.Tests/Issues/IssueTagTests.cs b/Engine/Source/Programs/Horde/Horde.Server.Tests/Issues/IssueTagTests.cs
--- a/Engine/Source/Programs/Horde/Horde.Server.Tests/Issues/IssueTagTests.cs
+++ b/Engine/Source/Programs/Horde/Horde.Server.Tests/Issues/IssueTagTests.cs
@@ -1,10 +1,7 @@
 // Copyright Epic Games, Inc. All Rights Reserved.
 
 using System;
-using System.Collections.Generic;
 using System.Linq;
-using System.Text;
-using System.Threading.Tasks;
 using Horde.Server.Issues;
 using Microsoft.VisualStudio.TestTools.UnitTesting;
 
