Horde: Include graphs in the IJob interface, and move methods for manipulating jobs onto the interface.

diff --git a/Engine/Source/Programs/Horde/Horde.Server.Tests/Devices/DeviceServiceTests.cs b/Engine/Source/Programs/Horde/Horde.Server.Tests/Devices/DeviceServiceTests.cs
--- a/Engine/Source/Programs/Horde/Horde.Server.Tests/Devices/DeviceServiceTests.cs
+++ b/Engine/Source/Programs/Horde/Horde.Server.Tests/Devices/DeviceServiceTests.cs
@@ -81,33 +81,33 @@ namespace Horde.Server.Tests.Devices
 			return node;
 		}
 
-		async Task<IJob> StartBatchAsync(IJob job, IGraph graph, int batchIdx)
+		static async Task<IJob> StartBatchAsync(IJob job, int batchIdx)
 		{
 			Assert.AreEqual(JobStepBatchState.Ready, job.Batches[batchIdx].State);
-			job = Deref(await JobCollection.TryUpdateBatchAsync(job, graph, job.Batches[batchIdx].Id, null, JobStepBatchState.Running, null));
+			job = Deref(await job.TryUpdateBatchAsync(job.Batches[batchIdx].Id, null, JobStepBatchState.Running, null));
 			Assert.AreEqual(JobStepBatchState.Running, job.Batches[batchIdx].State);
 			return job;
 		}
 
-		async Task<IJob> StartStepAsync(IJob job, IGraph graph, int batchIdx, int stepIdx)
+		static async Task<IJob> StartStepAsync(IJob job, int batchIdx, int stepIdx)
 		{
 			Assert.AreEqual(JobStepState.Ready, job.Batches[batchIdx].Steps[stepIdx].State);
-			job = Deref(await JobCollection.TryUpdateStepAsync(job, graph, job.Batches[batchIdx].Id, job.Batches[batchIdx].Steps[stepIdx].Id, JobStepState.Running, JobStepOutcome.Success));
+			job = Deref(await job.TryUpdateStepAsync(job.Batches[batchIdx].Id, job.Batches[batchIdx].Steps[stepIdx].Id, JobStepState.Running, JobStepOutcome.Success));
 			return job;
 		}
 
-		async Task<IJob> FinishStepAsync(IJob job, IGraph graph, int batchIdx, int stepIdx, JobStepOutcome outcome)
+		static async Task<IJob> FinishStepAsync(IJob job, int batchIdx, int stepIdx, JobStepOutcome outcome)
 		{
-			job = Deref(await JobCollection.TryUpdateStepAsync(job, graph, job.Batches[batchIdx].Id, job.Batches[batchIdx].Steps[stepIdx].Id, JobStepState.Completed, outcome));
+			job = Deref(await job.TryUpdateStepAsync(job.Batches[batchIdx].Id, job.Batches[batchIdx].Steps[stepIdx].Id, JobStepState.Completed, outcome));
 			Assert.AreEqual(JobStepState.Completed, job.Batches[batchIdx].Steps[stepIdx].State);
 			Assert.AreEqual(outcome, job.Batches[batchIdx].Steps[stepIdx].Outcome);
 			return job;
 		}
 
-		async Task<IJob> RunStepAsync(IJob job, IGraph graph, int batchIdx, int stepIdx, JobStepOutcome outcome)
+		static async Task<IJob> RunStepAsync(IJob job, int batchIdx, int stepIdx, JobStepOutcome outcome)
 		{
-			job = Deref(await StartStepAsync(job, graph, batchIdx, stepIdx));
-			return Deref(await FinishStepAsync(job, graph, batchIdx, stepIdx, outcome));
+			job = Deref(await StartStepAsync(job, batchIdx, stepIdx));
+			return Deref(await FinishStepAsync(job, batchIdx, stepIdx, outcome));
 		}
 
 		JobStepId GetStepId(IJob job, string nodeName)
@@ -198,8 +198,8 @@ namespace Horde.Server.Tests.Devices
 
 			IJob job = await JobCollection.AddAsync(JobIdUtils.GenerateNewId(), new StreamId("ue5-main"), new TemplateId("test-build"), ContentHash.SHA1("hello"), baseGraph, "Test job", 123, 123, options);
 
-			job = await StartBatchAsync(job, baseGraph, 0);
-			job = await RunStepAsync(job, baseGraph, 0, 0, JobStepOutcome.Success); // Setup Build
+			job = await StartBatchAsync(job, 0);
+			job = await RunStepAsync(job, 0, 0, JobStepOutcome.Success); // Setup Build
 
 			List<NewGroup> newGroups = new List<NewGroup>();
 
@@ -240,7 +240,7 @@ namespace Horde.Server.Tests.Devices
 			AddNode(testGroup, "Run Tests", new[] { "Run Test 1", "Run Test 2", "Run Test 3", "Run Test 4" });
 
 			_graph = await GraphCollection.AppendAsync(baseGraph, newGroups, null, null);
-			job = Deref(await JobCollection.TryUpdateGraphAsync(job, baseGraph, _graph));
+			job = Deref(await job.TryUpdateGraphAsync(_graph));
 
 			return job;
 		}
@@ -312,32 +312,30 @@ namespace Horde.Server.Tests.Devices
 			await SetupDevicesAsync();
 			IJob job = await SetupJobAsync("DeviceReserveNodes");
 
-			IGraph? graph = _graph!;
+			job = await StartBatchAsync(job, 1);
+			job = await RunStepAsync(job, 1, 0, JobStepOutcome.Success); // Update Version Files
+			job = await RunStepAsync(job, 1, 1, JobStepOutcome.Success); // Compile Editor
 
-			job = await StartBatchAsync(job, graph, 1);
-			job = await RunStepAsync(job, graph, 1, 0, JobStepOutcome.Success); // Update Version Files
-			job = await RunStepAsync(job, graph, 1, 1, JobStepOutcome.Success); // Compile Editor
+			job = await StartBatchAsync(job, 2);
+			job = await RunStepAsync(job, 2, 0, JobStepOutcome.Success); // Compile Client
 
-			job = await StartBatchAsync(job, graph, 2);
-			job = await RunStepAsync(job, graph, 2, 0, JobStepOutcome.Success); // Compile Client
+			job = await StartBatchAsync(job, 3);
+			job = await RunStepAsync(job, 3, 0, JobStepOutcome.Success); // Cook Client
 
-			job = await StartBatchAsync(job, graph, 3);
-			job = await RunStepAsync(job, graph, 3, 0, JobStepOutcome.Success); // Cook Client
-
-			job = await StartBatchAsync(job, graph, 4);
+			job = await StartBatchAsync(job, 4);
 
 			// Install  the build
 			JobStepId stepId = GetStepId(job, "Install Build");
-			job = await StartStepAsync(job, graph, 4, 0); // Install Build														  
+			job = await StartStepAsync(job, 4, 0); // Install Build														  
 			LegacyCreateReservationRequest request = SetupReservationTestAsync(job, stepId: stepId);
 			GetLegacyReservationResponse installReservation = ResultToValue(await DeviceController!.CreateDeviceReservationV1Async(request));
-			job = await FinishStepAsync(job, graph, 4, 0, JobStepOutcome.Success); // Install Build
+			job = await FinishStepAsync(job, 4, 0, JobStepOutcome.Success); // Install Build
 
 			for (int i = 1; i < 5; i++)
 			{
 				// Run Test 1
 				stepId = GetStepId(job, $"Run Test {i}");
-				job = await StartStepAsync(job, graph, 4, i);
+				job = await StartStepAsync(job, 4, i);
 
 				request = SetupReservationTestAsync(job, stepId: stepId);
 
@@ -352,7 +350,7 @@ namespace Horde.Server.Tests.Devices
 					Assert.AreEqual((result.Result as ConflictObjectResult)!.Value, "Reserved nodes must not run in parallel: Run Test 3,Run Test 4");
 
 					// finish the step
-					job = await FinishStepAsync(job, graph, 4, 3, JobStepOutcome.Success);
+					job = await FinishStepAsync(job, 4, 3, JobStepOutcome.Success);
 
 					result = await DeviceController!.CreateDeviceReservationV1Async(request);
 					reservation = ResultToValue(result);
@@ -369,7 +367,7 @@ namespace Horde.Server.Tests.Devices
 				if (i != 3)
 				{
 					await DeviceController!.DeleteReservationV1Async(reservation.Guid);
-					job = await FinishStepAsync(job, graph, 4, i, JobStepOutcome.Success);
+					job = await FinishStepAsync(job, 4, i, JobStepOutcome.Success);
 				}
 
 				await DeviceService.TickForTestingAsync();
@@ -394,23 +392,21 @@ namespace Horde.Server.Tests.Devices
 			await SetupDevicesAsync();
 			IJob job = await SetupJobAsync("DeviceReserve");
 
-			IGraph? graph = _graph!;
-
-			job = await StartBatchAsync(job, graph, 1);
-			job = await RunStepAsync(job, graph, 1, 0, JobStepOutcome.Success); // Update Version Files
-			job = await RunStepAsync(job, graph, 1, 1, JobStepOutcome.Success); // Compile Editor
+			job = await StartBatchAsync(job, 1);
+			job = await RunStepAsync(job, 1, 0, JobStepOutcome.Success); // Update Version Files
+			job = await RunStepAsync(job, 1, 1, JobStepOutcome.Success); // Compile Editor
 
-			job = await StartBatchAsync(job, graph, 2);
-			job = await RunStepAsync(job, graph, 2, 0, JobStepOutcome.Success); // Compile Client
+			job = await StartBatchAsync(job, 2);
+			job = await RunStepAsync(job, 2, 0, JobStepOutcome.Success); // Compile Client
 
-			job = await StartBatchAsync(job, graph, 3);
-			job = await RunStepAsync(job, graph, 3, 0, JobStepOutcome.Success); // Cook Client
+			job = await StartBatchAsync(job, 3);
+			job = await RunStepAsync(job, 3, 0, JobStepOutcome.Success); // Cook Client
 
-			job = await StartBatchAsync(job, graph, 4);
+			job = await StartBatchAsync(job, 4);
 
 			// Install  the build
 			JobStepId stepId = GetStepId(job, "Install Build");
-			job = await StartStepAsync(job, graph, 4, 0); // Install Build														  
+			job = await StartStepAsync(job, 4, 0); // Install Build														  
 			LegacyCreateReservationRequest request = SetupReservationTestAsync(job, stepId: stepId, modelId: "Base");
 			GetLegacyReservationResponse installReservation = ResultToValue(await DeviceController!.CreateDeviceReservationV1Async(request));
 			Assert.IsTrue(installReservation.InstallRequired);
@@ -422,13 +418,13 @@ namespace Horde.Server.Tests.Devices
 			Assert.IsTrue(installReservation.InstallRequired);
 			Assert.AreNotEqual(installProblemDeviceName, installReservation.DeviceNames[0]);
 
-			job = await FinishStepAsync(job, graph, 4, 0, JobStepOutcome.Success); // Install Build
+			job = await FinishStepAsync(job, 4, 0, JobStepOutcome.Success); // Install Build
 
 			for (int i = 1; i < 5; i++)
 			{
 				// Run Test 1
 				stepId = GetStepId(job, $"Run Test {i}");
-				job = await StartStepAsync(job, graph, 4, i);
+				job = await StartStepAsync(job, 4, i);
 
 				if (i == 2)
 				{
@@ -480,7 +476,7 @@ namespace Horde.Server.Tests.Devices
 				}
 
 				await DeviceController!.DeleteReservationV1Async(reservation.Guid);
-				job = await FinishStepAsync(job, graph, 4, i, JobStepOutcome.Success);
+				job = await FinishStepAsync(job, 4, i, JobStepOutcome.Success);
 
 				await DeviceService.TickForTestingAsync();
 			}
diff --git a/Engine/Source/Programs/Horde/Horde.Server.Tests/Fleet/JobQueueStrategyTest.cs b/Engine/Source/Programs/Horde/Horde.Server.Tests/Fleet/JobQueueStrategyTest.cs
--- a/Engine/Source/Programs/Horde/Horde.Server.Tests/Fleet/JobQueueStrategyTest.cs
+++ b/Engine/Source/Programs/Horde/Horde.Server.Tests/Fleet/JobQueueStrategyTest.cs
@@ -152,13 +152,13 @@ namespace Horde.Server.Tests.Fleet
 			for (int i = 0; i < numBatchesRunning; i++)
 			{
 				IJob job = await AddPlaceholderJobAsync(graph, streamConfig.Id, nodeForAgentType1);
-				await JobCollection.TryUpdateBatchAsync(job, graph, job.Batches[0].Id, null, JobStepBatchState.Running, null);
+				await job.TryUpdateBatchAsync(job.Batches[0].Id, null, JobStepBatchState.Running, null);
 			}
 
 			for (int i = 0; i < numBatchesReady; i++)
 			{
 				IJob job = await AddPlaceholderJobAsync(graph, streamConfig.Id, nodeForAgentType1);
-				await JobCollection.TryUpdateBatchAsync(job, graph, job.Batches[0].Id, null, JobStepBatchState.Ready, null);
+				await job.TryUpdateBatchAsync(job.Batches[0].Id, null, JobStepBatchState.Ready, null);
 			}
 
 			return (new(JobCollection, GraphCollection, StreamCollection, Clock, Cache, isDowntimeActive, GlobalConfig), poolSize, pool, agents);
diff --git a/Engine/Source/Programs/Horde/Horde.Server.Tests/Jobs/Bisect/BisectTests.cs b/Engine/Source/Programs/Horde/Horde.Server.Tests/Jobs/Bisect/BisectTests.cs
--- a/Engine/Source/Programs/Horde/Horde.Server.Tests/Jobs/Bisect/BisectTests.cs
+++ b/Engine/Source/Programs/Horde/Horde.Server.Tests/Jobs/Bisect/BisectTests.cs
@@ -177,17 +177,16 @@ namespace Horde.Server.Tests.Jobs.Bisect
 
 		async Task<IJob> SetJobOutcomeAsync(IJob job, IGraph graph, JobStepOutcome outcome)
 		{
-			IGraph oldGraph = await GraphCollection.GetAsync(job.GraphHash);
-			job = Deref(await JobCollection.TryUpdateGraphAsync(job, oldGraph, graph));
+			job = Deref(await job.TryUpdateGraphAsync(graph));
 
-			job = Deref(await JobCollection.TryUpdateStepAsync(job, graph, job.Batches[0].Id, job.Batches[0].Steps[0].Id, JobStepState.Completed, JobStepOutcome.Success));
-			job = Deref(await JobCollection.TryUpdateBatchAsync(job, graph, job.Batches[0].Id, null, JobStepBatchState.Complete, null));
+			job = Deref(await job.TryUpdateStepAsync(job.Batches[0].Id, job.Batches[0].Steps[0].Id, JobStepState.Completed, JobStepOutcome.Success));
+			job = Deref(await job.TryUpdateBatchAsync(job.Batches[0].Id, null, JobStepBatchState.Complete, null));
 
 			IJobStepBatch batch = job.Batches[^1];
 			IJobStep step = batch.Steps[^1];
 
-			job = Deref(await JobCollection.TryUpdateStepAsync(job, graph, batch.Id, step.Id, JobStepState.Completed, outcome));
-			job = Deref(await JobCollection.TryUpdateBatchAsync(job, graph, batch.Id, null, JobStepBatchState.Complete, null));
+			job = Deref(await job.TryUpdateStepAsync(batch.Id, step.Id, JobStepState.Completed, outcome));
+			job = Deref(await job.TryUpdateBatchAsync(batch.Id, null, JobStepBatchState.Complete, null));
 
 			INodeGroup group = graph.Groups[batch.GroupIdx];
 			INode node = group.Nodes[step.NodeIdx];
diff --git a/Engine/Source/Programs/Horde/Horde.Server.Tests/Jobs/JobCollectionTests.cs b/Engine/Source/Programs/Horde/Horde.Server.Tests/Jobs/JobCollectionTests.cs
--- a/Engine/Source/Programs/Horde/Horde.Server.Tests/Jobs/JobCollectionTests.cs
+++ b/Engine/Source/Programs/Horde/Horde.Server.Tests/Jobs/JobCollectionTests.cs
@@ -44,29 +44,29 @@ namespace Horde.Server.Tests.Jobs
 			return node;
 		}
 
-		Task<IJob> RunBatchAsync(IJob job, IGraph graph, int batchIdx)
+		static Task<IJob> RunBatchAsync(IJob job, int batchIdx)
 		{
-			return UpdateBatchAsync(job, graph, batchIdx, JobStepBatchState.Running);
+			return UpdateBatchAsync(job, batchIdx, JobStepBatchState.Running);
 		}
 
-		async Task<IJob> UpdateBatchAsync(IJob job, IGraph graph, int batchIdx, JobStepBatchState state)
+		static async Task<IJob> UpdateBatchAsync(IJob job, int batchIdx, JobStepBatchState state)
 		{
 			Assert.AreEqual(JobStepBatchState.Ready, job.Batches[batchIdx].State);
-			job = Deref(await JobCollection.TryAssignLeaseAsync(job, batchIdx, new PoolId("foo"), new AgentId("agent"), new SessionId(BinaryIdUtils.CreateNew()), new LeaseId(BinaryIdUtils.CreateNew()), new LogId(BinaryIdUtils.CreateNew())));
+			job = Deref(await job.TryAssignLeaseAsync(batchIdx, new PoolId("foo"), new AgentId("agent"), new SessionId(BinaryIdUtils.CreateNew()), new LeaseId(BinaryIdUtils.CreateNew()), new LogId(BinaryIdUtils.CreateNew())));
 			if (job.Batches[batchIdx].State != state)
 			{
-				job = Deref(await JobCollection.TryUpdateBatchAsync(job, graph, job.Batches[batchIdx].Id, null, state, null));
+				job = Deref(await job.TryUpdateBatchAsync(job.Batches[batchIdx].Id, null, state, null));
 				Assert.AreEqual(state, job.Batches[batchIdx].State);
 			}
 			return job;
 		}
 
-		async Task<IJob> RunStepAsync(IJob job, IGraph graph, int batchIdx, int stepIdx, JobStepOutcome outcome)
+		static async Task<IJob> RunStepAsync(IJob job, int batchIdx, int stepIdx, JobStepOutcome outcome)
 		{
 			Assert.AreEqual(JobStepState.Ready, job.Batches[batchIdx].Steps[stepIdx].State);
-			job = Deref(await JobCollection.TryUpdateStepAsync(job, graph, job.Batches[batchIdx].Id, job.Batches[batchIdx].Steps[stepIdx].Id, JobStepState.Running, JobStepOutcome.Success));
+			job = Deref(await job.TryUpdateStepAsync(job.Batches[batchIdx].Id, job.Batches[batchIdx].Steps[stepIdx].Id, JobStepState.Running, JobStepOutcome.Success));
 			Assert.AreEqual(JobStepState.Running, job.Batches[batchIdx].Steps[stepIdx].State);
-			job = Deref(await JobCollection.TryUpdateStepAsync(job, graph, job.Batches[batchIdx].Id, job.Batches[batchIdx].Steps[stepIdx].Id, JobStepState.Completed, outcome));
+			job = Deref(await job.TryUpdateStepAsync(job.Batches[batchIdx].Id, job.Batches[batchIdx].Steps[stepIdx].Id, JobStepState.Completed, outcome));
 			Assert.AreEqual(JobStepState.Completed, job.Batches[batchIdx].Steps[stepIdx].State);
 			Assert.AreEqual(outcome, job.Batches[batchIdx].Steps[stepIdx].Outcome);
 			return job;
@@ -86,8 +86,8 @@ namespace Horde.Server.Tests.Jobs
 
 			IJob job = await JobCollection.AddAsync(JobIdUtils.GenerateNewId(), new StreamId("ue4-main"), new TemplateId("test-build"), ContentHash.SHA1("hello"), baseGraph, "Test job", 123, 123, options);
 
-			job = await RunBatchAsync(job, baseGraph, 0);
-			job = await RunStepAsync(job, baseGraph, 0, 0, JobStepOutcome.Success); // Setup Build
+			job = await RunBatchAsync(job, 0);
+			job = await RunStepAsync(job, 0, 0, JobStepOutcome.Success); // Setup Build
 
 			List<NewGroup> newGroups = new List<NewGroup>();
 
@@ -104,17 +104,17 @@ namespace Horde.Server.Tests.Jobs
 			AddNode(publishGroup, "Post-Publish Client", null, x => x.OrderDependencies = new List<string> { "Publish Client" });
 
 			IGraph graph = await GraphCollection.AppendAsync(baseGraph, newGroups, null, null);
-			job = Deref(await JobCollection.TryUpdateGraphAsync(job, baseGraph, graph));
+			job = Deref(await job.TryUpdateGraphAsync(graph));
 
-			job = await RunBatchAsync(job, graph, 1);
-			job = await RunStepAsync(job, graph, 1, 0, JobStepOutcome.Success); // Update Version Files
-			job = await RunStepAsync(job, graph, 1, 1, JobStepOutcome.Success); // Compile Editor
+			job = await RunBatchAsync(job, 1);
+			job = await RunStepAsync(job, 1, 0, JobStepOutcome.Success); // Update Version Files
+			job = await RunStepAsync(job, 1, 1, JobStepOutcome.Success); // Compile Editor
 
-			job = await RunBatchAsync(job, graph, 2);
-			job = await RunStepAsync(job, graph, 2, 0, JobStepOutcome.Success); // Compile Client
+			job = await RunBatchAsync(job, 2);
+			job = await RunStepAsync(job, 2, 0, JobStepOutcome.Success); // Compile Client
 
-			job = await RunBatchAsync(job, graph, 3);
-			job = await RunStepAsync(job, graph, 3, 0, JobStepOutcome.Failure); // Cook Client
+			job = await RunBatchAsync(job, 3);
+			job = await RunStepAsync(job, 3, 0, JobStepOutcome.Failure); // Cook Client
 			Assert.AreEqual(JobStepState.Skipped, job.Batches[3].Steps[1].State); // Publish Client
 			Assert.AreEqual(JobStepState.Skipped, job.Batches[3].Steps[2].State); // Post-Publish Client
 		}
@@ -147,21 +147,21 @@ namespace Horde.Server.Tests.Jobs
 
 			IJob job = await JobCollection.AddAsync(JobIdUtils.GenerateNewId(), new StreamId("ue4-main"), new TemplateId("test-build"), ContentHash.SHA1("hello"), baseGraph, "Test job", 123, 123, options);
 
-			job = await RunBatchAsync(job, baseGraph, 0);
-			job = await RunStepAsync(job, baseGraph, 0, 0, JobStepOutcome.Success); // Setup Build
+			job = await RunBatchAsync(job, 0);
+			job = await RunStepAsync(job, 0, 0, JobStepOutcome.Success); // Setup Build
 
-			job = Deref(await JobCollection.TryUpdateGraphAsync(job, baseGraph, graph1));
+			job = Deref(await job.TryUpdateGraphAsync(graph1));
 			Assert.AreEqual(4, job.Batches.Count);
 			Assert.AreEqual(JobStepState.Completed, job.Batches[0].Steps[0].State); // Setup Build
 			Assert.AreEqual(JobStepState.Ready, job.Batches[1].Steps[0].State); // Initial Node
 			Assert.AreEqual(JobStepState.Waiting, job.Batches[2].Steps[0].State); // Split Multi-Agent Cook
 			Assert.AreEqual(JobStepState.Waiting, job.Batches[3].Steps[0].State); // Gather
 
-			job = await RunBatchAsync(job, graph1, 1);
-			job = await RunStepAsync(job, graph1, 1, 0, JobStepOutcome.Success); // Initial Node
+			job = await RunBatchAsync(job, 1);
+			job = await RunStepAsync(job, 1, 0, JobStepOutcome.Success); // Initial Node
 
-			job = await RunBatchAsync(job, graph1, 2);
-			job = await RunStepAsync(job, graph1, 2, 0, JobStepOutcome.Success); // Split Multi-Agent Cook
+			job = await RunBatchAsync(job, 2);
+			job = await RunStepAsync(job, 2, 0, JobStepOutcome.Success); // Split Multi-Agent Cook
 
 			Assert.AreEqual(4, job.Batches.Count);
 			Assert.AreEqual(JobStepState.Completed, job.Batches[0].Steps[0].State); // Setup Build
@@ -169,7 +169,7 @@ namespace Horde.Server.Tests.Jobs
 			Assert.AreEqual(JobStepState.Completed, job.Batches[2].Steps[0].State); // Split Multi-Agent Cook
 			Assert.AreEqual(JobStepState.Ready, job.Batches[3].Steps[0].State); // Gather
 
-			await JobCollection.TryUpdateGraphAsync(job, graph1, graph2);
+			job = Deref(await job.TryUpdateGraphAsync(graph2));
 
 			Assert.AreEqual(6, job.Batches.Count);
 			Assert.AreEqual(JobStepState.Completed, job.Batches[0].Steps[0].State); // Setup Build
@@ -179,16 +179,16 @@ namespace Horde.Server.Tests.Jobs
 			Assert.AreEqual(JobStepState.Ready, job.Batches[4].Steps[0].State); // Cook Item 2
 			Assert.AreEqual(JobStepState.Waiting, job.Batches[5].Steps[0].State); // Gather
 
-			job = await RunBatchAsync(job, graph2, 3);
-			job = await RunStepAsync(job, graph2, 3, 0, JobStepOutcome.Success); // Split Multi-Agent Cook
+			job = await RunBatchAsync(job, 3);
+			job = await RunStepAsync(job, 3, 0, JobStepOutcome.Success); // Split Multi-Agent Cook
 
 			Assert.AreEqual(6, job.Batches.Count);
 			Assert.AreEqual(JobStepState.Completed, job.Batches[3].Steps[0].State); // Cook Item 1
 			Assert.AreEqual(JobStepState.Ready, job.Batches[4].Steps[0].State); // Cook Item 2
 			Assert.AreEqual(JobStepState.Waiting, job.Batches[5].Steps[0].State); // Gather
 
-			job = await RunBatchAsync(job, graph2, 4);
-			job = await RunStepAsync(job, graph2, 4, 0, JobStepOutcome.Success); // Split Multi-Agent Cook
+			job = await RunBatchAsync(job, 4);
+			job = await RunStepAsync(job, 4, 0, JobStepOutcome.Success); // Split Multi-Agent Cook
 
 			Assert.AreEqual(6, job.Batches.Count);
 			Assert.AreEqual(JobStepState.Completed, job.Batches[3].Steps[0].State); // Cook Item 1
@@ -217,12 +217,12 @@ namespace Horde.Server.Tests.Jobs
 			IJob job = await JobCollection.AddAsync(JobIdUtils.GenerateNewId(), new StreamId("ue4-main"), new TemplateId("test-build"), ContentHash.SHA1("hello"), baseGraph, "Test job", 123, 123, options);
 
 			// Try a batch and fail it
-			job = await RunBatchAsync(job, baseGraph, 0);
-			job = Deref(await JobCollection.TryUpdateBatchAsync(job, baseGraph, job.Batches[0].Id, null, JobStepBatchState.Complete, JobStepBatchError.Incomplete));
+			job = await RunBatchAsync(job, 0);
+			job = Deref(await job.TryUpdateBatchAsync(job.Batches[0].Id, null, JobStepBatchState.Complete, JobStepBatchError.Incomplete));
 
 			// Start the replacement batch and update the graph
-			job = await RunBatchAsync(job, baseGraph, 1);
-			job = Deref(await JobCollection.TryUpdateGraphAsync(job, baseGraph, graph1));
+			job = await RunBatchAsync(job, 1);
+			job = Deref(await job.TryUpdateGraphAsync(graph1));
 
 			// Validate the new job state
 			Assert.AreEqual(3, job.Batches.Count);
@@ -243,12 +243,12 @@ namespace Horde.Server.Tests.Jobs
 			Fixture fixture = await CreateFixtureAsync();
 
 			SessionId sessionId1 = SessionIdUtils.GenerateNewId();
-			await JobCollection.TryAssignLeaseAsync(fixture.Job1, 0, new PoolId("foo"), fixture.Agent1.Id,
+			await fixture.Job1.TryAssignLeaseAsync(0, new PoolId("foo"), fixture.Agent1.Id,
 				sessionId1, new LeaseId(BinaryIdUtils.CreateNew()), LogIdUtils.GenerateNewId());
 
 			SessionId sessionId2 = SessionIdUtils.GenerateNewId();
 			IJob job = (await JobCollection.GetAsync(fixture.Job1.Id))!;
-			await JobCollection.TryAssignLeaseAsync(job, 0, new PoolId("foo"), fixture.Agent1.Id,
+			await job.TryAssignLeaseAsync(0, new PoolId("foo"), fixture.Agent1.Id,
 				sessionId2, new LeaseId(BinaryIdUtils.CreateNew()), LogIdUtils.GenerateNewId());
 
 			// Manually verify the log output
@@ -279,8 +279,8 @@ namespace Horde.Server.Tests.Jobs
 
 			IJob job = await JobCollection.AddAsync(JobIdUtils.GenerateNewId(), new StreamId("ue4-main"), new TemplateId("test-build"), ContentHash.SHA1("hello"), baseGraph, "Test job", 123, 123, options);
 
-			job = await RunBatchAsync(job, baseGraph, 0);
-			job = await RunStepAsync(job, baseGraph, 0, 0, JobStepOutcome.Success); // Setup Build
+			job = await RunBatchAsync(job, 0);
+			job = await RunStepAsync(job, 0, 0, JobStepOutcome.Success); // Setup Build
 
 			List<NewGroup> newGroups = new List<NewGroup>();
 
@@ -290,14 +290,14 @@ namespace Horde.Server.Tests.Jobs
 			AddNode(initialGroup, "Step 3", new[] { "Step 2" });
 
 			IGraph graph = await GraphCollection.AppendAsync(baseGraph, newGroups, null, null);
-			job = Deref(await JobCollection.TryUpdateGraphAsync(job, baseGraph, graph));
+			job = Deref(await job.TryUpdateGraphAsync(graph));
 
-			job = await RunBatchAsync(job, graph, 1);
-			job = await RunStepAsync(job, graph, 1, 0, JobStepOutcome.Success); // Step 1
-			job = await RunStepAsync(job, graph, 1, 1, JobStepOutcome.Success); // Step 2
+			job = await RunBatchAsync(job, 1);
+			job = await RunStepAsync(job, 1, 0, JobStepOutcome.Success); // Step 1
+			job = await RunStepAsync(job, 1, 1, JobStepOutcome.Success); // Step 2
 
 			// Force an error executing the batch
-			job = Deref(await JobCollection.TryUpdateBatchAsync(job, graph, job.Batches[1].Id, null, JobStepBatchState.Complete, JobStepBatchError.Incomplete));
+			job = Deref(await job.TryUpdateBatchAsync(job.Batches[1].Id, null, JobStepBatchState.Complete, JobStepBatchError.Incomplete));
 
 			// Check that it restarted all three nodes
 			IJob newJob = (await JobCollection.GetAsync(job.Id))!;
@@ -336,8 +336,8 @@ namespace Horde.Server.Tests.Jobs
 			IJob job = await JobCollection.AddAsync(JobIdUtils.GenerateNewId(), new StreamId("ue4-main"), new TemplateId("test-build"), ContentHash.SHA1("hello"), baseGraph, "Test job", 123, 123, options);
 			Assert.AreEqual(1, job.Batches.Count);
 
-			job = await RunBatchAsync(job, baseGraph, 0);
-			job = Deref(await JobCollection.TryUpdateBatchAsync(job, baseGraph, job.Batches[0].Id, null, JobStepBatchState.Complete, JobStepBatchError.Incomplete));
+			job = await RunBatchAsync(job, 0);
+			job = Deref(await job.TryUpdateBatchAsync(job.Batches[0].Id, null, JobStepBatchState.Complete, JobStepBatchError.Incomplete));
 
 			job = (await JobCollection.GetAsync(job.Id))!;
 			Assert.AreEqual(2, job.Batches.Count);
@@ -346,7 +346,7 @@ namespace Horde.Server.Tests.Jobs
 			Assert.AreEqual(JobStepBatchState.Ready, job.Batches[1].State);
 			Assert.AreEqual(1, job.Batches[1].Steps.Count);
 
-			job = Deref(await JobCollection.TryUpdateBatchAsync(job, baseGraph, job.Batches[1].Id, null, JobStepBatchState.Complete, JobStepBatchError.Incomplete));
+			job = Deref(await job.TryUpdateBatchAsync(job.Batches[1].Id, null, JobStepBatchState.Complete, JobStepBatchError.Incomplete));
 
 			job = (await JobCollection.GetAsync(job.Id))!;
 			Assert.AreEqual(3, job.Batches.Count);
@@ -357,7 +357,7 @@ namespace Horde.Server.Tests.Jobs
 			Assert.AreEqual(JobStepBatchState.Ready, job.Batches[2].State);
 			Assert.AreEqual(1, job.Batches[2].Steps.Count);
 
-			job = Deref(await JobCollection.TryUpdateBatchAsync(job, baseGraph, job.Batches[2].Id, null, JobStepBatchState.Complete, JobStepBatchError.Incomplete));
+			job = Deref(await job.TryUpdateBatchAsync(job.Batches[2].Id, null, JobStepBatchState.Complete, JobStepBatchError.Incomplete));
 
 			job = (await JobCollection.GetAsync(job.Id))!;
 			Assert.AreEqual(3, job.Batches.Count);
@@ -389,9 +389,9 @@ namespace Horde.Server.Tests.Jobs
 
 			// First retry
 
-			job = await RunBatchAsync(job, baseGraph, 0);
-			job = Deref(await JobCollection.TryUpdateStepAsync(job, baseGraph, job.Batches[0].Id, job.Batches[0].Steps[0].Id, newState: JobStepState.Running));
-			job = Deref(await JobCollection.TryUpdateBatchAsync(job, baseGraph, job.Batches[0].Id, null, JobStepBatchState.Complete, JobStepBatchError.Incomplete));
+			job = await RunBatchAsync(job, 0);
+			job = Deref(await job.TryUpdateStepAsync(job.Batches[0].Id, job.Batches[0].Steps[0].Id, newState: JobStepState.Running));
+			job = Deref(await job.TryUpdateBatchAsync(job.Batches[0].Id, null, JobStepBatchState.Complete, JobStepBatchError.Incomplete));
 
 			job = (await JobCollection.GetAsync(job.Id))!;
 			Assert.AreEqual(2, job.Batches.Count);
@@ -406,8 +406,8 @@ namespace Horde.Server.Tests.Jobs
 
 			// Second retry
 
-			job = Deref(await JobCollection.TryUpdateStepAsync(job, baseGraph, job.Batches[1].Id, job.Batches[1].Steps[0].Id, newState: JobStepState.Running));
-			job = Deref(await JobCollection.TryUpdateBatchAsync(job, baseGraph, job.Batches[1].Id, null, JobStepBatchState.Complete, JobStepBatchError.Incomplete));
+			job = Deref(await job.TryUpdateStepAsync(job.Batches[1].Id, job.Batches[1].Steps[0].Id, newState: JobStepState.Running));
+			job = Deref(await job.TryUpdateBatchAsync(job.Batches[1].Id, null, JobStepBatchState.Complete, JobStepBatchError.Incomplete));
 
 			job = (await JobCollection.GetAsync(job.Id))!;
 			Assert.AreEqual(3, job.Batches.Count);
@@ -427,8 +427,8 @@ namespace Horde.Server.Tests.Jobs
 
 			// Check it doesn't retry a third time
 
-			job = Deref(await JobCollection.TryUpdateStepAsync(job, baseGraph, job.Batches[2].Id, job.Batches[2].Steps[0].Id, newState: JobStepState.Running));
-			job = Deref(await JobCollection.TryUpdateBatchAsync(job, baseGraph, job.Batches[2].Id, null, JobStepBatchState.Complete, JobStepBatchError.Incomplete));
+			job = Deref(await job.TryUpdateStepAsync(job.Batches[2].Id, job.Batches[2].Steps[0].Id, newState: JobStepState.Running));
+			job = Deref(await job.TryUpdateBatchAsync(job.Batches[2].Id, null, JobStepBatchState.Complete, JobStepBatchError.Incomplete));
 
 			job = (await JobCollection.GetAsync(job.Id))!;
 			Assert.AreEqual(3, job.Batches.Count);
@@ -469,8 +469,8 @@ namespace Horde.Server.Tests.Jobs
 			IJob job = await JobCollection.AddAsync(JobIdUtils.GenerateNewId(), new StreamId("ue4-main"), new TemplateId("test-build"), ContentHash.SHA1("hello"), graph, "Test job", 123, 123, options);
 
 			// Fail the first step
-			job = await RunBatchAsync(job, graph, 0);
-			job = await RunStepAsync(job, graph, 0, 0, JobStepOutcome.Failure);
+			job = await RunBatchAsync(job, 0);
+			job = await RunStepAsync(job, 0, 0, JobStepOutcome.Failure);
 
 			Assert.AreEqual(1, job.Batches.Count);
 			Assert.AreEqual(2, job.Batches[0].Steps.Count);
@@ -480,7 +480,7 @@ namespace Horde.Server.Tests.Jobs
 			Assert.AreEqual(JobStepOutcome.Failure, job.Batches[0].Steps[1].Outcome);
 
 			// Retry the failed step (expect: new batch containing the retried step is created, skipped steps from original batch are removed)
-			job = Deref(await JobCollection.TryUpdateStepAsync(job, graph, job.Batches[0].Id, job.Batches[0].Steps[0].Id, newRetryByUserId: UserId.Anonymous));
+			job = Deref(await job.TryUpdateStepAsync(job.Batches[0].Id, job.Batches[0].Steps[0].Id, newRetryByUserId: UserId.Anonymous));
 			Assert.AreEqual(2, job.Batches.Count);
 			Assert.AreEqual(1, job.Batches[0].Steps.Count);
 			Assert.AreEqual(2, job.Batches[1].Steps.Count);
@@ -501,7 +501,7 @@ namespace Horde.Server.Tests.Jobs
 			Assert.AreEqual(JobStepOutcome.Success, step.Outcome);
 
 			// Fail the retried step
-			job = await RunStepAsync(job, graph, 1, 0, JobStepOutcome.Failure);
+			job = await RunStepAsync(job, 1, 0, JobStepOutcome.Failure);
 			Assert.AreEqual(2, job.Batches.Count);
 			Assert.AreEqual(1, job.Batches[0].Steps.Count);
 			Assert.AreEqual(2, job.Batches[1].Steps.Count);
@@ -522,7 +522,7 @@ namespace Horde.Server.Tests.Jobs
 			Assert.AreEqual(JobStepOutcome.Failure, step.Outcome);
 
 			// Retry the failed step for a second time
-			job = Deref(await JobCollection.TryUpdateStepAsync(job, graph, job.Batches[1].Id, job.Batches[1].Steps[0].Id, newRetryByUserId: UserId.Anonymous));
+			job = Deref(await job.TryUpdateStepAsync(job.Batches[1].Id, job.Batches[1].Steps[0].Id, newRetryByUserId: UserId.Anonymous));
 			Assert.AreEqual(3, job.Batches.Count);
 			Assert.AreEqual(1, job.Batches[0].Steps.Count);
 			Assert.AreEqual(1, job.Batches[1].Steps.Count);
@@ -571,9 +571,9 @@ namespace Horde.Server.Tests.Jobs
 			IJob job = await JobCollection.AddAsync(JobIdUtils.GenerateNewId(), new StreamId("ue4-main"), new TemplateId("test-build"), ContentHash.SHA1("hello"), graph, "Test job", 123, 123, options);
 
 			// Fail the first step
-			job = await RunBatchAsync(job, graph, 0);
-			job = await RunStepAsync(job, graph, 0, 0, JobStepOutcome.Failure);
-			job = Deref(await JobCollection.TryUpdateBatchAsync(job, graph, job.Batches[0].Id, null, JobStepBatchState.Complete, null));
+			job = await RunBatchAsync(job, 0);
+			job = await RunStepAsync(job, 0, 0, JobStepOutcome.Failure);
+			job = Deref(await job.TryUpdateBatchAsync(job.Batches[0].Id, null, JobStepBatchState.Complete, null));
 
 			Assert.AreEqual(2, job.Batches.Count);
 			Assert.AreEqual(1, job.Batches[0].Steps.Count);
@@ -588,7 +588,7 @@ namespace Horde.Server.Tests.Jobs
 			Assert.AreEqual(JobStepOutcome.Failure, step.Outcome);
 
 			// Retry the failed step
-			job = Deref(await JobCollection.TryUpdateStepAsync(job, graph, job.Batches[0].Id, job.Batches[0].Steps[0].Id, newRetryByUserId: UserId.Anonymous));
+			job = Deref(await job.TryUpdateStepAsync(job.Batches[0].Id, job.Batches[0].Steps[0].Id, newRetryByUserId: UserId.Anonymous));
 			Assert.AreEqual(3, job.Batches.Count);
 			Assert.AreEqual(1, job.Batches[0].Steps.Count);
 			Assert.AreEqual(1, job.Batches[1].Steps.Count);
@@ -634,8 +634,8 @@ namespace Horde.Server.Tests.Jobs
 			IJob job = await JobCollection.AddAsync(JobIdUtils.GenerateNewId(), new StreamId("ue4-main"), new TemplateId("test-build"), ContentHash.SHA1("hello"), graph, "Test job", 123, 123, options);
 
 			// Fail the first step
-			job = await RunBatchAsync(job, graph, 0);
-			job = await RunStepAsync(job, graph, 0, 0, JobStepOutcome.Failure);
+			job = await RunBatchAsync(job, 0);
+			job = await RunStepAsync(job, 0, 0, JobStepOutcome.Failure);
 
 			Assert.AreEqual(1, job.Batches.Count);
 			Assert.AreEqual(3, job.Batches[0].Steps.Count);
@@ -656,7 +656,7 @@ namespace Horde.Server.Tests.Jobs
 			Assert.AreEqual(JobStepOutcome.Failure, step.Outcome);
 
 			// Retry the last failed step. All the failed upstream steps should run again.
-			job = Deref(await JobCollection.TryUpdateStepAsync(job, graph, job.Batches[0].Id, job.Batches[0].Steps[2].Id, newRetryByUserId: UserId.Anonymous));
+			job = Deref(await job.TryUpdateStepAsync(job.Batches[0].Id, job.Batches[0].Steps[2].Id, newRetryByUserId: UserId.Anonymous));
 			Assert.AreEqual(2, job.Batches.Count);
 			Assert.AreEqual(1, job.Batches[0].Steps.Count);
 			Assert.AreEqual(3, job.Batches[1].Steps.Count);
@@ -703,8 +703,8 @@ namespace Horde.Server.Tests.Jobs
 			IJob job = await JobCollection.AddAsync(JobIdUtils.GenerateNewId(), new StreamId("ue4-main"), new TemplateId("test-build"), ContentHash.SHA1("hello"), graph, "Test job", 123, 123, options);
 
 			// Pass the first step
-			job = await RunBatchAsync(job, graph, 0);
-			job = await RunStepAsync(job, graph, 0, 0, JobStepOutcome.Success);
+			job = await RunBatchAsync(job, 0);
+			job = await RunStepAsync(job, 0, 0, JobStepOutcome.Success);
 
 			Assert.AreEqual(1, job.Batches.Count);
 			Assert.AreEqual(3, job.Batches[0].Steps.Count);
@@ -723,7 +723,7 @@ namespace Horde.Server.Tests.Jobs
 			Assert.AreEqual(JobStepState.Waiting, step.State);
 
 			// Update the priority of the second step. This should trigger a refresh of the graph structure, but NOT cause any new steps to be created.
-			job = Deref(await JobCollection.TryUpdateStepAsync(job, graph, job.Batches[0].Id, job.Batches[0].Steps[1].Id, newPriority: Priority.High));
+			job = Deref(await job.TryUpdateStepAsync(job.Batches[0].Id, job.Batches[0].Steps[1].Id, newPriority: Priority.High));
 			Assert.AreEqual(1, job.Batches.Count);
 			Assert.AreEqual(3, job.Batches[0].Steps.Count);
 
@@ -755,8 +755,8 @@ namespace Horde.Server.Tests.Jobs
 			IJob job = await JobCollection.AddAsync(JobIdUtils.GenerateNewId(), new StreamId("ue4-main"), new TemplateId("test-build"), ContentHash.SHA1("hello"), baseGraph, "Test job", 123, 123, options);
 			Assert.AreEqual(1, job.Batches.Count);
 
-			job = await RunBatchAsync(job, baseGraph, 0);
-			job = Deref(await JobCollection.TryUpdateBatchAsync(job, baseGraph, job.Batches[0].Id, null, JobStepBatchState.Complete, JobStepBatchError.UnknownShelf));
+			job = await RunBatchAsync(job, 0);
+			job = Deref(await job.TryUpdateBatchAsync(job.Batches[0].Id, null, JobStepBatchState.Complete, JobStepBatchError.UnknownShelf));
 
 			Assert.AreEqual(1, job.Batches.Count);
 		}
@@ -774,8 +774,8 @@ namespace Horde.Server.Tests.Jobs
 
 			IJob job = await JobCollection.AddAsync(JobIdUtils.GenerateNewId(), new StreamId("ue4-main"), new TemplateId("test-build"), ContentHash.SHA1("hello"), baseGraph, "Test job", 123, 123, options);
 
-			job = await RunBatchAsync(job, baseGraph, 0);
-			job = await RunStepAsync(job, baseGraph, 0, 0, JobStepOutcome.Success); // Setup Build
+			job = await RunBatchAsync(job, 0);
+			job = await RunStepAsync(job, 0, 0, JobStepOutcome.Success); // Setup Build
 
 			List<NewGroup> newGroups = new List<NewGroup>();
 
@@ -788,10 +788,10 @@ namespace Horde.Server.Tests.Jobs
 
 			// Check that we're only building the editor
 			IGraph graph = await GraphCollection.AppendAsync(baseGraph, newGroups, null, null);
-			job = Deref(await JobCollection.TryUpdateGraphAsync(job, baseGraph, graph));
+			job = Deref(await job.TryUpdateGraphAsync(graph));
 
 			// Move the initial batch to the starting state
-			job = await UpdateBatchAsync(job, graph, 1, JobStepBatchState.Starting);
+			job = await UpdateBatchAsync(job, 1, JobStepBatchState.Starting);
 			Assert.AreEqual(2, job.Batches.Count);
 			Assert.AreEqual(1, job.Batches[1].GroupIdx);
 			Assert.AreEqual(2, job.Batches[1].Steps.Count);
@@ -802,7 +802,7 @@ namespace Horde.Server.Tests.Jobs
 
 			// Force an update on the batches for the job
 			options.Arguments.Add("-Target=Compile Client");
-			job = Deref(await JobCollection.TryUpdateJobAsync(job, graph, arguments: options.Arguments));
+			job = Deref(await job.TryUpdateJobAsync(arguments: options.Arguments));
 
 			// Check that we're still using the original batch, and have added one more
 			Assert.AreEqual(3, job.Batches.Count);
diff --git a/Engine/Source/Programs/Horde/Horde.Server.Tests/Server/RpcServiceTest.cs b/Engine/Source/Programs/Horde/Horde.Server.Tests/Server/RpcServiceTest.cs
--- a/Engine/Source/Programs/Horde/Horde.Server.Tests/Server/RpcServiceTest.cs
+++ b/Engine/Source/Programs/Horde/Horde.Server.Tests/Server/RpcServiceTest.cs
@@ -480,7 +480,7 @@ namespace Horde.Server.Tests.Server
 			};
 
 			// Set the session ID on the job batch to pass auth later
-			Deref(await JobCollection.TryAssignLeaseAsync(fixture.Job1, 0, new PoolId("foo"),
+			Deref(await fixture.Job1.TryAssignLeaseAsync(0, new PoolId("foo"),
 				new AgentId("test"), sessionId,
 				new LeaseId(BinaryIdUtils.CreateNew()), LogIdUtils.GenerateNewId()));
 			/*
diff --git a/Engine/Source/Programs/Horde/Horde.Server/Issues/IssueTagService.cs b/Engine/Source/Programs/Horde/Horde.Server/Issues/IssueTagService.cs
--- a/Engine/Source/Programs/Horde/Horde.Server/Issues/IssueTagService.cs
+++ b/Engine/Source/Programs/Horde/Horde.Server/Issues/IssueTagService.cs
@@ -8,7 +8,6 @@ using System.Threading.Tasks;
 using EpicGames.Core;
 using EpicGames.Horde.Streams;
 using Horde.Server.Commits;
-using Horde.Server.Perforce;
 using Horde.Server.Server;
 using Horde.Server.Streams;
 using Horde.Server.Utilities;
diff --git a/Engine/Source/Programs/Horde/Horde.Server/Jobs/IJob.cs b/Engine/Source/Programs/Horde/Horde.Server/Jobs/IJob.cs
--- a/Engine/Source/Programs/Horde/Horde.Server/Jobs/IJob.cs
+++ b/Engine/Source/Programs/Horde/Horde.Server/Jobs/IJob.cs
@@ -5,6 +5,8 @@ using System.Collections.Generic;
 using System.Diagnostics.CodeAnalysis;
 using System.Linq;
 using System.Security.Claims;
+using System.Threading;
+using System.Threading.Tasks;
 using EpicGames.Core;
 using EpicGames.Horde.Acls;
 using EpicGames.Horde.Agents;
@@ -573,6 +575,11 @@ namespace Horde.Server.Jobs
 		/// </summary>
 		public ContentHash GraphHash { get; }
 
+		/// <summary>
+		/// Graph for this job
+		/// </summary>
+		public IGraph Graph { get; }
+
 		/// <summary>
 		/// Id of the user that started this job
 		/// </summary>
@@ -747,6 +754,122 @@ namespace Horde.Server.Jobs
 		/// Update counter for this document. Any updates should compare-and-swap based on the value of this counter, or increment it in the case of server-side updates.
 		/// </summary>
 		public int UpdateIndex { get; }
+
+		/// <summary>
+		/// Gets the latest job state
+		/// </summary>
+		/// <param name="cancellationToken">Cancellation token for the operation</param>
+		Task<IJob?> RefreshAsync(CancellationToken cancellationToken = default);
+
+		/// <summary>
+		/// Removes a job from the dispatch queue. Ignores the state of any batches still remaining to execute. Should only be used to correct for inconsistent state.
+		/// </summary>
+		/// <param name="cancellationToken">Cancellation token for the operation</param>
+		/// <returns></returns>
+		Task<IJob?> TryRemoveFromDispatchQueueAsync(CancellationToken cancellationToken = default);
+
+		/// <summary>
+		/// Updates a new job
+		/// </summary>
+		/// <param name="name">Name of the job</param>
+		/// <param name="priority">Priority of the job</param>
+		/// <param name="autoSubmit">Automatically submit the job on completion</param>
+		/// <param name="autoSubmitChange">Changelist that was automatically submitted</param>
+		/// <param name="autoSubmitMessage"></param>
+		/// <param name="abortedByUserId">Name of the user that aborted the job</param>
+		/// <param name="notificationTriggerId">Id for a notification trigger</param>
+		/// <param name="reports">New reports</param>
+		/// <param name="arguments">New arguments for the job</param>
+		/// <param name="labelIdxToTriggerId">New trigger ID for a label in the job</param>
+		/// <param name="jobTrigger">New downstream job id</param>
+		/// <param name="cancellationToken">Cancellation token for the operation</param>
+		Task<IJob?> TryUpdateJobAsync(string? name = null, Priority? priority = null, bool? autoSubmit = null, int? autoSubmitChange = null, string? autoSubmitMessage = null, UserId? abortedByUserId = null, ObjectId? notificationTriggerId = null, List<Report>? reports = null, List<string>? arguments = null, KeyValuePair<int, ObjectId>? labelIdxToTriggerId = null, KeyValuePair<TemplateId, JobId>? jobTrigger = null, CancellationToken cancellationToken = default);
+
+		/// <summary>
+		/// Updates the state of a batch
+		/// </summary>
+		/// <param name="batchId">Unique id of the batch to update</param>
+		/// <param name="newLogId">The new log file id</param>
+		/// <param name="newState">New state of the jobstep</param>
+		/// <param name="newError">Error code for the batch</param>
+		/// <param name="cancellationToken">Cancellation token for the operation</param>
+		/// <returns>True if the job was updated, false if it was deleted</returns>
+		Task<IJob?> TryUpdateBatchAsync(JobStepBatchId batchId, LogId? newLogId, JobStepBatchState? newState, JobStepBatchError? newError, CancellationToken cancellationToken = default);
+
+		/// <summary>
+		/// Update a jobstep state
+		/// </summary>
+		/// <param name="batchId">Unique id of the batch containing the step</param>
+		/// <param name="stepId">Unique id of the step to update</param>
+		/// <param name="newState">New state of the jobstep</param>
+		/// <param name="newOutcome">New outcome of the jobstep</param>
+		/// <param name="newError">New error annotation for this jobstep</param>
+		/// <param name="newAbortRequested">New state of request abort</param>
+		/// <param name="newAbortByUserId">New name of user that requested the abort</param>
+		/// <param name="newLogId">New log id for the jobstep</param>
+		/// <param name="newNotificationTriggerId">New id for a notification trigger</param>
+		/// <param name="newRetryByUserId">Whether the step should be retried</param>
+		/// <param name="newPriority">New priority for this step</param>
+		/// <param name="newReports">New report documents</param>
+		/// <param name="newProperties">Property changes. Any properties with a null value will be removed.</param>
+		/// <param name="cancellationToken">Cancellation token for the operation</param>
+		/// <returns>True if the job was updated, false if it was deleted in the meantime</returns>
+		Task<IJob?> TryUpdateStepAsync(JobStepBatchId batchId, JobStepId stepId, JobStepState newState = default, JobStepOutcome newOutcome = default, JobStepError? newError = null, bool? newAbortRequested = null, UserId? newAbortByUserId = null, LogId? newLogId = null, ObjectId? newNotificationTriggerId = null, UserId? newRetryByUserId = null, Priority? newPriority = null, List<Report>? newReports = null, Dictionary<string, string?>? newProperties = null, CancellationToken cancellationToken = default);
+
+		/// <summary>
+		/// Attempts to update the node groups to be executed for a job. Fails if another write happens in the meantime.
+		/// </summary>
+		/// <param name="newGraph">New graph for this job</param>
+		/// <param name="cancellationToken">Cancellation token for the operation</param>
+		/// <returns>True if the groups were updated to the given list. False if another write happened first.</returns>
+		Task<IJob?> TryUpdateGraphAsync(IGraph newGraph, CancellationToken cancellationToken = default);
+
+		/// <summary>
+		/// Marks a job as skipped
+		/// </summary>
+		/// <param name="reason">Reason for this batch being failed</param>
+		/// <param name="cancellationToken">Cancellation token for the operation</param>
+		/// <returns>Updated version of the job</returns>
+		Task<IJob?> TrySkipAllBatchesAsync(JobStepBatchError reason, CancellationToken cancellationToken = default);
+
+		/// <summary>
+		/// Marks a batch as skipped
+		/// </summary>
+		/// <param name="batchId">The batch to mark as skipped</param>
+		/// <param name="reason">Reason for this batch being failed</param>
+		/// <param name="cancellationToken">Cancellation token for the operation</param>
+		/// <returns>Updated version of the job</returns>
+		Task<IJob?> TrySkipBatchAsync(JobStepBatchId batchId, JobStepBatchError reason, CancellationToken cancellationToken = default);
+
+		/// <summary>
+		/// Abort an agent's lease, and update the payload accordingly
+		/// </summary>
+		/// <param name="batchIdx">Index of the batch to cancel</param>
+		/// <param name="reason">Reason for this batch being failed</param>
+		/// <param name="cancellationToken">Cancellation token for the operation</param>
+		/// <returns>True if the job is updated</returns>
+		Task<IJob?> TryFailBatchAsync(int batchIdx, JobStepBatchError reason, CancellationToken cancellationToken = default);
+
+		/// <summary>
+		/// Attempt to assign a lease to execute a batch
+		/// </summary>
+		/// <param name="batchIdx">Index of the batch</param>
+		/// <param name="poolId">The pool id</param>
+		/// <param name="agentId">New agent to execute the batch</param>
+		/// <param name="sessionId">Session of the agent that is to execute the batch</param>
+		/// <param name="leaseId">The lease unique id</param>
+		/// <param name="logId">Unique id of the log for the batch</param>
+		/// <param name="cancellationToken">Cancellation token for the operation</param>
+		/// <returns>True if the batch is updated</returns>
+		Task<IJob?> TryAssignLeaseAsync(int batchIdx, PoolId poolId, AgentId agentId, SessionId sessionId, LeaseId leaseId, LogId logId, CancellationToken cancellationToken = default);
+
+		/// <summary>
+		/// Cancel a lease reservation on a batch (before it has started)
+		/// </summary>
+		/// <param name="batchIdx">Index of the batch to cancel</param>
+		/// <param name="cancellationToken">Cancellation token for the operation</param>
+		/// <returns>True if the job is updated</returns>
+		Task<IJob?> TryCancelLeaseAsync(int batchIdx, CancellationToken cancellationToken = default);
 	}
 
 	/// <summary>
@@ -1589,5 +1712,47 @@ namespace Horde.Server.Jobs
 		{
 			return $"job:{job.Id}/step:{jobStep.Id}";
 		}
+
+		/// <inheritdoc cref="IJob.TrySkipBatchAsync(JobStepBatchId, JobStepBatchError, CancellationToken)"/>
+		public static async Task<IJob?> SkipBatchAsync(this IJob job, JobStepBatchId batchId, JobStepBatchError error, CancellationToken cancellationToken)
+		{
+			for (; ; )
+			{
+				IJob? newJob = await job.TrySkipBatchAsync(batchId, error, cancellationToken);
+				if (newJob != null)
+				{
+					return newJob;
+				}
+
+				newJob = await job.RefreshAsync(cancellationToken);
+				if (newJob == null)
+				{
+					return null;
+				}
+
+				job = newJob;
+			}
+		}
+
+		/// <inheritdoc cref="IJob.TrySkipAllBatchesAsync(JobStepBatchError, CancellationToken)"/>
+		public static async Task<IJob?> SkipAllBatchesAsync(this IJob job, JobStepBatchError reason, CancellationToken cancellationToken)
+		{
+			for(; ;)
+			{
+				IJob? newJob = await job.TrySkipAllBatchesAsync(reason, cancellationToken);
+				if (newJob != null)
+				{
+					return newJob;
+				}
+
+				newJob = await job.RefreshAsync(cancellationToken);
+				if (newJob == null)
+				{
+					return null;
+				}
+
+				job = newJob;
+			}
+		}
 	}
 }
diff --git a/Engine/Source/Programs/Horde/Horde.Server/Jobs/IJobCollection.cs b/Engine/Source/Programs/Horde/Horde.Server/Jobs/IJobCollection.cs
--- a/Engine/Source/Programs/Horde/Horde.Server/Jobs/IJobCollection.cs
+++ b/Engine/Source/Programs/Horde/Horde.Server/Jobs/IJobCollection.cs
@@ -5,20 +5,14 @@ using System.Collections.Generic;
 using System.Threading;
 using System.Threading.Tasks;
 using EpicGames.Core;
-using EpicGames.Horde.Agents;
-using EpicGames.Horde.Agents.Leases;
-using EpicGames.Horde.Agents.Pools;
-using EpicGames.Horde.Agents.Sessions;
 using EpicGames.Horde.Jobs;
 using EpicGames.Horde.Jobs.Bisect;
 using EpicGames.Horde.Jobs.Templates;
-using EpicGames.Horde.Logs;
 using EpicGames.Horde.Streams;
 using EpicGames.Horde.Users;
 using Horde.Server.Acls;
 using Horde.Server.Jobs.Graphs;
 using Horde.Server.Streams;
-using MongoDB.Bson;
 
 namespace Horde.Server.Jobs
 {
@@ -205,78 +199,6 @@ namespace Horde.Server.Jobs
 		/// <param name="cancellationToken">Cancellation token for the operation</param>
 		Task<IReadOnlyList<IJob>> FindLatestByStreamWithTemplatesAsync(StreamId streamId, TemplateId[] templates, UserId? preflightStartedByUser = null, DateTimeOffset? maxCreateTime = null, DateTimeOffset? modifiedAfter = null, int? index = null, int? count = null, bool consistentRead = false, CancellationToken cancellationToken = default);
 
-		/// <summary>
-		/// Updates a new job
-		/// </summary>
-		/// <param name="job">The job document to update</param>
-		/// <param name="graph">Graph for the job</param>
-		/// <param name="name">Name of the job</param>
-		/// <param name="priority">Priority of the job</param>
-		/// <param name="autoSubmit">Automatically submit the job on completion</param>
-		/// <param name="autoSubmitChange">Changelist that was automatically submitted</param>
-		/// <param name="autoSubmitMessage"></param>
-		/// <param name="abortedByUserId">Name of the user that aborted the job</param>
-		/// <param name="notificationTriggerId">Id for a notification trigger</param>
-		/// <param name="reports">New reports</param>
-		/// <param name="arguments">New arguments for the job</param>
-		/// <param name="labelIdxToTriggerId">New trigger ID for a label in the job</param>
-		/// <param name="jobTrigger">New downstream job id</param>
-		/// <param name="cancellationToken">Cancellation token for the operation</param>
-		Task<IJob?> TryUpdateJobAsync(IJob job, IGraph graph, string? name = null, Priority? priority = null, bool? autoSubmit = null, int? autoSubmitChange = null, string? autoSubmitMessage = null, UserId? abortedByUserId = null, ObjectId? notificationTriggerId = null, List<Report>? reports = null, List<string>? arguments = null, KeyValuePair<int, ObjectId>? labelIdxToTriggerId = null, KeyValuePair<TemplateId, JobId>? jobTrigger = null, CancellationToken cancellationToken = default);
-
-		/// <summary>
-		/// Updates the state of a batch
-		/// </summary>
-		/// <param name="job">Job to update</param>
-		/// <param name="graph">The graph for this job</param>
-		/// <param name="batchId">Unique id of the batch to update</param>
-		/// <param name="newLogId">The new log file id</param>
-		/// <param name="newState">New state of the jobstep</param>
-		/// <param name="newError">Error code for the batch</param>
-		/// <param name="cancellationToken">Cancellation token for the operation</param>
-		/// <returns>True if the job was updated, false if it was deleted</returns>
-		Task<IJob?> TryUpdateBatchAsync(IJob job, IGraph graph, JobStepBatchId batchId, LogId? newLogId, JobStepBatchState? newState, JobStepBatchError? newError, CancellationToken cancellationToken = default);
-
-		/// <summary>
-		/// Update a jobstep state
-		/// </summary>
-		/// <param name="job">Job to update</param>
-		/// <param name="graph">The graph for this job</param>
-		/// <param name="batchId">Unique id of the batch containing the step</param>
-		/// <param name="stepId">Unique id of the step to update</param>
-		/// <param name="newState">New state of the jobstep</param>
-		/// <param name="newOutcome">New outcome of the jobstep</param>
-		/// <param name="newError">New error annotation for this jobstep</param>
-		/// <param name="newAbortRequested">New state of request abort</param>
-		/// <param name="newAbortByUserId">New name of user that requested the abort</param>
-		/// <param name="newLogId">New log id for the jobstep</param>
-		/// <param name="newNotificationTriggerId">New id for a notification trigger</param>
-		/// <param name="newRetryByUserId">Whether the step should be retried</param>
-		/// <param name="newPriority">New priority for this step</param>
-		/// <param name="newReports">New report documents</param>
-		/// <param name="newProperties">Property changes. Any properties with a null value will be removed.</param>
-		/// <param name="cancellationToken">Cancellation token for the operation</param>
-		/// <returns>True if the job was updated, false if it was deleted in the meantime</returns>
-		Task<IJob?> TryUpdateStepAsync(IJob job, IGraph graph, JobStepBatchId batchId, JobStepId stepId, JobStepState newState = default, JobStepOutcome newOutcome = default, JobStepError? newError = null, bool? newAbortRequested = null, UserId? newAbortByUserId = null, LogId? newLogId = null, ObjectId? newNotificationTriggerId = null, UserId? newRetryByUserId = null, Priority? newPriority = null, List<Report>? newReports = null, Dictionary<string, string?>? newProperties = null, CancellationToken cancellationToken = default);
-
-		/// <summary>
-		/// Attempts to update the node groups to be executed for a job. Fails if another write happens in the meantime.
-		/// </summary>
-		/// <param name="job">The job to update</param>
-		/// <param name="oldGraph">Old graph for this job</param>
-		/// <param name="newGraph">New graph for this job</param>
-		/// <param name="cancellationToken">Cancellation token for the operation</param>
-		/// <returns>True if the groups were updated to the given list. False if another write happened first.</returns>
-		Task<IJob?> TryUpdateGraphAsync(IJob job, IGraph oldGraph, IGraph newGraph, CancellationToken cancellationToken = default);
-
-		/// <summary>
-		/// Removes a job from the dispatch queue. Ignores the state of any batches still remaining to execute. Should only be used to correct for inconsistent state.
-		/// </summary>
-		/// <param name="job">Job</param>
-		/// <param name="cancellationToken">Cancellation token for the operation</param>
-		/// <returns></returns>
-		Task<IJob?> TryRemoveFromDispatchQueueAsync(IJob job, CancellationToken cancellationToken = default);
-
 		/// <summary>
 		/// Adds an issue to a job
 		/// </summary>
@@ -293,61 +215,6 @@ namespace Horde.Server.Jobs
 		/// <returns>Sorted list of jobs to execute</returns>
 		Task<IReadOnlyList<IJob>> GetDispatchQueueAsync(CancellationToken cancellationToken = default);
 
-		/// <summary>
-		/// Marks a job as skipped
-		/// </summary>
-		/// <param name="job">The job to update</param>
-		/// <param name="graph">Graph for the job</param>
-		/// <param name="reason">Reason for this batch being failed</param>
-		/// <param name="cancellationToken">Cancellation token for the operation</param>
-		/// <returns>Updated version of the job</returns>
-		Task<IJob?> SkipAllBatchesAsync(IJob? job, IGraph graph, JobStepBatchError reason, CancellationToken cancellationToken = default);
-
-		/// <summary>
-		/// Marks a batch as skipped
-		/// </summary>
-		/// <param name="job">The job to update</param>
-		/// <param name="batchId">The batch to mark as skipped</param>
-		/// <param name="graph">Graph for the job</param>
-		/// <param name="reason">Reason for this batch being failed</param>
-		/// <param name="cancellationToken">Cancellation token for the operation</param>
-		/// <returns>Updated version of the job</returns>
-		Task<IJob?> SkipBatchAsync(IJob? job, JobStepBatchId batchId, IGraph graph, JobStepBatchError reason, CancellationToken cancellationToken = default);
-
-		/// <summary>
-		/// Abort an agent's lease, and update the payload accordingly
-		/// </summary>
-		/// <param name="job">The job to update</param>
-		/// <param name="batchIdx">Index of the batch to cancel</param>
-		/// <param name="graph">Graph for the job</param>
-		/// <param name="reason">Reason for this batch being failed</param>
-		/// <param name="cancellationToken">Cancellation token for the operation</param>
-		/// <returns>True if the job is updated</returns>
-		Task<IJob?> TryFailBatchAsync(IJob job, int batchIdx, IGraph graph, JobStepBatchError reason, CancellationToken cancellationToken = default);
-
-		/// <summary>
-		/// Attempt to assign a lease to execute a batch
-		/// </summary>
-		/// <param name="job">The job containing the batch</param>
-		/// <param name="batchIdx">Index of the batch</param>
-		/// <param name="poolId">The pool id</param>
-		/// <param name="agentId">New agent to execute the batch</param>
-		/// <param name="sessionId">Session of the agent that is to execute the batch</param>
-		/// <param name="leaseId">The lease unique id</param>
-		/// <param name="logId">Unique id of the log for the batch</param>
-		/// <param name="cancellationToken">Cancellation token for the operation</param>
-		/// <returns>True if the batch is updated</returns>
-		Task<IJob?> TryAssignLeaseAsync(IJob job, int batchIdx, PoolId poolId, AgentId agentId, SessionId sessionId, LeaseId leaseId, LogId logId, CancellationToken cancellationToken = default);
-
-		/// <summary>
-		/// Cancel a lease reservation on a batch (before it has started)
-		/// </summary>
-		/// <param name="job">The job containing the lease</param>
-		/// <param name="batchIdx">Index of the batch to cancel</param>
-		/// <param name="cancellationToken">Cancellation token for the operation</param>
-		/// <returns>True if the job is updated</returns>
-		Task<IJob?> TryCancelLeaseAsync(IJob job, int batchIdx, CancellationToken cancellationToken = default);
-
 		/// <summary>
 		/// Upgrade all documents in the collection
 		/// </summary>
diff --git a/Engine/Source/Programs/Horde/Horde.Server/Jobs/JobCollection.cs b/Engine/Source/Programs/Horde/Horde.Server/Jobs/JobCollection.cs
--- a/Engine/Source/Programs/Horde/Horde.Server/Jobs/JobCollection.cs
+++ b/Engine/Source/Programs/Horde/Horde.Server/Jobs/JobCollection.cs
@@ -2,6 +2,7 @@
 
 using System;
 using System.Collections.Generic;
+using System.Diagnostics.CodeAnalysis;
 using System.IO;
 using System.Linq;
 using System.Runtime.CompilerServices;
@@ -228,7 +229,7 @@ namespace Horde.Server.Jobs
 			public ObjectId _triggerId;
 		}
 
-		class JobDocument : IJob
+		class JobDocument
 		{
 			[BsonRequired, BsonId]
 			public JobId Id { get; set; }
@@ -283,7 +284,6 @@ namespace Horde.Server.Jobs
 			public JobOptions? JobOptions { get; set; }
 
 			public List<AclClaimConfig> Claims { get; set; } = new List<AclClaimConfig>();
-			IReadOnlyList<AclClaimConfig> IJob.Claims => Claims;
 
 			[BsonIgnoreIfNull]
 			public DateTime? CreateTimeUtc { get; set; }
@@ -294,7 +294,6 @@ namespace Horde.Server.Jobs
 
 			[BsonDictionaryOptions(Representation = DictionaryRepresentation.ArrayOfDocuments)]
 			public Dictionary<ParameterId, string> Parameters { get; set; } = new Dictionary<ParameterId, string>();
-			IReadOnlyDictionary<ParameterId, string> IJob.Parameters => Parameters;
 
 			public List<string> Arguments { get; set; } = new List<string>();
 
@@ -323,16 +322,6 @@ namespace Horde.Server.Jobs
 			[BsonRequired]
 			public int UpdateIndex { get; set; }
 
-			DateTime IJob.CreateTimeUtc => CreateTimeUtc ?? CreateTime?.UtcDateTime ?? DateTime.UnixEpoch;
-			DateTime IJob.UpdateTimeUtc => UpdateTimeUtc ?? UpdateTime?.UtcDateTime ?? DateTime.UnixEpoch;
-			IReadOnlyList<IJobStepBatch> IJob.Batches => (IReadOnlyList<JobStepBatchDocument>)Batches;
-			IReadOnlyList<IReport>? IJob.Reports => Reports;
-			IReadOnlyList<string> IJob.Arguments => Arguments;
-			IReadOnlyList<int> IJob.Issues => ReferencedByIssues;
-			IReadOnlyDictionary<int, ObjectId> IJob.LabelIdxToTriggerId => _labelNotifications.ToDictionary(x => x._labelIdx, x => x._triggerId);
-			IReadOnlyDictionary<string, string> IJob.Environment => Environment;
-			IReadOnlyList<IChainedJob> IJob.ChainedJobs => ChainedJobs;
-
 			[BsonConstructor]
 			private JobDocument()
 			{
@@ -378,6 +367,130 @@ namespace Horde.Server.Jobs
 				NextSubResourceId = SubResourceId.GenerateNewId();
 				UpdateTimeUtc = createTimeUtc;
 			}
+
+			public DateTime GetCreateTimeOrDefault()
+				=> CreateTimeUtc ?? CreateTime?.UtcDateTime ?? DateTime.UnixEpoch;
+		}
+
+		class Job : IJob
+		{
+			readonly JobCollection _collection;
+			readonly JobDocument _document;
+			readonly IGraph _graph;
+
+			public Job(JobCollection collection, JobDocument document, IGraph graph)
+			{
+				_collection = collection;
+				_document = document;
+				_graph = graph;
+			}
+
+			public JobDocument Document => _document;
+
+			public JobId Id => _document.Id;
+			public StreamId StreamId => _document.StreamId;
+			public TemplateId TemplateId => _document.TemplateId;
+			public ContentHash? TemplateHash => _document.TemplateHash;
+			public ContentHash GraphHash => _document.GraphHash;
+			public IGraph Graph => _graph;
+			public UserId? StartedByUserId => _document.StartedByUserId;
+			public UserId? AbortedByUserId => _document.AbortedByUserId;
+			public BisectTaskId? StartedByBisectTaskId => _document.StartedByBisectTaskId;
+			public string Name => _document.Name;
+			public int Change => _document.Change;
+			public int CodeChange => _document.CodeChange;
+			public int PreflightChange => _document.PreflightChange;
+			public int ClonedPreflightChange => _document.ClonedPreflightChange;
+			public string? PreflightDescription => _document.PreflightDescription;
+			public Priority Priority => _document.Priority;
+			public bool AutoSubmit => _document.AutoSubmit;
+			public int? AutoSubmitChange => _document.AutoSubmitChange;
+			public string? AutoSubmitMessage => _document.AutoSubmitMessage;
+			public bool UpdateIssues => _document.UpdateIssues;
+			public bool PromoteIssuesByDefault => _document.PromoteIssuesByDefault;
+			public DateTime CreateTimeUtc => _document.GetCreateTimeOrDefault();
+			public JobOptions? JobOptions => _document.JobOptions;
+			public IReadOnlyList<AclClaimConfig> Claims => _document.Claims;
+			public int SchedulePriority => _document.SchedulePriority;
+			public IReadOnlyList<IJobStepBatch> Batches => _document.Batches;
+			public IReadOnlyDictionary<ParameterId, string> Parameters => _document.Parameters;
+			public IReadOnlyList<string> Arguments => _document.Arguments;
+			public IReadOnlyDictionary<string, string> Environment => _document.Environment;
+			public IReadOnlyList<int> Issues => _document.ReferencedByIssues;
+			public ObjectId? NotificationTriggerId => _document.NotificationTriggerId;
+			public bool ShowUgsBadges => _document.ShowUgsBadges;
+			public bool ShowUgsAlerts => _document.ShowUgsAlerts;
+			public string? NotificationChannel => _document.NotificationChannel;
+			public string? NotificationChannelFilter => _document.NotificationChannelFilter;
+			public IReadOnlyDictionary<int, ObjectId> LabelIdxToTriggerId => _document._labelNotifications.ToDictionary(x => x._labelIdx, x => x._triggerId);
+			public IReadOnlyList<IReport>? Reports => _document.Reports;
+			public IReadOnlyList<IChainedJob> ChainedJobs => _document.ChainedJobs;
+			public SubResourceId NextSubResourceId => _document.NextSubResourceId;
+			public DateTime UpdateTimeUtc => _document.UpdateTimeUtc ?? _document.UpdateTime?.UtcDateTime ?? DateTime.UnixEpoch;
+			public int UpdateIndex => _document.UpdateIndex;
+
+			public Task<IJob?> RefreshAsync(CancellationToken cancellationToken = default)
+				=> _collection.GetAsync(Id, cancellationToken);
+
+			public async Task<IJob?> TryAssignLeaseAsync(int batchIdx, PoolId poolId, AgentId agentId, SessionId sessionId, LeaseId leaseId, LogId logId, CancellationToken cancellationToken = default)
+			{
+				JobDocument? newDocument = await _collection.TryAssignLeaseAsync(Document, batchIdx, poolId, agentId, sessionId, leaseId, logId, cancellationToken);
+				return _collection.CreateJobObject(newDocument, _graph);
+			}
+
+			public async Task<IJob?> TryCancelLeaseAsync(int batchIdx, CancellationToken cancellationToken = default)
+			{
+				JobDocument? newDocument = await _collection.TryCancelLeaseAsync(Document, batchIdx, cancellationToken);
+				return _collection.CreateJobObject(newDocument, _graph);
+			}
+
+			public async Task<IJob?> TryFailBatchAsync(int batchIdx, JobStepBatchError reason, CancellationToken cancellationToken = default)
+			{
+				JobDocument? newDocument = await _collection.TryFailBatchAsync(Document, batchIdx, Graph, reason, cancellationToken);
+				return _collection.CreateJobObject(newDocument, _graph);
+			}
+
+			public async Task<IJob?> TryRemoveFromDispatchQueueAsync(CancellationToken cancellationToken = default)
+			{
+				JobDocument? newDocument = await _collection.TryRemoveFromDispatchQueueAsync(Document, cancellationToken);
+				return _collection.CreateJobObject(newDocument, _graph);
+			}
+
+			public async Task<IJob?> TrySkipAllBatchesAsync(JobStepBatchError reason, CancellationToken cancellationToken = default)
+			{
+				JobDocument? newDocument = await _collection.TrySkipAllBatchesAsync(Document, Graph, reason, cancellationToken);
+				return _collection.CreateJobObject(newDocument, _graph);
+			}
+
+			public async Task<IJob?> TrySkipBatchAsync(JobStepBatchId batchId, JobStepBatchError reason, CancellationToken cancellationToken = default)
+			{
+				JobDocument? newDocument = await _collection.TrySkipBatchAsync(Document, batchId, Graph, reason, cancellationToken);
+				return _collection.CreateJobObject(newDocument, _graph);
+			}
+
+			public async Task<IJob?> TryUpdateBatchAsync(JobStepBatchId batchId, LogId? newLogId, JobStepBatchState? newState, JobStepBatchError? newError, CancellationToken cancellationToken = default)
+			{
+				JobDocument? newDocument = await _collection.TryUpdateBatchAsync(Document, Graph, batchId, newLogId, newState, newError, cancellationToken);
+				return _collection.CreateJobObject(newDocument, _graph);
+			}
+
+			public async Task<IJob?> TryUpdateGraphAsync(IGraph newGraph, CancellationToken cancellationToken = default)
+			{
+				JobDocument? newDocument = await _collection.TryUpdateGraphAsync(Document, Graph, newGraph, cancellationToken);
+				return _collection.CreateJobObject(newDocument, newGraph);
+			}
+
+			public async Task<IJob?> TryUpdateJobAsync(string? name = null, Priority? priority = null, bool? autoSubmit = null, int? autoSubmitChange = null, string? autoSubmitMessage = null, UserId? abortedByUserId = null, ObjectId? notificationTriggerId = null, List<Report>? reports = null, List<string>? arguments = null, KeyValuePair<int, ObjectId>? labelIdxToTriggerId = null, KeyValuePair<TemplateId, JobId>? jobTrigger = null, CancellationToken cancellationToken = default)
+			{
+				JobDocument? newDocument = await _collection.TryUpdateJobAsync(Document, Graph, name, priority, autoSubmit, autoSubmitChange, autoSubmitMessage, abortedByUserId, notificationTriggerId, reports, arguments, labelIdxToTriggerId, jobTrigger, cancellationToken);
+				return _collection.CreateJobObject(newDocument, Graph);
+			}
+
+			public async Task<IJob?> TryUpdateStepAsync(JobStepBatchId batchId, JobStepId stepId, JobStepState newState = JobStepState.Unspecified, JobStepOutcome newOutcome = JobStepOutcome.Unspecified, JobStepError? newError = null, bool? newAbortRequested = null, UserId? newAbortByUserId = null, LogId? newLogId = null, ObjectId? newNotificationTriggerId = null, UserId? newRetryByUserId = null, Priority? newPriority = null, List<Report>? newReports = null, Dictionary<string, string?>? newProperties = null, CancellationToken cancellationToken = default)
+			{
+				JobDocument? newDocument = await _collection.TryUpdateStepAsync(_document, _graph, batchId, stepId, newState, newOutcome, newError, newAbortRequested, newAbortByUserId, newLogId, newNotificationTriggerId, newRetryByUserId, newPriority, newReports, newProperties, cancellationToken);
+				return _collection.CreateJobObject(newDocument, Graph);
+			}
 		}
 
 		/// <summary>
@@ -392,6 +505,7 @@ namespace Horde.Server.Jobs
 		readonly MongoIndex<JobDocument> _startedByBisectTaskIdIndex;
 		readonly ITelemetrySink _telemetrySink;
 		readonly IClock _clock;
+		readonly IGraphCollection _graphCollection;
 		readonly Tracer _tracer;
 		readonly IOptionsMonitor<GlobalConfig> _globalConfig;
 		readonly ILogger<JobCollection> _logger;
@@ -399,9 +513,10 @@ namespace Horde.Server.Jobs
 		/// <summary>
 		/// Constructor
 		/// </summary>
-		public JobCollection(MongoService mongoService, IClock clock, ITelemetrySink telemetrySink, IOptionsMonitor<GlobalConfig> globalConfig, Tracer tracer, ILogger<JobCollection> logger)
+		public JobCollection(MongoService mongoService, IClock clock, IGraphCollection graphCollection, ITelemetrySink telemetrySink, IOptionsMonitor<GlobalConfig> globalConfig, Tracer tracer, ILogger<JobCollection> logger)
 		{
 			_clock = clock;
+			_graphCollection = graphCollection;
 			_telemetrySink = telemetrySink;
 			_globalConfig = globalConfig;
 			_tracer = tracer;
@@ -473,7 +588,7 @@ namespace Horde.Server.Jobs
 				});
 			}
 
-			return newJob;
+			return CreateJobObject(newJob, graph);
 		}
 
 		/// <inheritdoc/>
@@ -484,7 +599,7 @@ namespace Horde.Server.Jobs
 			{
 				await PostLoadAsync(job);
 			}
-			return job;
+			return await CreateJobObjectAsync(job, cancellationToken);
 		}
 
 		/// <inheritdoc/>
@@ -582,17 +697,20 @@ namespace Horde.Server.Jobs
 				filter &= filterBuilder.ElemMatch(x => x.Batches, batch => batch.State == batchState);
 			}
 
-			List<JobDocument> results;
+			List<JobDocument> documents;
 			using (TelemetrySpan _ = _tracer.StartActiveSpan($"{nameof(JobCollection)}.{nameof(FindAsync)}"))
 			{
 				IMongoCollection<JobDocument> collection = consistentRead ? _jobs : _jobs.WithReadPreference(ReadPreference.SecondaryPreferred);
-				results = await collection.FindWithHintAsync(filter, indexHint, x => x.SortByDescending(x => x.CreateTimeUtc!).Range(index, count).ToListAsync());
+				documents = await collection.FindWithHintAsync(filter, indexHint, x => x.SortByDescending(x => x.CreateTimeUtc!).Range(index, count).ToListAsync());
 			}
-			foreach (JobDocument result in results)
+
+			List<IJob> jobs = new List<IJob>();
+			foreach (JobDocument document in documents)
 			{
-				await PostLoadAsync(result);
+				await PostLoadAsync(document);
+				jobs.Add(await CreateJobObjectAsync(document, cancellationToken));
 			}
-			return results;
+			return jobs;
 		}
 
 		/// <inheritdoc/>
@@ -607,9 +725,11 @@ namespace Horde.Server.Jobs
 			List<JobDocument> results = await _jobs.FindWithHintAsync(filter, _startedByBisectTaskIdIndex.Name, x => x.SortByDescending(x => x.CreateTimeUtc!).ToListAsync(cancellationToken));
 			foreach (JobDocument jobDoc in results)
 			{
+				Job job = await CreateJobObjectAsync(jobDoc, cancellationToken);
+
 				if (running.HasValue && running.Value)
 				{
-					JobState state = jobDoc.GetState();
+					JobState state = job.GetState();
 					if (state == JobState.Complete)
 					{
 						continue;
@@ -617,7 +737,7 @@ namespace Horde.Server.Jobs
 				}
 
 				await PostLoadAsync(jobDoc);
-				yield return jobDoc;
+				yield return job;
 			}
 		}
 
@@ -634,7 +754,7 @@ namespace Horde.Server.Jobs
 		}
 
 		/// <inheritdoc/>
-		public async Task<IJob?> TryUpdateJobAsync(IJob inJob, IGraph graph, string? name, Priority? priority, bool? autoSubmit, int? autoSubmitChange, string? autoSubmitMessage, UserId? abortedByUserId, ObjectId? notificationTriggerId, List<Report>? reports, List<string>? arguments, KeyValuePair<int, ObjectId>? labelIdxToTriggerId, KeyValuePair<TemplateId, JobId>? jobTrigger, CancellationToken cancellationToken)
+		async Task<JobDocument?> TryUpdateJobAsync(JobDocument jobDocument, IGraph graph, string? name, Priority? priority, bool? autoSubmit, int? autoSubmitChange, string? autoSubmitMessage, UserId? abortedByUserId, ObjectId? notificationTriggerId, List<Report>? reports, List<string>? arguments, KeyValuePair<int, ObjectId>? labelIdxToTriggerId, KeyValuePair<TemplateId, JobId>? jobTrigger, CancellationToken cancellationToken)
 		{
 			// Create the update 
 			UpdateDefinitionBuilder<JobDocument> updateBuilder = Builders<JobDocument>.Update;
@@ -644,7 +764,7 @@ namespace Horde.Server.Jobs
 			bool updateBatches = false;
 
 			// Build the update list
-			JobDocument jobDocument = Clone((JobDocument)inJob);
+			jobDocument = Clone(jobDocument);
 			if (name != null)
 			{
 				jobDocument.Name = name;
@@ -738,9 +858,9 @@ namespace Horde.Server.Jobs
 		}
 
 		/// <inheritdoc/>
-		public async Task<IJob?> TryUpdateBatchAsync(IJob job, IGraph graph, JobStepBatchId batchId, LogId? newLogId, JobStepBatchState? newState, JobStepBatchError? newError, CancellationToken cancellationToken)
+		async Task<JobDocument?> TryUpdateBatchAsync(JobDocument jobDocument, IGraph graph, JobStepBatchId batchId, LogId? newLogId, JobStepBatchState? newState, JobStepBatchError? newError, CancellationToken cancellationToken)
 		{
-			JobDocument jobDocument = Clone((JobDocument)job);
+			jobDocument = Clone(jobDocument);
 
 			// Find the index of the appropriate batch
 			int batchIdx = jobDocument.Batches.FindIndex(x => x.Id == batchId);
@@ -872,9 +992,9 @@ namespace Horde.Server.Jobs
 		}
 
 		/// <inheritdoc/>
-		public Task<IJob?> TryUpdateStepAsync(IJob job, IGraph graph, JobStepBatchId batchId, JobStepId stepId, JobStepState newState, JobStepOutcome newOutcome, JobStepError? newError, bool? newAbortRequested, UserId? newAbortByUserId, LogId? newLogId, ObjectId? newNotificationTriggerId, UserId? newRetryByUserId, Priority? newPriority, List<Report>? newReports, Dictionary<string, string?>? newProperties, CancellationToken cancellationToken)
+		async Task<JobDocument?> TryUpdateStepAsync(JobDocument jobDocument, IGraph graph, JobStepBatchId batchId, JobStepId stepId, JobStepState newState, JobStepOutcome newOutcome, JobStepError? newError, bool? newAbortRequested, UserId? newAbortByUserId, LogId? newLogId, ObjectId? newNotificationTriggerId, UserId? newRetryByUserId, Priority? newPriority, List<Report>? newReports, Dictionary<string, string?>? newProperties, CancellationToken cancellationToken)
 		{
-			JobDocument jobDocument = Clone((JobDocument)job);
+			jobDocument = Clone(jobDocument);
 
 			// Create the update 
 			UpdateDefinitionBuilder<JobDocument> updateBuilder = Builders<JobDocument>.Update;
@@ -1062,14 +1182,34 @@ namespace Horde.Server.Jobs
 			}
 
 			// Update the new list of job steps
-			return TryUpdateAsync(jobDocument, updates, cancellationToken);
+			return await TryUpdateAsync(jobDocument, updates, cancellationToken);
+		}
+
+		[return: NotNullIfNotNull("document")]
+		Job? CreateJobObject(JobDocument? document, IGraph graph)
+		{
+			if (document == null)
+			{
+				return null;
+			}
+			else
+			{
+				return new Job(this, document, graph);
+			}
 		}
 
-		Task<IJob?> TryUpdateAsync(JobDocument job, List<UpdateDefinition<JobDocument>> updates, CancellationToken cancellationToken)
+		[return: NotNullIfNotNull("document")]
+		async Task<Job> CreateJobObjectAsync(JobDocument document, CancellationToken cancellationToken)
+		{
+			IGraph graph = await _graphCollection.GetAsync(document.GraphHash, cancellationToken);
+			return new Job(this, document, graph);
+		}
+
+		Task<JobDocument?> TryUpdateAsync(JobDocument job, List<UpdateDefinition<JobDocument>> updates, CancellationToken cancellationToken)
 		{
 			if (updates.Count == 0)
 			{
-				return Task.FromResult<IJob?>(job);
+				return Task.FromResult<JobDocument?>(job);
 			}
 			else
 			{
@@ -1077,7 +1217,7 @@ namespace Horde.Server.Jobs
 			}
 		}
 
-		async Task<IJob?> TryUpdateAsync(JobDocument job, UpdateDefinition<JobDocument> update, CancellationToken cancellationToken)
+		async Task<JobDocument?> TryUpdateAsync(JobDocument job, UpdateDefinition<JobDocument> update, CancellationToken cancellationToken)
 		{
 			int newUpdateIndex = job.UpdateIndex + 1;
 			update = update.Set(x => x.UpdateIndex, newUpdateIndex);
@@ -1085,20 +1225,13 @@ namespace Horde.Server.Jobs
 			DateTime newUpdateTimeUtc = DateTime.UtcNow;
 			update = update.Set(x => x.UpdateTimeUtc, newUpdateTimeUtc);
 
-			UpdateResult result = await _jobs.UpdateOneAsync<JobDocument>(x => x.Id == job.Id && x.UpdateIndex == job.UpdateIndex, update, null, cancellationToken);
-			if (result.ModifiedCount > 0)
-			{
-				job.UpdateIndex = newUpdateIndex;
-				job.UpdateTimeUtc = newUpdateTimeUtc;
-				return job;
-			}
-			return null;
+			return await _jobs.FindOneAndUpdateAsync<JobDocument>(x => x.Id == job.Id && x.UpdateIndex == job.UpdateIndex, update, new FindOneAndUpdateOptions<JobDocument, JobDocument> { ReturnDocument = ReturnDocument.After }, cancellationToken);
 		}
 
 		/// <inheritdoc/>
-		public Task<IJob?> TryRemoveFromDispatchQueueAsync(IJob job, CancellationToken cancellationToken)
+		async Task<JobDocument?> TryRemoveFromDispatchQueueAsync(JobDocument jobDocument, CancellationToken cancellationToken)
 		{
-			JobDocument jobDocument = Clone((JobDocument)job);
+			jobDocument = Clone(jobDocument);
 
 			// Create the update 
 			UpdateDefinitionBuilder<JobDocument> updateBuilder = Builders<JobDocument>.Update;
@@ -1108,21 +1241,15 @@ namespace Horde.Server.Jobs
 			updates.Add(updateBuilder.Set(x => x.SchedulePriority, jobDocument.SchedulePriority));
 
 			// Update the new list of job steps
-			return TryUpdateAsync(jobDocument, updates, cancellationToken);
+			return await TryUpdateAsync(jobDocument, updates, cancellationToken);
 		}
 
 		/// <inheritdoc/>
-		public Task<IJob?> TryUpdateGraphAsync(IJob job, IGraph oldGraph, IGraph newGraph, CancellationToken cancellationToken)
+		async Task<JobDocument?> TryUpdateGraphAsync(JobDocument jobDocument, IGraph oldGraph, IGraph newGraph, CancellationToken cancellationToken)
 		{
-			JobDocument jobDocument = (JobDocument)job;
-			if (job.GraphHash != oldGraph.Id)
-			{
-				throw new InvalidOperationException($"Job {job.Id} uses graph {job.GraphHash}, not {oldGraph.Id}");
-			}
-
 			// Update all the references in the job to use references within the new graph
 			List<int> newNodeIndexes = new List<int>();
-			foreach (JobStepBatchDocument batch in job.Batches)
+			foreach (JobStepBatchDocument batch in jobDocument.Batches)
 			{
 				INodeGroup oldGroup = oldGraph.Groups[batch.GroupIdx];
 				newNodeIndexes.Clear();
@@ -1181,12 +1308,12 @@ namespace Horde.Server.Jobs
 			List<UpdateDefinition<JobDocument>> updates = new List<UpdateDefinition<JobDocument>>();
 
 			jobDocument.GraphHash = newGraph.Id;
-			updates.Add(updateBuilder.Set(x => x.GraphHash, job.GraphHash));
+			updates.Add(updateBuilder.Set(x => x.GraphHash, jobDocument.GraphHash));
 
 			UpdateBatches(jobDocument, newGraph, updates, _logger);
 
 			// Update the new list of job steps
-			return TryUpdateAsync(jobDocument, updates, cancellationToken);
+			return await TryUpdateAsync(jobDocument, updates, cancellationToken);
 		}
 
 		static bool NodesMatch(IGraph oldGraph, INode oldNode, IGraph newGraph, INode newNode)
@@ -1225,97 +1352,76 @@ namespace Horde.Server.Jobs
 		public async Task<IReadOnlyList<IJob>> GetDispatchQueueAsync(CancellationToken cancellationToken)
 		{
 			List<JobDocument> newJobs = await _jobs.Find(x => x.SchedulePriority > 0).SortByDescending(x => x.SchedulePriority).ThenBy(x => x.CreateTimeUtc).ToListAsync(cancellationToken);
+
+			List<Job> results = new List<Job>();
 			foreach (JobDocument result in newJobs)
 			{
 				await PostLoadAsync(result);
+				results.Add(await CreateJobObjectAsync(result, cancellationToken));
 			}
-			return newJobs;
+			return results;
 		}
 
 		/// <inheritdoc/>
-		public async Task<IJob?> SkipAllBatchesAsync(IJob? job, IGraph graph, JobStepBatchError reason, CancellationToken cancellationToken)
+		async Task<JobDocument?> TrySkipAllBatchesAsync(JobDocument jobDocument, IGraph graph, JobStepBatchError reason, CancellationToken cancellationToken)
 		{
-			while (job != null)
+			for (int batchIdx = 0; batchIdx < jobDocument.Batches.Count; batchIdx++)
 			{
-				JobDocument jobDocument = Clone((JobDocument)job);
-
-				for (int batchIdx = 0; batchIdx < jobDocument.Batches.Count; batchIdx++)
+				JobStepBatchDocument batch = jobDocument.Batches[batchIdx];
+				if (batch.State == JobStepBatchState.Ready || batch.State == JobStepBatchState.Waiting)
 				{
-					JobStepBatchDocument batch = jobDocument.Batches[batchIdx];
-					if (batch.State == JobStepBatchState.Ready || batch.State == JobStepBatchState.Waiting)
-					{
-						batch.State = JobStepBatchState.Complete;
-						batch.Error = reason;
-						batch.FinishTime = DateTimeOffset.UtcNow;
+					batch.State = JobStepBatchState.Complete;
+					batch.Error = reason;
+					batch.FinishTime = DateTimeOffset.UtcNow;
 
-						for (int stepIdx = 0; stepIdx < batch.Steps.Count; stepIdx++)
+					for (int stepIdx = 0; stepIdx < batch.Steps.Count; stepIdx++)
+					{
+						JobStepDocument step = batch.Steps[stepIdx];
+						if (step.State == JobStepState.Ready || step.State == JobStepState.Waiting)
 						{
-							JobStepDocument step = batch.Steps[stepIdx];
-							if (step.State == JobStepState.Ready || step.State == JobStepState.Waiting)
-							{
-								step.State = JobStepState.Completed;
-								step.Outcome = JobStepOutcome.Failure;
-							}
+							step.State = JobStepState.Completed;
+							step.Outcome = JobStepOutcome.Failure;
 						}
 					}
 				}
+			}
 
-				List<UpdateDefinition<JobDocument>> updates = new List<UpdateDefinition<JobDocument>>();
-				UpdateBatches(jobDocument, graph, updates, _logger);
-
-				IJob? newJob = await TryUpdateAsync(jobDocument, updates, cancellationToken);
-				if (newJob != null)
-				{
-					return newJob;
-				}
+			List<UpdateDefinition<JobDocument>> updates = new List<UpdateDefinition<JobDocument>>();
+			UpdateBatches(jobDocument, graph, updates, _logger);
 
-				job = await GetAsync(job.Id, cancellationToken);
-			}
-			return job;
+			return await TryUpdateAsync(jobDocument, updates, cancellationToken);
 		}
 
 		/// <inheritdoc/>
-		public async Task<IJob?> SkipBatchAsync(IJob? job, JobStepBatchId batchId, IGraph graph, JobStepBatchError reason, CancellationToken cancellationToken)
+		async Task<JobDocument?> TrySkipBatchAsync(JobDocument jobDocument, JobStepBatchId batchId, IGraph graph, JobStepBatchError reason, CancellationToken cancellationToken)
 		{
-			while (job != null)
+			JobStepBatchDocument? batch = jobDocument.Batches.FirstOrDefault(x => x.Id == batchId);
+			if (batch == null)
 			{
-				JobDocument jobDocument = (JobDocument)job;
-
-				JobStepBatchDocument? batch = jobDocument.Batches.FirstOrDefault(x => x.Id == batchId);
-				if (batch == null)
-				{
-					return job;
-				}
+				return jobDocument;
+			}
 
-				batch.State = JobStepBatchState.Complete;
-				batch.Error = reason;
-				batch.FinishTime = DateTimeOffset.UtcNow;
+			batch.State = JobStepBatchState.Complete;
+			batch.Error = reason;
+			batch.FinishTime = DateTimeOffset.UtcNow;
 
-				foreach (JobStepDocument step in batch.Steps)
+			foreach (JobStepDocument step in batch.Steps)
+			{
+				if (step.State != JobStepState.Skipped)
 				{
-					if (step.State != JobStepState.Skipped)
-					{
-						step.State = JobStepState.Skipped;
-						step.Outcome = JobStepOutcome.Failure;
-					}
+					step.State = JobStepState.Skipped;
+					step.Outcome = JobStepOutcome.Failure;
 				}
+			}
 
-				List<UpdateDefinition<JobDocument>> updates = new List<UpdateDefinition<JobDocument>>();
-				UpdateBatches(jobDocument, graph, updates, _logger);
-
-				IJob? newJob = await TryUpdateAsync(jobDocument, updates, cancellationToken);
-				if (newJob != null)
-				{
-					return newJob;
-				}
+			List<UpdateDefinition<JobDocument>> updates = new List<UpdateDefinition<JobDocument>>();
+			UpdateBatches(jobDocument, graph, updates, _logger);
 
-				job = await GetAsync(job.Id, cancellationToken);
-			}
-			return job;
+			return await TryUpdateAsync(jobDocument, updates, cancellationToken);
 		}
 
 		/// <inheritdoc/>
-		public async Task<IJob?> TryAssignLeaseAsync(IJob job, int batchIdx, PoolId poolId, AgentId agentId, SessionId sessionId, LeaseId leaseId, LogId logId, CancellationToken cancellationToken)
+		async Task<JobDocument?> TryAssignLeaseAsync(JobDocument jobDocument, int batchIdx, PoolId poolId, AgentId agentId, SessionId sessionId, LeaseId leaseId, LogId logId, CancellationToken cancellationToken)
 		{
 			// Try to update the job with this agent id
 			UpdateDefinitionBuilder<JobDocument> updateBuilder = Builders<JobDocument>.Update;
@@ -1328,34 +1434,22 @@ namespace Horde.Server.Jobs
 			updates.Add(updateBuilder.Set(x => x.Batches[batchIdx].LogId, logId));
 
 			// Extra logs for catching why session IDs sometimes doesn't match in prod. Resulting in PermissionDenied errors.
-			if (batchIdx < job.Batches.Count && job.Batches[batchIdx].SessionId != null)
+			if (batchIdx < jobDocument.Batches.Count && jobDocument.Batches[batchIdx].SessionId != null)
 			{
-				string currentSessionId = job.Batches[batchIdx].SessionId!.Value.ToString();
-				_logger.LogError("Attempt to replace current session ID {CurrSessionId} with {NewSessionId} for batch {JobId}:{BatchId}", currentSessionId, sessionId.ToString(), job.Id.ToString(), job.Batches[batchIdx].Id);
+				string currentSessionId = jobDocument.Batches[batchIdx].SessionId!.Value.ToString();
+				_logger.LogError("Attempt to replace current session ID {CurrSessionId} with {NewSessionId} for batch {JobId}:{BatchId}", currentSessionId, sessionId.ToString(), jobDocument.Id.ToString(), jobDocument.Batches[batchIdx].Id);
 				return null;
 			}
 
-			JobDocument jobDocument = Clone((JobDocument)job);
-			if (await TryUpdateAsync(jobDocument, updates, cancellationToken) == null)
-			{
-				return null;
-			}
-
-			JobStepBatchDocument batch = jobDocument.Batches[batchIdx];
-			batch.AgentId = agentId;
-			batch.SessionId = sessionId;
-			batch.LeaseId = leaseId;
-			batch.LogId = logId;
-
-			return jobDocument;
+			return await TryUpdateAsync(jobDocument, updates, cancellationToken);
 		}
 
 		/// <inheritdoc/>
-		public async Task<IJob?> TryCancelLeaseAsync(IJob job, int batchIdx, CancellationToken cancellationToken)
+		async Task<JobDocument?> TryCancelLeaseAsync(JobDocument jobDocument, int batchIdx, CancellationToken cancellationToken)
 		{
-			_logger.LogInformation("Cancelling lease {LeaseId} for agent {AgentId}", job.Batches[batchIdx].LeaseId, job.Batches[batchIdx].AgentId);
+			_logger.LogInformation("Cancelling lease {LeaseId} for agent {AgentId}", jobDocument.Batches[batchIdx].LeaseId, jobDocument.Batches[batchIdx].AgentId);
 
-			JobDocument jobDocument = Clone((JobDocument)job);
+			jobDocument = Clone(jobDocument);
 
 			UpdateDefinition<JobDocument> update = Builders<JobDocument>.Update.Unset(x => x.Batches[batchIdx].AgentId).Unset(x => x.Batches[batchIdx].SessionId).Unset(x => x.Batches[batchIdx].LeaseId);
 			if (await TryUpdateAsync(jobDocument, update, cancellationToken) == null)
@@ -1372,11 +1466,12 @@ namespace Horde.Server.Jobs
 		}
 
 		/// <inheritdoc/>
-		public Task<IJob?> TryFailBatchAsync(IJob job, int batchIdx, IGraph graph, JobStepBatchError error, CancellationToken cancellationToken)
+		async Task<JobDocument?> TryFailBatchAsync(JobDocument jobDocument, int batchIdx, IGraph graph, JobStepBatchError error, CancellationToken cancellationToken)
 		{
-			JobDocument jobDocument = Clone((JobDocument)job);
+			jobDocument = Clone(jobDocument);
+
 			JobStepBatchDocument batch = jobDocument.Batches[batchIdx];
-			_logger.LogInformation("Failing batch {JobId}:{BatchId} with error {Error}", job.Id, batch.Id, error);
+			_logger.LogInformation("Failing batch {JobId}:{BatchId} with error {Error}", jobDocument.Id, batch.Id, error);
 
 			UpdateDefinitionBuilder<JobDocument> updateBuilder = new UpdateDefinitionBuilder<JobDocument>();
 			List<UpdateDefinition<JobDocument>> updates = new List<UpdateDefinition<JobDocument>>();
@@ -1424,7 +1519,7 @@ namespace Horde.Server.Jobs
 			RefreshDependentJobSteps(jobDocument, graph, updates, _logger);
 			RefreshJobPriority(jobDocument, updates);
 
-			return TryUpdateAsync(jobDocument, updateBuilder.Combine(updates), cancellationToken);
+			return await TryUpdateAsync(jobDocument, updateBuilder.Combine(updates), cancellationToken);
 		}
 
 		/// <summary>
@@ -1850,6 +1945,11 @@ namespace Horde.Server.Jobs
 		/// </summary>
 		/// <param name="job">Job to consider</param>
 		public static int GetSchedulePriority(IJob job)
+		{
+			return GetSchedulePriority(((Job)job).Document);
+		}
+
+		static int GetSchedulePriority(JobDocument job)
 		{
 			int newSchedulePriority = 0;
 			foreach (IJobStepBatch batch in job.Batches)
@@ -1985,7 +2085,7 @@ namespace Horde.Server.Jobs
 		/// <param name="stepForNode">Array mapping each index to the appropriate step for that node</param>
 		/// <param name="outReadyTimeUtc">Receives the time at which the batch was ready to execute</param>
 		/// <returns>True if the batch is ready, false otherwise</returns>
-		static JobStepBatchState GetBatchState(IJob job, IGraph graph, IJobStepBatch batch, Dictionary<INode, IJobStep> stepForNode, out DateTime? outReadyTimeUtc)
+		static JobStepBatchState GetBatchState(JobDocument job, IGraph graph, IJobStepBatch batch, Dictionary<INode, IJobStep> stepForNode, out DateTime? outReadyTimeUtc)
 		{
 			// Check if the batch is already complete
 			if (batch.Steps.All(x => x.State == JobStepState.Skipped || x.State == JobStepState.Completed || x.State == JobStepState.Aborted))
@@ -1998,7 +2098,7 @@ namespace Horde.Server.Jobs
 			List<INode> nodeDependencies = batch.GetStartDependencies(graph.Groups).ToList();
 
 			// Check if we're still waiting on anything
-			DateTime readyTimeUtc = job.CreateTimeUtc;
+			DateTime readyTimeUtc = job.GetCreateTimeOrDefault();
 			foreach (INode nodeDependency in nodeDependencies)
 			{
 				IJobStep? stepDependency;
diff --git a/Engine/Source/Programs/Horde/Horde.Server/Jobs/JobRpcService.cs b/Engine/Source/Programs/Horde/Horde.Server/Jobs/JobRpcService.cs
--- a/Engine/Source/Programs/Horde/Horde.Server/Jobs/JobRpcService.cs
+++ b/Engine/Source/Programs/Horde/Horde.Server/Jobs/JobRpcService.cs
@@ -866,7 +866,7 @@ namespace Horde.Server.Jobs
 				IGraph newGraph = await _graphs.AppendAsync(null, newGroups, newAggregates, newLabels, newArtifacts);
 
 				// Try to update the graph with the new value
-				IJob? newJob = await _jobService.TryUpdateGraphAsync(job, oldGraph, newGraph, context.CancellationToken);
+				IJob? newJob = await _jobService.TryUpdateGraphAsync(job, newGraph, context.CancellationToken);
 				if (newJob != null)
 				{
 					_logger.LogInformation("Updating graph for {JobId} from {OldGraphHash} to {NewGraphHash}", job.Id, oldGraph.Id, newJob.GraphHash);
diff --git a/Engine/Source/Programs/Horde/Horde.Server/Jobs/JobService.cs b/Engine/Source/Programs/Horde/Horde.Server/Jobs/JobService.cs
--- a/Engine/Source/Programs/Horde/Horde.Server/Jobs/JobService.cs
+++ b/Engine/Source/Programs/Horde/Horde.Server/Jobs/JobService.cs
@@ -374,7 +374,7 @@ namespace Horde.Server.Jobs
 				IReadOnlyList<(LabelState, LabelOutcome)> oldLabelStates = newJob.GetLabelStates(graph);
 
 				// Update the new list of job steps
-				newJob = await _jobs.TryUpdateJobAsync(newJob, graph, name, priority, autoSubmit, null, null, abortedByUserId, onCompleteTriggerId, reports, arguments, labelIdxToTriggerId, cancellationToken: cancellationToken);
+				newJob = await newJob.TryUpdateJobAsync(name, priority, autoSubmit, null, null, abortedByUserId, onCompleteTriggerId, reports, arguments, labelIdxToTriggerId, cancellationToken: cancellationToken);
 				if (newJob != null)
 				{
 					// Update any badges that have been modified
@@ -675,11 +675,10 @@ namespace Horde.Server.Jobs
 		/// Attempts to update the node groups to be executed for a job. Fails if another write happens in the meantime.
 		/// </summary>
 		/// <param name="job">The job to update</param>
-		/// <param name="oldGraph">Old graph for this job</param>
 		/// <param name="newGraph">New graph for this job</param>
 		/// <param name="cancellationToken">Cancellation token for the operation</param>
 		/// <returns>True if the groups were updated to the given list. False if another write happened first.</returns>
-		public async Task<IJob?> TryUpdateGraphAsync(IJob job, IGraph oldGraph, IGraph newGraph, CancellationToken cancellationToken)
+		public async Task<IJob?> TryUpdateGraphAsync(IJob job, IGraph newGraph, CancellationToken cancellationToken)
 		{
 			using TelemetrySpan span = _tracer.StartActiveSpan($"{nameof(JobService)}.{nameof(TryUpdateGraphAsync)}");
 			span.SetAttribute("Job", job.Id.ToString());
@@ -689,7 +688,7 @@ namespace Horde.Server.Jobs
 
 			IReadOnlyList<(LabelState, LabelOutcome)> oldLabelStates = job.GetLabelStates(newGraph);
 
-			IJob? newJob = await _jobs.TryUpdateGraphAsync(job, oldGraph, newGraph, cancellationToken);
+			IJob? newJob = await job.TryUpdateGraphAsync(newGraph, cancellationToken);
 			if (newJob != null)
 			{
 				await _jobTaskSource.UpdateUgsBadgesAsync(newJob, newGraph, oldLabelStates, cancellationToken);
@@ -953,8 +952,7 @@ namespace Horde.Server.Jobs
 		/// <returns>The updated job, otherwise null</returns>
 		public async Task<IJob?> TryUpdateBatchAsync(IJob job, JobStepBatchId batchId, LogId? newLogId = null, JobStepBatchState? newState = null, JobStepBatchError? newError = null, CancellationToken cancellationToken = default)
 		{
-			IGraph graph = await GetGraphAsync(job, cancellationToken);
-			return await _jobs.TryUpdateBatchAsync(job, graph, batchId, newLogId, newState, newError, cancellationToken);
+			return await job.TryUpdateBatchAsync(batchId, newLogId, newState, newError, cancellationToken);
 		}
 
 		/// <summary>
@@ -1055,7 +1053,7 @@ namespace Horde.Server.Jobs
 			}
 
 			// Update the step
-			IJob? newJob = await _jobs.TryUpdateStepAsync(job, graph, batchId, stepId, newState, newOutcome, newError, newAbortRequested, newAbortByUserId, newLogId, newTriggerId, newRetryByUserId, newPriority, newReports, newProperties, cancellationToken);
+			IJob? newJob = await job.TryUpdateStepAsync(batchId, stepId, newState, newOutcome, newError, newAbortRequested, newAbortByUserId, newLogId, newTriggerId, newRetryByUserId, newPriority, newReports, newProperties, cancellationToken);
 			if (newJob != null)
 			{
 				job = newJob;
@@ -1243,7 +1241,7 @@ namespace Horde.Server.Jobs
 
 			for (; ; )
 			{
-				IJob? newJob = await _jobs.TryUpdateJobAsync(job, graph, autoSubmitChange: change, autoSubmitMessage: message, cancellationToken: cancellationToken);
+				IJob? newJob = await job.TryUpdateJobAsync(autoSubmitChange: change, autoSubmitMessage: message, cancellationToken: cancellationToken);
 				if (newJob != null)
 				{
 					return newJob;
@@ -1331,7 +1329,7 @@ namespace Horde.Server.Jobs
 				// Update the job
 				JobId chainedJobId = JobIdUtils.GenerateNewId();
 
-				IJob? newJob = await _jobs.TryUpdateJobAsync(job, graph, jobTrigger: new KeyValuePair<TemplateId, JobId>(jobTrigger.TemplateRefId, chainedJobId));
+				IJob? newJob = await job.TryUpdateJobAsync(jobTrigger: new KeyValuePair<TemplateId, JobId>(jobTrigger.TemplateRefId, chainedJobId));
 				if (newJob != null)
 				{
 					TemplateRefConfig? templateRefConfig;
diff --git a/Engine/Source/Programs/Horde/Horde.Server/Jobs/JobTaskSource.cs b/Engine/Source/Programs/Horde/Horde.Server/Jobs/JobTaskSource.cs
--- a/Engine/Source/Programs/Horde/Horde.Server/Jobs/JobTaskSource.cs
+++ b/Engine/Source/Programs/Horde/Horde.Server/Jobs/JobTaskSource.cs
@@ -404,13 +404,13 @@ namespace Horde.Server.Jobs
 				if (newJob.GraphHash == null)
 				{
 					_logger.LogError("Job {JobId} has a null graph hash and can't be started.", newJob.Id);
-					await _jobs.TryRemoveFromDispatchQueueAsync(newJob, cancellationToken);
+					await newJob.TryRemoveFromDispatchQueueAsync(cancellationToken);
 					continue;
 				}
 				if (newJob.AbortedByUserId != null)
 				{
 					_logger.LogError("Job {JobId} was aborted but not removed from dispatch queue", newJob.Id);
-					await _jobs.TryRemoveFromDispatchQueueAsync(newJob, cancellationToken);
+					await newJob.TryRemoveFromDispatchQueueAsync(cancellationToken);
 					continue;
 				}
 
@@ -421,7 +421,7 @@ namespace Horde.Server.Jobs
 				IStream? stream;
 				if (!streams.TryGetValue(newJob.StreamId, out stream))
 				{
-					newJob = await _jobs.SkipAllBatchesAsync(newJob, graph, JobStepBatchError.UnknownStream, cancellationToken);
+					newJob = await newJob.SkipAllBatchesAsync(JobStepBatchError.UnknownStream, cancellationToken);
 					continue;
 				}
 
@@ -469,7 +469,7 @@ namespace Horde.Server.Jobs
 									if (step != null)
 									{
 										JobId jobId = newJob.Id;
-										newJob = await _jobs.TryUpdateStepAsync(newJob, graph, batch.Id, step.Id, JobStepState.Skipped, newError: JobStepError.Paused, cancellationToken: cancellationToken);
+										newJob = await newJob.TryUpdateStepAsync(batch.Id, step.Id, JobStepState.Skipped, newError: JobStepError.Paused, cancellationToken: cancellationToken);
 										if (newJob == null)
 										{
 											_logger.LogError("Job {JobId} failed to update step {StepName} pause state", jobId, state.Name);
@@ -504,7 +504,7 @@ namespace Horde.Server.Jobs
 					if (!newJob.Batches.Any(batch => batch.State == JobStepBatchState.Ready || batch.State == JobStepBatchState.Starting || batch.State == JobStepBatchState.Running || batch.State == JobStepBatchState.Stopping))
 					{
 						_logger.LogError("Job {JobId} is in dispatch queue but not currently executing", newJob.Id);
-						await _jobs.TryRemoveFromDispatchQueueAsync(newJob, cancellationToken);
+						await newJob.TryRemoveFromDispatchQueueAsync(cancellationToken);
 					}
 				}
 			}
@@ -544,7 +544,7 @@ namespace Horde.Server.Jobs
 			_logger.LogInformation("Skipping batch {BatchId} for job {JobId} (reason: {Reason})", batchId, job.Id, reason);
 
 			IReadOnlyList<(LabelState, LabelOutcome)> oldLabelStates = job.GetLabelStates(graph);
-			IJob? newJob = await _jobs.SkipBatchAsync(job, batchId, graph, reason, cancellationToken);
+			IJob? newJob = await job.SkipBatchAsync(batchId, reason, cancellationToken);
 			if (newJob != null)
 			{
 				IReadOnlyList<(LabelState, LabelOutcome)> newLabelStates = newJob.GetLabelStates(graph);
@@ -785,7 +785,7 @@ namespace Horde.Server.Jobs
 			LogId logId = LogIdUtils.GenerateNewId();
 
 			// Try to update the job with this agent id
-			IJob? newJob = await _jobs.TryAssignLeaseAsync(item._job, item._batchIdx, item._poolId, agent.Id, agent.SessionId!.Value, leaseId, logId, cancellationToken);
+			IJob? newJob = await item._job.TryAssignLeaseAsync(item._batchIdx, item._poolId, agent.Id, agent.SessionId!.Value, leaseId, logId, cancellationToken);
 			if (newJob != null)
 			{
 				job = newJob;
@@ -1127,7 +1127,7 @@ namespace Horde.Server.Jobs
 					}
 
 					IGraph graph = await _graphs.GetAsync(job.GraphHash, cancellationToken);
-					job = await _jobs.TryFailBatchAsync(job, batchIdx, graph, error, cancellationToken);
+					job = await job.TryFailBatchAsync(batchIdx, error, cancellationToken);
 
 					if (job != null)
 					{
@@ -1189,7 +1189,7 @@ namespace Horde.Server.Jobs
 					break;
 				}
 
-				IJob? newJob = await _jobs.TryCancelLeaseAsync(job, batchIdx, cancellationToken);
+				IJob? newJob = await job.TryCancelLeaseAsync(batchIdx, cancellationToken);
 				if (newJob != null)
 				{
 					break;
diff --git a/Engine/Source/Programs/Horde/Horde.Server/Jobs/JobsController.cs b/Engine/Source/Programs/Horde/Horde.Server/Jobs/JobsController.cs
--- a/Engine/Source/Programs/Horde/Horde.Server/Jobs/JobsController.cs
+++ b/Engine/Source/Programs/Horde/Horde.Server/Jobs/JobsController.cs
@@ -1134,10 +1134,9 @@ namespace Horde.Server.Jobs
 					return Forbid(JobAclAction.ExecuteJob, jobId);
 				}
 
-				IGraph oldGraph = await _jobService.GetGraphAsync(job, cancellationToken);
-				IGraph newGraph = await _graphs.AppendAsync(oldGraph, requests, null, null, null, cancellationToken);
+				IGraph newGraph = await _graphs.AppendAsync(job.Graph, requests, null, null, null, cancellationToken);
 
-				IJob? newJob = await _jobService.TryUpdateGraphAsync(job, oldGraph, newGraph, cancellationToken);
+				IJob? newJob = await _jobService.TryUpdateGraphAsync(job, newGraph, cancellationToken);
 				if (newJob != null)
 				{
 					return Ok();
