Horde: Fix static analysis warning.

diff --git a/Engine/Source/Programs/Horde/Horde.Server/Jobs/JobCollection.cs b/Engine/Source/Programs/Horde/Horde.Server/Jobs/JobCollection.cs
--- a/Engine/Source/Programs/Horde/Horde.Server/Jobs/JobCollection.cs
+++ b/Engine/Source/Programs/Horde/Horde.Server/Jobs/JobCollection.cs
@@ -595,10 +595,12 @@ namespace Horde.Server.Jobs
 		public async Task<IJob?> GetAsync(JobId jobId, CancellationToken cancellationToken)
 		{
 			JobDocument? job = await _jobs.Find<JobDocument>(x => x.Id == jobId).FirstOrDefaultAsync(cancellationToken);
-			if (job != null)
+			if (job == null)
 			{
-				await PostLoadAsync(job);
+				return null;
 			}
+
+			await PostLoadAsync(job);
 			return await CreateJobObjectAsync(job, cancellationToken);
 		}
 
