Horde: Manually spawn multiple tasks for checking issue tags to avoid parallel loop limiting concurrency.

#rnx

diff --git a/Engine/Source/Programs/Horde/Horde.Server/Issues/IssueTagService.cs b/Engine/Source/Programs/Horde/Horde.Server/Issues/IssueTagService.cs
--- a/Engine/Source/Programs/Horde/Horde.Server/Issues/IssueTagService.cs
+++ b/Engine/Source/Programs/Horde/Horde.Server/Issues/IssueTagService.cs
@@ -2,6 +2,7 @@
 
 using System;
 using System.Collections.Generic;
+using System.Linq;
 using System.Text.RegularExpressions;
 using System.Threading;
 using System.Threading.Tasks;
@@ -69,9 +70,31 @@ namespace Horde.Server.Issues
 			State initialState = await _state.GetAsync(cancellationToken);
 
 			GlobalConfig globalConfig = _globalConfig.CurrentValue;
+
+			List<Task> tasks = new List<Task>();
 			try
 			{
-				await Parallel.ForEachAsync(globalConfig.Streams, cancellationSource.Token, async (stream, ctx) => await TickStreamAsync(stream, initialState, ctx));
+				foreach (StreamConfig streamConfig in globalConfig.Streams)
+				{
+					tasks.Add(Task.Run(() => TickStreamGuardedAsync(streamConfig, initialState, cancellationSource.Token), cancellationSource.Token));
+				}
+				await Task.WhenAny(tasks);
+			}
+			finally
+			{
+				await cancellationSource.CancelAsync();
+				await Task.WhenAll(tasks);
+			}
+		}
+
+		async ValueTask TickStreamGuardedAsync(StreamConfig streamConfig, State initialState, CancellationToken cancellationToken)
+		{
+			try
+			{
+				await TickStreamAsync(streamConfig, initialState, cancellationToken);
+			}
+			catch (OperationCanceledException)
+			{
 			}
 			catch (Exception ex)
 			{
@@ -91,6 +114,7 @@ namespace Horde.Server.Issues
 
 			await foreach (ICommit commit in commits.SubscribeAsync(minChange + 1, null, cancellationToken))
 			{
+				_logger.LogDebug("Checking commit {Change} in {StreamId}", commit.Number, streamConfig.Id);
 				foreach (int issueId in ParseTags(_globalConfig.CurrentValue.IssueFixedTag, commit.Description))
 				{
 					for (; ; )
@@ -98,14 +122,14 @@ namespace Horde.Server.Issues
 						IIssue? issue = await _issueCollection.GetIssueAsync(issueId, cancellationToken);
 						if (issue == null)
 						{
-							_logger.LogInformation("Commit {Change} by {Author} has invalid issue id {IssueId}", commit.Number, commit.AuthorId, issueId);
+							_logger.LogInformation("Commit {Change} by {Author} in {StreamId} has invalid issue id {IssueId}", commit.Number, commit.AuthorId, streamConfig.Id, issueId);
 							break;
 						}
 
 						issue = await _issueCollection.TryUpdateIssueAsync(issue, commit.AuthorId, newFixChange: commit.Number, newResolvedById: commit.AuthorId, cancellationToken: cancellationToken);
 						if (issue != null)
 						{
-							_logger.LogInformation("Commit {Change} by {Author} fixes issue id {IssueId}", commit.Number, commit.AuthorId, issueId);
+							_logger.LogInformation("Commit {Change} by {Author} in {StreamId} fixes issue id {IssueId}", commit.Number, commit.AuthorId, streamConfig.Id, issueId);
 							break;
 						}
 					}
