Horde: Suppress cancellation exceptions in issuetagservice.

#rnx

diff --git a/Engine/Source/Programs/Horde/Horde.Server/Issues/IssueTagService.cs b/Engine/Source/Programs/Horde/Horde.Server/Issues/IssueTagService.cs
--- a/Engine/Source/Programs/Horde/Horde.Server/Issues/IssueTagService.cs
+++ b/Engine/Source/Programs/Horde/Horde.Server/Issues/IssueTagService.cs
@@ -67,7 +67,7 @@ namespace Horde.Server.Issues
 			using CancellationTokenSource cancellationSource = CancellationTokenSource.CreateLinkedTokenSource(cancellationToken);
 			using IDisposable? listener = _globalConfig.OnChange((_, _) => cancellationSource.Cancel());
 
-			State initialState = await _state.GetAsync(cancellationToken);
+			State initialState = await _state.GetAsync(cancellationSource.Token);
 
 			GlobalConfig globalConfig = _globalConfig.CurrentValue;
 
@@ -82,8 +82,14 @@ namespace Horde.Server.Issues
 			}
 			finally
 			{
-				await cancellationSource.CancelAsync();
-				await Task.WhenAll(tasks);
+				try
+				{
+					await cancellationSource.CancelAsync();
+					await Task.WhenAll(tasks);
+				}
+				catch
+				{
+				}
 			}
 		}
 
