Horde: Fix static analysis warning.

#rnx

diff --git a/Engine/Source/Programs/Horde/Horde.Server/Issues/IssueTagService.cs b/Engine/Source/Programs/Horde/Horde.Server/Issues/IssueTagService.cs
--- a/Engine/Source/Programs/Horde/Horde.Server/Issues/IssueTagService.cs
+++ b/Engine/Source/Programs/Horde/Horde.Server/Issues/IssueTagService.cs
@@ -2,7 +2,6 @@
 
 using System;
 using System.Collections.Generic;
-using System.Linq;
 using System.Text.RegularExpressions;
 using System.Threading;
 using System.Threading.Tasks;
