Horde: Updating schema documentation

#rnx

diff --git a/Engine/Source/Programs/Horde/Docs/Config/Schema/Globals.md b/Engine/Source/Programs/Horde/Docs/Config/Schema/Globals.md
--- a/Engine/Source/Programs/Horde/Docs/Config/Schema/Globals.md
+++ b/Engine/Source/Programs/Horde/Docs/Config/Schema/Globals.md
@@ -28,6 +28,7 @@ Name | Description
 `telemetryStores` | [`TelemetryStoreConfig`](Telemetry.md)`[]`<br>Metrics to aggregate on the Horde server
 `parameters` | `object`<br>General parameters for other tools. Can be queried through the api/v1/parameters endpoint.
 `acl` | [`AclConfig`](#aclconfig)<br>Access control list
+`issueFixedTag` | `string`<br>Commit tag to use for marking issues as fixed
 
 ## ConfigInclude
 
