Horde: Add min scratch space to workspace config

diff --git a/Engine/Source/Programs/Horde/Horde.Agent.Tests/WorkspaceInfoTest.cs b/Engine/Source/Programs/Horde/Horde.Agent.Tests/WorkspaceInfoTest.cs
--- a/Engine/Source/Programs/Horde/Horde.Agent.Tests/WorkspaceInfoTest.cs
+++ b/Engine/Source/Programs/Horde/Horde.Agent.Tests/WorkspaceInfoTest.cs
@@ -32,6 +32,7 @@ public class WorkspaceInfoTest
 		Assert.AreEqual(defaultOptions, GetMwOptions(null));
 		Assert.AreEqual(defaultOptions, GetMwOptions(""));
 		Assert.AreEqual(defaultOptions, GetMwOptions("name=somethingElse&numParallelSyncThreads=111"));
+		Assert.AreEqual(new MWO { MinScratchSpace = 444 }, GetMwOptions("name=managedWorkspace&minScratchSpace=444", 10));
 		Assert.AreEqual(new MWO { NumParallelSyncThreads = 111 }, GetMwOptions("name=managedWorkspace&numParallelSyncThreads=111"));
 		Assert.AreEqual(new MWO { MaxFileConcurrency = 222 }, GetMwOptions("name=managedWorkspace&maxFileConcurrency=222"));
 		Assert.AreEqual(new MWO { MinScratchSpace = 333 }, GetMwOptions("name=managedWorkspace&minScratchSpace=333"));
@@ -41,10 +42,9 @@ public class WorkspaceInfoTest
 		Assert.AreEqual(new MWO { UseHaveTable = true }, GetMwOptions("name=managedWorkspace&useHaveTable=TrUE"));
 	}
 
-	static ManagedWorkspaceOptions GetMwOptions(string? method)
+	static ManagedWorkspaceOptions GetMwOptions(string? method, long? minScratchSpace = null)
 	{
-		HordeCommon.Rpc.Messages.RpcAgentWorkspace workspace = new HordeCommon.Rpc.Messages.RpcAgentWorkspace();
-		workspace.Method = method ?? "";
+		HordeCommon.Rpc.Messages.RpcAgentWorkspace workspace = new () { Method = method ?? "", MinScratchSpace = minScratchSpace ?? 0 };
 		return WorkspaceInfo.GetMwOptions(workspace);
 	}
 }
diff --git a/Engine/Source/Programs/Horde/Horde.Agent/Utility/WorkspaceInfo.cs b/Engine/Source/Programs/Horde/Horde.Agent/Utility/WorkspaceInfo.cs
--- a/Engine/Source/Programs/Horde/Horde.Agent/Utility/WorkspaceInfo.cs
+++ b/Engine/Source/Programs/Horde/Horde.Agent/Utility/WorkspaceInfo.cs
@@ -291,10 +291,10 @@ namespace Horde.Agent.Utility
 		}
 
 		/// <summary>
-		/// Create ManagedWorkspace options from a URL-encoded query string.
+		/// Create manage workspace options
 		/// </summary>
-		/// <param name="workspace">Workspace settings</param>
-		/// <returns></returns>
+		/// <param name="workspace">Agent workspace RPC message</param>
+		/// <returns>Options for a managed workspace</returns>
 		public static ManagedWorkspaceOptions GetMwOptions(RpcAgentWorkspace workspace)
 		{
 			const string NameKey = "name";
@@ -308,6 +308,11 @@ namespace Horde.Agent.Utility
 			ManagedWorkspaceOptions options = new ManagedWorkspaceOptions();
 			options = options with { Partitioned = workspace.Partitioned };
 
+			if (workspace.MinScratchSpace > 0)
+			{
+				options = options with { MinScratchSpace = workspace.MinScratchSpace };
+			}
+
 			string? method = workspace.Method;
 			if (!String.IsNullOrEmpty(method))
 			{
diff --git a/Engine/Source/Programs/Horde/Horde.Server/Agents/IAgent.cs b/Engine/Source/Programs/Horde/Horde.Server/Agents/IAgent.cs
--- a/Engine/Source/Programs/Horde/Horde.Server/Agents/IAgent.cs
+++ b/Engine/Source/Programs/Horde/Horde.Server/Agents/IAgent.cs
@@ -70,6 +70,11 @@ namespace Horde.Server.Agents
 		/// </summary>
 		public string? Method { get; set; }
 
+		/// <summary>
+		/// Minimum disk space that must be available *after* syncing this workspace (in megabytes)
+		/// </summary>
+		public long? MinScratchSpace { get; set; }
+
 		/// <summary>
 		/// Constructor
 		/// </summary>
@@ -80,7 +85,8 @@ namespace Horde.Server.Agents
 		/// <param name="view">Custom view for the workspace</param>
 		/// <param name="incremental">Whether to use an incremental workspace</param>
 		/// <param name="method">Method to use when syncing/materializing data from Perforce</param>
-		public AgentWorkspaceInfo(string? cluster, string? userName, string identifier, string stream, List<string>? view, bool incremental, string? method)
+		/// <param name="minScratchSpace">Minimum disk space that must be available *after* syncing this workspace (in megabytes)</param>
+		public AgentWorkspaceInfo(string? cluster, string? userName, string identifier, string stream, List<string>? view, bool incremental, string? method, long? minScratchSpace)
 		{
 			if (!String.IsNullOrEmpty(cluster))
 			{
@@ -95,6 +101,7 @@ namespace Horde.Server.Agents
 			View = view;
 			Incremental = incremental;
 			Method = method;
+			MinScratchSpace = minScratchSpace;
 		}
 
 		/// <summary>
@@ -102,8 +109,10 @@ namespace Horde.Server.Agents
 		/// </summary>
 		/// <param name="workspace">RPC message to construct from</param>
 		public AgentWorkspaceInfo(RpcAgentWorkspace workspace)
-			: this(workspace.ConfiguredCluster, workspace.ConfiguredUserName, workspace.Identifier, workspace.Stream, (workspace.View.Count > 0) ? workspace.View.ToList() : null, workspace.Incremental, workspace.Method)
+			: this(workspace.ConfiguredCluster, workspace.ConfiguredUserName, workspace.Identifier, workspace.Stream, (workspace.View.Count > 0) ? workspace.View.ToList() : null, workspace.Incremental, workspace.Method, workspace.MinScratchSpace)
 		{
+			// Treat zero value as null as Protobuf cannot store null 
+			MinScratchSpace = workspace.MinScratchSpace == 0 ? null : workspace.MinScratchSpace;
 		}
 
 		/// <summary>
@@ -179,7 +188,8 @@ namespace Horde.Server.Agents
 				Stream = Stream,
 				Incremental = Incremental,
 				Partitioned = server.SupportsPartitionedWorkspaces,
-				Method = Method ?? String.Empty
+				Method = Method ?? String.Empty,
+				MinScratchSpace = MinScratchSpace ?? 0
 			};
 
 			if (View != null)
@@ -949,7 +959,7 @@ namespace Horde.Server.Agents
 			{
 				if (autoSdk.Stream != null && autoSdk.Properties.All(x => agent.Properties.Contains(x)))
 				{
-					return new AgentWorkspaceInfo(cluster.Name, autoSdk.UserName, autoSdk.Name ?? "AutoSDK", autoSdk.Stream!, autoSdkConfig.View.ToList(), true, null);
+					return new AgentWorkspaceInfo(cluster.Name, autoSdk.UserName, autoSdk.Name ?? "AutoSDK", autoSdk.Stream!, autoSdkConfig.View.ToList(), true, null, null);
 				}
 			}
 			return null;
@@ -1052,7 +1062,7 @@ namespace Horde.Server.Agents
 			if (agentType.Workspace == null)
 			{
 				// Use the default settings (fast switching workspace, clean)
-				workspace = new AgentWorkspaceInfo(streamConfig.ClusterName, null, streamConfig.GetDefaultWorkspaceIdentifier(), streamConfig.Name, null, false, null);
+				workspace = new AgentWorkspaceInfo(streamConfig.ClusterName, null, streamConfig.GetDefaultWorkspaceIdentifier(), streamConfig.Name, null, false, null, null);
 				autoSdkConfig = AutoSdkConfig.Full;
 				return true;
 			}
@@ -1084,7 +1094,7 @@ namespace Horde.Server.Agents
 
 				// Create the new workspace
 				string cluster = workspaceConfig.Cluster ?? streamConfig.ClusterName;
-				workspace = new AgentWorkspaceInfo(cluster, workspaceConfig.UserName, identifier, workspaceConfig.Stream ?? streamConfig.Name, workspaceConfig.View, workspaceConfig.Incremental ?? false, workspaceConfig.Method);
+				workspace = new AgentWorkspaceInfo(cluster, workspaceConfig.UserName, identifier, workspaceConfig.Stream ?? streamConfig.Name, workspaceConfig.View, workspaceConfig.Incremental ?? false, workspaceConfig.Method, workspaceConfig.MinScratchSpace);
 				autoSdkConfig = GetAutoSdkConfig(workspaceConfig, streamConfig);
 
 				return true;
diff --git a/Engine/Source/Programs/Horde/Horde.Server/Streams/StreamConfig.cs b/Engine/Source/Programs/Horde/Horde.Server/Streams/StreamConfig.cs
--- a/Engine/Source/Programs/Horde/Horde.Server/Streams/StreamConfig.cs
+++ b/Engine/Source/Programs/Horde/Horde.Server/Streams/StreamConfig.cs
@@ -709,6 +709,12 @@ namespace Horde.Server.Streams
 		/// Method to use when syncing/materializing data from Perforce
 		/// </summary>
 		public string? Method { get; set; } = null;
+
+		/// <summary>
+		/// Minimum disk space that must be available *after* syncing this workspace (in megabytes)
+		/// If not available, the job will be aborted.
+		/// </summary>
+		public long? MinScratchSpace { get; set; } = null;
 	}
 
 	/// <summary>
diff --git a/Engine/Source/Programs/Shared/EpicGames.Perforce.Managed/ManagedWorkspace.cs b/Engine/Source/Programs/Shared/EpicGames.Perforce.Managed/ManagedWorkspace.cs
--- a/Engine/Source/Programs/Shared/EpicGames.Perforce.Managed/ManagedWorkspace.cs
+++ b/Engine/Source/Programs/Shared/EpicGames.Perforce.Managed/ManagedWorkspace.cs
@@ -84,7 +84,7 @@ namespace EpicGames.Perforce.Managed
 	/// </summary>
 	/// <param name="NumParallelSyncThreads">Maximum number of threads to sync in parallel</param>
 	/// <param name="MaxFileConcurrency">Maximum number of concurrent file system operations (copying, moving, deleting etc)</param>
-	/// <param name="MinScratchSpace">Minimum amount of space that must be on a drive after a branch is synced</param>
+	/// <param name="MinScratchSpace">Minimum amount of disk space that must be available after a sync (in megabytes)</param>
 	/// <param name="UseHaveTable">
 	///		Use the client's have table when syncing.
 	///		
@@ -93,11 +93,11 @@ namespace EpicGames.Perforce.Managed
 	///	</param>
 	/// <param name="Partitioned">Whether to allow using partitioned workspaces</param>
 	/// <param name="PreferNativeClient">Whether to prefer the native p4 client</param>
-	public record class ManagedWorkspaceOptions
+	public record ManagedWorkspaceOptions
 	(
 		int NumParallelSyncThreads = 4,
 		int MaxFileConcurrency = 4,
-		long MinScratchSpace = 50L * 1024 * 1024 * 1024,
+		long MinScratchSpace = 50L * 1024, // 50 GB
 		bool UseHaveTable = true,
 		bool Partitioned = false,
 		bool PreferNativeClient = false
@@ -1964,9 +1964,10 @@ namespace EpicGames.Perforce.Managed
 				if (RuntimeInformation.IsOSPlatform(OSPlatform.Windows))
 				{
 					long freeSpace = new DriveInfo(Path.GetPathRoot(_baseDir.FullName)!).AvailableFreeSpace;
-					if (freeSpace - syncSize < _options.MinScratchSpace)
+					long minScratchSpaceBytes = _options.MinScratchSpace * 1024 * 1024;
+					if (freeSpace - syncSize < minScratchSpaceBytes)
 					{
-						throw new InsufficientSpaceException($"Not enough space to sync new files (free space: {freeSpace / (1024.0 * 1024.0):n1}mb, sync size: {syncSize / (1024.0 * 1024.0):n1}mb, min scratch space: {_options.MinScratchSpace / (1024.0 * 1024.0):n1}mb)");
+						throw new InsufficientSpaceException($"Not enough space to sync new files (free space: {freeSpace / (1024.0 * 1024.0):n1}mb, sync size: {syncSize / (1024.0 * 1024.0):n1}mb, min scratch space: {minScratchSpaceBytes / (1024.0 * 1024.0):n1}mb)");
 					}
 				}
 
