Horde: Move KnownPropertyNames to EpicGames.Horde

Both agent and server are making use of these.

#rnx

diff --git a/Engine/Source/Programs/Horde/Horde.Server/Agents/Fleet/ComputeQueueAwsMetricStrategy.cs b/Engine/Source/Programs/Horde/Horde.Server/Agents/Fleet/ComputeQueueAwsMetricStrategy.cs
--- a/Engine/Source/Programs/Horde/Horde.Server/Agents/Fleet/ComputeQueueAwsMetricStrategy.cs
+++ b/Engine/Source/Programs/Horde/Horde.Server/Agents/Fleet/ComputeQueueAwsMetricStrategy.cs
@@ -9,6 +9,7 @@ using System.Threading;
 using System.Threading.Tasks;
 using Amazon.CloudWatch;
 using Amazon.CloudWatch.Model;
+using EpicGames.Horde.Agents;
 using EpicGames.Horde.Compute;
 using Horde.Server.Agents.Pools;
 using Horde.Server.Compute;
diff --git a/Engine/Source/Programs/Horde/Horde.Server/Agents/IAgent.cs b/Engine/Source/Programs/Horde/Horde.Server/Agents/IAgent.cs
--- a/Engine/Source/Programs/Horde/Horde.Server/Agents/IAgent.cs
+++ b/Engine/Source/Programs/Horde/Horde.Server/Agents/IAgent.cs
@@ -422,67 +422,6 @@ namespace Horde.Server.Agents
 		}
 	}
 
-	/// <summary>
-	/// Well-known property names for agents
-	/// </summary>
-	static class KnownPropertyNames
-	{
-		/// <summary>
-		/// The agent id
-		/// </summary>
-		public const string Id = "Id";
-
-		/// <summary>
-		/// The operating system (Linux, MacOS, Windows)
-		/// </summary>
-		public const string OsFamily = "OSFamily";
-
-		/// <summary>
-		/// Whether the agent is a .NET self-contained app
-		/// </summary>
-		public const string SelfContained = "SelfContained";
-
-		/// <summary>
-		/// Pools that this agent belongs to
-		/// </summary>
-		public const string Pool = "Pool";
-
-		/// <summary>
-		/// Pools requested by the agent to join when registering with server
-		/// </summary>
-		public const string RequestedPools = "RequestedPools";
-
-		/// <summary>
-		/// Number of logical cores
-		/// </summary>
-		public const string LogicalCores = "LogicalCores";
-
-		/// <summary>
-		/// Amount of RAM, in GB
-		/// </summary>
-		public const string Ram = "RAM";
-
-		/// <summary>
-		/// IP address used for sending compute task payloads
-		/// </summary>
-		public const string ComputeIp = "ComputeIp";
-		
-		/// <summary>
-		/// Port used for sending compute task payloads
-		/// </summary>
-		public const string ComputePort = "ComputePort";
-
-		/// <summary>
-		/// AWS: Instance ID
-		/// </summary>
-		public const string AwsInstanceId = "aws-instance-id";
-
-		/// <summary>
-		/// AWS: Instance type
-		/// </summary>
-		public const string AwsInstanceType = "aws-instance-type";
-	}
-
 	/// <summary>
 	/// Mirrors an Agent document in the database
 	/// </summary>
diff --git a/Engine/Source/Programs/Horde/Horde.Server/Compute/ComputeController.cs b/Engine/Source/Programs/Horde/Horde.Server/Compute/ComputeController.cs
--- a/Engine/Source/Programs/Horde/Horde.Server/Compute/ComputeController.cs
+++ b/Engine/Source/Programs/Horde/Horde.Server/Compute/ComputeController.cs
@@ -6,10 +6,10 @@ using System.Net;
 using System.Threading;
 using System.Threading.Tasks;
 using EpicGames.Core;
+using EpicGames.Horde.Agents;
 using EpicGames.Horde.Agents.Pools;
 using EpicGames.Horde.Compute;
 using Horde.Server.Acls;
-using Horde.Server.Agents;
 using Horde.Server.Server;
 using Horde.Server.Utilities;
 using Microsoft.AspNetCore.Authorization;
diff --git a/Engine/Source/Programs/Horde/Horde.Server/Compute/ComputeService.cs b/Engine/Source/Programs/Horde/Horde.Server/Compute/ComputeService.cs
--- a/Engine/Source/Programs/Horde/Horde.Server/Compute/ComputeService.cs
+++ b/Engine/Source/Programs/Horde/Horde.Server/Compute/ComputeService.cs
@@ -11,6 +11,7 @@ using System.Text.Json;
 using System.Threading;
 using System.Threading.Tasks;
 using EpicGames.Core;
+using EpicGames.Horde.Agents;
 using EpicGames.Horde.Agents.Leases;
 using EpicGames.Horde.Compute;
 using EpicGames.Horde.Compute.Clients;
diff --git a/Engine/Source/Programs/Shared/EpicGames.Horde/Agents/KnownPropertyNames.cs b/Engine/Source/Programs/Shared/EpicGames.Horde/Agents/KnownPropertyNames.cs
new file mode 100644
--- /dev/null
+++ b/Engine/Source/Programs/Shared/EpicGames.Horde/Agents/KnownPropertyNames.cs
@@ -0,0 +1,64 @@
+// Copyright Epic Games, Inc. All Rights Reserved.
+
+namespace EpicGames.Horde.Agents;
+
+/// <summary>
+/// Well-known property names for agents
+/// </summary>
+public static class KnownPropertyNames
+{
+	/// <summary>
+	/// The agent id
+	/// </summary>
+	public const string Id = "Id";
+
+	/// <summary>
+	/// The operating system (Linux, MacOS, Windows)
+	/// </summary>
+	public const string OsFamily = "OSFamily";
+
+	/// <summary>
+	/// Whether the agent is a .NET self-contained app
+	/// </summary>
+	public const string SelfContained = "SelfContained";
+
+	/// <summary>
+	/// Pools that this agent belongs to
+	/// </summary>
+	public const string Pool = "Pool";
+
+	/// <summary>
+	/// Pools requested by the agent to join when registering with server
+	/// </summary>
+	public const string RequestedPools = "RequestedPools";
+
+	/// <summary>
+	/// Number of logical cores
+	/// </summary>
+	public const string LogicalCores = "LogicalCores";
+
+	/// <summary>
+	/// Amount of RAM, in GB
+	/// </summary>
+	public const string Ram = "RAM";
+
+	/// <summary>
+	/// IP address used for sending compute task payloads
+	/// </summary>
+	public const string ComputeIp = "ComputeIp";
+		
+	/// <summary>
+	/// Port used for sending compute task payloads
+	/// </summary>
+	public const string ComputePort = "ComputePort";
+
+	/// <summary>
+	/// AWS: Instance ID
+	/// </summary>
+	public const string AwsInstanceId = "aws-instance-id";
+
+	/// <summary>
+	/// AWS: Instance type
+	/// </summary>
+	public const string AwsInstanceType = "aws-instance-type";
+}
\ No newline at end of file
