Horde: Fix issue tag tasks exiting immediately due to unintended call of wrong Task.Run overload.

#rnx

diff --git a/Engine/Source/Programs/Horde/Horde.Server/Issues/IssueTagService.cs b/Engine/Source/Programs/Horde/Horde.Server/Issues/IssueTagService.cs
--- a/Engine/Source/Programs/Horde/Horde.Server/Issues/IssueTagService.cs
+++ b/Engine/Source/Programs/Horde/Horde.Server/Issues/IssueTagService.cs
@@ -48,7 +48,7 @@ namespace Horde.Server.Issues
 			_commitService = commitService;
 			_issueCollection = issueCollection;
 			_globalConfig = globalConfig;
-			_ticker = clock.AddSharedTicker<IssueTagService>(TimeSpan.FromMinutes(1.0), TickAsync, logger);
+			_ticker = clock.AddSharedTicker<IssueTagService>(TimeSpan.FromSeconds(30.0), TickAsync, logger);
 			_logger = logger;
 		}
 
@@ -92,7 +92,7 @@ namespace Horde.Server.Issues
 			}
 		}
 
-		async ValueTask TickStreamGuardedAsync(StreamConfig streamConfig, State initialState, CancellationToken cancellationToken)
+		async Task TickStreamGuardedAsync(StreamConfig streamConfig, State initialState, CancellationToken cancellationToken)
 		{
 			try
 			{
