Horde: Use a standard message for server subsystems that are operating normally.

diff --git a/Engine/Source/Programs/Horde/Horde.Server/Perforce/PerforceLoadBalancer.cs b/Engine/Source/Programs/Horde/Horde.Server/Perforce/PerforceLoadBalancer.cs
--- a/Engine/Source/Programs/Horde/Horde.Server/Perforce/PerforceLoadBalancer.cs
+++ b/Engine/Source/Programs/Horde/Horde.Server/Perforce/PerforceLoadBalancer.cs
@@ -469,13 +469,13 @@ namespace Horde.Server.Perforce
 
 			// Update Horde's health monitor for subsystems
 			list = await _serverListSingleton.GetAsync(cancellationToken);
-			(HealthStatus health, string message) = GetPerforceHealth(list.Servers);
+			(HealthStatus health, string? message) = GetPerforceHealth(list.Servers);
 			span.SetAttribute("health.status", health.ToString());
-			span.SetAttribute("health.message", message);
+			span.SetAttribute("health.message", message ?? String.Empty);
 			await _health.UpdateAsync(health, message);
 		}
 
-		static (HealthStatus health, string message) GetPerforceHealth(List<PerforceServerEntry> servers)
+		static (HealthStatus health, string? message) GetPerforceHealth(List<PerforceServerEntry> servers)
 		{
 			HealthStatus result = HealthStatus.Healthy;
 			foreach (PerforceServerEntry server in servers)
@@ -491,11 +491,11 @@ namespace Horde.Server.Perforce
 				result = serverHealth < result ? serverHealth : result;
 			}
 
-			string message = result switch
+			string? message = result switch
 			{
 				HealthStatus.Unhealthy => "One or more Perforce servers are unhealthy. Check Perforce servers page for details.",
 				HealthStatus.Degraded => "One or more Perforce servers are degraded. Check Perforce servers page for details.",
-				HealthStatus.Healthy => "All Perforce servers are healthy",
+				HealthStatus.Healthy => null, 
 				_ => throw new ArgumentOutOfRangeException($"Unknown health status: {result}")
 			};
 
diff --git a/Engine/Source/Programs/Horde/Horde.Server/Server/ServerStatusController.cs b/Engine/Source/Programs/Horde/Horde.Server/Server/ServerStatusController.cs
--- a/Engine/Source/Programs/Horde/Horde.Server/Server/ServerStatusController.cs
+++ b/Engine/Source/Programs/Horde/Horde.Server/Server/ServerStatusController.cs
@@ -100,7 +100,7 @@ public class ServerStatusController : Controller
 						u => new ServerStatusUpdate()
 						{
 							Result = ConvertSubsystemResult(u.Result),
-							Message = u.Message,
+							Message = u.Message ?? "Operating normally.",
 							UpdatedAt = u.UpdatedAt
 						}).ToArray()
 				};
