Horde: Add auto conform for agents with workspaces below a certain free disk space threshold

diff --git a/Engine/Source/Programs/Horde/Horde.Agent/Services/CapabilitiesService.cs b/Engine/Source/Programs/Horde/Horde.Agent/Services/CapabilitiesService.cs
--- a/Engine/Source/Programs/Horde/Horde.Agent/Services/CapabilitiesService.cs
+++ b/Engine/Source/Programs/Horde/Horde.Agent/Services/CapabilitiesService.cs
@@ -10,6 +10,7 @@ using Amazon.EC2;
 using Amazon.EC2.Model;
 using Amazon.Util;
 using EpicGames.Core;
+using EpicGames.Horde.Agents;
 using EpicGames.Horde.Compute;
 using Horde.Agent.Execution;
 using HordeCommon.Rpc.Messages;
@@ -70,7 +71,7 @@ namespace Horde.Agent.Services
 
 			// Create the primary device
 			RpcDeviceCapabilities primaryDevice = new RpcDeviceCapabilities();
-			primaryDevice.Handle = "Primary";
+			primaryDevice.Handle = KnownPropertyNames.PrimaryDeviceHandle;
 
 			List<RpcDeviceCapabilities> otherDevices = new List<RpcDeviceCapabilities>();
 			if (RuntimeInformation.IsOSPlatform(OSPlatform.Windows))
@@ -407,8 +408,8 @@ namespace Horde.Agent.Services
 					try
 					{
 						DriveInfo info = new DriveInfo(driveName);
-						primaryDevice.Properties.Add($"DiskFreeSpace={info.AvailableFreeSpace}");
-						primaryDevice.Properties.Add($"DiskTotalSize={info.TotalSize}");
+						primaryDevice.Properties.Add($"{KnownPropertyNames.DiskFreeSpace}={info.AvailableFreeSpace}");
+						primaryDevice.Properties.Add($"{KnownPropertyNames.DiskTotalSize}={info.TotalSize}");
 					}
 					catch (Exception ex)
 					{
diff --git a/Engine/Source/Programs/Horde/Horde.Server.Tests/Agents/Pools/PoolUpdateServiceTest.cs b/Engine/Source/Programs/Horde/Horde.Server.Tests/Agents/Pools/PoolUpdateServiceTest.cs
--- a/Engine/Source/Programs/Horde/Horde.Server.Tests/Agents/Pools/PoolUpdateServiceTest.cs
+++ b/Engine/Source/Programs/Horde/Horde.Server.Tests/Agents/Pools/PoolUpdateServiceTest.cs
@@ -1,98 +1,128 @@
 // Copyright Epic Games, Inc. All Rights Reserved.
 
 using System;
+using System.Collections.Generic;
+using System.Linq;
 using System.Threading;
 using System.Threading.Tasks;
+using EpicGames.Horde.Agents;
 using Horde.Server.Agents;
 using Horde.Server.Agents.Pools;
 using Microsoft.Extensions.Logging.Abstractions;
 using Microsoft.VisualStudio.TestTools.UnitTesting;
 
-namespace Horde.Server.Tests.Agents.Pools
+namespace Horde.Server.Tests.Agents.Pools;
+
+[TestClass]
+public class PoolUpdateServiceTest : TestSetup
 {
-	[TestClass]
-	public class PoolUpdateServiceTest : TestSetup
+	private readonly PoolUpdateService _pus;
+	private IPool _pool = default!;
+	private IAgent _enabledAgent = default!;
+	private IAgent _disabledAgent = default!;
+	private IAgent _disabledAgentBeyondGracePeriod = default!;
+
+	public PoolUpdateServiceTest()
+	{
+		UpdateConfig(x => x.Pools.Clear());
+		_pus = new(AgentCollection, PoolCollection, Clock, GlobalConfig, Tracer, new NullLogger<PoolUpdateService>());
+	}
+
+	[TestInitialize]
+	public async Task SetupAsync()
+	{
+		_pool = await CreatePoolAsync(new PoolConfig { Name = "testPool", EnableAutoscaling = true });
+		_enabledAgent = await CreateAgentAsync(_pool, true);
+		_disabledAgent = await CreateAgentAsync(_pool, false);
+		_disabledAgentBeyondGracePeriod = await CreateAgentAsync(_pool, enabled: false, adjustClockBy: -TimeSpan.FromHours(9));
+	}
+
+	private async Task RefreshAgentsAsync()
+	{
+		_enabledAgent = (await AgentService.GetAgentAsync(_enabledAgent.Id))!;
+		_disabledAgent = (await AgentService.GetAgentAsync(_disabledAgent.Id))!;
+		_disabledAgentBeyondGracePeriod = (await AgentService.GetAgentAsync(_disabledAgentBeyondGracePeriod.Id))!;
+	}
+
+	public override async ValueTask DisposeAsync()
+	{
+		GC.SuppressFinalize(this);
+		await base.DisposeAsync();
+
+		await _pus.DisposeAsync();
+	}
+
+	[TestMethod]
+	public async Task ShutdownDisabledAgents_WithGlobalGracePeriod_RequestsShutdownAsync()
+	{
+		// Act
+		await _pus.ShutdownDisabledAgentsAsync(CancellationToken.None);
+		await RefreshAgentsAsync();
+
+		// Assert
+		Assert.IsFalse(_enabledAgent.RequestShutdown);
+		Assert.IsFalse(_disabledAgent.RequestShutdown);
+		Assert.IsTrue(_disabledAgentBeyondGracePeriod.RequestShutdown);
+	}
+
+	[TestMethod]
+	public async Task ShutdownDisabledAgents_WithPerPoolGracePeriod_DoesNotRequestShutdownAsync()
+	{
+		// Arrange
+		// Explicitly set the grace period for the pool to be longer than the default of 8 hours
+		UpdateConfig(config => config.Pools[0].ShutdownIfDisabledGracePeriod = TimeSpan.FromHours(24.0));
+
+		// Act
+		await _pus.ShutdownDisabledAgentsAsync(CancellationToken.None);
+		await RefreshAgentsAsync();
+
+		// Assert
+		Assert.IsFalse(_enabledAgent.RequestShutdown);
+		Assert.IsFalse(_disabledAgent.RequestShutdown);
+		Assert.IsFalse(_disabledAgentBeyondGracePeriod.RequestShutdown);
+	}
+
+	[TestMethod]
+	public async Task ShutdownDisabledAgents_WithAutoScalingOff_DoesNotRequestShutdownAsync()
+	{
+		// Arrange
+		UpdateConfig(config => config.Pools[0].EnableAutoscaling = false);
+
+		// Act
+		await _pus.ShutdownDisabledAgentsAsync(CancellationToken.None);
+		await RefreshAgentsAsync();
+
+		// Assert
+		Assert.IsFalse(_enabledAgent.RequestShutdown);
+		Assert.IsFalse(_disabledAgent.RequestShutdown);
+		Assert.IsFalse(_disabledAgentBeyondGracePeriod.RequestShutdown);
+	}
+	
+	[TestMethod]
+	[DataRow(false, 50)]
+	[DataRow(false, null)]
+	[DataRow(false, 0)]
+	[DataRow(false, 50, null, 0, 2)]
+	[DataRow(true, 200)]
+	[DataRow(true, 200, 50)]
+	[DataRow(true, 200, null)]
+	[DataRow(true, null, 0, 50, 300)]
+	public async Task AutoConformAgentsAsync(bool conformRequested, params int?[] autoConformThresholdsM)
+	{
+		// Arrange
+		IAgent agent = await CreateAutoConformAgentAsync(100, autoConformThresholdsM);
+		
+		// Act
+		await _pus.AutoConformAgentsAsync(CancellationToken.None);
+		agent = (await AgentService.GetAgentAsync(agent.Id))!;
+		
+		// Assert
+		Assert.AreEqual(conformRequested, agent.RequestConform);
+	}
+
+	private async Task<IAgent> CreateAutoConformAgentAsync(int freeDiskSpaceMb, params int?[] autoConformThresholdsMb)
 	{
-		private readonly PoolUpdateService _pus;
-		private IPool _pool = default!;
-		private IAgent _enabledAgent = default!;
-		private IAgent _disabledAgent = default!;
-		private IAgent _disabledAgentBeyondGracePeriod = default!;
-
-		public PoolUpdateServiceTest()
-		{
-			UpdateConfig(x => x.Pools.Clear());
-			_pus = new(AgentCollection, PoolCollection, Clock, GlobalConfig, Tracer, new NullLogger<PoolUpdateService>());
-		}
-
-		[TestInitialize]
-		public async Task SetupAsync()
-		{
-			_pool = await CreatePoolAsync(new PoolConfig { Name = "testPool", EnableAutoscaling = true });
-			_enabledAgent = await CreateAgentAsync(_pool, true);
-			_disabledAgent = await CreateAgentAsync(_pool, false);
-			_disabledAgentBeyondGracePeriod = await CreateAgentAsync(_pool, enabled: false, adjustClockBy: -TimeSpan.FromHours(9));
-		}
-
-		private async Task RefreshAgentsAsync()
-		{
-			_enabledAgent = (await AgentService.GetAgentAsync(_enabledAgent.Id))!;
-			_disabledAgent = (await AgentService.GetAgentAsync(_disabledAgent.Id))!;
-			_disabledAgentBeyondGracePeriod = (await AgentService.GetAgentAsync(_disabledAgentBeyondGracePeriod.Id))!;
-		}
-
-		public override async ValueTask DisposeAsync()
-		{
-			GC.SuppressFinalize(this);
-			await base.DisposeAsync();
-
-			await _pus.DisposeAsync();
-		}
-
-		[TestMethod]
-		public async Task ShutdownDisabledAgents_WithGlobalGracePeriod_RequestsShutdownAsync()
-		{
-			// Act
-			await _pus.ShutdownDisabledAgentsAsync(CancellationToken.None);
-			await RefreshAgentsAsync();
-
-			// Assert
-			Assert.IsFalse(_enabledAgent.RequestShutdown);
-			Assert.IsFalse(_disabledAgent.RequestShutdown);
-			Assert.IsTrue(_disabledAgentBeyondGracePeriod.RequestShutdown);
-		}
-
-		[TestMethod]
-		public async Task ShutdownDisabledAgents_WithPerPoolGracePeriod_DoesNotRequestShutdownAsync()
-		{
-			// Arrange
-			// Explicitly set the grace period for the pool to be longer than the default of 8 hours
-			UpdateConfig(config => config.Pools[0].ShutdownIfDisabledGracePeriod = TimeSpan.FromHours(24.0));
-
-			// Act
-			await _pus.ShutdownDisabledAgentsAsync(CancellationToken.None);
-			await RefreshAgentsAsync();
-
-			// Assert
-			Assert.IsFalse(_enabledAgent.RequestShutdown);
-			Assert.IsFalse(_disabledAgent.RequestShutdown);
-			Assert.IsFalse(_disabledAgentBeyondGracePeriod.RequestShutdown);
-		}
-
-		[TestMethod]
-		public async Task ShutdownDisabledAgents_WithAutoScalingOff_DoesNotRequestShutdownAsync()
-		{
-			// Arrange
-			UpdateConfig(config => config.Pools[0].EnableAutoscaling = false);
-
-			// Act
-			await _pus.ShutdownDisabledAgentsAsync(CancellationToken.None);
-			await RefreshAgentsAsync();
-
-			// Assert
-			Assert.IsFalse(_enabledAgent.RequestShutdown);
-			Assert.IsFalse(_disabledAgent.RequestShutdown);
-			Assert.IsFalse(_disabledAgentBeyondGracePeriod.RequestShutdown);
-		}
+		IEnumerable<AgentWorkspaceInfo> workspaces = autoConformThresholdsMb.Select(x => new AgentWorkspaceInfo(null, null, "someWorkspace", "//Some/Stream", null, false, null, null, x));
+		return await CreateAgentAsync(_pool, properties: [$"{KnownPropertyNames.DiskFreeSpace}={freeDiskSpaceMb * 1024 * 1024}"], workspaces: workspaces.ToList());
 	}
 }
diff --git a/Engine/Source/Programs/Horde/Horde.Server.Tests/TestSetup.cs b/Engine/Source/Programs/Horde/Horde.Server.Tests/TestSetup.cs
--- a/Engine/Source/Programs/Horde/Horde.Server.Tests/TestSetup.cs
+++ b/Engine/Source/Programs/Horde/Horde.Server.Tests/TestSetup.cs
@@ -369,9 +369,9 @@ namespace Horde.Server.Tests
 		}
 
 		private static int s_agentIdCounter = 1;
-		public Task<IAgent> CreateAgentAsync(IPool pool, bool enabled = true, bool requestShutdown = false, List<string>? properties = null, TimeSpan? adjustClockBy = null)
+		public Task<IAgent> CreateAgentAsync(IPool pool, bool enabled = true, bool requestShutdown = false, List<string>? properties = null, List<AgentWorkspaceInfo>? workspaces = null, TimeSpan? adjustClockBy = null)
 		{
-			return CreateAgentAsync(pool.Id, enabled, requestShutdown, properties, adjustClockBy);
+			return CreateAgentAsync(pool.Id, enabled, requestShutdown, properties, workspaces, adjustClockBy);
 		}
 
 		/// <summary>
@@ -381,14 +381,22 @@ namespace Horde.Server.Tests
 		/// <param name="enabled">Whether set the agent as enabled</param>
 		/// <param name="requestShutdown">Mark it with a request for shutdown</param>
 		/// <param name="properties">Any properties to assign</param>
+		/// <param name="workspaces">Any workspaces to assign</param>
 		/// <param name="adjustClockBy">Time span to temporarily skew the clock when creating the agent</param>
 		/// <param name="awsInstanceId">AWS instance ID for the agent (will be set in properties)</param>
 		/// <param name="lease">A lease to assign the agent</param>
 		/// <param name="ephemeral">Whether the agent is ephemeral</param>
 		/// <returns>A new agent</returns>
 		public async Task<IAgent> CreateAgentAsync(
-			PoolId poolId, bool enabled = true, bool requestShutdown = false, List<string>? properties = null,
-			TimeSpan? adjustClockBy = null, string? awsInstanceId = null, AgentLease? lease = null, bool ephemeral = false)
+			PoolId poolId,
+			bool enabled = true,
+			bool requestShutdown = false,
+			List<string>? properties = null,
+			List<AgentWorkspaceInfo>? workspaces = null,
+			TimeSpan? adjustClockBy = null,
+			string? awsInstanceId = null,
+			AgentLease? lease = null,
+			bool ephemeral = false)
 		{
 			DateTime now = Clock.UtcNow;
 			if (adjustClockBy != null)
@@ -412,6 +420,11 @@ namespace Horde.Server.Tests
 			agent = await AgentService.CreateSessionAsync(agent, AgentStatus.Ok, tempProps, resources, null);
 			Assert.IsNotNull(agent);
 
+			if (workspaces is { Count: > 0 })
+			{
+				await AgentCollection.TryUpdateWorkspacesAsync(agent, workspaces, false);
+			}
+			
 			if (requestShutdown)
 			{
 				await AgentCollection.TryUpdateSettingsAsync(agent, requestShutdown: true);
diff --git a/Engine/Source/Programs/Horde/Horde.Server/Agents/IAgent.cs b/Engine/Source/Programs/Horde/Horde.Server/Agents/IAgent.cs
--- a/Engine/Source/Programs/Horde/Horde.Server/Agents/IAgent.cs
+++ b/Engine/Source/Programs/Horde/Horde.Server/Agents/IAgent.cs
@@ -75,6 +75,11 @@ namespace Horde.Server.Agents
 		/// </summary>
 		public long? MinScratchSpace { get; set; }
 
+		/// <summary>
+		/// Threshold for when to trigger an automatic conform of agent. Measured in megabytes free on disk.
+		/// </summary>
+		public long? ConformDiskFreeSpace { get; set; }
+
 		/// <summary>
 		/// Constructor
 		/// </summary>
@@ -86,7 +91,8 @@ namespace Horde.Server.Agents
 		/// <param name="incremental">Whether to use an incremental workspace</param>
 		/// <param name="method">Method to use when syncing/materializing data from Perforce</param>
 		/// <param name="minScratchSpace">Minimum disk space that must be available *after* syncing this workspace (in megabytes)</param>
-		public AgentWorkspaceInfo(string? cluster, string? userName, string identifier, string stream, List<string>? view, bool incremental, string? method, long? minScratchSpace)
+		/// <param name="conformDiskFreeSpace">Threshold for when to trigger an automatic conform of agent. Measured in megabytes free on disk</param>
+		public AgentWorkspaceInfo(string? cluster, string? userName, string identifier, string stream, List<string>? view, bool incremental, string? method, long? minScratchSpace = null, long? conformDiskFreeSpace = null)
 		{
 			if (!String.IsNullOrEmpty(cluster))
 			{
@@ -102,6 +108,7 @@ namespace Horde.Server.Agents
 			Incremental = incremental;
 			Method = method;
 			MinScratchSpace = minScratchSpace;
+			ConformDiskFreeSpace = conformDiskFreeSpace;
 		}
 
 		/// <summary>
@@ -109,10 +116,11 @@ namespace Horde.Server.Agents
 		/// </summary>
 		/// <param name="workspace">RPC message to construct from</param>
 		public AgentWorkspaceInfo(RpcAgentWorkspace workspace)
-			: this(workspace.ConfiguredCluster, workspace.ConfiguredUserName, workspace.Identifier, workspace.Stream, (workspace.View.Count > 0) ? workspace.View.ToList() : null, workspace.Incremental, workspace.Method, workspace.MinScratchSpace)
+			: this(workspace.ConfiguredCluster, workspace.ConfiguredUserName, workspace.Identifier, workspace.Stream, (workspace.View.Count > 0) ? workspace.View.ToList() : null, workspace.Incremental, workspace.Method, workspace.MinScratchSpace, workspace.ConformDiskFreeSpace)
 		{
 			// Treat zero value as null as Protobuf cannot store null 
 			MinScratchSpace = workspace.MinScratchSpace == 0 ? null : workspace.MinScratchSpace;
+			ConformDiskFreeSpace = workspace.ConformDiskFreeSpace == 0 ? null : workspace.ConformDiskFreeSpace;
 		}
 
 		/// <summary>
@@ -685,6 +693,17 @@ namespace Horde.Server.Agents
 			List<string> values = agent.GetPropertyValues(KnownPropertyNames.SelfContained).ToList();
 			return values.Count > 0 && values[0].Equals("true", StringComparison.OrdinalIgnoreCase);
 		}
+		
+		/// <summary>
+		/// Get free disk space on agent (as reported through primary device capabilities)
+		/// </summary>
+		/// <param name="agent">Agent to query</param>
+		/// <returns>Amount of free disk space in bytes</returns>
+		public static long? GetDiskFreeSpace(this IAgent agent)
+		{
+			List<string> values = agent.GetPropertyValues(KnownPropertyNames.DiskFreeSpace).ToList();
+			return values.Count > 0 && Int64.TryParse(values[0], out long amount) ? amount : null;
+		}
 
 		/// <summary>
 		/// Get operating system family of agent
diff --git a/Engine/Source/Programs/Horde/Horde.Server/Agents/Pools/PoolUpdateService.cs b/Engine/Source/Programs/Horde/Horde.Server/Agents/Pools/PoolUpdateService.cs
--- a/Engine/Source/Programs/Horde/Horde.Server/Agents/Pools/PoolUpdateService.cs
+++ b/Engine/Source/Programs/Horde/Horde.Server/Agents/Pools/PoolUpdateService.cs
@@ -6,6 +6,7 @@ using System.Linq;
 using System.Threading;
 using System.Threading.Tasks;
 using EpicGames.Core;
+using EpicGames.Horde.Agents;
 using EpicGames.Horde.Agents.Pools;
 using Horde.Server.Server;
 using Horde.Server.Streams;
@@ -30,6 +31,7 @@ namespace Horde.Server.Agents.Pools
 		readonly Tracer _tracer;
 		readonly ITicker _updatePoolsTicker;
 		readonly ITicker _shutdownDisabledAgentsTicker;
+		readonly ITicker _autoConformAgentsTicker;
 
 		/// <summary>
 		/// Constructor
@@ -44,6 +46,7 @@ namespace Horde.Server.Agents.Pools
 			_logger = logger;
 			_updatePoolsTicker = clock.AddTicker($"{nameof(PoolUpdateService)}.{nameof(UpdatePoolsAsync)}", TimeSpan.FromSeconds(30.0), UpdatePoolsAsync, logger);
 			_shutdownDisabledAgentsTicker = clock.AddSharedTicker($"{nameof(PoolUpdateService)}.{nameof(ShutdownDisabledAgentsAsync)}", TimeSpan.FromHours(1), ShutdownDisabledAgentsAsync, logger);
+			_autoConformAgentsTicker = clock.AddSharedTicker($"{nameof(PoolUpdateService)}.{nameof(AutoConformAgentsAsync)}", TimeSpan.FromHours(1), AutoConformAgentsAsync, logger);
 		}
 
 		/// <inheritdoc/>
@@ -51,6 +54,7 @@ namespace Horde.Server.Agents.Pools
 		{
 			await _updatePoolsTicker.StartAsync();
 			await _shutdownDisabledAgentsTicker.StartAsync();
+			await _autoConformAgentsTicker.StartAsync();
 		}
 
 		/// <inheritdoc/>
@@ -58,6 +62,7 @@ namespace Horde.Server.Agents.Pools
 		{
 			await _updatePoolsTicker.StopAsync();
 			await _shutdownDisabledAgentsTicker.StopAsync();
+			await _autoConformAgentsTicker.StopAsync();
 		}
 
 		/// <inheritdoc/>
@@ -65,6 +70,7 @@ namespace Horde.Server.Agents.Pools
 		{
 			await _updatePoolsTicker.DisposeAsync();
 			await _shutdownDisabledAgentsTicker.DisposeAsync();
+			await _autoConformAgentsTicker.DisposeAsync();
 		}
 
 		/// <summary>
@@ -118,6 +124,40 @@ namespace Horde.Server.Agents.Pools
 			IEnumerable<PoolId> poolIds = agent.ExplicitPools.Concat(agent.DynamicPools);
 			return pools.Any(x => poolIds.Contains(x.Id) && x.EnableAutoscaling);
 		}
+		
+		/// <summary>
+		/// Find and conform any agents below the disk free conform threshold for a workspace
+		/// </summary>
+		/// <param name="cancellationToken">Cancellation token for the async task</param>
+		/// <returns>Async task</returns>
+		internal async ValueTask AutoConformAgentsAsync(CancellationToken cancellationToken)
+		{
+			using TelemetrySpan span = _tracer.StartActiveSpan($"{nameof(PoolUpdateService)}.{nameof(AutoConformAgentsAsync)}");
+			List<IAgent> agents = (await _agents.FindAsync(status: AgentStatus.Ok, enabled: true, cancellationToken: cancellationToken)).Where(x => !x.RequestShutdown).ToList();
+
+			long MegabytesToBytes (long v) => v * 1024 * 1024;
+			foreach (IAgent agent in agents)
+			{
+				long? freeDiskSpace = agent.GetDiskFreeSpace();
+				long maxConformDiskSpace = 0;
+
+				// Find the largest conform disk space amount needed, if any
+				foreach (AgentWorkspaceInfo workspace in agent.Workspaces)
+				{
+					if (workspace.ConformDiskFreeSpace is > 0)
+					{
+						maxConformDiskSpace = Math.Max(maxConformDiskSpace, MegabytesToBytes(workspace.ConformDiskFreeSpace.Value));
+					}
+				}
+
+				if (freeDiskSpace != null && maxConformDiskSpace > 0 && freeDiskSpace < maxConformDiskSpace && agent.ConformAttemptCount is null or 0)
+				{
+					await _agents.TryUpdateSettingsAsync(agent, requestConform: true, cancellationToken: cancellationToken);
+					_logger.LogInformation("Auto-conforming {AgentId} as workspace conform disk space needed ({ConformDiskSpace:F1} MB) is less than free disk space ({FreeDiskSpace:F1} MB)",
+						agent.Id.ToString(), maxConformDiskSpace / 1024.0 / 1024.0, freeDiskSpace.Value / 1024.0 / 1024.0);
+				}
+			}
+		}
 
 		/// <summary>
 		/// Execute the background task
diff --git a/Engine/Source/Programs/Horde/Horde.Server/Streams/StreamConfig.cs b/Engine/Source/Programs/Horde/Horde.Server/Streams/StreamConfig.cs
--- a/Engine/Source/Programs/Horde/Horde.Server/Streams/StreamConfig.cs
+++ b/Engine/Source/Programs/Horde/Horde.Server/Streams/StreamConfig.cs
@@ -715,6 +715,12 @@ namespace Horde.Server.Streams
 		/// If not available, the job will be aborted.
 		/// </summary>
 		public long? MinScratchSpace { get; set; } = null;
+
+		/// <summary>
+		/// Threshold for when to trigger an automatic conform of agent. Measured in megabytes free on disk.
+		/// Set to null or 0 to disable.
+		/// </summary>
+		public long? ConformDiskFreeSpace { get; set; } = null;
 	}
 
 	/// <summary>
diff --git a/Engine/Source/Programs/Shared/EpicGames.Horde/Agents/KnownPropertyNames.cs b/Engine/Source/Programs/Shared/EpicGames.Horde/Agents/KnownPropertyNames.cs
--- a/Engine/Source/Programs/Shared/EpicGames.Horde/Agents/KnownPropertyNames.cs
+++ b/Engine/Source/Programs/Shared/EpicGames.Horde/Agents/KnownPropertyNames.cs
@@ -11,6 +11,11 @@ public static class KnownPropertyNames
 	/// The agent id
 	/// </summary>
 	public const string Id = "Id";
+	
+	/// <summary>
+	/// Handle of the primary device (usually the agent VM itself)
+	/// </summary>
+	public const string PrimaryDeviceHandle = "Primary";
 
 	/// <summary>
 	/// The operating system (Linux, MacOS, Windows)
@@ -41,6 +46,16 @@ public static class KnownPropertyNames
 	/// Amount of RAM, in GB
 	/// </summary>
 	public const string Ram = "RAM";
+	
+	/// <summary>
+	/// The total size of storage space on drive, in bytes
+	/// </summary>
+	public const string DiskTotalSize = "DiskTotalSize";
+
+	/// <summary>
+	/// Amount of available free space on drive, in bytes
+	/// </summary>
+	public const string DiskFreeSpace = "DiskFreeSpace";
 
 	/// <summary>
 	/// IP address used for sending compute task payloads
