Horde: Add simple test for workspace inheritance in config

#rnx

diff --git a/Engine/Source/Programs/Horde/Horde.Server.Tests/Configuration/ConfigTests.cs b/Engine/Source/Programs/Horde/Horde.Server.Tests/Configuration/ConfigTests.cs
--- a/Engine/Source/Programs/Horde/Horde.Server.Tests/Configuration/ConfigTests.cs
+++ b/Engine/Source/Programs/Horde/Horde.Server.Tests/Configuration/ConfigTests.cs
@@ -9,7 +9,9 @@ using System.Threading;
 using System.Threading.Tasks;
 using EpicGames.Core;
 using Horde.Server.Configuration;
+using Horde.Server.Projects;
 using Horde.Server.Server;
+using Horde.Server.Streams;
 using Microsoft.Extensions.Logging.Abstractions;
 using Microsoft.VisualStudio.TestTools.UnitTesting;
 
@@ -166,6 +168,25 @@ namespace Horde.Server.Tests.Configuration
 			Assert.AreEqual("global", GetNetworkConfig(gc, "15.3.4.5")!.Id);
 		}
 
+		[TestMethod]
+		public void WorkspaceConfig()
+		{
+			GlobalConfig gc = new();
+			Dictionary<string, WorkspaceConfig> inputWorkspaces = new();
+			inputWorkspaces["base"] = new WorkspaceConfig { Identifier = "base", Cluster = "myCluster", MinScratchSpace = 111 };
+			inputWorkspaces["subType"] = new WorkspaceConfig { Base = "base", Identifier = "subType", ConformDiskFreeSpace = 222 };
+			inputWorkspaces["subSubType"] = new WorkspaceConfig { Base = "subType", Identifier = "subSubType" };
+			
+			gc.Projects.Add(new ProjectConfig { Streams = [new StreamConfig { WorkspaceTypes = inputWorkspaces }]});
+			gc.PostLoad(new ServerSettings());
+
+			Dictionary<string,WorkspaceConfig> workspaces = gc.Projects[0].Streams[0].WorkspaceTypes;
+			Assert.AreEqual(3, workspaces.Count);
+			Assert.AreEqual(111, workspaces["subType"].MinScratchSpace);
+			Assert.AreEqual(111, workspaces["subSubType"].MinScratchSpace);
+			Assert.AreEqual(222, workspaces["subSubType"].ConformDiskFreeSpace);
+		}
+
 		class ObjectValue
 		{
 			public string Value { get; set; } = "";
