Horde: Add a service to expire jobs after a period of time.

Jobs are kept forever by default, though this may be modified through the ExpireAfterDays parameter in the JobOptions object in streams or projects.

diff --git a/Engine/Source/Programs/Horde/Horde.Server.Tests/Jobs/JobExpirationTests.cs b/Engine/Source/Programs/Horde/Horde.Server.Tests/Jobs/JobExpirationTests.cs
new file mode 100644
--- /dev/null
+++ b/Engine/Source/Programs/Horde/Horde.Server.Tests/Jobs/JobExpirationTests.cs
@@ -0,0 +1,74 @@
+// Copyright Epic Games, Inc. All Rights Reserved.
+
+using System;
+using System.Collections.Generic;
+using System.Linq;
+using System.Text;
+using System.Threading.Tasks;
+using EpicGames.Horde.Jobs.Templates;
+using EpicGames.Horde.Projects;
+using EpicGames.Horde.Streams;
+using Horde.Server.Jobs;
+using Horde.Server.Jobs.Graphs;
+using Horde.Server.Jobs.Templates;
+using Horde.Server.Projects;
+using Horde.Server.Server;
+using Horde.Server.Streams;
+using Microsoft.Extensions.DependencyInjection;
+using Microsoft.VisualStudio.TestTools.UnitTesting;
+
+namespace Horde.Server.Tests.Jobs
+{
+	[TestClass]
+	public class JobExpirationTests : TestSetup
+	{
+		async Task<IJob> CreateJobAsync()
+		{
+			ProjectId projectId = new ProjectId("ue5");
+			StreamId streamId = new StreamId("ue5-main");
+			TemplateId templateId = new TemplateId("template1");
+
+			StreamConfig streamConfig = new StreamConfig { Id = streamId };
+			streamConfig.JobOptions.ExpireAfterDays = 2;
+			streamConfig.Templates.Add(new TemplateRefConfig { Id = templateId, Name = "Test Template" });
+
+			ProjectConfig projectConfig = new ProjectConfig { Id = projectId };
+			projectConfig.Streams.Add(streamConfig);
+
+			GlobalConfig globalConfig = new GlobalConfig();
+			globalConfig.Projects.Add(projectConfig);
+
+			SetConfig(globalConfig);
+
+			CreateJobOptions options = new CreateJobOptions();
+			options.PreflightChange = 999;
+
+			ITemplate template = await TemplateCollection.GetOrAddAsync(streamConfig.Templates[0]);
+
+			IGraph graph = await GraphCollection.AddAsync(template, null);
+
+			return await JobService.CreateJobAsync(null, streamConfig, templateId, template.Hash, graph, "Hello", 1234, 1233, options);
+		}
+
+		[TestMethod]
+		public async Task TestJobExpiryAsync()
+		{
+			await ServiceProvider.GetRequiredService<JobExpirationService>().StartAsync(default);
+
+			IJob job = await CreateJobAsync();
+
+			IJob? newJob = await JobCollection.GetAsync(job.Id);
+			Assert.IsNotNull(newJob);
+
+			await Clock.AdvanceAsync(TimeSpan.FromDays(1.0));
+
+			newJob = await JobCollection.GetAsync(job.Id);
+			Assert.IsNotNull(newJob);
+
+			await Clock.AdvanceAsync(TimeSpan.FromDays(2.0));
+
+			newJob = await JobCollection.GetAsync(job.Id);
+			Assert.IsNull(newJob);
+		}
+	}
+}
diff --git a/Engine/Source/Programs/Horde/Horde.Server.Tests/Jobs/JobTaskSourceTests.cs b/Engine/Source/Programs/Horde/Horde.Server.Tests/Jobs/JobTaskSourceTests.cs
--- a/Engine/Source/Programs/Horde/Horde.Server.Tests/Jobs/JobTaskSourceTests.cs
+++ b/Engine/Source/Programs/Horde/Horde.Server.Tests/Jobs/JobTaskSourceTests.cs
@@ -122,7 +122,7 @@ namespace Horde.Server.Tests.Jobs
 			IReadOnlyList<IJob> jobs = await JobCollection.FindAsync();
 			for (int i = 0; i < jobs.Count; i++)
 			{
-				await JobCollection.RemoveAsync(jobs[i]);
+				await jobs[i].TryDeleteAsync();
 			}
 
 			// create a new job
diff --git a/Engine/Source/Programs/Horde/Horde.Server.Tests/TestSetup.cs b/Engine/Source/Programs/Horde/Horde.Server.Tests/TestSetup.cs
--- a/Engine/Source/Programs/Horde/Horde.Server.Tests/TestSetup.cs
+++ b/Engine/Source/Programs/Horde/Horde.Server.Tests/TestSetup.cs
@@ -264,6 +264,7 @@ namespace Horde.Server.Tests
 			services.AddSingleton<IDowntimeService, DowntimeServiceStub>();
 			services.AddSingleton<IssueService>();
 			services.AddSingleton<JobService>();
+			services.AddSingleton<JobExpirationService>();
 			services.AddSingleton<LifetimeService>();
 			services.AddSingleton<ILogStorage, NullLogStorage>();
 			services.AddSingleton<ILogFileService, LogFileService>();
diff --git a/Engine/Source/Programs/Horde/Horde.Server/Jobs/IJob.cs b/Engine/Source/Programs/Horde/Horde.Server/Jobs/IJob.cs
--- a/Engine/Source/Programs/Horde/Horde.Server/Jobs/IJob.cs
+++ b/Engine/Source/Programs/Horde/Horde.Server/Jobs/IJob.cs
@@ -761,6 +761,13 @@ namespace Horde.Server.Jobs
 		/// <param name="cancellationToken">Cancellation token for the operation</param>
 		Task<IJob?> RefreshAsync(CancellationToken cancellationToken = default);
 
+		/// <summary>
+		/// Attempt to delete the job
+		/// </summary>
+		/// <param name="cancellationToken">Cancellation token for the operation</param>
+		/// <returns>True if the job was deleted. False if the job is not the latest revision.</returns>
+		Task<bool> TryDeleteAsync(CancellationToken cancellationToken = default);
+
 		/// <summary>
 		/// Removes a job from the dispatch queue. Ignores the state of any batches still remaining to execute. Should only be used to correct for inconsistent state.
 		/// </summary>
diff --git a/Engine/Source/Programs/Horde/Horde.Server/Jobs/IJobCollection.cs b/Engine/Source/Programs/Horde/Horde.Server/Jobs/IJobCollection.cs
--- a/Engine/Source/Programs/Horde/Horde.Server/Jobs/IJobCollection.cs
+++ b/Engine/Source/Programs/Horde/Horde.Server/Jobs/IJobCollection.cs
@@ -134,13 +134,6 @@ namespace Horde.Server.Jobs
 		/// <returns>Information about the given job</returns>
 		Task<IJob?> GetAsync(JobId jobId, CancellationToken cancellationToken = default);
 
-		/// <summary>
-		/// Deletes a job
-		/// </summary>
-		/// <param name="job">The job to remove</param>
-		/// <param name="cancellationToken">Cancellation token for the operation</param>
-		Task<bool> RemoveAsync(IJob job, CancellationToken cancellationToken = default);
-
 		/// <summary>
 		/// Delete all the jobs for a stream
 		/// </summary>
diff --git a/Engine/Source/Programs/Horde/Horde.Server/Jobs/JobCollection.cs b/Engine/Source/Programs/Horde/Horde.Server/Jobs/JobCollection.cs
--- a/Engine/Source/Programs/Horde/Horde.Server/Jobs/JobCollection.cs
+++ b/Engine/Source/Programs/Horde/Horde.Server/Jobs/JobCollection.cs
@@ -444,6 +444,9 @@ namespace Horde.Server.Jobs
 				return _collection.CreateJobObject(newDocument, _graph);
 			}
 
+			public Task<bool> TryDeleteAsync(CancellationToken cancellationToken = default)
+				=> _collection.TryDeleteAsync(Document, cancellationToken);
+
 			public async Task<IJob?> TryFailBatchAsync(int batchIdx, JobStepBatchError reason, CancellationToken cancellationToken = default)
 			{
 				JobDocument? newDocument = await _collection.TryFailBatchAsync(Document, batchIdx, Graph, reason, cancellationToken);
@@ -605,7 +608,7 @@ namespace Horde.Server.Jobs
 		}
 
 		/// <inheritdoc/>
-		public async Task<bool> RemoveAsync(IJob job, CancellationToken cancellationToken)
+		async Task<bool> TryDeleteAsync(JobDocument job, CancellationToken cancellationToken)
 		{
 			DeleteResult result = await _jobs.DeleteOneAsync(x => x.Id == job.Id && x.UpdateIndex == job.UpdateIndex, null, cancellationToken);
 			return result.DeletedCount > 0;
diff --git a/Engine/Source/Programs/Horde/Horde.Server/Jobs/JobExpirationService.cs b/Engine/Source/Programs/Horde/Horde.Server/Jobs/JobExpirationService.cs
new file mode 100644
--- /dev/null
+++ b/Engine/Source/Programs/Horde/Horde.Server/Jobs/JobExpirationService.cs
@@ -0,0 +1,81 @@
+﻿// Copyright Epic Games, Inc. All Rights Reserved.
+
+using System;
+using System.Collections.Generic;
+using System.Threading;
+using System.Threading.Tasks;
+using EpicGames.Core;
+using Horde.Server.Server;
+using Horde.Server.Streams;
+using HordeCommon;
+using Microsoft.Extensions.Hosting;
+using Microsoft.Extensions.Logging;
+using Microsoft.Extensions.Options;
+
+namespace Horde.Server.Jobs
+{
+	/// <summary>
+	/// Expires jobs a certain amount of time after they have run
+	/// </summary>
+	public sealed class JobExpirationService : IHostedService, IAsyncDisposable
+	{
+		readonly IJobCollection _jobCollection;
+		readonly IClock _clock;
+		readonly IOptionsMonitor<GlobalConfig> _globalConfig;
+		readonly ITicker _ticker;
+		readonly ILogger _logger;
+
+		/// <summary>
+		/// Constructor
+		/// </summary>
+		public JobExpirationService(IJobCollection jobCollection, IClock clock, IOptionsMonitor<GlobalConfig> globalConfig, ILogger<JobExpirationService> logger)
+		{
+			_clock = clock;
+			_globalConfig = globalConfig;
+			_jobCollection = jobCollection;
+			_ticker = clock.AddSharedTicker<JobExpirationService>(TimeSpan.FromHours(1.0), TickAsync, logger);
+			_logger = logger;
+		}
+
+		/// <inheritdoc/>
+		public ValueTask DisposeAsync() => _ticker.DisposeAsync();
+
+		/// <inheritdoc/>
+		public Task StartAsync(CancellationToken cancellationToken) => _ticker.StartAsync();
+
+		/// <inheritdoc/>
+		public Task StopAsync(CancellationToken cancellationToken) => _ticker.StopAsync();
+
+		async ValueTask TickAsync(CancellationToken cancellationToken)
+		{
+			DateTime currentTime = _clock.UtcNow;
+
+			GlobalConfig globalConfig = _globalConfig.CurrentValue;
+			foreach (StreamConfig streamConfig in globalConfig.Streams)
+			{
+				int expireAfterDays = streamConfig.JobOptions.ExpireAfterDays ?? 0;
+				if (expireAfterDays != 0)
+				{
+					DateTimeOffset maxCreateTime = new DateTimeOffset(currentTime) - TimeSpan.FromDays(expireAfterDays);
+					for (; ; )
+					{
+						IReadOnlyList<IJob> jobs = await _jobCollection.FindAsync(streamId: streamConfig.Id, maxCreateTime: maxCreateTime, count: 50, cancellationToken: cancellationToken);
+						if (jobs.Count == 0)
+						{
+							break;
+						}
+
+						foreach (IJob job in jobs)
+						{
+							if (await job.TryDeleteAsync(cancellationToken))
+							{
+								int age = (int)(currentTime - job.CreateTimeUtc).TotalDays;
+								_logger.LogDebug("Deleted job {JobId} ({NumDays}d >= {ExpireDays}d)", job.Id, age, expireAfterDays);
+							}
+						}
+					}
+				}
+			}
+		}
+	}
+}
diff --git a/Engine/Source/Programs/Horde/Horde.Server/Jobs/JobService.cs b/Engine/Source/Programs/Horde/Horde.Server/Jobs/JobService.cs
--- a/Engine/Source/Programs/Horde/Horde.Server/Jobs/JobService.cs
+++ b/Engine/Source/Programs/Horde/Horde.Server/Jobs/JobService.cs
@@ -307,9 +307,9 @@ namespace Horde.Server.Jobs
 			using IDisposable? scope = _logger.BeginScope("DeleteJobAsync({JobId})", job.Id);
 
 			// Delete the job
-			while (!await _jobs.RemoveAsync(job, cancellationToken))
+			while (!await job.TryDeleteAsync(cancellationToken))
 			{
-				IJob? newJob = await _jobs.GetAsync(job.Id, cancellationToken);
+				IJob? newJob = await job.RefreshAsync(cancellationToken);
 				if (newJob == null)
 				{
 					return false;
diff --git a/Engine/Source/Programs/Horde/Horde.Server/Startup.cs b/Engine/Source/Programs/Horde/Horde.Server/Startup.cs
--- a/Engine/Source/Programs/Horde/Horde.Server/Startup.cs
+++ b/Engine/Source/Programs/Horde/Horde.Server/Startup.cs
@@ -864,6 +864,7 @@ namespace Horde.Server
 				services.AddHostedService(provider => provider.GetRequiredService<IssueService>());
 				services.AddHostedService<IssueReportService>();
 				services.AddHostedService<IssueTagService>();
+				services.AddHostedService<JobExpirationService>();
 				services.AddHostedService<MetricService>();
 				services.AddHostedService(provider => provider.GetRequiredService<PerforceLoadBalancer>());
 				services.AddHostedService<PoolUpdateService>();
diff --git a/Engine/Source/Programs/Shared/EpicGames.Horde/Jobs/JobMessages.cs b/Engine/Source/Programs/Shared/EpicGames.Horde/Jobs/JobMessages.cs
--- a/Engine/Source/Programs/Shared/EpicGames.Horde/Jobs/JobMessages.cs
+++ b/Engine/Source/Programs/Shared/EpicGames.Horde/Jobs/JobMessages.cs
@@ -71,16 +71,25 @@ namespace EpicGames.Horde.Jobs
 		/// <summary>
 		/// Options for executing a job inside a container
 		/// </summary>
-		public JobContainerOptions? Container { get; set; }
+		public JobContainerOptions Container { get; set; } = new JobContainerOptions();
+
+		/// <summary>
+		/// Number of days after which to expire jobs
+		/// </summary>
+		public int? ExpireAfterDays { get; set; }
 
 		/// <summary>
 		/// Merge defaults from another options object
 		/// </summary>
-		/// <param name="other"></param>
 		public void MergeDefaults(JobOptions other)
 		{
 			Executor ??= other.Executor;
 			UseNewTempStorage ??= other.UseNewTempStorage;
+			UseWine ??= other.UseWine;
+			RunInSeparateProcess ??= other.RunInSeparateProcess;
+			WorkspaceMaterializer ??= other.WorkspaceMaterializer;
+			Container.MergeDefaults(other.Container);
+			ExpireAfterDays ??= other.ExpireAfterDays;
 		}
 	}
 
@@ -108,6 +117,17 @@ namespace EpicGames.Horde.Jobs
 		/// Additional arguments to pass to container engine
 		/// </summary>
 		public string? ExtraArguments { get; set; }
+
+		/// <summary>
+		/// Merge defaults from another options object
+		/// </summary>
+		public void MergeDefaults(JobContainerOptions other)
+		{
+			Enabled ??= other.Enabled;
+			ImageUrl ??= other.ImageUrl;
+			ContainerEngineExecutable ??= other.ContainerEngineExecutable;
+			ExtraArguments ??= other.ExtraArguments;
+		}
 	}
 
 	/// <summary>
