Horde: Rename uses of "LogFile" to "Log". We started renaming all these things some time ago (eg. LogId) but didn't finish it.

diff --git a/Engine/Source/Programs/Horde/Horde.Server.Tests/Compute/ComputeServiceTest.cs b/Engine/Source/Programs/Horde/Horde.Server.Tests/Compute/ComputeServiceTest.cs
--- a/Engine/Source/Programs/Horde/Horde.Server.Tests/Compute/ComputeServiceTest.cs
+++ b/Engine/Source/Programs/Horde/Horde.Server.Tests/Compute/ComputeServiceTest.cs
@@ -287,7 +287,7 @@ namespace Horde.Server.Tests.Compute
 		private async Task<ComputeService> CreateComputeServiceAsync(string? tunnelAddress)
 		{
 			ServerSettings ss = new() { ComputeTunnelAddress = tunnelAddress };
-			ComputeService cs = new(AgentCollection, LogFileService, AgentService, AgentRelayService, GetRedisServiceSingleton(),
+			ComputeService cs = new(AgentCollection, LogService, AgentService, AgentRelayService, GetRedisServiceSingleton(),
 				new TestOptionsMonitor<ServerSettings>(ss), GlobalConfig, Clock, Tracer, Meter,
 				NullLogger<ComputeService>.Instance);
 			List<string> props = new() { "ComputeIp=11.0.0.1", "ComputePort=5000" };
diff --git a/Engine/Source/Programs/Horde/Horde.Server.Tests/Issues/IssueServiceTests.cs b/Engine/Source/Programs/Horde/Horde.Server.Tests/Issues/IssueServiceTests.cs
--- a/Engine/Source/Programs/Horde/Horde.Server.Tests/Issues/IssueServiceTests.cs
+++ b/Engine/Source/Programs/Horde/Horde.Server.Tests/Issues/IssueServiceTests.cs
@@ -44,7 +44,7 @@ namespace Horde.Server.Tests.Issues
 	{
 		class TestJsonLogger : ILogger, IAsyncDisposable
 		{
-			readonly ILogFileService _logFileService;
+			readonly ILogService _logFileService;
 			readonly LogId _logId;
 			readonly LogBuilder _builder;
 			readonly List<(LogLevel, ReadOnlyMemory<byte>)> _events = new List<(LogLevel, ReadOnlyMemory<byte>)>();
@@ -52,7 +52,7 @@ namespace Horde.Server.Tests.Issues
 
 			int _lineIndex;
 
-			public TestJsonLogger(ILogFileService logFileService, LogId logId, IStorageClient storageClient)
+			public TestJsonLogger(ILogService logFileService, LogId logId, IStorageClient storageClient)
 			{
 				_logFileService = logFileService;
 				_logId = logId;
@@ -73,7 +73,7 @@ namespace Horde.Server.Tests.Issues
 				await using (IBlobWriter writer = _storageClient.CreateBlobWriter())
 				{
 					IBlobRef<LogNode> handle = await _builder.FlushAsync(writer, true, CancellationToken.None);
-					ILogFile? logFile = await _logFileService.GetLogFileAsync(_logId, CancellationToken.None);
+					ILog? logFile = await _logFileService.GetLogAsync(_logId, CancellationToken.None);
 					await _storageClient.WriteRefAsync(logFile!.RefName, handle);
 				}
 
@@ -246,7 +246,7 @@ namespace Horde.Server.Tests.Issues
 				{
 					JobStepId stepId = new JobStepId((ushort)((groupIdx * 100) + nodeIdx));
 
-					ILogFile logFile = LogFileService.CreateLogFileAsync(jobId, null, null, LogType.Json).Result;
+					ILog logFile = LogService.CreateLogAsync(jobId, null, null, LogType.Json).Result;
 
 					Mock<IJobStep> step = new Mock<IJobStep>(MockBehavior.Strict);
 					step.SetupGet(x => x.Id).Returns(stepId);
@@ -323,7 +323,7 @@ namespace Horde.Server.Tests.Issues
 		{
 			LogId logId = job.Batches[batchIdx].Steps[stepIdx].LogId!.Value;
 
-			ILogFile logFile = (await LogFileService.GetLogFileAsync(logId, CancellationToken.None))!;
+			ILog logFile = (await LogService.GetLogAsync(logId, CancellationToken.None))!;
 
 			using (IStorageClient storageClient = StorageService.CreateClient(Namespace.Logs))
 			{
@@ -337,14 +337,14 @@ namespace Horde.Server.Tests.Issues
 				}
 			}
 
-			await LogFileService.CreateEventsAsync(new List<NewLogEventData> { new NewLogEventData { LogId = logId, LineIndex = 0, LineCount = 1, Severity = severity } }, CancellationToken.None);
+			await LogService.CreateEventsAsync(new List<NewLogEventData> { new NewLogEventData { LogId = logId, LineIndex = 0, LineCount = 1, Severity = severity } }, CancellationToken.None);
 		}
 
 		private TestJsonLogger CreateLogger(IJob job, int batchIdx, int stepIdx)
 		{
 			LogId logId = job.Batches[batchIdx].Steps[stepIdx].LogId!.Value;
 			IStorageClient storageClient = StorageService.CreateClient(Namespace.Logs);
-			return new TestJsonLogger(LogFileService, logId, storageClient);
+			return new TestJsonLogger(LogService, logId, storageClient);
 		}
 
 		private async Task ParseEventsAsync(IJob job, int batchIdx, int stepIdx, string[] lines)
@@ -352,7 +352,7 @@ namespace Horde.Server.Tests.Issues
 			LogId logId = job.Batches[batchIdx].Steps[stepIdx].LogId!.Value;
 
 			IStorageClient storageClient = StorageService.CreateClient(Namespace.Logs);
-			await using (TestJsonLogger logger = new TestJsonLogger(LogFileService, logId, storageClient))
+			await using (TestJsonLogger logger = new TestJsonLogger(LogService, logId, storageClient))
 			{
 				PerforceMetadataLogger perforceLogger = new PerforceMetadataLogger(logger);
 				perforceLogger.AddClientView(_autoSdkDir, "//depot/CarefullyRedist/...", 12345);
@@ -669,8 +669,8 @@ namespace Horde.Server.Tests.Issues
 				await ParseEventsAsync(job, 0, 0, lines);
 				await UpdateCompleteStepAsync(job, 0, 0, JobStepOutcome.Warnings);
 
-				ILogFile? log = await LogFileService.GetLogFileAsync(job.Batches[0].Steps[0].LogId!.Value, CancellationToken.None);
-				List<ILogEvent> events = await LogFileService.FindEventsAsync(log!);
+				ILog? log = await LogService.GetLogAsync(job.Batches[0].Steps[0].LogId!.Value, CancellationToken.None);
+				List<ILogEvent> events = await LogService.FindEventsAsync(log!);
 				Assert.AreEqual(1, events.Count);
 				Assert.AreEqual(2, events[0].LineCount);
 
@@ -711,8 +711,8 @@ namespace Horde.Server.Tests.Issues
 				await AddEventAsync(job, 0, 0, LogLevel.Warning, id: KnownLogEvents.Gauntlet_TestEvent.Id);
 				await UpdateCompleteStepAsync(job, 0, 0, JobStepOutcome.Warnings);
 
-				ILogFile? log = await LogFileService.GetLogFileAsync(job.Batches[0].Steps[0].LogId!.Value, CancellationToken.None);
-				List<ILogEvent> events = await LogFileService.FindEventsAsync(log!);
+				ILog? log = await LogService.GetLogAsync(job.Batches[0].Steps[0].LogId!.Value, CancellationToken.None);
+				List<ILogEvent> events = await LogService.FindEventsAsync(log!);
 				Assert.AreEqual(1, events.Count);
 				Assert.AreEqual(1, events[0].LineCount);
 
diff --git a/Engine/Source/Programs/Horde/Horde.Server.Tests/Jobs/TestData/TestDataTests.cs b/Engine/Source/Programs/Horde/Horde.Server.Tests/Jobs/TestData/TestDataTests.cs
--- a/Engine/Source/Programs/Horde/Horde.Server.Tests/Jobs/TestData/TestDataTests.cs
+++ b/Engine/Source/Programs/Horde/Horde.Server.Tests/Jobs/TestData/TestDataTests.cs
@@ -123,7 +123,7 @@ namespace Horde.Server.Tests.Jobs.TestData
 				{
 					JobStepId stepId = new JobStepId((ushort)((groupIdx * 100) + nodeIdx));
 
-					ILogFile logFile = LogFileService.CreateLogFileAsync(jobId, null, null, LogType.Json).Result;
+					ILog logFile = LogService.CreateLogAsync(jobId, null, null, LogType.Json).Result;
 
 					Mock<IJobStep> step = new Mock<IJobStep>(MockBehavior.Strict);
 					step.SetupGet(x => x.Id).Returns(stepId);
diff --git a/Engine/Source/Programs/Horde/Horde.Server.Tests/Logs/LogFileServiceTests.cs b/Engine/Source/Programs/Horde/Horde.Server.Tests/Logs/LogFileServiceTests.cs
deleted file mode 100644
diff --git a/Engine/Source/Programs/Horde/Horde.Server.Tests/Logs/LogIndexingTests.cs b/Engine/Source/Programs/Horde/Horde.Server.Tests/Logs/LogIndexingTests.cs
--- a/Engine/Source/Programs/Horde/Horde.Server.Tests/Logs/LogIndexingTests.cs
+++ b/Engine/Source/Programs/Horde/Horde.Server.Tests/Logs/LogIndexingTests.cs
@@ -20,14 +20,14 @@ namespace Horde.Server.Tests.Logs
 {
 	class TestLogWriter : IAsyncDisposable
 	{
-		ILogFile _logFile;
-		readonly ILogFileCollection _logCollection;
+		ILog _logFile;
+		readonly ILogCollection _logCollection;
 		readonly IStorageClient _storageClient;
 		readonly IBlobWriter _blobWriter;
 		readonly LogBuilder _builder;
 		int _lineCount;
 
-		public TestLogWriter(ILogFile logFile, ILogFileCollection logCollection, StorageService storageService)
+		public TestLogWriter(ILog logFile, ILogCollection logCollection, StorageService storageService)
 		{
 			_logFile = logFile;
 			_logCollection = logCollection;
@@ -42,7 +42,7 @@ namespace Horde.Server.Tests.Logs
 			_storageClient.Dispose();
 		}
 
-		public async Task<ILogFile> WriteDataAsync(ReadOnlyMemory<byte> data)
+		public async Task<ILog> WriteDataAsync(ReadOnlyMemory<byte> data)
 		{
 			_builder.WriteData(data);
 			_lineCount += data.Span.Count((byte)'\n');
@@ -50,7 +50,7 @@ namespace Horde.Server.Tests.Logs
 			return await FlushAsync(false);
 		}
 
-		public async Task<ILogFile> FlushAsync(bool complete = true)
+		public async Task<ILog> FlushAsync(bool complete = true)
 		{
 			IBlobRef<LogNode> handle = await _builder.FlushAsync(_blobWriter, complete, CancellationToken.None);
 			await _storageClient.WriteRefAsync(_logFile.RefName, handle);
@@ -68,10 +68,10 @@ namespace Horde.Server.Tests.Logs
 		public async Task IndexTestsAsync()
 		{
 			JobId jobId = JobIdUtils.GenerateNewId();
-			ILogFile logFile = await LogFileService.CreateLogFileAsync(jobId, null, null, LogType.Text);
+			ILog logFile = await LogService.CreateLogAsync(jobId, null, null, LogType.Text);
 
 			// Write the test data to the log file in blocks
-			await using (TestLogWriter logWriter = new TestLogWriter(logFile, LogFileCollection, StorageService))
+			await using (TestLogWriter logWriter = new TestLogWriter(logFile, LogCollection, StorageService))
 			{
 				int offset = 0;
 				int lineIndex = 0;
@@ -100,7 +100,7 @@ namespace Horde.Server.Tests.Logs
 
 			// Read the data back out and check it's the same
 			byte[] readData = new byte[_data.Length];
-			using (Stream stream = await LogFileService.OpenRawStreamAsync(logFile))
+			using (Stream stream = await LogService.OpenRawStreamAsync(logFile))
 			{
 				int readSize = await stream.ReadAsync(readData, 0, readData.Length);
 				Assert.AreEqual(readData.Length, readSize);
@@ -139,7 +139,7 @@ namespace Horde.Server.Tests.Logs
 		public async Task PartialTokenTestsAsync()
 		{
 			JobId jobId = JobIdUtils.GenerateNewId();
-			ILogFile logFile = await LogFileService.CreateLogFileAsync(jobId, null, null, LogType.Text);
+			ILog logFile = await LogService.CreateLogAsync(jobId, null, null, LogType.Text);
 
 			string[] lines =
 			{
@@ -150,7 +150,7 @@ namespace Horde.Server.Tests.Logs
 			};
 
 			int length = 0;
-			await using (TestLogWriter logWriter = new TestLogWriter(logFile, LogFileCollection, StorageService))
+			await using (TestLogWriter logWriter = new TestLogWriter(logFile, LogCollection, StorageService))
 			{
 				for (int lineIdx = 0; lineIdx < lines.Length; lineIdx++)
 				{
@@ -169,7 +169,7 @@ namespace Horde.Server.Tests.Logs
 						string str = lines[lineIdx].Substring(strOfs, strLen);
 
 						SearchStats stats = new SearchStats();
-						List<int> results = await LogFileService.SearchLogDataAsync(logFile, str, 0, 5, stats, CancellationToken.None);
+						List<int> results = await LogService.SearchLogDataAsync(logFile, str, 0, 5, stats, CancellationToken.None);
 						Assert.AreEqual(1, results.Count);
 						Assert.AreEqual(lineIdx, results[0]);
 
@@ -185,9 +185,9 @@ namespace Horde.Server.Tests.Logs
 		public async Task AppendIndexTestsAsync()
 		{
 			JobId jobId = JobIdUtils.GenerateNewId();
-			ILogFile logFile = await LogFileService.CreateLogFileAsync(jobId, null, null, LogType.Text);
+			ILog logFile = await LogService.CreateLogAsync(jobId, null, null, LogType.Text);
 
-			await using (TestLogWriter writer = new TestLogWriter(logFile, LogFileCollection, StorageService))
+			await using (TestLogWriter writer = new TestLogWriter(logFile, LogCollection, StorageService))
 			{
 				logFile = await writer.WriteDataAsync(Encoding.UTF8.GetBytes("abc\n"));
 				logFile = await writer.WriteDataAsync(Encoding.UTF8.GetBytes("def\n"));
@@ -197,7 +197,7 @@ namespace Horde.Server.Tests.Logs
 
 			{
 				SearchStats stats = new SearchStats();
-				List<int> results = await LogFileService.SearchLogDataAsync(logFile, "abc", 0, 5, stats, CancellationToken.None);
+				List<int> results = await LogService.SearchLogDataAsync(logFile, "abc", 0, 5, stats, CancellationToken.None);
 				Assert.AreEqual(1, results.Count);
 				Assert.AreEqual(0, results[0]);
 
@@ -207,7 +207,7 @@ namespace Horde.Server.Tests.Logs
 			}
 			{
 				SearchStats stats = new SearchStats();
-				List<int> results = await LogFileService.SearchLogDataAsync(logFile, "def", 0, 5, stats, CancellationToken.None);
+				List<int> results = await LogService.SearchLogDataAsync(logFile, "def", 0, 5, stats, CancellationToken.None);
 				Assert.AreEqual(1, results.Count);
 				Assert.AreEqual(1, results[0]);
 
@@ -217,7 +217,7 @@ namespace Horde.Server.Tests.Logs
 			}
 			{
 				SearchStats stats = new SearchStats();
-				List<int> results = await LogFileService.SearchLogDataAsync(logFile, "ghi", 0, 5, stats, CancellationToken.None);
+				List<int> results = await LogService.SearchLogDataAsync(logFile, "ghi", 0, 5, stats, CancellationToken.None);
 				Assert.AreEqual(1, results.Count);
 				Assert.AreEqual(2, results[0]);
 
@@ -229,7 +229,7 @@ namespace Horde.Server.Tests.Logs
 
 		async Task SearchLogDataTestAsync(LogId logId)
 		{
-			ILogFile? logFile = await LogFileService.GetLogFileAsync(logId, CancellationToken.None);
+			ILog? logFile = await LogService.GetLogAsync(logId, CancellationToken.None);
 			Assert.IsNotNull(logFile);
 
 			await SearchLogDataTestAsync(logFile, "HISPANIOLA", 0, 4, new[] { 1503, 1520, 1525, 1595 });
@@ -239,10 +239,10 @@ namespace Horde.Server.Tests.Logs
 			await SearchLogDataTestAsync(logFile, "NEWSLETTER", 0, 100, new[] { 7886 });
 		}
 
-		async Task SearchLogDataTestAsync(ILogFile logFile, string text, int firstLine, int count, int[] expectedLines)
+		async Task SearchLogDataTestAsync(ILog logFile, string text, int firstLine, int count, int[] expectedLines)
 		{
 			SearchStats stats = new SearchStats();
-			List<int> lines = await LogFileService.SearchLogDataAsync(logFile, text, firstLine, count, stats, CancellationToken.None);
+			List<int> lines = await LogService.SearchLogDataAsync(logFile, text, firstLine, count, stats, CancellationToken.None);
 			Assert.IsTrue(lines.SequenceEqual(expectedLines));
 		}
 	}
diff --git a/Engine/Source/Programs/Horde/Horde.Server.Tests/Logs/LogServiceTests.cs b/Engine/Source/Programs/Horde/Horde.Server.Tests/Logs/LogServiceTests.cs
new file mode 100644
--- /dev/null
+++ b/Engine/Source/Programs/Horde/Horde.Server.Tests/Logs/LogServiceTests.cs
@@ -0,0 +1,245 @@
+// Copyright Epic Games, Inc. All Rights Reserved.
+
+using System;
+using System.Collections.Generic;
+using System.IO;
+using System.Linq;
+using System.Security.Claims;
+using System.Text;
+using System.Threading;
+using System.Threading.Tasks;
+using EpicGames.Core;
+using EpicGames.Horde.Agents.Sessions;
+using EpicGames.Horde.Jobs;
+using EpicGames.Horde.Logs;
+using Horde.Server.Agents.Sessions;
+using Horde.Server.Jobs;
+using Horde.Server.Logs;
+using Horde.Server.Logs.Data;
+using Horde.Server.Utilities;
+using Microsoft.VisualStudio.TestTools.UnitTesting;
+
+namespace Horde.Server.Tests.Logs
+{
+	[TestClass]
+	public sealed class LogServiceTest : TestSetup
+	{
+		[TestMethod]
+		public void PlainTextDecoder()
+		{
+			string jsonText = "{\"time\":\"2022-01-26T14:18:29\",\"level\":\"Error\",\"message\":\"       \\tdepends on: /mnt/horde/U5\\u002BR5.0\\u002BInc\\u002BMin/Sync/Engine/Binaries/Linux/libUnrealEditor-Core.so\",\"id\":1,\"line\":5,\"lineCount\":1028}";
+			byte[] jsonData = Encoding.UTF8.GetBytes(jsonText);
+
+			byte[] outputData = new byte[jsonData.Length];
+			int outputLength = LogText.ConvertToPlainText(jsonData, outputData, 0);
+			string outputText = Encoding.UTF8.GetString(outputData, 0, outputLength);
+
+			string expectText = "       \tdepends on: /mnt/horde/U5+R5.0+Inc+Min/Sync/Engine/Binaries/Linux/libUnrealEditor-Core.so\n";
+			Assert.AreEqual(outputText, expectText);
+		}
+
+		[TestMethod]
+		public async Task WriteLogLifecycleOldTestAsync()
+		{
+			JobId jobId = JobIdUtils.GenerateNewId();
+			ILog logFile = await LogService.CreateLogAsync(jobId, null, null, LogType.Text, logId: null, cancellationToken: CancellationToken.None);
+
+			await using (TestLogWriter writer = new TestLogWriter(logFile, LogCollection, StorageService))
+			{
+				await writer.WriteDataAsync(Encoding.ASCII.GetBytes("hello\n"));
+				await writer.FlushAsync(false);
+
+				await writer.WriteDataAsync(Encoding.ASCII.GetBytes("foo\nbar\n"));
+				await writer.FlushAsync(false);
+
+				await writer.WriteDataAsync(Encoding.ASCII.GetBytes("baz\n"));
+				await writer.FlushAsync(true);
+			}
+
+			logFile = (await LogService.GetLogAsync(logFile.Id, CancellationToken.None))!;
+
+			Assert.AreEqual("hello", await ReadLogAsync(LogService, logFile, 0, 5));
+			Assert.AreEqual("foo\nbar\nbaz\n", await ReadLogAsync(LogService, logFile, 6, 12));
+
+			List<Utf8String> lines = await LogService.ReadLinesAsync(logFile, 0, 100, CancellationToken.None);
+			Assert.AreEqual("hello", lines[0].ToString());
+			Assert.AreEqual("foo", lines[1].ToString());
+			Assert.AreEqual("bar", lines[2].ToString());
+			Assert.AreEqual("baz", lines[3].ToString());
+		}
+
+		private static async Task AssertLineOffsetAsync(ILogService logService, LogId logId, int lineIndex, int clampedLineIndex, long offset)
+		{
+			ILog? logFile = await logService.GetLogAsync(logId, CancellationToken.None);
+
+			List<Utf8String> lines = await logService.ReadLinesAsync(logFile!, 0, lineIndex + 1, CancellationToken.None);
+			Assert.AreEqual(clampedLineIndex, Math.Min(lineIndex, lines.Count));
+			Assert.AreEqual(offset, lines.Take(lineIndex).Sum(x => x.Length + 1));
+		}
+
+		static async Task<string> ReadLogAsync(ILogService logService, ILog log, long offset, long length)
+		{
+			using Stream stream = await logService.OpenRawStreamAsync(log);
+
+			byte[] prefix = new byte[(int)offset];
+			await stream.ReadGreedyAsync(prefix);
+
+			byte[] data = new byte[(int)length];
+			int actualLength = await stream.ReadGreedyAsync(data);
+
+			return Encoding.UTF8.GetString(data.AsSpan(0, actualLength));
+		}
+
+		[TestMethod]
+		public async Task WriteLogLifecycleAsync()
+		{
+			JobId jobId = JobIdUtils.GenerateNewId();
+			ILog log = await LogService.CreateLogAsync(jobId, null, null, LogType.Text);
+
+			string str1 = "hello\n";
+			string str2 = "foo\nbar\n";
+			string str3 = "baz\nqux\nquux\n";
+			string str4 = "quuz\n";
+
+			int lineIndex = 0;
+			int offset = 0;
+
+			await using TestLogWriter logWriter = new TestLogWriter(log, LogCollection, StorageService);
+
+			// First write with flush. Will become chunk #1
+			log = await logWriter.WriteDataAsync(Encoding.ASCII.GetBytes(str1));
+			Assert.AreEqual(str1, await ReadLogAsync(LogService, log, 0, str1.Length));
+			Assert.AreEqual(str1, await ReadLogAsync(LogService, log, 0, str1.Length + 100)); // Reading too far is valid?
+			await AssertLineOffsetAsync(LogService, log.Id, 0, 0, 0);
+			await AssertLineOffsetAsync(LogService, log.Id, 1, 1, 6);
+			await AssertLineOffsetAsync(LogService, log.Id, 1, 1, 6);
+			await AssertLineOffsetAsync(LogService, log.Id, 1, 1, 6);
+
+			// Second write without flushing. Will become chunk #2
+			offset += str1.Length;
+			lineIndex += str1.Count(f => f == '\n');
+			log = await logWriter.WriteDataAsync(Encoding.ASCII.GetBytes(str2));
+
+			Assert.AreEqual(str1 + str2, await ReadLogAsync(LogService, log, 0, str1.Length + str2.Length));
+			await AssertLineOffsetAsync(LogService, log.Id, 0, 0, 0);
+			await AssertLineOffsetAsync(LogService, log.Id, 1, 1, 6);
+			await AssertLineOffsetAsync(LogService, log.Id, 2, 2, 10);
+			await AssertLineOffsetAsync(LogService, log.Id, 3, 3, 14);
+
+			// Third write without flushing. Will become chunk #2
+			offset += str2.Length;
+			lineIndex += str3.Count(f => f == '\n');
+			log = await logWriter.WriteDataAsync(Encoding.ASCII.GetBytes(str3));
+			Assert.AreEqual(str1 + str2 + str3, await ReadLogAsync(LogService, log, 0, str1.Length + str2.Length + str3.Length));
+			await AssertLineOffsetAsync(LogService, log.Id, 0, 0, 0);
+			await AssertLineOffsetAsync(LogService, log.Id, 1, 1, 6);
+			await AssertLineOffsetAsync(LogService, log.Id, 2, 2, 10);
+			await AssertLineOffsetAsync(LogService, log.Id, 3, 3, 14);
+			await AssertLineOffsetAsync(LogService, log.Id, 4, 4, 18);
+
+			// Fourth write with flush. Will become chunk #2
+			offset += str3.Length;
+			lineIndex += str4.Count(f => f == '\n');
+			log = await logWriter.WriteDataAsync(Encoding.ASCII.GetBytes(str4));
+			Assert.AreEqual(str1 + str2 + str3 + str4,
+			await ReadLogAsync(LogService, log, 0, str1.Length + str2.Length + str3.Length + str4.Length));
+			Assert.AreEqual(str3 + str4,
+			await ReadLogAsync(LogService, log, str1.Length + str2.Length, str3.Length + str4.Length));
+			await AssertLineOffsetAsync(LogService, log.Id, 0, 0, 0);
+			await AssertLineOffsetAsync(LogService, log.Id, 1, 1, 6);
+			await AssertLineOffsetAsync(LogService, log.Id, 2, 2, 10);
+			await AssertLineOffsetAsync(LogService, log.Id, 3, 3, 14);
+			await AssertLineOffsetAsync(LogService, log.Id, 4, 4, 18);
+			await AssertLineOffsetAsync(LogService, log.Id, 5, 5, 22);
+			await AssertLineOffsetAsync(LogService, log.Id, 6, 6, 27);
+			await AssertLineOffsetAsync(LogService, log.Id, 7, 7, 32);
+
+			// Fifth write with flush and data that will span more than chunk. Will become chunk #3
+			string a = "Lorem ipsum dolor sit amet\n";
+			string b = "consectetur adipiscing\n";
+			string str5 = a + b;
+
+			offset += str4.Length;
+			lineIndex += str5.Count(f => f == '\n');
+
+			// Using this single write below will fail the ReadLogFile assert below. A bug?
+			//await LogFileService.WriteLogDataAsync(LogFile, Offset, LineIndex, Encoding.ASCII.GetBytes(Str5), true);
+
+			// Dividing it in two like this will work however
+			log = await logWriter.WriteDataAsync(Encoding.ASCII.GetBytes(a));
+			log = await logWriter.WriteDataAsync(Encoding.ASCII.GetBytes(b));
+
+			Assert.AreEqual(str5, await ReadLogAsync(LogService, log, offset, str5.Length));
+		}
+
+		[TestMethod]
+		public async Task GetLogFileTestAsync()
+		{
+			await GetMongoServiceSingleton().Database.DropCollectionAsync("LogFiles");
+			Assert.AreEqual(0, (await LogService.GetLogsAsync()).Count);
+
+			// Will implicitly test GetLogFileAsync(), AddCachedLogFile()
+			JobId jobId = JobIdUtils.GenerateNewId();
+			SessionId sessionId = SessionIdUtils.GenerateNewId();
+			ILog a = await LogService.CreateLogAsync(jobId, null, sessionId, LogType.Text, logId: null, CancellationToken.None);
+			ILog b = (await LogService.GetLogAsync(a.Id, CancellationToken.None))!;
+			Assert.AreEqual(a.JobId, b.JobId);
+			Assert.AreEqual(a.SessionId, b.SessionId);
+			Assert.AreEqual(a.Type, b.Type);
+
+			ILog? notFound = await LogService.GetLogAsync(LogIdUtils.GenerateNewId(), CancellationToken.None);
+			Assert.IsNull(notFound);
+
+			await LogService.CreateLogAsync(JobIdUtils.GenerateNewId(), null, SessionIdUtils.GenerateNewId(), LogType.Text, logId: null, cancellationToken: CancellationToken.None);
+			Assert.AreEqual(2, (await LogService.GetLogsAsync()).Count);
+		}
+
+		[TestMethod]
+		public async Task AuthorizeForSessionAsync()
+		{
+			JobId jobId = JobIdUtils.GenerateNewId();
+			SessionId sessionId = SessionIdUtils.GenerateNewId();
+			ILog log = await LogService.CreateLogAsync(jobId, null, sessionId, LogType.Text, logId: null, cancellationToken: CancellationToken.None);
+			ILog logNoSession = await LogService.CreateLogAsync(jobId, null, null, LogType.Text, logId: null, cancellationToken: CancellationToken.None);
+
+			ClaimsPrincipal hasClaim = new ClaimsPrincipal(new ClaimsIdentity(new List<Claim>
+			{
+				new Claim(HordeClaimTypes.AgentSessionId, sessionId.ToString()),
+			}, "TestAuthType"));
+			ClaimsPrincipal hasNoClaim = new ClaimsPrincipal(new ClaimsIdentity(new List<Claim>
+			{
+				new Claim(HordeClaimTypes.AgentSessionId, "invalid-session-id"),
+			}, "TestAuthType"));
+
+			Assert.IsTrue(Horde.Server.Logs.LogService.AuthorizeForSession(log, hasClaim));
+			Assert.IsFalse(Horde.Server.Logs.LogService.AuthorizeForSession(log, hasNoClaim));
+			Assert.IsFalse(Horde.Server.Logs.LogService.AuthorizeForSession(logNoSession, hasClaim));
+		}
+
+		[TestMethod]
+		public async Task ChunkSplittingAsync()
+		{
+			JobId jobId = JobIdUtils.GenerateNewId();
+			ILog logFile = await LogService.CreateLogAsync(jobId, null, null, LogType.Text, logId: null, cancellationToken: CancellationToken.None);
+			await using TestLogWriter logWriter = new TestLogWriter(logFile, LogCollection, StorageService);
+
+			byte[] line1 = Encoding.UTF8.GetBytes("hello world\n");
+			await logWriter.WriteDataAsync(line1);
+
+			byte[] line2 = Encoding.UTF8.GetBytes("ab\n");
+			await logWriter.WriteDataAsync(line2);
+
+			byte[] line3 = Encoding.UTF8.GetBytes("a\n");
+			await logWriter.WriteDataAsync(line3);
+
+			byte[] line4 = Encoding.UTF8.GetBytes("b\n");
+			await logWriter.WriteDataAsync(line4);
+
+			byte[] line5 = Encoding.UTF8.GetBytes("a\nb\n");
+			await logWriter.WriteDataAsync(line5);
+
+			logFile = await logWriter.FlushAsync();
+			Assert.AreEqual(6, logFile.LineCount);
+		}
+	}
+}
\ No newline at end of file
diff --git a/Engine/Source/Programs/Horde/Horde.Server.Tests/TestSetup.cs b/Engine/Source/Programs/Horde/Horde.Server.Tests/TestSetup.cs
--- a/Engine/Source/Programs/Horde/Horde.Server.Tests/TestSetup.cs
+++ b/Engine/Source/Programs/Horde/Horde.Server.Tests/TestSetup.cs
@@ -95,7 +95,7 @@ namespace Horde.Server.Tests
 		public IIssueCollection IssueCollection => ServiceProvider.GetRequiredService<IIssueCollection>();
 		public IPoolCollection PoolCollection => ServiceProvider.GetRequiredService<IPoolCollection>();
 		public ILeaseCollection LeaseCollection => ServiceProvider.GetRequiredService<ILeaseCollection>();
-		public ILogFileCollection LogFileCollection => ServiceProvider.GetRequiredService<ILogFileCollection>();
+		public ILogCollection LogCollection => ServiceProvider.GetRequiredService<ILogCollection>();
 		public ISessionCollection SessionCollection => ServiceProvider.GetRequiredService<ISessionCollection>();
 		public ITestDataCollection TestDataCollection => ServiceProvider.GetRequiredService<ITestDataCollection>();
 		public IUserCollection UserCollection => ServiceProvider.GetRequiredService<IUserCollection>();
@@ -112,7 +112,7 @@ namespace Horde.Server.Tests
 		public MongoService MongoService => ServiceProvider.GetRequiredService<MongoService>();
 		public ITemplateCollection TemplateCollection => ServiceProvider.GetRequiredService<ITemplateCollection>();
 		internal PerforceServiceStub PerforceService => (PerforceServiceStub)ServiceProvider.GetRequiredService<IPerforceService>();
-		public ILogFileService LogFileService => ServiceProvider.GetRequiredService<ILogFileService>();
+		public ILogService LogService => ServiceProvider.GetRequiredService<ILogService>();
 		public ISubscriptionCollection SubscriptionCollection => ServiceProvider.GetRequiredService<ISubscriptionCollection>();
 		public INotificationService NotificationService => ServiceProvider.GetRequiredService<INotificationService>();
 		public IssueService IssueService => ServiceProvider.GetRequiredService<IssueService>();
@@ -223,7 +223,7 @@ namespace Horde.Server.Tests
 			services.AddSingleton<IJobTimingCollection, JobTimingCollection>();
 			services.AddSingleton<ILeaseCollection, LeaseCollection>();
 			services.AddSingleton<ILogEventCollection, LogEventCollection>();
-			services.AddSingleton<ILogFileCollection, LogFileCollection>();
+			services.AddSingleton<ILogCollection, LogCollection>();
 			services.AddSingleton<INotificationTriggerCollection, NotificationTriggerCollection>();
 			services.AddSingleton<IPoolCollection, PoolCollection>();
 			services.AddSingleton<IBisectTaskCollection, BisectTaskCollection>();
@@ -267,7 +267,7 @@ namespace Horde.Server.Tests
 			services.AddSingleton<JobExpirationService>();
 			services.AddSingleton<LifetimeService>();
 			services.AddSingleton<ILogStorage, NullLogStorage>();
-			services.AddSingleton<ILogFileService, LogFileService>();
+			services.AddSingleton<ILogService, LogService>();
 			services.AddSingleton<LogTailService>();
 			services.AddSingleton<INotificationService, NotificationService>();
 			services.AddSingleton<IPerforceService, PerforceServiceStub>();
diff --git a/Engine/Source/Programs/Horde/Horde.Server/Compute/ComputeService.cs b/Engine/Source/Programs/Horde/Horde.Server/Compute/ComputeService.cs
--- a/Engine/Source/Programs/Horde/Horde.Server/Compute/ComputeService.cs
+++ b/Engine/Source/Programs/Horde/Horde.Server/Compute/ComputeService.cs
@@ -167,7 +167,7 @@ namespace Horde.Server.Compute
 		public event ResourceNeedEvent? OnResourceNeedsUpdated;
 
 		readonly IAgentCollection _agentCollection;
-		readonly ILogFileService _logService;
+		readonly ILogService _logService;
 		readonly AgentService _agentService;
 		readonly AgentRelayService _agentRelayService;
 		readonly RedisService _redisService;
@@ -189,7 +189,7 @@ namespace Horde.Server.Compute
 		/// </summary>
 		public ComputeService(
 			IAgentCollection agentCollection,
-			ILogFileService logService,
+			ILogService logService,
 			AgentService agentService,
 			AgentRelayService agentRelayService,
 			RedisService redisService,
@@ -404,7 +404,7 @@ namespace Horde.Server.Compute
 						}
 
 						LeaseId leaseId = new LeaseId(BinaryIdUtils.CreateNew());
-						ILogFile? log = await _logService.CreateLogFileAsync(JobId.Empty, leaseId, agent.SessionId, LogType.Json, cancellationToken: cancellationToken);
+						ILog? log = await _logService.CreateLogAsync(JobId.Empty, leaseId, agent.SessionId, LogType.Json, cancellationToken: cancellationToken);
 
 						using TelemetrySpan createTaskSpan = _tracer.StartActiveSpan("CreateComputeTask");
 
diff --git a/Engine/Source/Programs/Horde/Horde.Server/Issues/IssueService.cs b/Engine/Source/Programs/Horde/Horde.Server/Issues/IssueService.cs
--- a/Engine/Source/Programs/Horde/Horde.Server/Issues/IssueService.cs
+++ b/Engine/Source/Programs/Horde/Horde.Server/Issues/IssueService.cs
@@ -203,7 +203,7 @@ namespace Horde.Server.Issues
 		readonly ICommitService _commitService;
 		readonly IStreamCollection _streams;
 		readonly IUserCollection _userCollection;
-		readonly ILogFileService _logFileService;
+		readonly ILogService _logService;
 		readonly IClock _clock;
 		readonly ITicker _ticker;
 
@@ -255,7 +255,7 @@ namespace Horde.Server.Issues
 		/// <summary>
 		/// Constructor
 		/// </summary>
-		public IssueService(IIssueCollection issueCollection, ICommitService commitService, IJobStepRefCollection jobStepRefs, IStreamCollection streams, IUserCollection userCollection, ILogFileService logFileService, IClock clock, IOptionsMonitor<GlobalConfig> globalConfig, Tracer tracer, ILogger<IssueService> logger)
+		public IssueService(IIssueCollection issueCollection, ICommitService commitService, IJobStepRefCollection jobStepRefs, IStreamCollection streams, IUserCollection userCollection, ILogService logService, IClock clock, IOptionsMonitor<GlobalConfig> globalConfig, Tracer tracer, ILogger<IssueService> logger)
 		{
 			Type[] issueTypes = Assembly.GetExecutingAssembly().GetTypes().Where(x => !x.IsAbstract && typeof(IIssue).IsAssignableFrom(x)).ToArray();
 			foreach (Type issueType in issueTypes)
@@ -269,7 +269,7 @@ namespace Horde.Server.Issues
 			_jobStepRefs = jobStepRefs;
 			_streams = streams;
 			_userCollection = userCollection;
-			_logFileService = logFileService;
+			_logService = logService;
 			_clock = clock;
 			_ticker = clock.AddTicker<IssueService>(TimeSpan.FromMinutes(1.0), TickAsync, logger);
 			_globalConfig = globalConfig;
@@ -679,8 +679,8 @@ namespace Horde.Server.Issues
 			}
 
 			// Find the log file for this step
-			ILogFile? logFile = await _logFileService.GetLogFileAsync(step.LogId.Value, CancellationToken.None);
-			if (logFile == null)
+			ILog? log = await _logService.GetLogAsync(step.LogId.Value, CancellationToken.None);
+			if (log == null)
 			{
 				throw new ArgumentException($"Unable to retrieve log {step.LogId}");
 			}
@@ -702,10 +702,10 @@ namespace Horde.Server.Issues
 			}
 
 			// Create all the issue definitions by passing each log event to the handlers in order until one attaches it to an issue
-			List<ILogEvent> stepEvents = await _logFileService.FindEventsAsync(logFile, cancellationToken: cancellationToken);
+			List<ILogEvent> stepEvents = await _logService.FindEventsAsync(log, cancellationToken: cancellationToken);
 			foreach (ILogEvent stepEvent in stepEvents)
 			{
-				ILogEventData stepEventData = await _logFileService.GetEventDataAsync(logFile, stepEvent.LineIndex, stepEvent.LineCount, cancellationToken);
+				ILogEventData stepEventData = await _logService.GetEventDataAsync(log, stepEvent.LineIndex, stepEvent.LineCount, cancellationToken);
 
 				IssueEventInternal issueEvent = new IssueEventInternal(stepEvent, stepEventData);
 				foreach (IssueHandler handler in handlers)
@@ -809,7 +809,7 @@ namespace Horde.Server.Issues
 						}
 
 						// Assign all the events to the span
-						await _logFileService.AddSpanToEventsAsync(matchEventGroups.SelectMany(x => x.Events.Select(x => x.Event)), newSpan.Id, cancellationToken);
+						await _logService.AddSpanToEventsAsync(matchEventGroups.SelectMany(x => x.Events.Select(x => x.Event)), newSpan.Id, cancellationToken);
 
 						// Remove the matches from the set of events
 						newEventGroups.ExceptWith(matchEventGroups);
@@ -907,7 +907,7 @@ namespace Horde.Server.Issues
 
 				// Update the log events
 				_logger.LogDebug("Created new span {SpanId} from event group {Group}", newSpan.Id, eventGroup.Id.ToString());
-				await _logFileService.AddSpanToEventsAsync(eventGroup.Events.Select(x => x.Event), newSpan.Id, cancellationToken);
+				await _logService.AddSpanToEventsAsync(eventGroup.Events.Select(x => x.Event), newSpan.Id, cancellationToken);
 
 				// Remove the events from the remaining list of events to match
 				newEventGroups.ExceptWith(sourceEventGroups);
diff --git a/Engine/Source/Programs/Horde/Horde.Server/Issues/IssuesController.cs b/Engine/Source/Programs/Horde/Horde.Server/Issues/IssuesController.cs
--- a/Engine/Source/Programs/Horde/Horde.Server/Issues/IssuesController.cs
+++ b/Engine/Source/Programs/Horde/Horde.Server/Issues/IssuesController.cs
@@ -42,20 +42,20 @@ namespace Horde.Server.Issues
 		private readonly IExternalIssueService _externalIssueService;
 		private readonly JobService _jobService;
 		private readonly IUserCollection _userCollection;
-		private readonly ILogFileService _logFileService;
+		private readonly ILogService _logService;
 		private readonly IOptionsSnapshot<GlobalConfig> _globalConfig;
 		private readonly ILogger<IssuesController> _logger;
 
 		/// <summary>
 		/// Constructor
 		/// </summary>
-		public IssuesController(IIssueCollection issueCollection, IssueService issueService, JobService jobService, IUserCollection userCollection, ILogFileService logFileService, IExternalIssueService externalIssueService, IOptionsSnapshot<GlobalConfig> globalConfig, ILogger<IssuesController> logger)
+		public IssuesController(IIssueCollection issueCollection, IssueService issueService, JobService jobService, IUserCollection userCollection, ILogService logService, IExternalIssueService externalIssueService, IOptionsSnapshot<GlobalConfig> globalConfig, ILogger<IssuesController> logger)
 		{
 			_issueCollection = issueCollection;
 			_issueService = issueService;
 			_jobService = jobService;
 			_userCollection = userCollection;
-			_logFileService = logFileService;
+			_logService = logService;
 			_externalIssueService = externalIssueService;
 			_globalConfig = globalConfig;
 			_logger = logger;
@@ -728,23 +728,23 @@ namespace Horde.Server.Issues
 			}
 
 			IReadOnlyList<IIssueSpan> spans = await _issueCollection.FindSpansAsync(issueId, cancellationToken);
-			List<ILogEvent> events = await _logFileService.FindEventsForSpansAsync(spans.Select(x => x.Id), logIdValues.ToArray(), index, count, cancellationToken);
+			List<ILogEvent> events = await _logService.FindEventsForSpansAsync(spans.Select(x => x.Id), logIdValues.ToArray(), index, count, cancellationToken);
 
 			JobPermissionsCache permissionsCache = new JobPermissionsCache();
-			Dictionary<LogId, ILogFile?> logFiles = new Dictionary<LogId, ILogFile?>();
+			Dictionary<LogId, ILog?> logs = new Dictionary<LogId, ILog?>();
 
 			List<object> responses = new List<object>();
 			foreach (ILogEvent logEvent in events)
 			{
-				ILogFile? logFile;
-				if (!logFiles.TryGetValue(logEvent.LogId, out logFile))
+				ILog? log;
+				if (!logs.TryGetValue(logEvent.LogId, out log))
 				{
-					logFile = await _logFileService.GetLogFileAsync(logEvent.LogId, cancellationToken);
-					logFiles[logEvent.LogId] = logFile;
+					log = await _logService.GetLogAsync(logEvent.LogId, cancellationToken);
+					logs[logEvent.LogId] = log;
 				}
-				if (logFile != null && await _jobService.AuthorizeAsync(logFile.JobId, LogAclAction.ViewLog, User, _globalConfig.Value, cancellationToken))
+				if (log != null && await _jobService.AuthorizeAsync(log.JobId, LogAclAction.ViewLog, User, _globalConfig.Value, cancellationToken))
 				{
-					ILogEventData data = await _logFileService.GetEventDataAsync(logFile, logEvent.LineIndex, logEvent.LineCount, cancellationToken);
+					ILogEventData data = await _logService.GetEventDataAsync(log, logEvent.LineIndex, logEvent.LineCount, cancellationToken);
 					GetLogEventResponse response = LogsController.CreateGetLogEventResponse(logEvent, data, issueId);
 					responses.Add(PropertyFilter.Apply(response, filter));
 				}
diff --git a/Engine/Source/Programs/Horde/Horde.Server/Jobs/JobRpcService.cs b/Engine/Source/Programs/Horde/Horde.Server/Jobs/JobRpcService.cs
--- a/Engine/Source/Programs/Horde/Horde.Server/Jobs/JobRpcService.cs
+++ b/Engine/Source/Programs/Horde/Horde.Server/Jobs/JobRpcService.cs
@@ -245,7 +245,7 @@ namespace Horde.Server.Jobs
 	{
 		readonly JobService _jobService;
 		readonly IArtifactCollectionV1 _artifactCollection;
-		readonly ILogFileService _logFileService;
+		readonly ILogService _logService;
 		readonly IGraphCollection _graphs;
 		readonly ITestDataCollection _testData;
 		readonly IJobStepRefCollection _jobStepRefCollection;
@@ -258,11 +258,11 @@ namespace Horde.Server.Jobs
 		/// <summary>
 		/// Constructor
 		/// </summary>
-		public JobRpcCommon(JobService jobService, IArtifactCollectionV1 artifactCollection, ILogFileService logFileService, IGraphCollection graphs, ITestDataCollection testData, IJobStepRefCollection jobStepRefCollection, ITemplateCollection templateCollection, HttpClient httpClient, IClock clock, IOptionsSnapshot<GlobalConfig> globalConfig, ILogger<JobRpcCommon> logger)
+		public JobRpcCommon(JobService jobService, IArtifactCollectionV1 artifactCollection, ILogService logService, IGraphCollection graphs, ITestDataCollection testData, IJobStepRefCollection jobStepRefCollection, ITemplateCollection templateCollection, HttpClient httpClient, IClock clock, IOptionsSnapshot<GlobalConfig> globalConfig, ILogger<JobRpcCommon> logger)
 		{
 			_jobService = jobService;
 			_artifactCollection = artifactCollection;
-			_logFileService = logFileService;
+			_logService = logService;
 			_graphs = graphs;
 			_testData = testData;
 			_jobStepRefCollection = jobStepRefCollection;
@@ -427,7 +427,7 @@ namespace Horde.Server.Jobs
 		/// <returns>Information about the new agent</returns>
 		public async Task<RpcBeginStepResponse> BeginStepAsync(RpcBeginStepRequest request, ServerCallContext context)
 		{
-			Boxed<ILogFile?> log = new Boxed<ILogFile?>(null);
+			Boxed<ILog?> log = new Boxed<ILog?>(null);
 			for (; ; )
 			{
 				RpcBeginStepResponse? response = await TryBeginStepAsync(request, log, context);
@@ -438,7 +438,7 @@ namespace Horde.Server.Jobs
 			}
 		}
 
-		async Task<RpcBeginStepResponse?> TryBeginStepAsync(RpcBeginStepRequest request, Boxed<ILogFile?> log, ServerCallContext context)
+		async Task<RpcBeginStepResponse?> TryBeginStepAsync(RpcBeginStepRequest request, Boxed<ILog?> log, ServerCallContext context)
 		{
 			// Check the job exists and we can access it
 			IJob? job = await _jobService.GetJobAsync(JobId.Parse(request.JobId));
@@ -487,7 +487,7 @@ namespace Horde.Server.Jobs
 			}
 
 			// Create a log file if necessary
-			log.Value ??= await _logFileService.CreateLogFileAsync(job.Id, batch.LeaseId, batch.SessionId, LogType.Json);
+			log.Value ??= await _logService.CreateLogAsync(job.Id, batch.LeaseId, batch.SessionId, LogType.Json);
 
 			// Get the node for this step
 			IGraph graph = await _jobService.GetGraphAsync(job);
@@ -898,7 +898,7 @@ namespace Horde.Server.Jobs
 				newEvent.LineCount = createEvent.LineCount;
 				newEvents.Add(newEvent);
 			}
-			await _logFileService.CreateEventsAsync(newEvents, context.CancellationToken);
+			await _logService.CreateEventsAsync(newEvents, context.CancellationToken);
 			return new Empty();
 		}
 
diff --git a/Engine/Source/Programs/Horde/Horde.Server/Jobs/JobTaskSource.cs b/Engine/Source/Programs/Horde/Horde.Server/Jobs/JobTaskSource.cs
--- a/Engine/Source/Programs/Horde/Horde.Server/Jobs/JobTaskSource.cs
+++ b/Engine/Source/Programs/Horde/Horde.Server/Jobs/JobTaskSource.cs
@@ -190,7 +190,7 @@ namespace Horde.Server.Jobs
 
 		readonly AclService _aclService;
 		readonly IStreamCollection _streamCollection;
-		readonly ILogFileService _logFileService;
+		readonly ILogService _logService;
 		readonly IAgentCollection _agentsCollection;
 		readonly IJobCollection _jobs;
 		readonly IJobStepRefCollection _jobStepRefs;
@@ -238,7 +238,7 @@ namespace Horde.Server.Jobs
 		/// <summary>
 		/// Constructor
 		/// </summary>
-		public JobTaskSource(AclService aclService, IAgentCollection agents, IJobCollection jobs, IJobStepRefCollection jobStepRefs, IBisectTaskCollection bisectTasks, IGraphCollection graphs, IPoolCollection pools, PoolService poolService, IUgsMetadataCollection ugsMetadataCollection, IStreamCollection streamCollection, ILogFileService logFileService, PerforceLoadBalancer perforceLoadBalancer, IClock clock, IOptionsMonitor<ServerSettings> settings, IOptionsMonitor<GlobalConfig> globalConfig, Tracer tracer, ILogger<JobTaskSource> logger)
+		public JobTaskSource(AclService aclService, IAgentCollection agents, IJobCollection jobs, IJobStepRefCollection jobStepRefs, IBisectTaskCollection bisectTasks, IGraphCollection graphs, IPoolCollection pools, PoolService poolService, IUgsMetadataCollection ugsMetadataCollection, IStreamCollection streamCollection, ILogService logService, PerforceLoadBalancer perforceLoadBalancer, IClock clock, IOptionsMonitor<ServerSettings> settings, IOptionsMonitor<GlobalConfig> globalConfig, Tracer tracer, ILogger<JobTaskSource> logger)
 		{
 			_aclService = aclService;
 			_agentsCollection = agents;
@@ -250,7 +250,7 @@ namespace Horde.Server.Jobs
 			_poolService = poolService;
 			_ugsMetadataCollection = ugsMetadataCollection;
 			_streamCollection = streamCollection;
-			_logFileService = logFileService;
+			_logService = logService;
 			_perforceLoadBalancer = perforceLoadBalancer;
 			_clock = clock;
 			_ticker = clock.AddTicker<JobTaskSource>(s_refreshInterval, TickAsync, logger);
@@ -827,7 +827,7 @@ namespace Horde.Server.Jobs
 					if (waiter.LeaseSource.TrySetResult(lease))
 					{
 						_logger.LogInformation("Assigned lease to agent");
-						await _logFileService.CreateLogFileAsync(job.Id, leaseId, agent.SessionId, LogType.Json, logId, cancellationToken);
+						await _logService.CreateLogAsync(job.Id, leaseId, agent.SessionId, LogType.Json, logId, cancellationToken);
 						return lease;
 					}
 				}
diff --git a/Engine/Source/Programs/Horde/Horde.Server/Logs/ILog.cs b/Engine/Source/Programs/Horde/Horde.Server/Logs/ILog.cs
new file mode 100644
--- /dev/null
+++ b/Engine/Source/Programs/Horde/Horde.Server/Logs/ILog.cs
@@ -0,0 +1,148 @@
+// Copyright Epic Games, Inc. All Rights Reserved.
+
+using System;
+using System.Collections.Generic;
+using EpicGames.Core;
+using EpicGames.Horde.Agents.Leases;
+using EpicGames.Horde.Agents.Sessions;
+using EpicGames.Horde.Jobs;
+using EpicGames.Horde.Logs;
+using EpicGames.Horde.Storage;
+
+namespace Horde.Server.Logs
+{
+	/// <summary>
+	/// Information about a log file chunk
+	/// </summary>
+	public interface ILogChunk
+	{
+		/// <summary>
+		/// Offset of the chunk within the log
+		/// </summary>
+		long Offset { get; }
+
+		/// <summary>
+		/// Length of this chunk. If zero, the chunk is still being written to.
+		/// </summary>
+		int Length { get; }
+
+		/// <summary>
+		/// Index of the first line within this chunk. If a line straddles two chunks, this is the index of the split line.
+		/// </summary>
+		int LineIndex { get; }
+
+		/// <summary>
+		/// If the chunk has yet to be pushed to persistent storage, includes the name of the server that is currently storing it.
+		/// </summary>
+		string? Server { get; }
+	}
+
+	/// <summary>
+	/// Information about a log file
+	/// </summary>
+	public interface ILog
+	{
+		/// <summary>
+		/// Identifier for the log. Randomly generated.
+		/// </summary>
+		public LogId Id { get; }
+
+		/// <summary>
+		/// Unique id of the job containing this log
+		/// </summary>
+		public JobId JobId { get; }
+
+		/// <summary>
+		/// The lease allowed to write to this log
+		/// </summary>
+		public LeaseId? LeaseId { get; }
+
+		/// <summary>
+		/// The session allowed to write to this log
+		/// </summary>
+		public SessionId? SessionId { get; }
+
+		/// <summary>
+		/// Whether to use the new storage backend for log data
+		/// </summary>
+		public bool UseNewStorageBackend { get; }
+
+		/// <summary>
+		/// Maximum line index in the file
+		/// </summary>
+		public int? MaxLineIndex { get; }
+
+		/// <summary>
+		/// Length of the file which is indexed
+		/// </summary>
+		public long? IndexLength { get; }
+
+		/// <summary>
+		/// Type of data stored in this log 
+		/// </summary>
+		public LogType Type { get; }
+
+		/// <summary>
+		/// Chunks within this file
+		/// </summary>
+		public IReadOnlyList<ILogChunk> Chunks { get; }
+
+		/// <summary>
+		/// Namespace containing the log data
+		/// </summary>
+		public NamespaceId NamespaceId { get; }
+
+		/// <summary>
+		/// Name of the ref used to store data for this log
+		/// </summary>
+		public RefName RefName { get; }
+
+		/// <summary>
+		/// Number of lines (V2 storage backend)
+		/// </summary>
+		public int LineCount { get; }
+
+		/// <summary>
+		/// Whether the log is complete (V2 storage backend)
+		/// </summary>
+		public bool Complete { get; }
+	}
+
+	/// <summary>
+	/// Extension methods for log files
+	/// </summary>
+	public static class LogExtensions
+	{
+		/// <summary>
+		/// Gets the chunk index containing the given offset.
+		/// </summary>
+		/// <param name="chunks">The chunks to search</param>
+		/// <param name="offset">The offset to search for</param>
+		/// <returns>The chunk index containing the given offset</returns>
+		public static int GetChunkForOffset(this IReadOnlyList<ILogChunk> chunks, long offset)
+		{
+			int chunkIndex = chunks.BinarySearch(x => x.Offset, offset);
+			if (chunkIndex < 0)
+			{
+				chunkIndex = ~chunkIndex - 1;
+			}
+			return chunkIndex;
+		}
+
+		/// <summary>
+		/// Gets the starting chunk index for the given line
+		/// </summary>
+		/// <param name="chunks">The chunks to search</param>
+		/// <param name="lineIndex">Index of the line to query</param>
+		/// <returns>Index of the chunk to fetch</returns>
+		public static int GetChunkForLine(this IReadOnlyList<ILogChunk> chunks, int lineIndex)
+		{
+			int chunkIndex = chunks.BinarySearch(x => x.LineIndex, lineIndex);
+			if (chunkIndex < 0)
+			{
+				chunkIndex = ~chunkIndex - 1;
+			}
+			return chunkIndex;
+		}
+	}
+}
diff --git a/Engine/Source/Programs/Horde/Horde.Server/Logs/ILogCollection.cs b/Engine/Source/Programs/Horde/Horde.Server/Logs/ILogCollection.cs
new file mode 100644
--- /dev/null
+++ b/Engine/Source/Programs/Horde/Horde.Server/Logs/ILogCollection.cs
@@ -0,0 +1,119 @@
+// Copyright Epic Games, Inc. All Rights Reserved.
+
+using System.Collections.Generic;
+using System.Threading;
+using System.Threading.Tasks;
+using EpicGames.Horde.Agents.Leases;
+using EpicGames.Horde.Agents.Sessions;
+using EpicGames.Horde.Jobs;
+using EpicGames.Horde.Logs;
+
+namespace Horde.Server.Logs
+{
+	/// <summary>
+	/// Updates a log file chunk
+	/// </summary>
+	public class CompleteLogChunkUpdate
+	{
+		/// <summary>
+		/// Index of the chunk
+		/// </summary>
+		public int Index { get; set; }
+
+		/// <summary>
+		/// New length for the chunk
+		/// </summary>
+		public int Length { get; set; }
+
+		/// <summary>
+		/// Number of lines in the chunk
+		/// </summary>
+		public int LineCount { get; set; }
+
+		/// <summary>
+		/// Constructor
+		/// </summary>
+		/// <param name="index">Index of the chunk</param>
+		/// <param name="length">New length for the chunk</param>
+		/// <param name="lineCount">Number of lines in the chunk</param>
+		public CompleteLogChunkUpdate(int index, int length, int lineCount)
+		{
+			Index = index;
+			Length = length;
+			LineCount = lineCount;
+		}
+	}
+
+	/// <summary>
+	/// Wrapper around the jobs collection in a mongo DB
+	/// </summary>
+	public interface ILogCollection
+	{
+		/// <summary>
+		/// Creates a new log
+		/// </summary>
+		/// <param name="jobId">Unique id of the job that owns this log file</param>
+		/// <param name="leaseId">Agent lease allowed to update the log</param>
+		/// <param name="sessionId">Agent session allowed to update the log</param>
+		/// <param name="type">Type of events to be stored in the log</param>
+		/// <param name="logId">ID of the log file (optional)</param>
+		/// <param name="cancellationToken">Cancellation token for the call</param>
+		/// <returns>The new log file document</returns>
+		Task<ILog> CreateLogAsync(JobId jobId, LeaseId? leaseId, SessionId? sessionId, LogType type, LogId? logId = null, CancellationToken cancellationToken = default);
+
+		/// <summary>
+		/// Updates the line count for a log file (v2 backend only)
+		/// </summary>
+		/// <param name="log">The current log file</param>
+		/// <param name="lineCount">New line count for the log file</param>
+		/// <param name="complete">Flag indicating whether the log is complete, or can still be tailed for additional data</param>
+		/// <param name="cancellationToken">Cancellation token for the call</param>
+		/// <returns>The updated log file document</returns>
+		Task<ILog> UpdateLineCountAsync(ILog log, int lineCount, bool complete, CancellationToken cancellationToken);
+
+		/// <summary>
+		/// Adds a new chunk
+		/// </summary>
+		/// <param name="log">The current log file</param>
+		/// <param name="offset">Offset of the new chunk</param>
+		/// <param name="lineIndex">Line index for the start of the chunk</param>
+		/// <param name="cancellationToken">Cancellation token for the call</param>
+		/// <returns>The updated log file document</returns>
+		Task<ILog?> TryAddChunkAsync(ILog log, long offset, int lineIndex, CancellationToken cancellationToken);
+
+		/// <summary>
+		/// Update the log file with final information about certain chunks
+		/// </summary>
+		/// <param name="log">The current log file</param>
+		/// <param name="chunks">Chunks to update. New chunks will be inserted</param>
+		/// <param name="cancellationToken">Cancellation token for the call</param>
+		/// <returns>The updated log file document</returns>
+		Task<ILog?> TryCompleteChunksAsync(ILog log, IEnumerable<CompleteLogChunkUpdate> chunks, CancellationToken cancellationToken);
+
+		/// <summary>
+		/// Update the log file with final information about the index
+		/// </summary>
+		/// <param name="log">The current log file</param>
+		/// <param name="newIndexLength">New length of the index</param>
+		/// <param name="cancellationToken">Cancellation token for the call</param>
+		/// <returns>The updated log file document</returns>
+		Task<ILog?> TryUpdateIndexAsync(ILog log, long newIndexLength, CancellationToken cancellationToken);
+
+		/// <summary>
+		/// Gets a log by ID
+		/// </summary>
+		/// <param name="logId">Unique id of the log file</param>
+		/// <param name="cancellationToken">Cancellation token for the call</param>
+		/// <returns>The log instance</returns>
+		Task<ILog?> GetLogAsync(LogId logId, CancellationToken cancellationToken);
+
+		/// <summary>
+		/// Gets all the log files
+		/// </summary>
+		/// <param name="index">Index of the first result to return</param>
+		/// <param name="count">Number of results to return</param>
+		/// <param name="cancellationToken">Cancellation token for the call</param>
+		/// <returns>List of log files</returns>
+		Task<List<ILog>> GetLogsAsync(int? index = null, int? count = null, CancellationToken cancellationToken = default);
+	}
+}
diff --git a/Engine/Source/Programs/Horde/Horde.Server/Logs/ILogFile.cs b/Engine/Source/Programs/Horde/Horde.Server/Logs/ILogFile.cs
deleted file mode 100644
diff --git a/Engine/Source/Programs/Horde/Horde.Server/Logs/ILogFileCollection.cs b/Engine/Source/Programs/Horde/Horde.Server/Logs/ILogFileCollection.cs
deleted file mode 100644
diff --git a/Engine/Source/Programs/Horde/Horde.Server/Logs/LogCollection.cs b/Engine/Source/Programs/Horde/Horde.Server/Logs/LogCollection.cs
new file mode 100644
--- /dev/null
+++ b/Engine/Source/Programs/Horde/Horde.Server/Logs/LogCollection.cs
@@ -0,0 +1,259 @@
+// Copyright Epic Games, Inc. All Rights Reserved.
+
+using System.Collections.Generic;
+using System.Linq;
+using System.Net;
+using System.Threading;
+using System.Threading.Tasks;
+using EpicGames.Horde.Agents.Leases;
+using EpicGames.Horde.Agents.Sessions;
+using EpicGames.Horde.Jobs;
+using EpicGames.Horde.Logs;
+using EpicGames.Horde.Storage;
+using Horde.Server.Server;
+using Horde.Server.Storage;
+using MongoDB.Bson.Serialization.Attributes;
+using MongoDB.Driver;
+
+namespace Horde.Server.Logs
+{
+	/// <summary>
+	/// Wrapper around the jobs collection in a mongo DB
+	/// </summary>
+	public class LogCollection : ILogCollection
+	{
+		class LogChunkDocument : ILogChunk
+		{
+			public long Offset { get; set; }
+			public int Length { get; set; }
+			public int LineIndex { get; set; }
+
+			[BsonIgnoreIfNull]
+			public string? Server { get; set; }
+
+			[BsonConstructor]
+			public LogChunkDocument()
+			{
+			}
+
+			public LogChunkDocument(LogChunkDocument other)
+			{
+				Offset = other.Offset;
+				Length = other.Length;
+				LineIndex = other.LineIndex;
+				Server = other.Server;
+			}
+
+			public LogChunkDocument Clone()
+			{
+				return (LogChunkDocument)MemberwiseClone();
+			}
+		}
+
+		class LogDocument : ILog
+		{
+			[BsonRequired, BsonId]
+			public LogId Id { get; set; }
+
+			[BsonRequired]
+			public JobId JobId { get; set; }
+
+			[BsonIgnoreIfNull]
+			public LeaseId? LeaseId { get; set; }
+
+			[BsonIgnoreIfNull]
+			public SessionId? SessionId { get; set; }
+
+			public LogType Type { get; set; }
+			public bool UseNewStorageBackend { get; set; }
+
+			[BsonIgnoreIfNull]
+			public int? MaxLineIndex { get; set; }
+
+			[BsonIgnoreIfNull]
+			public long? IndexLength { get; set; }
+
+			public List<LogChunkDocument> Chunks { get; set; } = new List<LogChunkDocument>();
+
+			public int LineCount { get; set; }
+
+			public NamespaceId NamespaceId { get; set; } = Namespace.Logs;
+			public RefName RefName { get; set; }
+
+			[BsonIgnoreIfDefault]
+			public bool Complete { get; set; }
+
+			[BsonRequired]
+			public int UpdateIndex { get; set; }
+
+			IReadOnlyList<ILogChunk> ILog.Chunks => Chunks;
+
+			[BsonConstructor]
+			private LogDocument()
+			{
+			}
+
+			public LogDocument(JobId jobId, LeaseId? leaseId, SessionId? sessionId, LogType type, LogId? logId, NamespaceId namespaceId)
+			{
+				Id = logId ?? LogIdUtils.GenerateNewId();
+				JobId = jobId;
+				LeaseId = leaseId;
+				SessionId = sessionId;
+				Type = type;
+				UseNewStorageBackend = true;
+				MaxLineIndex = 0;
+				NamespaceId = namespaceId;
+				RefName = new RefName(Id.ToString());
+			}
+
+			public LogDocument Clone()
+			{
+				LogDocument document = (LogDocument)MemberwiseClone();
+				document.Chunks = document.Chunks.ConvertAll(x => x.Clone());
+				return document;
+			}
+		}
+
+		/// <summary>
+		/// The jobs collection
+		/// </summary>
+		readonly IMongoCollection<LogDocument> _logCollection;
+
+		/// <summary>
+		/// Hostname for the current server
+		/// </summary>
+		readonly string _hostName;
+
+		/// <summary>
+		/// Constructor
+		/// </summary>
+		/// <param name="mongoService">The database service singleton</param>
+		public LogCollection(MongoService mongoService)
+		{
+			_logCollection = mongoService.GetCollection<LogDocument>("LogFiles");
+			_hostName = Dns.GetHostName();
+		}
+
+		/// <inheritdoc/>
+		public async Task<ILog> CreateLogAsync(JobId jobId, LeaseId? leaseId, SessionId? sessionId, LogType type, LogId? logId, CancellationToken cancellationToken)
+		{
+			LogDocument newLog = new LogDocument(jobId, leaseId, sessionId, type, logId, Namespace.Logs);
+			await _logCollection.InsertOneAsync(newLog, null, cancellationToken);
+			return newLog;
+		}
+
+		/// <inheritdoc/>
+		public async Task<ILog> UpdateLineCountAsync(ILog log, int lineCount, bool complete, CancellationToken cancellationToken)
+		{
+			FilterDefinition<LogDocument> filter = Builders<LogDocument>.Filter.Eq(x => x.Id, log.Id);
+			UpdateDefinition<LogDocument> update = Builders<LogDocument>.Update.Set(x => x.LineCount, lineCount).Set(x => x.Complete, complete).Inc(x => x.UpdateIndex, 1);
+			return await _logCollection.FindOneAndUpdateAsync(filter, update, new FindOneAndUpdateOptions<LogDocument, LogDocument> { ReturnDocument = ReturnDocument.After }, cancellationToken);
+		}
+
+		/// <inheritdoc/>
+		public async Task<ILog?> TryAddChunkAsync(ILog logInterface, long offset, int lineIndex, CancellationToken cancellationToken)
+		{
+			LogDocument log = ((LogDocument)logInterface).Clone();
+
+			int chunkIdx = log.Chunks.GetChunkForOffset(offset) + 1;
+
+			LogChunkDocument chunk = new LogChunkDocument();
+			chunk.Offset = offset;
+			chunk.LineIndex = lineIndex;
+			chunk.Server = _hostName;
+			log.Chunks.Insert(chunkIdx, chunk);
+
+			UpdateDefinition<LogDocument> update = Builders<LogDocument>.Update.Set(x => x.Chunks, log.Chunks);
+			if (chunkIdx == log.Chunks.Count - 1)
+			{
+				log.MaxLineIndex = null;
+				update = update.Unset(x => x.MaxLineIndex);
+			}
+
+			if (!await TryUpdateLogAsync(log, update, cancellationToken))
+			{
+				return null;
+			}
+
+			return log;
+		}
+
+		/// <inheritdoc/>
+		public async Task<ILog?> TryCompleteChunksAsync(ILog logInterface, IEnumerable<CompleteLogChunkUpdate> chunkUpdates, CancellationToken cancellationToken)
+		{
+			LogDocument log = ((LogDocument)logInterface).Clone();
+
+			// Update the length of any complete chunks
+			UpdateDefinitionBuilder<LogDocument> updateBuilder = Builders<LogDocument>.Update;
+			List<UpdateDefinition<LogDocument>> updates = new List<UpdateDefinition<LogDocument>>();
+			foreach (CompleteLogChunkUpdate chunkUpdate in chunkUpdates)
+			{
+				LogChunkDocument chunk = log.Chunks[chunkUpdate.Index];
+				chunk.Length = chunkUpdate.Length;
+				updates.Add(updateBuilder.Set(x => x.Chunks[chunkUpdate.Index].Length, chunkUpdate.Length));
+
+				if (chunkUpdate.Index == log.Chunks.Count - 1)
+				{
+					log.MaxLineIndex = chunk.LineIndex + chunkUpdate.LineCount;
+					updates.Add(updateBuilder.Set(x => x.MaxLineIndex, log.MaxLineIndex));
+				}
+			}
+
+			// Try to apply the updates
+			if (updates.Count > 0 && !await TryUpdateLogAsync(log, updateBuilder.Combine(updates), cancellationToken))
+			{
+				return null;
+			}
+
+			return log;
+		}
+
+		/// <inheritdoc/>
+		public async Task<ILog?> TryUpdateIndexAsync(ILog logInterface, long newIndexLength, CancellationToken cancellationToken)
+		{
+			LogDocument log = ((LogDocument)logInterface).Clone();
+
+			UpdateDefinition<LogDocument> update = Builders<LogDocument>.Update.Set(x => x.IndexLength, newIndexLength);
+			if (!await TryUpdateLogAsync(log, update, cancellationToken))
+			{
+				return null;
+			}
+
+			log.IndexLength = newIndexLength;
+			return log;
+		}
+
+		/// <inheritdoc/>
+		private async Task<bool> TryUpdateLogAsync(LogDocument current, UpdateDefinition<LogDocument> update, CancellationToken cancellationToken)
+		{
+			int prevUpdateIndex = current.UpdateIndex;
+			current.UpdateIndex++;
+			UpdateResult result = await _logCollection.UpdateOneAsync<LogDocument>(x => x.Id == current.Id && x.UpdateIndex == prevUpdateIndex, update.Set(x => x.UpdateIndex, current.UpdateIndex), cancellationToken: cancellationToken);
+			return result.ModifiedCount == 1;
+		}
+
+		/// <inheritdoc/>
+		public async Task<ILog?> GetLogAsync(LogId logId, CancellationToken cancellationToken)
+		{
+			LogDocument log = await _logCollection.Find<LogDocument>(x => x.Id == logId).FirstOrDefaultAsync(cancellationToken);
+			return log;
+		}
+
+		/// <inheritdoc/>
+		public async Task<List<ILog>> GetLogsAsync(int? index = null, int? count = null, CancellationToken cancellationToken = default)
+		{
+			IFindFluent<LogDocument, LogDocument> query = _logCollection.Find(FilterDefinition<LogDocument>.Empty);
+			if (index != null)
+			{
+				query = query.Skip(index.Value);
+			}
+			if (count != null)
+			{
+				query = query.Limit(count.Value);
+			}
+
+			List<LogDocument> results = await query.ToListAsync(cancellationToken);
+			return results.ConvertAll<ILog>(x => x);
+		}
+	}
+}
diff --git a/Engine/Source/Programs/Horde/Horde.Server/Logs/LogFileCollection.cs b/Engine/Source/Programs/Horde/Horde.Server/Logs/LogFileCollection.cs
deleted file mode 100644
diff --git a/Engine/Source/Programs/Horde/Horde.Server/Logs/LogFileService.cs b/Engine/Source/Programs/Horde/Horde.Server/Logs/LogFileService.cs
deleted file mode 100644
diff --git a/Engine/Source/Programs/Horde/Horde.Server/Logs/LogRpcService.cs b/Engine/Source/Programs/Horde/Horde.Server/Logs/LogRpcService.cs
--- a/Engine/Source/Programs/Horde/Horde.Server/Logs/LogRpcService.cs
+++ b/Engine/Source/Programs/Horde/Horde.Server/Logs/LogRpcService.cs
@@ -21,8 +21,8 @@ namespace Horde.Server.Logs
 	[Authorize]
 	public class LogRpcService : LogRpc.LogRpcBase
 	{
-		readonly ILogFileService _logFileService;
-		readonly ILogFileCollection _logFileCollection;
+		readonly ILogService _logService;
+		readonly ILogCollection _logCollection;
 		readonly LogTailService _logTailService;
 		readonly StorageService _storageService;
 		readonly ILogger<LogRpcService> _logger;
@@ -30,10 +30,10 @@ namespace Horde.Server.Logs
 		/// <summary>
 		/// Constructor
 		/// </summary>
-		public LogRpcService(ILogFileService logFileService, ILogFileCollection logFileCollection, LogTailService logTailService, StorageService storageService, ILogger<LogRpcService> logger)
+		public LogRpcService(ILogService logService, ILogCollection logCollection, LogTailService logTailService, StorageService storageService, ILogger<LogRpcService> logger)
 		{
-			_logFileService = logFileService;
-			_logFileCollection = logFileCollection;
+			_logService = logService;
+			_logCollection = logCollection;
 			_logTailService = logTailService;
 			_storageService = storageService;
 			_logger = logger;
@@ -42,12 +42,12 @@ namespace Horde.Server.Logs
 		/// <inheritdoc/>
 		public override async Task<UpdateLogResponse> UpdateLog(UpdateLogRequest request, ServerCallContext context)
 		{
-			ILogFile? logFile = await _logFileService.GetLogFileAsync(LogId.Parse(request.LogId), context.CancellationToken);
-			if (logFile == null)
+			ILog? log = await _logService.GetLogAsync(LogId.Parse(request.LogId), context.CancellationToken);
+			if (log == null)
 			{
 				throw new StructuredRpcException(StatusCode.NotFound, "Resource not found");
 			}
-			if (!LogFileService.AuthorizeForSession(logFile, context.GetHttpContext().User))
+			if (!LogService.AuthorizeForSession(log, context.GetHttpContext().User))
 			{
 				throw new StructuredRpcException(StatusCode.PermissionDenied, "Access denied");
 			}
@@ -65,9 +65,9 @@ namespace Horde.Server.Logs
 			IBlobRef target = store.CreateBlobRef(hash, new BlobLocator(request.TargetLocator));
 			await store.WriteRefAsync(new RefName(request.LogId), target);
 
-			await _logFileCollection.UpdateLineCountAsync(logFile, request.LineCount, request.Complete, CancellationToken.None);
+			await _logCollection.UpdateLineCountAsync(log, request.LineCount, request.Complete, CancellationToken.None);
 
-			await _logTailService.FlushAsync(logFile.Id, request.LineCount);
+			await _logTailService.FlushAsync(log.Id, request.LineCount);
 
 			return new UpdateLogResponse();
 		}
diff --git a/Engine/Source/Programs/Horde/Horde.Server/Logs/LogService.cs b/Engine/Source/Programs/Horde/Horde.Server/Logs/LogService.cs
new file mode 100644
--- /dev/null
+++ b/Engine/Source/Programs/Horde/Horde.Server/Logs/LogService.cs
@@ -0,0 +1,1249 @@
+// Copyright Epic Games, Inc. All Rights Reserved.
+
+using System;
+using System.Collections.Generic;
+using System.Diagnostics;
+using System.IO;
+using System.Linq;
+using System.Runtime.CompilerServices;
+using System.Security.Claims;
+using System.Text;
+using System.Text.Json;
+using System.Threading;
+using System.Threading.Tasks;
+using EpicGames.Core;
+using EpicGames.Horde.Agents.Leases;
+using EpicGames.Horde.Agents.Sessions;
+using EpicGames.Horde.Jobs;
+using EpicGames.Horde.Logs;
+using EpicGames.Horde.Storage;
+using Horde.Server.Acls;
+using Horde.Server.Logs.Data;
+using Horde.Server.Storage;
+using Horde.Server.Utilities;
+using Microsoft.Extensions.Caching.Memory;
+using Microsoft.Extensions.Logging;
+using Microsoft.Extensions.Options;
+using MongoDB.Bson;
+using OpenTelemetry.Trace;
+
+namespace Horde.Server.Logs
+{
+	using Stream = System.IO.Stream;
+
+	/// <summary>
+	/// Metadata about a log file
+	/// </summary>
+	public class LogMetadata
+	{
+		/// <summary>
+		/// Length of the log file
+		/// </summary>
+		public long Length { get; set; }
+
+		/// <summary>
+		/// Number of lines in the log file
+		/// </summary>
+		public int MaxLineIndex { get; set; }
+	}
+
+	/// <summary>
+	/// Interface for the log file service
+	/// </summary>
+	public interface ILogService
+	{
+		/// <summary>
+		/// Creates a new log
+		/// </summary>
+		/// <param name="jobId">Unique id of the job that owns this log file</param>
+		/// <param name="leaseId">Lease allowed to update the log</param>
+		/// <param name="sessionId">Agent session allowed to update the log</param>
+		/// <param name="type">Type of events to be stored in the log</param>
+		/// <param name="cancellationToken">Cancellation token for the call</param>
+		/// <param name="logId">ID of the log file (optional)</param>
+		/// <returns>The new log file document</returns>
+		Task<ILog> CreateLogAsync(JobId jobId, LeaseId? leaseId, SessionId? sessionId, LogType type, LogId? logId = null, CancellationToken cancellationToken = default);
+
+		/// <summary>
+		/// Gets a log by ID
+		/// </summary>
+		/// <param name="logId">Unique id of the log file</param>
+		/// <param name="cancellationToken">Cancellation token for the call</param>
+		/// <returns>The log  document</returns>
+		Task<ILog?> GetLogAsync(LogId logId, CancellationToken cancellationToken);
+
+		/// <summary>
+		/// Returns a list of log files
+		/// </summary>
+		/// <param name="index">Index of the first result to return</param>
+		/// <param name="count">Number of results to return</param>
+		/// <param name="cancellationToken">Cancellation token for the call</param>
+		/// <returns>List of log documents</returns>
+		Task<List<ILog>> GetLogsAsync(int? index = null, int? count = null, CancellationToken cancellationToken = default);
+
+		/// <summary>
+		/// Read a set of lines from the given log file
+		/// </summary>
+		/// <param name="log">Log file to read</param>
+		/// <param name="index">Index of the first line to read</param>
+		/// <param name="count">Maximum number of lines to return</param>
+		/// <param name="cancellationToken">Cancellation token for the operation</param>
+		/// <returns>List of lines</returns>
+		Task<List<Utf8String>> ReadLinesAsync(ILog log, int index, int count, CancellationToken cancellationToken = default);
+
+		/// <summary>
+		/// Gets metadata about the log file
+		/// </summary>
+		/// <param name="log">The log file to query</param>
+		/// <param name="cancellationToken">Cancellation token for the call</param>
+		/// <returns>Metadata about the log file</returns>
+		Task<LogMetadata> GetMetadataAsync(ILog log, CancellationToken cancellationToken);
+
+		/// <summary>
+		/// Creates new log events
+		/// </summary>
+		/// <param name="newEvents">List of events</param>
+		/// <param name="cancellationToken">Cancellation token for the call</param>
+		/// <returns>Async task</returns>
+		Task CreateEventsAsync(List<NewLogEventData> newEvents, CancellationToken cancellationToken);
+
+		/// <summary>
+		/// Find events for a particular log file
+		/// </summary>
+		/// <param name="log">The log file instance</param>
+		/// <param name="spanId">Issue span to return events for</param>
+		/// <param name="index">Index of the first event to retrieve</param>
+		/// <param name="count">Number of events to retrieve</param>
+		/// <param name="cancellationToken">Cancellation token for the call</param>
+		/// <returns>List of log events</returns>
+		Task<List<ILogEvent>> FindEventsAsync(ILog log, ObjectId? spanId = null, int? index = null, int? count = null, CancellationToken cancellationToken = default);
+
+		/// <summary>
+		/// Adds events to a log span
+		/// </summary>
+		/// <param name="events">The events to add</param>
+		/// <param name="spanId">The span id</param>
+		/// <param name="cancellationToken">Cancellation token for the call</param>
+		/// <returns>Async task</returns>
+		Task AddSpanToEventsAsync(IEnumerable<ILogEvent> events, ObjectId spanId, CancellationToken cancellationToken = default);
+
+		/// <summary>
+		/// Find events for an issue
+		/// </summary>
+		/// <param name="spanIds">The span ids</param>
+		/// <param name="logIds">Log ids to include</param>
+		/// <param name="index">Index within the events for results to return</param>
+		/// <param name="count">Number of results to return</param>
+		/// <param name="cancellationToken">Cancellation token for the call</param>
+		/// <returns>Async task</returns>
+		Task<List<ILogEvent>> FindEventsForSpansAsync(IEnumerable<ObjectId> spanIds, LogId[]? logIds, int index, int count, CancellationToken cancellationToken = default);
+
+		/// <summary>
+		/// Gets the data for an event
+		/// </summary>
+		/// <param name="log">The log file instance</param>
+		/// <param name="lineIndex">Index of the line in the file</param>
+		/// <param name="lineCount">Number of lines in the event</param>
+		/// <param name="cancellationToken">Cancellation token for the call</param>
+		/// <returns>New event data instance</returns>
+		Task<ILogEventData> GetEventDataAsync(ILog log, int lineIndex, int lineCount, CancellationToken cancellationToken = default);
+
+		/// <summary>
+		/// Gets lines from the given log 
+		/// </summary>
+		/// <param name="log">The log file</param>
+		/// <param name="cancellationToken">Cancellation token for the call</param>
+		/// <returns>Data for the requested range</returns>
+		Task<Stream> OpenRawStreamAsync(ILog log, CancellationToken cancellationToken = default);
+
+		/// <summary>
+		/// Parses a stream of json text and outputs plain text
+		/// </summary>
+		/// <param name="log">The log file to query</param>
+		/// <param name="outputStream">Output stream to receive the text data</param>
+		/// <param name="cancellationToken">Cancellation token for the call</param>
+		/// <returns>Async text</returns>
+		Task CopyPlainTextStreamAsync(ILog log, Stream outputStream, CancellationToken cancellationToken = default);
+
+		/// <summary>
+		/// Search for the specified text in a log file
+		/// </summary>
+		/// <param name="log">The log file to search</param>
+		/// <param name="text">Text to search for</param>
+		/// <param name="firstLine">Line to start search from</param>
+		/// <param name="count">Number of results to return</param>
+		/// <param name="stats">Receives stats for the search</param>
+		/// <param name="cancellationToken">Cancellation token for the call</param>
+		/// <returns>List of line numbers containing the given term</returns>
+		Task<List<int>> SearchLogDataAsync(ILog log, string text, int firstLine, int count, SearchStats stats, CancellationToken cancellationToken);
+	}
+
+	/// <summary>
+	/// Extension methods for dealing with log files
+	/// </summary>
+	public static class LogServiceExtensions
+	{
+		/// <summary>
+		/// Parses a stream of json text and outputs plain text
+		/// </summary>
+		/// <param name="logService">The log file service</param>
+		/// <param name="log">The log file to query</param>
+		/// <param name="outputStream">Output stream to receive the text data</param>
+		/// <param name="cancellationToken">Cancellation token for the call</param>
+		/// <returns>Async text</returns>
+		public static async Task CopyRawStreamAsync(this ILogService logService, ILog log, Stream outputStream, CancellationToken cancellationToken)
+		{
+			await using Stream stream = await logService.OpenRawStreamAsync(log, cancellationToken);
+			await stream.CopyToAsync(outputStream, cancellationToken);
+		}
+	}
+
+	/// <summary>
+	/// Wraps functionality for manipulating logs
+	/// </summary>
+	public sealed class LogService : ILogService, IDisposable
+	{
+		private readonly Tracer _tracer;
+		private readonly ILogger<LogService> _logger;
+		private readonly ILogCollection _logCollection;
+		private readonly ILogEventCollection _logEventCollection;
+		private readonly ILogStorage _storage;
+		private readonly StorageService _storageService;
+		private readonly IOptions<ServerSettings> _settings;
+		private readonly IMemoryCache _logCache;
+
+		/// <summary>
+		/// Streams log data to a caller
+		/// </summary>
+		class ResponseStream : Stream
+		{
+			/// <summary>
+			/// The log file service that created this stream
+			/// </summary>
+			readonly LogService _logService;
+
+			/// <summary>
+			/// The log file being read
+			/// </summary>
+			readonly ILog _log;
+
+			/// <summary>
+			/// Starting offset within the file of the data to return 
+			/// </summary>
+			readonly long _responseOffset;
+
+			/// <summary>
+			/// Length of data to return
+			/// </summary>
+			readonly long _responseLength;
+
+			/// <summary>
+			/// Current offset within the stream
+			/// </summary>
+			long _currentOffset;
+
+			/// <summary>
+			/// The current chunk index
+			/// </summary>
+			int _chunkIdx;
+
+			/// <summary>
+			/// Buffer containing a message for missing data
+			/// </summary>
+			ReadOnlyMemory<byte> _sourceBuffer;
+
+			/// <summary>
+			/// Offset within the source buffer
+			/// </summary>
+			int _sourcePos;
+
+			/// <summary>
+			/// Length of the source buffer being copied from
+			/// </summary>
+			int _sourceEnd;
+
+			/// <summary>
+			/// Constructor
+			/// </summary>
+			/// <param name="logService">The log file service, for q</param>
+			/// <param name="log"></param>
+			/// <param name="offset"></param>
+			/// <param name="length"></param>
+			public ResponseStream(LogService logService, ILog log, long offset, long length)
+			{
+				_logService = logService;
+				_log = log;
+
+				_responseOffset = offset;
+				_responseLength = length;
+
+				_currentOffset = offset;
+
+				_chunkIdx = log.Chunks.GetChunkForOffset(offset);
+				_sourceBuffer = null!;
+			}
+
+			/// <inheritdoc/>
+			public override bool CanRead => true;
+
+			/// <inheritdoc/>
+			public override bool CanSeek => false;
+
+			/// <inheritdoc/>
+			public override bool CanWrite => false;
+
+			/// <inheritdoc/>
+			public override long Length => _responseLength;
+
+			/// <inheritdoc/>
+			public override long Position
+			{
+				get => _currentOffset - _responseOffset;
+				set => throw new NotImplementedException();
+			}
+
+			/// <inheritdoc/>
+			public override void Flush()
+			{
+			}
+
+			/// <inheritdoc/>
+			public override int Read(byte[] buffer, int offset, int count)
+			{
+#pragma warning disable VSTHRD002
+				return ReadAsync(buffer, offset, count, CancellationToken.None).Result;
+#pragma warning restore VSTHRD002
+			}
+
+			/// <inheritdoc/>
+			public override async Task<int> ReadAsync(byte[] buffer, int offset, int length, CancellationToken cancellationToken)
+			{
+				return await ReadAsync(buffer.AsMemory(offset, length), cancellationToken);
+			}
+
+			/// <inheritdoc/>
+			public override async ValueTask<int> ReadAsync(Memory<byte> buffer, CancellationToken cancellationToken)
+			{
+				int readBytes = 0;
+				while (readBytes < buffer.Length)
+				{
+					if (_sourcePos < _sourceEnd)
+					{
+						// Try to copy from the current buffer
+						int blockSize = Math.Min(_sourceEnd - _sourcePos, buffer.Length - readBytes);
+						_sourceBuffer.Slice(_sourcePos, blockSize).Span.CopyTo(buffer.Slice(readBytes).Span);
+						_currentOffset += blockSize;
+						readBytes += blockSize;
+						_sourcePos += blockSize;
+					}
+					else if (_currentOffset < _responseOffset + _responseLength)
+					{
+						// Move to the right chunk
+						while (_chunkIdx + 1 < _log.Chunks.Count && _currentOffset >= _log.Chunks[_chunkIdx + 1].Offset)
+						{
+							_chunkIdx++;
+						}
+
+						// Get the chunk data
+						ILogChunk chunk = _log.Chunks[_chunkIdx];
+						LogChunkData chunkData = await _logService.ReadChunkAsync(_log, _chunkIdx);
+
+						// Figure out which sub-chunk to use
+						int subChunkIdx = chunkData.GetSubChunkForOffsetWithinChunk((int)(_currentOffset - chunk.Offset));
+						LogSubChunkData subChunkData = chunkData.SubChunks[subChunkIdx];
+
+						// Get the source data
+						long subChunkOffset = chunk.Offset + chunkData.SubChunkOffset[subChunkIdx];
+						_sourceBuffer = subChunkData.InflateText().Data;
+						_sourcePos = (int)(_currentOffset - subChunkOffset);
+						_sourceEnd = (int)Math.Min(_sourceBuffer.Length, (_responseOffset + _responseLength) - subChunkOffset);
+					}
+					else
+					{
+						// End of the log
+						break;
+					}
+				}
+				return readBytes;
+			}
+
+			/// <inheritdoc/>
+			public override long Seek(long offset, SeekOrigin origin) => throw new NotImplementedException();
+
+			/// <inheritdoc/>
+			public override void SetLength(long value) => throw new NotImplementedException();
+
+			/// <inheritdoc/>
+			public override void Write(byte[] buffer, int offset, int count) => throw new NotImplementedException();
+		}
+
+		/// <summary>
+		/// Streams log data to a caller
+		/// </summary>
+		class NewLoggerResponseStream : Stream
+		{
+			readonly LogNode _rootNode;
+
+			/// <summary>
+			/// Starting offset within the file of the data to return 
+			/// </summary>
+			readonly long _responseOffset;
+
+			/// <summary>
+			/// Length of data to return
+			/// </summary>
+			readonly long _responseLength;
+
+			/// <summary>
+			/// Current offset within the stream
+			/// </summary>
+			long _currentOffset;
+
+			/// <summary>
+			/// The current chunk index
+			/// </summary>
+			int _chunkIdx;
+
+			/// <summary>
+			/// Buffer containing a message for missing data
+			/// </summary>
+			ReadOnlyMemory<byte> _sourceBuffer;
+
+			/// <summary>
+			/// Offset within the source buffer
+			/// </summary>
+			int _sourcePos;
+
+			/// <summary>
+			/// Length of the source buffer being copied from
+			/// </summary>
+			int _sourceEnd;
+
+			/// <summary>
+			/// Constructor
+			/// </summary>
+			public NewLoggerResponseStream(LogNode rootNode, long offset, long length)
+			{
+				_rootNode = rootNode;
+
+				_responseOffset = offset;
+				_responseLength = length;
+
+				_currentOffset = offset;
+
+				_chunkIdx = rootNode.TextChunkRefs.GetChunkForOffset(offset);
+				_sourceBuffer = null!;
+			}
+
+			/// <inheritdoc/>
+			public override bool CanRead => true;
+
+			/// <inheritdoc/>
+			public override bool CanSeek => false;
+
+			/// <inheritdoc/>
+			public override bool CanWrite => false;
+
+			/// <inheritdoc/>
+			public override long Length => _responseLength;
+
+			/// <inheritdoc/>
+			public override long Position
+			{
+				get => _currentOffset - _responseOffset;
+				set => throw new NotImplementedException();
+			}
+
+			/// <inheritdoc/>
+			public override void Flush()
+			{
+			}
+
+			/// <inheritdoc/>
+			public override int Read(byte[] buffer, int offset, int count)
+			{
+#pragma warning disable VSTHRD002
+				return ReadAsync(buffer, offset, count, CancellationToken.None).Result;
+#pragma warning restore VSTHRD002
+			}
+
+			/// <inheritdoc/>
+			public override async Task<int> ReadAsync(byte[] buffer, int offset, int length, CancellationToken cancellationToken)
+			{
+				return await ReadAsync(buffer.AsMemory(offset, length), cancellationToken);
+			}
+
+			/// <inheritdoc/>
+			public override async ValueTask<int> ReadAsync(Memory<byte> buffer, CancellationToken cancellationToken)
+			{
+				int readBytes = 0;
+				while (readBytes < buffer.Length)
+				{
+					if (_sourcePos < _sourceEnd)
+					{
+						// Try to copy from the current buffer
+						int blockSize = Math.Min(_sourceEnd - _sourcePos, buffer.Length - readBytes);
+						_sourceBuffer.Slice(_sourcePos, blockSize).Span.CopyTo(buffer.Slice(readBytes).Span);
+						_currentOffset += blockSize;
+						readBytes += blockSize;
+						_sourcePos += blockSize;
+					}
+					else if (_currentOffset < _responseOffset + _responseLength)
+					{
+						// Move to the right chunk
+						while (_chunkIdx + 1 < _rootNode.TextChunkRefs.Count && _currentOffset >= _rootNode.TextChunkRefs[_chunkIdx + 1].Offset)
+						{
+							_chunkIdx++;
+						}
+
+						// Get the chunk data
+						LogChunkRef chunk = _rootNode.TextChunkRefs[_chunkIdx];
+						LogChunkNode chunkNode = await chunk.Target.ReadBlobAsync(cancellationToken: cancellationToken);
+
+						// Get the source data
+						_sourceBuffer = chunkNode.Data;
+						_sourcePos = (int)(_currentOffset - chunk.Offset);
+						_sourceEnd = (int)Math.Min(_sourceBuffer.Length, (_responseOffset + _responseLength) - chunk.Offset);
+					}
+					else
+					{
+						// End of the log
+						break;
+					}
+				}
+				return readBytes;
+			}
+
+			/// <inheritdoc/>
+			public override long Seek(long offset, SeekOrigin origin) => throw new NotImplementedException();
+
+			/// <inheritdoc/>
+			public override void SetLength(long value) => throw new NotImplementedException();
+
+			/// <inheritdoc/>
+			public override void Write(byte[] buffer, int offset, int count) => throw new NotImplementedException();
+		}
+
+		readonly LogTailService _logTailService;
+
+		/// <summary>
+		/// Constructor
+		/// </summary>
+		public LogService(ILogCollection logCollection, ILogEventCollection logEventCollection, ILogStorage storage, LogTailService logTailService, StorageService storageService, IOptions<ServerSettings> settings, Tracer tracer, ILogger<LogService> logger)
+		{
+			_logCollection = logCollection;
+			_logEventCollection = logEventCollection;
+			_logCache = new MemoryCache(new MemoryCacheOptions());
+			_storage = storage;
+			_logTailService = logTailService;
+			_storageService = storageService;
+			_settings = settings;
+			_tracer = tracer;
+			_logger = logger;
+		}
+
+		/// <inheritdoc/>
+		public void Dispose()
+		{
+			_logCache.Dispose();
+			_storage.Dispose();
+		}
+
+		/// <inheritdoc/>
+		public Task<ILog> CreateLogAsync(JobId jobId, LeaseId? leaseId, SessionId? sessionId, LogType type, LogId? logId, CancellationToken cancellationToken)
+		{
+			return _logCollection.CreateLogAsync(jobId, leaseId, sessionId, type, logId, cancellationToken);
+		}
+
+		/// <inheritdoc/>
+		public async Task<ILog?> GetLogAsync(LogId logId, CancellationToken cancellationToken)
+		{
+			return await _logCollection.GetLogAsync(logId, cancellationToken);
+		}
+
+		/// <inheritdoc/>
+		public Task<List<ILog>> GetLogsAsync(int? index = null, int? count = null, CancellationToken cancellationToken = default)
+		{
+			return _logCollection.GetLogsAsync(index, count, cancellationToken);
+		}
+
+		/// <inheritdoc/>
+		public async Task<List<Utf8String>> ReadLinesAsync(ILog log, int index, int count, CancellationToken cancellationToken)
+		{
+			List<Utf8String> lines = new List<Utf8String>();
+
+			if (log.UseNewStorageBackend)
+			{
+				using IStorageClient storageClient = _storageService.CreateClient(log.NamespaceId);
+
+				int maxIndex = index + count;
+				bool complete = log.Complete;
+
+				LogNode? root = await storageClient.TryReadRefTargetAsync<LogNode>(log.RefName, cancellationToken: cancellationToken);
+				if (root != null)
+				{
+					int chunkIdx = root.TextChunkRefs.GetChunkForLine(index);
+					for (; index < maxIndex && chunkIdx < root.TextChunkRefs.Count; chunkIdx++)
+					{
+						LogChunkRef chunk = root.TextChunkRefs[chunkIdx];
+						LogChunkNode chunkData = await chunk.Target.ReadBlobAsync(cancellationToken: cancellationToken);
+
+						for (; index < maxIndex && index < chunk.LineIndex; index++)
+						{
+							lines.Add(new Utf8String($"Internal error; missing data for line {index}\n"));
+						}
+
+						for (; index < maxIndex && index < chunk.LineIndex + chunk.LineCount; index++)
+						{
+							lines.Add(chunkData.GetLine(index - chunk.LineIndex));
+						}
+					}
+					complete |= root.Complete;
+				}
+
+				if (!complete)
+				{
+					await _logTailService.EnableTailingAsync(log.Id, root?.LineCount ?? 0, cancellationToken);
+					if (index < maxIndex)
+					{
+						await _logTailService.ReadAsync(log.Id, index, maxIndex - index, lines);
+					}
+				}
+			}
+			else
+			{
+				(_, long minOffset) = await GetLineOffsetAsync(log, index, cancellationToken);
+				(_, long maxOffset) = await GetLineOffsetAsync(log, index + Math.Min(count, Int32.MaxValue - index), cancellationToken);
+
+				byte[] result;
+				using (System.IO.Stream stream = await OpenRawStreamAsync(log, minOffset, maxOffset - minOffset, cancellationToken))
+				{
+					result = new byte[stream.Length];
+					await stream.ReadFixedSizeDataAsync(result, 0, result.Length, cancellationToken);
+				}
+
+				int offset = 0;
+				for (int idx = 0; idx < result.Length; idx++)
+				{
+					if (result[idx] == (byte)'\n')
+					{
+						lines.Add(new Utf8String(result.AsMemory(offset, idx - offset)));
+						offset = idx + 1;
+					}
+				}
+			}
+
+			return lines;
+		}
+
+		/// <inheritdoc/>
+		public async Task<LogMetadata> GetMetadataAsync(ILog log, CancellationToken cancellationToken)
+		{
+			LogMetadata metadata = new LogMetadata();
+			if (log.UseNewStorageBackend)
+			{
+				if (log.Complete)
+				{
+					metadata.MaxLineIndex = log.LineCount;
+				}
+				else
+				{
+					metadata.MaxLineIndex = await _logTailService.GetFullLineCountAsync(log.Id, log.LineCount, cancellationToken);
+				}
+			}
+			else
+			{
+				if (log.Chunks.Count > 0)
+				{
+					ILogChunk chunk = log.Chunks[log.Chunks.Count - 1];
+					if (log.MaxLineIndex == null || chunk.Length == 0)
+					{
+						LogChunkData chunkData = await ReadChunkAsync(log, log.Chunks.Count - 1);
+						metadata.Length = chunk.Offset + chunkData.Length;
+						metadata.MaxLineIndex = chunk.LineIndex + chunkData.LineCount;
+					}
+					else
+					{
+						metadata.Length = chunk.Offset + chunk.Length;
+						metadata.MaxLineIndex = log.MaxLineIndex.Value;
+					}
+				}
+			}
+			return metadata;
+		}
+
+		/// <inheritdoc/>
+		public Task CreateEventsAsync(List<NewLogEventData> newEvents, CancellationToken cancellationToken)
+		{
+			return _logEventCollection.AddManyAsync(newEvents);
+		}
+
+		/// <inheritdoc/>
+		public Task<List<ILogEvent>> FindEventsAsync(ILog log, ObjectId? spanId = null, int? index = null, int? count = null, CancellationToken cancellationToken = default)
+		{
+			return _logEventCollection.FindAsync(log.Id, spanId, index, count);
+		}
+
+		class LogEventData : ILogEventData
+		{
+			public string? _message;
+			public IReadOnlyList<JsonLogEvent> Lines { get; }
+
+			EventId? ILogEventData.EventId => (Lines.Count > 0) ? Lines[0].EventId : null;
+			LogEventSeverity ILogEventData.Severity => (Lines.Count == 0) ? LogEventSeverity.Information : (Lines[0].Level == LogLevel.Warning) ? LogEventSeverity.Warning : LogEventSeverity.Error;
+
+			public LogEventData(IReadOnlyList<JsonLogEvent> lines)
+			{
+				Lines = lines;
+			}
+
+			string ILogEventData.Message
+			{
+				get
+				{
+					_message ??= String.Join("\n", Lines.Select(x => x.GetRenderedMessage().ToString()));
+					return _message;
+				}
+			}
+		}
+
+		/// <inheritdoc/>
+		public Task AddSpanToEventsAsync(IEnumerable<ILogEvent> events, ObjectId spanId, CancellationToken cancellationToken)
+		{
+			return _logEventCollection.AddSpanToEventsAsync(events, spanId);
+		}
+
+		/// <inheritdoc/>
+		public Task<List<ILogEvent>> FindEventsForSpansAsync(IEnumerable<ObjectId> spanIds, LogId[]? logIds, int index, int count, CancellationToken cancellationToken)
+		{
+			return _logEventCollection.FindEventsForSpansAsync(spanIds, logIds, index, count);
+		}
+
+		/// <inheritdoc/>
+		public async Task<ILogEventData> GetEventDataAsync(ILog log, int lineIndex, int lineCount, CancellationToken cancellationToken)
+		{
+			using TelemetrySpan span = _tracer.StartActiveSpan($"{nameof(LogService)}.{nameof(GetEventDataAsync)}");
+			span.SetAttribute("logId", log.Id.ToString());
+			span.SetAttribute("lineIndex", lineIndex);
+			span.SetAttribute("lineCount", lineCount);
+
+			List<Utf8String> lines = await ReadLinesAsync(log, lineIndex, lineCount, cancellationToken);
+			List<JsonLogEvent> jsonLines = new List<JsonLogEvent>(lines.Count);
+
+			foreach (Utf8String line in lines)
+			{
+				try
+				{
+					jsonLines.Add(JsonLogEvent.Parse(line.Memory));
+				}
+				catch (JsonException ex)
+				{
+					_logger.LogWarning(ex, "Unable to parse line from log file: {Line}", line);
+				}
+			}
+
+			return new LogEventData(jsonLines);
+		}
+
+		/// <inheritdoc/>
+		public Task<Stream> OpenRawStreamAsync(ILog log, CancellationToken cancellationToken)
+		{
+			return OpenRawStreamAsync(log, 0, Int64.MaxValue, cancellationToken);
+		}
+
+		/// <inheritdoc/>
+		public async Task<Stream> OpenRawStreamAsync(ILog log, long offset, long length, CancellationToken cancellationToken)
+		{
+			if (log.UseNewStorageBackend)
+			{
+				using IStorageClient storageClient = _storageService.CreateClient(log.NamespaceId);
+
+				LogNode? root = await storageClient.TryReadRefTargetAsync<LogNode>(log.RefName, cancellationToken: cancellationToken);
+				if (root == null || root.TextChunkRefs.Count == 0)
+				{
+					return new MemoryStream(Array.Empty<byte>(), false);
+				}
+				else
+				{
+					int lastChunkIdx = root.TextChunkRefs.Count - 1;
+
+					// Clamp the length of the request
+					LogChunkRef lastChunk = root.TextChunkRefs[lastChunkIdx];
+					if (length > lastChunk.Offset)
+					{
+						long lastChunkLength = lastChunk.Length;
+						if (lastChunkLength <= 0)
+						{
+							LogChunkNode lastChunkNode = await lastChunk.Target.ReadBlobAsync(cancellationToken: cancellationToken);
+							lastChunkLength = lastChunkNode.Length;
+						}
+						length = Math.Min(length, (lastChunk.Offset + lastChunkLength) - offset);
+					}
+
+					// Create the new stream
+					return new NewLoggerResponseStream(root, offset, length);
+				}
+			}
+			else
+			{
+				if (log.Chunks.Count == 0)
+				{
+					return new MemoryStream(Array.Empty<byte>(), false);
+				}
+				else
+				{
+					int lastChunkIdx = log.Chunks.Count - 1;
+
+					// Clamp the length of the request
+					ILogChunk lastChunk = log.Chunks[lastChunkIdx];
+					if (length > lastChunk.Offset)
+					{
+						long lastChunkLength = lastChunk.Length;
+						if (lastChunkLength <= 0)
+						{
+							LogChunkData lastChunkData = await ReadChunkAsync(log, lastChunkIdx);
+							lastChunkLength = lastChunkData.Length;
+						}
+						length = Math.Min(length, (lastChunk.Offset + lastChunkLength) - offset);
+					}
+
+					// Create the new stream
+					return new ResponseStream(this, log, offset, length);
+				}
+			}
+		}
+
+		/// <summary>
+		/// Helper method for catching exceptions in <see cref="LogText.ConvertToPlainText(ReadOnlySpan{Byte}, Byte[], Int32)"/>
+		/// </summary>
+		public static int GuardedConvertToPlainText(ReadOnlySpan<byte> input, byte[] output, int outputOffset, ILogger logger)
+		{
+			try
+			{
+				return LogText.ConvertToPlainText(input, output, outputOffset);
+			}
+			catch (Exception ex)
+			{
+				logger.LogWarning(ex, "Unable to convert log line to plain text: {Line}", Encoding.UTF8.GetString(input));
+				output[outputOffset] = (byte)'\n';
+				return outputOffset + 1;
+			}
+		}
+
+		/// <inheritdoc/>
+		public async Task CopyPlainTextStreamAsync(ILog log, Stream outputStream, CancellationToken cancellationToken)
+		{
+			long offset = 0;
+			long length = Int64.MaxValue;
+
+			using (Stream stream = await OpenRawStreamAsync(log, 0, Int64.MaxValue, cancellationToken))
+			{
+				byte[] readBuffer = new byte[4096];
+				int readBufferLength = 0;
+
+				byte[] writeBuffer = new byte[4096];
+				int writeBufferLength = 0;
+
+				while (length > 0)
+				{
+					// Add more data to the buffer
+					int readBytes = await stream.ReadAsync(readBuffer.AsMemory(readBufferLength, readBuffer.Length - readBufferLength), cancellationToken);
+					readBufferLength += readBytes;
+
+					// Copy as many lines as possible to the output
+					int convertedBytes = 0;
+					for (int endIdx = 1; endIdx < readBufferLength; endIdx++)
+					{
+						if (readBuffer[endIdx] == '\n')
+						{
+							writeBufferLength = GuardedConvertToPlainText(readBuffer.AsSpan(convertedBytes, endIdx - convertedBytes), writeBuffer, writeBufferLength, _logger);
+							convertedBytes = endIdx + 1;
+						}
+					}
+
+					// If there's anything in the write buffer, write it out
+					if (writeBufferLength > 0)
+					{
+						if (offset < writeBufferLength)
+						{
+							int writeLength = (int)Math.Min((long)writeBufferLength - offset, length);
+							await outputStream.WriteAsync(writeBuffer.AsMemory((int)offset, writeLength), cancellationToken);
+							length -= writeLength;
+						}
+						offset = Math.Max(offset - writeBufferLength, 0);
+						writeBufferLength = 0;
+					}
+
+					// If we were able to read something, shuffle down the rest of the buffer. Otherwise expand the read buffer.
+					if (convertedBytes > 0)
+					{
+						Buffer.BlockCopy(readBuffer, convertedBytes, readBuffer, 0, readBufferLength - convertedBytes);
+						readBufferLength -= convertedBytes;
+					}
+					else if (readBufferLength > 0)
+					{
+						Array.Resize(ref readBuffer, readBuffer.Length + 128);
+						writeBuffer = new byte[readBuffer.Length];
+					}
+
+					// Exit if we didn't read anything in this iteration
+					if (readBytes == 0)
+					{
+						break;
+					}
+				}
+			}
+		}
+
+		/// <inheritdoc/>
+		public async Task<(int, long)> GetLineOffsetAsync(ILog log, int lineIdx, CancellationToken cancellationToken)
+		{
+			if (log.UseNewStorageBackend)
+			{
+				using IStorageClient storageClient = _storageService.CreateClient(log.NamespaceId);
+
+				LogNode? root = await storageClient.TryReadRefTargetAsync<LogNode>(log.RefName, cancellationToken: cancellationToken);
+				if (root == null)
+				{
+					return (0, 0);
+				}
+
+				int chunkIdx = root.TextChunkRefs.GetChunkForLine(lineIdx);
+				LogChunkRef chunk = root.TextChunkRefs[chunkIdx];
+				LogChunkNode chunkData = await chunk.Target.ReadBlobAsync(cancellationToken: cancellationToken);
+
+				if (lineIdx < chunk.LineIndex)
+				{
+					lineIdx = chunk.LineIndex;
+				}
+
+				int maxLineIndex = chunk.LineIndex + chunkData.LineCount;
+				if (lineIdx >= maxLineIndex)
+				{
+					lineIdx = maxLineIndex;
+				}
+
+				long offset = chunk.Offset + chunkData.LineOffsets[lineIdx - chunk.LineIndex];
+				return (lineIdx, offset);
+			}
+			else
+			{
+				int chunkIdx = log.Chunks.GetChunkForLine(lineIdx);
+
+				ILogChunk chunk = log.Chunks[chunkIdx];
+				LogChunkData chunkData = await ReadChunkAsync(log, chunkIdx);
+
+				if (lineIdx < chunk.LineIndex)
+				{
+					lineIdx = chunk.LineIndex;
+				}
+
+				int maxLineIndex = chunk.LineIndex + chunkData.LineCount;
+				if (lineIdx >= maxLineIndex)
+				{
+					lineIdx = maxLineIndex;
+				}
+
+				long offset = chunk.Offset + chunkData.GetLineOffsetWithinChunk(lineIdx - chunk.LineIndex);
+				return (lineIdx, offset);
+			}
+		}
+
+		/// <summary>
+		/// Reads a chunk from storage
+		/// </summary>
+		/// <param name="log">Log file to read from</param>
+		/// <param name="chunkIdx">The chunk to read</param>
+		/// <returns>Chunk data</returns>
+		private async Task<LogChunkData> ReadChunkAsync(ILog log, int chunkIdx)
+		{
+			ILogChunk chunk = log.Chunks[chunkIdx];
+
+			// Try to read the chunk data from storage
+			LogChunkData? chunkData = null;
+			try
+			{
+				chunkData = await _storage.ReadChunkAsync(log.Id, chunk.Offset, chunk.LineIndex);
+			}
+			catch (Exception ex)
+			{
+				_logger.LogError(ex, "Unable to read log {LogId} at offset {Offset}", log.Id, chunk.Offset);
+			}
+
+			// Get the minimum length and line count for the chunk
+			if (chunkIdx + 1 < log.Chunks.Count)
+			{
+				ILogChunk nextChunk = log.Chunks[chunkIdx + 1];
+				chunkData = RepairChunkData(log, chunkIdx, chunkData, (int)(nextChunk.Offset - chunk.Offset), nextChunk.LineIndex - chunk.LineIndex, $"before next");
+			}
+			else
+			{
+				if (log.MaxLineIndex != null && chunk.Length != 0)
+				{
+					chunkData = RepairChunkData(log, chunkIdx, chunkData, chunk.Length, log.MaxLineIndex.Value - chunk.LineIndex, $"last chunk (max line index = {log.MaxLineIndex})");
+				}
+				else
+				{
+					chunkData ??= RepairChunkData(log, chunkIdx, chunkData, 1024, 1, "default");
+				}
+			}
+
+			return chunkData;
+		}
+
+		/// <summary>
+		/// Validates the given chunk data, and fix it up if necessary
+		/// </summary>
+		/// <param name="log">The log file instance</param>
+		/// <param name="chunkIdx">Index of the chunk within the log</param>
+		/// <param name="chunkData">The chunk data that was read</param>
+		/// <param name="length">Expected length of the data</param>
+		/// <param name="lineCount">Expected number of lines in the data</param>
+		/// <param name="context">Context string for diagnostic output</param>
+		/// <returns>Repaired chunk data</returns>
+		LogChunkData RepairChunkData(ILog log, int chunkIdx, LogChunkData? chunkData, int length, int lineCount, string context)
+		{
+			int currentLength = 0;
+			int currentLineCount = 0;
+			if (chunkData != null)
+			{
+				currentLength = chunkData.Length;
+				currentLineCount = chunkData.LineCount;
+			}
+
+			if (chunkData == null || currentLength < length || currentLineCount < lineCount)
+			{
+				_logger.LogWarning("Creating placeholder subchunk for log {LogId} chunk {ChunkIdx} (length {Length} vs expected {ExpLength}, lines {LineCount} vs expected {ExpLineCount}, context {Context})", log.Id, chunkIdx, currentLength, length, currentLineCount, lineCount, context);
+
+				List<LogSubChunkData> subChunks = new List<LogSubChunkData>();
+				if (chunkData != null && chunkData.Length < length && chunkData.LineCount < lineCount)
+				{
+					subChunks.AddRange(chunkData.SubChunks);
+				}
+
+				LogText text = new LogText();
+				text.AppendMissingDataInfo(chunkIdx, log.Chunks[chunkIdx].Server, length - currentLength, lineCount - currentLineCount);
+				subChunks.Add(new LogSubChunkData(log.Type, currentLength, currentLineCount, text));
+
+				ILogChunk chunk = log.Chunks[chunkIdx];
+				chunkData = new LogChunkData(chunk.Offset, chunk.LineIndex, subChunks);
+			}
+			return chunkData;
+		}
+
+		/// <summary>
+		/// Reads a chunk from storage
+		/// </summary>
+		/// <param name="log">Log file to read from</param>
+		/// <param name="length">Length of the log covered by the index</param>
+		/// <returns>Chunk data</returns>
+		private async Task<LogIndexData?> ReadIndexAsync(ILog log, long length)
+		{
+			try
+			{
+				LogIndexData? index = await _storage.ReadIndexAsync(log.Id, length);
+				return index;
+			}
+			catch (Exception ex)
+			{
+				_logger.LogError(ex, "Unable to read log {LogId} index at length {Length}", log.Id, length);
+				return null;
+			}
+		}
+
+		/// <summary>
+		/// Determines if the user is authorized to perform an action on a particular template
+		/// </summary>
+		/// <param name="log">The template to check</param>
+		/// <param name="user">The principal to authorize</param>
+		/// <returns>True if the action is authorized</returns>
+		public static bool AuthorizeForSession(ILog log, ClaimsPrincipal user)
+		{
+			if (log.SessionId != null && user.HasSessionClaim(log.SessionId.Value))
+			{
+				return true;
+			}
+			if (log.LeaseId != null && user.HasLeaseClaim(log.LeaseId.Value))
+			{
+				return true;
+			}
+			return false;
+		}
+
+		/// <inheritdoc/>
+		public async Task<List<int>> SearchLogDataAsync(ILog log, string text, int firstLine, int count, SearchStats searchStats, CancellationToken cancellationToken)
+		{
+			Stopwatch timer = Stopwatch.StartNew();
+
+			using TelemetrySpan span = _tracer.StartActiveSpan($"{nameof(LogService)}.{nameof(SearchLogDataAsync)}");
+			span.SetAttribute("logId", log.Id.ToString());
+			span.SetAttribute("text", text);
+			span.SetAttribute("count", count);
+
+			List<int> results = new List<int>();
+			if (count > 0)
+			{
+				IAsyncEnumerable<int> enumerable = (log.UseNewStorageBackend) ?
+						SearchLogDataInternalNewAsync(log, text, firstLine, searchStats, cancellationToken) :
+						SearchLogDataInternalAsync(log, text, firstLine, searchStats);
+
+				await using IAsyncEnumerator<int> enumerator = enumerable.GetAsyncEnumerator(cancellationToken);
+				while (await enumerator.MoveNextAsync() && results.Count < count)
+				{
+					results.Add(enumerator.Current);
+				}
+			}
+
+			_logger.LogDebug("Search for \"{SearchText}\" in log {LogId} found {NumResults}/{MaxResults} results, took {Time}ms ({@Stats})", text, log.Id, results.Count, count, timer.ElapsedMilliseconds, searchStats);
+			return results;
+		}
+
+		async IAsyncEnumerable<int> SearchLogDataInternalNewAsync(ILog log, string text, int firstLine, SearchStats searchStats, [EnumeratorCancellation] CancellationToken cancellationToken)
+		{
+			SearchTerm searchText = new SearchTerm(text);
+			using IStorageClient storageClient = _storageService.CreateClient(log.NamespaceId);
+
+			// Search the index
+			if (log.LineCount > 0)
+			{
+				LogNode? root = await storageClient.ReadRefTargetAsync<LogNode>(log.RefName, cancellationToken: cancellationToken);
+				if (root != null)
+				{
+					LogIndexNode index = await root.IndexRef.ReadBlobAsync(cancellationToken: cancellationToken);
+					await foreach (int lineIdx in index.SearchAsync(firstLine, searchText, searchStats, cancellationToken: cancellationToken))
+					{
+						yield return lineIdx;
+					}
+					if (root.Complete)
+					{
+						yield break;
+					}
+					firstLine = root.LineCount;
+				}
+			}
+
+			// Search any tail data we have
+			if (!log.Complete)
+			{
+				for (; ; )
+				{
+					Utf8String[] lines = await ReadTailAsync(log, firstLine, cancellationToken);
+					if (lines.Length == 0)
+					{
+						break;
+					}
+
+					for (int idx = 0; idx < lines.Length; idx++)
+					{
+						if (SearchTerm.FindNextOcurrence(lines[idx].Span, 0, searchText) != -1)
+						{
+							yield return firstLine + idx;
+						}
+					}
+
+					firstLine += lines.Length;
+				}
+			}
+		}
+
+		async Task<Utf8String[]> ReadTailAsync(ILog log, int index, CancellationToken cancellationToken)
+		{
+			_ = cancellationToken;
+
+			const int BatchSize = 128;
+
+			if (log.Complete)
+			{
+				return Array.Empty<Utf8String>();
+			}
+
+			string cacheKey = $"{log.Id}@{index}";
+			if (_logCache.TryGetValue(cacheKey, out Utf8String[]? lines))
+			{
+				return lines!;
+			}
+
+			lines = (await _logTailService.ReadAsync(log.Id, index, BatchSize)).ToArray();
+			if (log.Type == LogType.Json)
+			{
+				LogChunkBuilder builder = new LogChunkBuilder(lines.Sum(x => x.Length));
+				foreach (Utf8String line in lines)
+				{
+					builder.AppendJsonAsPlainText(line.Span, _logger);
+				}
+				lines = lines.ToArray();
+			}
+
+			if (lines.Length == BatchSize)
+			{
+				int length = lines.Sum(x => x.Length);
+				using (ICacheEntry entry = _logCache.CreateEntry(cacheKey))
+				{
+					entry.SetSlidingExpiration(TimeSpan.FromMinutes(1.0));
+					entry.SetSize(length);
+					entry.SetValue(lines);
+				}
+			}
+
+			return lines;
+		}
+
+		async IAsyncEnumerable<int> SearchLogDataInternalAsync(ILog log, string text, int firstLine, SearchStats searchStats)
+		{
+			SearchText searchText = new SearchText(text);
+
+			// Read the index for this log file
+			if (log.IndexLength != null)
+			{
+				LogIndexData? indexData = await ReadIndexAsync(log, log.IndexLength.Value);
+				if (indexData != null && firstLine < indexData.LineCount)
+				{
+					using TelemetrySpan span = _tracer.StartActiveSpan($"{nameof(LogService)}.{nameof(SearchLogDataInternalAsync)}.Indexed");
+					span.SetAttribute("lineCount", indexData.LineCount);
+
+					foreach (int lineIndex in indexData.Search(firstLine, searchText, searchStats))
+					{
+						yield return lineIndex;
+					}
+
+					firstLine = indexData.LineCount;
+				}
+			}
+
+			// Manually search through the rest of the log
+			int chunkIdx = log.Chunks.GetChunkForLine(firstLine);
+			for (; chunkIdx < log.Chunks.Count; chunkIdx++)
+			{
+				ILogChunk chunk = log.Chunks[chunkIdx];
+
+				// Read the chunk data
+				LogChunkData chunkData = await ReadChunkAsync(log, chunkIdx);
+				if (firstLine < chunkData.LineIndex + chunkData.LineCount)
+				{
+					// Find the first sub-chunk we're looking for
+					int subChunkIdx = 0;
+					if (firstLine > chunk.LineIndex)
+					{
+						subChunkIdx = chunkData.GetSubChunkForLine(firstLine - chunk.LineIndex);
+					}
+
+					// Search through the sub-chunks
+					for (; subChunkIdx < chunkData.SubChunks.Count; subChunkIdx++)
+					{
+						LogSubChunkData subChunkData = chunkData.SubChunks[subChunkIdx];
+						if (firstLine < subChunkData.LineIndex + subChunkData.LineCount)
+						{
+							// Create an index containing just this sub-chunk
+							LogIndexData index = subChunkData.BuildIndex(_logger);
+							foreach (int lineIndex in index.Search(firstLine, searchText, searchStats))
+							{
+								yield return lineIndex;
+							}
+						}
+					}
+				}
+			}
+		}
+	}
+}
+
diff --git a/Engine/Source/Programs/Horde/Horde.Server/Logs/LogsController.cs b/Engine/Source/Programs/Horde/Horde.Server/Logs/LogsController.cs
--- a/Engine/Source/Programs/Horde/Horde.Server/Logs/LogsController.cs
+++ b/Engine/Source/Programs/Horde/Horde.Server/Logs/LogsController.cs
@@ -52,7 +52,7 @@ namespace Horde.Server.Logs
 	[Route("[controller]")]
 	public class LogsController : ControllerBase
 	{
-		private readonly ILogFileService _logFileService;
+		private readonly ILogService _logService;
 		private readonly IIssueCollection _issueCollection;
 		private readonly JobService _jobService;
 		private readonly StorageService _storageService;
@@ -61,9 +61,9 @@ namespace Horde.Server.Logs
 		/// <summary>
 		/// Constructor
 		/// </summary>
-		public LogsController(ILogFileService logFileService, IIssueCollection issueCollection, JobService jobService, StorageService storageService, IOptionsSnapshot<GlobalConfig> globalConfig)
+		public LogsController(ILogService logService, IIssueCollection issueCollection, JobService jobService, StorageService storageService, IOptionsSnapshot<GlobalConfig> globalConfig)
 		{
-			_logFileService = logFileService;
+			_logService = logService;
 			_issueCollection = issueCollection;
 			_jobService = jobService;
 			_storageService = storageService;
@@ -73,37 +73,37 @@ namespace Horde.Server.Logs
 		/// <summary>
 		/// Retrieve metadata about a specific log file
 		/// </summary>
-		/// <param name="logFileId">Id of the log file to get information about</param>
+		/// <param name="logId">Id of the log file to get information about</param>
 		/// <param name="filter">Filter for the properties to return</param>
 		/// <param name="cancellationToken">Cancellation token for the request</param>
 		/// <returns>Information about the requested project</returns>
 		[HttpGet]
-		[Route("/api/v1/logs/{logFileId}")]
+		[Route("/api/v1/logs/{logId}")]
 		[ProducesResponseType(typeof(GetLogResponse), 200)]
-		public async Task<ActionResult<object>> GetLogAsync(LogId logFileId, [FromQuery] PropertyFilter? filter = null, CancellationToken cancellationToken = default)
+		public async Task<ActionResult<object>> GetLogAsync(LogId logId, [FromQuery] PropertyFilter? filter = null, CancellationToken cancellationToken = default)
 		{
-			ILogFile? logFile = await _logFileService.GetLogFileAsync(logFileId, cancellationToken);
-			if (logFile == null)
+			ILog? log = await _logService.GetLogAsync(logId, cancellationToken);
+			if (log == null)
 			{
 				return NotFound();
 			}
-			if (!await AuthorizeAsync(logFile, LogAclAction.ViewLog, User, cancellationToken))
+			if (!await AuthorizeAsync(log, LogAclAction.ViewLog, User, cancellationToken))
 			{
 				return Forbid();
 			}
 
-			LogMetadata metadata = await _logFileService.GetMetadataAsync(logFile, cancellationToken);
-			return CreateGetLogFileResponse(logFile, metadata).ApplyFilter(filter);
+			LogMetadata metadata = await _logService.GetMetadataAsync(log, cancellationToken);
+			return CreateGetLogResponse(log, metadata).ApplyFilter(filter);
 		}
 
-		static GetLogResponse CreateGetLogFileResponse(ILogFile logFile, LogMetadata metadata)
+		static GetLogResponse CreateGetLogResponse(ILog log, LogMetadata metadata)
 		{
 			GetLogResponse response = new GetLogResponse();
-			response.Id = logFile.Id;
-			response.JobId = logFile.JobId;
-			response.LeaseId = logFile.LeaseId;
-			response.SessionId = logFile.SessionId;
-			response.Type = logFile.Type;
+			response.Id = log.Id;
+			response.JobId = log.JobId;
+			response.LeaseId = log.LeaseId;
+			response.SessionId = log.SessionId;
+			response.Type = log.Type;
 			response.LineCount = metadata.MaxLineIndex;
 			return response;
 		}
@@ -111,21 +111,21 @@ namespace Horde.Server.Logs
 		/// <summary>
 		/// Uploads a blob for a log file. See /api/v1/storage/XXX/blobs.
 		/// </summary>
-		/// <param name="logFileId">Id of the log file to get information about</param>
+		/// <param name="logId">Id of the log file to get information about</param>
 		/// <param name="request">Request for the upload</param>
 		/// <param name="cancellationToken">Cancellation token for the request</param>
 		/// <returns>Information about the requested project</returns>
 		[HttpPost]
-		[Route("/api/v1/logs/{logFileId}/blobs")]
+		[Route("/api/v1/logs/{logId}/blobs")]
 		[ProducesResponseType(typeof(WriteBlobResponse), 200)]
-		public async Task<ActionResult<WriteBlobResponse>> WriteLogBlobAsync(LogId logFileId, WriteBlobRequest request, CancellationToken cancellationToken = default)
+		public async Task<ActionResult<WriteBlobResponse>> WriteLogBlobAsync(LogId logId, WriteBlobRequest request, CancellationToken cancellationToken = default)
 		{
-			ILogFile? logFile = await _logFileService.GetLogFileAsync(logFileId, cancellationToken);
-			if (logFile == null)
+			ILog? log = await _logService.GetLogAsync(logId, cancellationToken);
+			if (log == null)
 			{
 				return NotFound();
 			}
-			if (!await AuthorizeAsync(logFile, LogAclAction.WriteLogData, User, cancellationToken))
+			if (!await AuthorizeAsync(log, LogAclAction.WriteLogData, User, cancellationToken))
 			{
 				return Forbid();
 			}
@@ -134,7 +134,7 @@ namespace Horde.Server.Logs
 				return BadRequest("Cannot specify prefix for logs");
 			}
 
-			request.Prefix = $"{logFile.RefName}";
+			request.Prefix = $"{log.RefName}";
 
 			IStorageBackend storageBackend = _storageService.CreateBackend(Namespace.Logs);
 			return await StorageController.WriteBlobAsync(storageBackend, request, cancellationToken);
@@ -143,69 +143,69 @@ namespace Horde.Server.Logs
 		/// <summary>
 		/// Retrieve raw data for a log file
 		/// </summary>
-		/// <param name="logFileId">Id of the log file to get information about</param>
+		/// <param name="logId">Id of the log file to get information about</param>
 		/// <param name="format">Format for the returned data</param>
 		/// <param name="fileName">Name of the default filename to download</param>
 		/// <param name="download">Whether to download the file rather than display in the browser</param>
 		/// <param name="cancellationToken">Cancellation token for the request</param>
 		/// <returns>Raw log data for the requested range</returns>
 		[HttpGet]
-		[Route("/api/v1/logs/{logFileId}/data")]
+		[Route("/api/v1/logs/{logId}/data")]
 		public async Task<ActionResult> GetLogDataAsync(
-			LogId logFileId,
+			LogId logId,
 			[FromQuery] LogOutputFormat format = LogOutputFormat.Raw,
 			[FromQuery] string? fileName = null,
 			[FromQuery] bool download = false,
 			CancellationToken cancellationToken = default)
 		{
-			ILogFile? logFile = await _logFileService.GetLogFileAsync(logFileId, cancellationToken);
-			if (logFile == null)
+			ILog? log = await _logService.GetLogAsync(logId, cancellationToken);
+			if (log == null)
 			{
 				return NotFound();
 			}
-			if (!await AuthorizeAsync(logFile, LogAclAction.ViewLog, User, cancellationToken))
+			if (!await AuthorizeAsync(log, LogAclAction.ViewLog, User, cancellationToken))
 			{
 				return Forbid();
 			}
 
 			Func<Stream, ActionContext, Task> copyTask;
-			if (format == LogOutputFormat.Text && logFile.Type == LogType.Json)
+			if (format == LogOutputFormat.Text && log.Type == LogType.Json)
 			{
-				copyTask = (outputStream, context) => _logFileService.CopyPlainTextStreamAsync(logFile, outputStream, cancellationToken);
+				copyTask = (outputStream, context) => _logService.CopyPlainTextStreamAsync(log, outputStream, cancellationToken);
 			}
 			else
 			{
-				copyTask = (outputStream, context) => _logFileService.CopyRawStreamAsync(logFile, outputStream, cancellationToken);
+				copyTask = (outputStream, context) => _logService.CopyRawStreamAsync(log, outputStream, cancellationToken);
 			}
 
-			return new CustomFileCallbackResult(fileName ?? $"log-{logFileId}.txt", "text/plain", !download, copyTask);
+			return new CustomFileCallbackResult(fileName ?? $"log-{logId}.txt", "text/plain", !download, copyTask);
 		}
 
 		/// <summary>
-		/// Retrieve line data for a logfile
+		/// Retrieve line data for a log
 		/// </summary>
-		/// <param name="logFileId">Id of the log file to get information about</param>
+		/// <param name="logId">Id of the log file to get information about</param>
 		/// <param name="index">Index of the first line to retrieve</param>
 		/// <param name="count">Number of lines to retrieve</param>
 		/// <param name="cancellationToken">Cancellation token for the request</param>
 		/// <returns>Information about the requested project</returns>
 		[HttpGet]
-		[Route("/api/v1/logs/{logFileId}/lines")]
-		public async Task<ActionResult> GetLogLinesAsync(LogId logFileId, [FromQuery] int index = 0, [FromQuery] int count = 100, CancellationToken cancellationToken = default)
+		[Route("/api/v1/logs/{logId}/lines")]
+		public async Task<ActionResult> GetLogLinesAsync(LogId logId, [FromQuery] int index = 0, [FromQuery] int count = 100, CancellationToken cancellationToken = default)
 		{
-			ILogFile? logFile = await _logFileService.GetLogFileAsync(logFileId, cancellationToken);
-			if (logFile == null)
+			ILog? log = await _logService.GetLogAsync(logId, cancellationToken);
+			if (log == null)
 			{
 				return NotFound();
 			}
-			if (!await AuthorizeAsync(logFile, LogAclAction.ViewLog, User, cancellationToken))
+			if (!await AuthorizeAsync(log, LogAclAction.ViewLog, User, cancellationToken))
 			{
 				return Forbid();
 			}
 
-			LogMetadata metadata = await _logFileService.GetMetadataAsync(logFile, cancellationToken);
+			LogMetadata metadata = await _logService.GetMetadataAsync(log, cancellationToken);
 
-			List<Utf8String> lines = await _logFileService.ReadLinesAsync(logFile, index, count, cancellationToken);
+			List<Utf8String> lines = await _logService.ReadLinesAsync(log, index, count, cancellationToken);
 			using (MemoryStream stream = new MemoryStream(lines.Sum(x => x.Length) + (lines.Count * 20)))
 			{
 				stream.WriteByte((byte)'{');
@@ -213,7 +213,7 @@ namespace Horde.Server.Logs
 				stream.Write(Encoding.UTF8.GetBytes($"\"index\":{index},"));
 				stream.Write(Encoding.UTF8.GetBytes($"\"count\":{lines.Count},"));
 				stream.Write(Encoding.UTF8.GetBytes($"\"maxLineIndex\":{Math.Max(metadata.MaxLineIndex, index + lines.Count)},"));
-				stream.Write(Encoding.UTF8.GetBytes($"\"format\":{(logFile.Type == LogType.Json ? "\"JSON\"" : "\"TEXT\"")},"));
+				stream.Write(Encoding.UTF8.GetBytes($"\"format\":{(log.Type == LogType.Json ? "\"JSON\"" : "\"TEXT\"")},"));
 
 				stream.Write(Encoding.UTF8.GetBytes($"\"lines\":["));
 				stream.WriteByte((byte)'\n');
@@ -225,7 +225,7 @@ namespace Horde.Server.Logs
 					stream.WriteByte((byte)' ');
 					stream.WriteByte((byte)' ');
 
-					if (logFile.Type == LogType.Json)
+					if (log.Type == LogType.Json)
 					{
 						await stream.WriteAsync(line.Memory, cancellationToken);
 					}
@@ -255,7 +255,7 @@ namespace Horde.Server.Logs
 					stream.WriteByte((byte)'\n');
 				}
 
-				if (logFile.Type == LogType.Json)
+				if (log.Type == LogType.Json)
 				{
 					stream.Write(Encoding.UTF8.GetBytes($"]"));
 				}
@@ -273,68 +273,68 @@ namespace Horde.Server.Logs
 		/// <summary>
 		/// Search log data
 		/// </summary>
-		/// <param name="logFileId">Id of the log file to get information about</param>
+		/// <param name="logId">Id of the log file to get information about</param>
 		/// <param name="text">Text to search for</param>
 		/// <param name="firstLine">First line to search from</param>
 		/// <param name="count">Number of results to return</param>
 		/// <param name="cancellationToken">Cancellation token for the request</param>
 		/// <returns>Raw log data for the requested range</returns>
 		[HttpGet]
-		[Route("/api/v1/logs/{logFileId}/search")]
-		public async Task<ActionResult<SearchLogResponse>> SearchLogFileAsync(
-			LogId logFileId,
+		[Route("/api/v1/logs/{logId}/search")]
+		public async Task<ActionResult<SearchLogResponse>> SearchLogAsync(
+			LogId logId,
 			[FromQuery] string text,
 			[FromQuery] int firstLine = 0,
 			[FromQuery] int count = 5,
 			CancellationToken cancellationToken = default)
 		{
-			ILogFile? logFile = await _logFileService.GetLogFileAsync(logFileId, cancellationToken);
-			if (logFile == null)
+			ILog? log = await _logService.GetLogAsync(logId, cancellationToken);
+			if (log == null)
 			{
 				return NotFound();
 			}
-			if (!await AuthorizeAsync(logFile, LogAclAction.ViewLog, User, cancellationToken))
+			if (!await AuthorizeAsync(log, LogAclAction.ViewLog, User, cancellationToken))
 			{
 				return Forbid();
 			}
 
 			SearchLogResponse response = new SearchLogResponse();
 			response.Stats = new SearchStats();
-			response.Lines = await _logFileService.SearchLogDataAsync(logFile, text, firstLine, count, response.Stats, cancellationToken);
+			response.Lines = await _logService.SearchLogDataAsync(log, text, firstLine, count, response.Stats, cancellationToken);
 			return response;
 		}
 
 		/// <summary>
-		/// Retrieve events for a logfile
+		/// Retrieve events for a log
 		/// </summary>
-		/// <param name="logFileId">Id of the log file to get information about</param>
+		/// <param name="logId">Id of the log file to get information about</param>
 		/// <param name="index">Index of the first line to retrieve</param>
 		/// <param name="count">Number of lines to retrieve</param>
 		/// <param name="cancellationToken">Cancellation token for the request</param>
 		/// <returns>Information about the requested project</returns>
 		[HttpGet]
-		[Route("/api/v1/logs/{logFileId}/events")]
+		[Route("/api/v1/logs/{logId}/events")]
 		[ProducesResponseType(typeof(List<GetLogEventResponse>), 200)]
-		public async Task<ActionResult<List<GetLogEventResponse>>> GetEventsAsync(LogId logFileId, [FromQuery] int? index = null, [FromQuery] int? count = null, CancellationToken cancellationToken = default)
+		public async Task<ActionResult<List<GetLogEventResponse>>> GetEventsAsync(LogId logId, [FromQuery] int? index = null, [FromQuery] int? count = null, CancellationToken cancellationToken = default)
 		{
-			ILogFile? logFile = await _logFileService.GetLogFileAsync(logFileId, cancellationToken);
-			if (logFile == null)
+			ILog? log = await _logService.GetLogAsync(logId, cancellationToken);
+			if (log == null)
 			{
 				return NotFound();
 			}
-			if (!await AuthorizeAsync(logFile, LogAclAction.ViewLog, User, cancellationToken))
+			if (!await AuthorizeAsync(log, LogAclAction.ViewLog, User, cancellationToken))
 			{
 				return Forbid();
 			}
 
-			List<ILogEvent> logEvents = await _logFileService.FindEventsAsync(logFile, null, index, count, cancellationToken);
+			List<ILogEvent> logEvents = await _logService.FindEventsAsync(log, null, index, count, cancellationToken);
 
 			Dictionary<ObjectId, int?> spanIdToIssueId = new Dictionary<ObjectId, int?>();
 
 			List<GetLogEventResponse> responses = new List<GetLogEventResponse>();
 			foreach (ILogEvent logEvent in logEvents)
 			{
-				ILogEventData logEventData = await _logFileService.GetEventDataAsync(logFile, logEvent.LineIndex, logEvent.LineCount, cancellationToken);
+				ILogEventData logEventData = await _logService.GetEventDataAsync(log, logEvent.LineIndex, logEvent.LineCount, cancellationToken);
 
 				int? issueId = null;
 				if (logEvent.SpanId != null && !spanIdToIssueId.TryGetValue(logEvent.SpanId.Value, out issueId))
@@ -370,31 +370,31 @@ namespace Horde.Server.Logs
 		/// <summary>
 		/// Determines if the user is authorized to perform an action on a particular template
 		/// </summary>
-		/// <param name="logFile">The template to check</param>
+		/// <param name="log">The template to check</param>
 		/// <param name="action">The action being performed</param>
 		/// <param name="user">The principal to authorize</param>
 		/// <param name="cancellationToken">Cancellation token for the operation</param>
 		/// <returns>True if the action is authorized</returns>
-		async Task<bool> AuthorizeAsync(ILogFile logFile, AclAction action, ClaimsPrincipal user, CancellationToken cancellationToken)
+		async Task<bool> AuthorizeAsync(ILog log, AclAction action, ClaimsPrincipal user, CancellationToken cancellationToken)
 		{
 			GlobalConfig globalConfig = _globalConfig.Value;
 			if (user.HasAdminClaim())
 			{
 				return true;
 			}
-			if (logFile.LeaseId != null && user.HasLeaseClaim(logFile.LeaseId.Value))
+			if (log.LeaseId != null && user.HasLeaseClaim(log.LeaseId.Value))
 			{
 				return true;
 			}
-			if (logFile.SessionId != null && LogFileService.AuthorizeForSession(logFile, user))
+			if (log.SessionId != null && LogService.AuthorizeForSession(log, user))
 			{
 				return true;
 			}
-			if (logFile.JobId != JobId.Empty && await _jobService.AuthorizeAsync(logFile.JobId, action, user, globalConfig, cancellationToken))
+			if (log.JobId != JobId.Empty && await _jobService.AuthorizeAsync(log.JobId, action, user, globalConfig, cancellationToken))
 			{
 				return true;
 			}
-			if (action == LogAclAction.ViewLog && logFile.SessionId != null && globalConfig.Authorize(SessionAclAction.ViewSession, user))
+			if (action == LogAclAction.ViewLog && log.SessionId != null && globalConfig.Authorize(SessionAclAction.ViewSession, user))
 			{
 				return true;
 			}
diff --git a/Engine/Source/Programs/Horde/Horde.Server/Notifications/NotificationService.cs b/Engine/Source/Programs/Horde/Horde.Server/Notifications/NotificationService.cs
--- a/Engine/Source/Programs/Horde/Horde.Server/Notifications/NotificationService.cs
+++ b/Engine/Source/Programs/Horde/Horde.Server/Notifications/NotificationService.cs
@@ -76,9 +76,9 @@ namespace Horde.Server.Notifications
 		private readonly IssueService _issueService;
 
 		/// <summary>
-		/// Instance of the <see cref="_logFileService"/>.
+		/// Instance of the <see cref="_logService"/>.
 		/// </summary>
-		private readonly ILogFileService _logFileService;
+		private readonly ILogService _logService;
 
 		/// <summary>
 		/// Cache for de-duplicating queued notifications
@@ -145,7 +145,7 @@ namespace Horde.Server.Notifications
 			IUserCollection userCollection,
 			JobService jobService,
 			IssueService issueService,
-			ILogFileService logFileService,
+			ILogService logService,
 			Meter meter,
 			IMemoryCache cache,
 			RedisService redisService,
@@ -161,7 +161,7 @@ namespace Horde.Server.Notifications
 			_userCollection = userCollection;
 			_jobService = jobService;
 			_issueService = issueService;
-			_logFileService = logFileService;
+			_logService = logService;
 			_cache = cache;
 			_redisConnectionPool = redisService.ConnectionPool;
 			_backgroundTask = new BackgroundTask(ExecuteAsync);
@@ -676,18 +676,18 @@ namespace Horde.Server.Notifications
 				return;
 			}
 
-			ILogFile? logFile = await _logFileService.GetLogFileAsync(step.LogId.Value, cancellationToken);
-			if (logFile == null)
+			ILog? log = await _logService.GetLogAsync(step.LogId.Value, cancellationToken);
+			if (log == null)
 			{
 				_logger.LogError("Step does not have a log file");
 				return;
 			}
 
-			List<ILogEvent> jobStepEvents = await _logFileService.FindEventsAsync(logFile, cancellationToken: cancellationToken);
+			List<ILogEvent> jobStepEvents = await _logService.FindEventsAsync(log, cancellationToken: cancellationToken);
 			List<ILogEventData> jobStepEventData = new List<ILogEventData>();
 			foreach (ILogEvent logEvent in jobStepEvents)
 			{
-				ILogEventData eventData = await _logFileService.GetEventDataAsync(logFile, logEvent.LineIndex, logEvent.LineCount, cancellationToken);
+				ILogEventData eventData = await _logService.GetEventDataAsync(log, logEvent.LineIndex, logEvent.LineCount, cancellationToken);
 				jobStepEventData.Add(eventData);
 			}
 
diff --git a/Engine/Source/Programs/Horde/Horde.Server/Notifications/Sinks/SlackNotificationSink.cs b/Engine/Source/Programs/Horde/Horde.Server/Notifications/Sinks/SlackNotificationSink.cs
--- a/Engine/Source/Programs/Horde/Horde.Server/Notifications/Sinks/SlackNotificationSink.cs
+++ b/Engine/Source/Programs/Horde/Horde.Server/Notifications/Sinks/SlackNotificationSink.cs
@@ -214,7 +214,7 @@ namespace Horde.Server.Notifications.Sinks
 		readonly RedisService _redisService;
 		readonly IssueService _issueService;
 		readonly IUserCollection _userCollection;
-		readonly ILogFileService _logFileService;
+		readonly ILogService _logService;
 		readonly IWebHostEnvironment _environment;
 		readonly ServerSettings _settings;
 		readonly IMongoCollection<MessageStateDocument> _messageStates;
@@ -246,12 +246,12 @@ namespace Horde.Server.Notifications.Sinks
 		/// <summary>
 		/// Constructor
 		/// </summary>
-		public SlackNotificationSink(MongoService mongoService, RedisService redisService, IssueService issueService, IUserCollection userCollection, ILogFileService logFileService, IExternalIssueService externalIssueService, IWebHostEnvironment environment, IOptions<ServerSettings> settings, IClock clock, IOptionsMonitor<GlobalConfig> globalConfig, ILogger<SlackNotificationSink> logger)
+		public SlackNotificationSink(MongoService mongoService, RedisService redisService, IssueService issueService, IUserCollection userCollection, ILogService logService, IExternalIssueService externalIssueService, IWebHostEnvironment environment, IOptions<ServerSettings> settings, IClock clock, IOptionsMonitor<GlobalConfig> globalConfig, ILogger<SlackNotificationSink> logger)
 		{
 			_redisService = redisService;
 			_issueService = issueService;
 			_userCollection = userCollection;
-			_logFileService = logFileService;
+			_logService = logService;
 			_externalIssueService = externalIssueService;
 			_environment = environment;
 			_settings = settings.Value;
@@ -1014,10 +1014,10 @@ namespace Horde.Server.Notifications.Sinks
 				if (span.FirstFailure.LogId != null)
 				{
 					LogId logId = span.FirstFailure.LogId.Value;
-					ILogFile? logFile = await _logFileService.GetLogFileAsync(logId, cancellationToken);
-					if (logFile != null)
+					ILog? log = await _logService.GetLogAsync(logId, cancellationToken);
+					if (log != null)
 					{
-						events = await _logFileService.FindEventsAsync(logFile, span.Id, 0, 50, cancellationToken);
+						events = await _logService.FindEventsAsync(log, span.Id, 0, 50, cancellationToken);
 						if (events.Any(x => x.Severity == LogEventSeverity.Error))
 						{
 							events.RemoveAll(x => x.Severity == LogEventSeverity.Warning);
@@ -1026,7 +1026,7 @@ namespace Horde.Server.Notifications.Sinks
 						List<string> eventStrings = new List<string>();
 						for (int idx = 0; idx < Math.Min(events.Count, 3); idx++)
 						{
-							ILogEventData data = await _logFileService.GetEventDataAsync(logFile, events[idx].LineIndex, events[idx].LineCount, cancellationToken);
+							ILogEventData data = await _logService.GetEventDataAsync(log, events[idx].LineIndex, events[idx].LineCount, cancellationToken);
 							eventDataItems.Add(data);
 						}
 					}
@@ -1450,10 +1450,10 @@ namespace Horde.Server.Notifications.Sinks
 			if (lastSpan != null && lastSpan.LastFailure.LogId != null)
 			{
 				LogId logId = lastSpan.LastFailure.LogId.Value;
-				ILogFile? logFile = await _logFileService.GetLogFileAsync(logId, cancellationToken);
-				if (logFile != null)
+				ILog? log = await _logService.GetLogAsync(logId, cancellationToken);
+				if (log != null)
 				{
-					List<ILogEvent> events = await _logFileService.FindEventsAsync(logFile, lastSpan.Id, 0, 20, cancellationToken);
+					List<ILogEvent> events = await _logService.FindEventsAsync(log, lastSpan.Id, 0, 20, cancellationToken);
 					if (events.Any(x => x.Severity == LogEventSeverity.Error))
 					{
 						events.RemoveAll(x => x.Severity == LogEventSeverity.Warning);
@@ -1461,7 +1461,7 @@ namespace Horde.Server.Notifications.Sinks
 
 					for (int idx = 0; idx < Math.Min(events.Count, 3); idx++)
 					{
-						ILogEventData data = await _logFileService.GetEventDataAsync(logFile, events[idx].LineIndex, events[idx].LineCount, cancellationToken);
+						ILogEventData data = await _logService.GetEventDataAsync(log, events[idx].LineIndex, events[idx].LineCount, cancellationToken);
 						attachment.AddSection(QuoteText(data.Message));
 					}
 					if (events.Count > 3)
diff --git a/Engine/Source/Programs/Horde/Horde.Server/Properties/launchSettings.json b/Engine/Source/Programs/Horde/Horde.Server/Properties/launchSettings.json
--- a/Engine/Source/Programs/Horde/Horde.Server/Properties/launchSettings.json
+++ b/Engine/Source/Programs/Horde/Horde.Server/Properties/launchSettings.json
@@ -4,7 +4,8 @@
       "commandName": "Project",
       "commandLineArgs": "",
       "environmentVariables": {
-        "ASPNETCORE_ENVIRONMENT": "Local"
+        "ASPNETCORE_ENVIRONMENT": "Local",
+        "Horde__Perforce__1__Id": "psyonix"
       },
       "hotReloadEnabled": false
     }
diff --git a/Engine/Source/Programs/Horde/Horde.Server/Server/DebugController.cs b/Engine/Source/Programs/Horde/Horde.Server/Server/DebugController.cs
--- a/Engine/Source/Programs/Horde/Horde.Server/Server/DebugController.cs
+++ b/Engine/Source/Programs/Horde/Horde.Server/Server/DebugController.cs
@@ -260,7 +260,7 @@ namespace Horde.Server.Server
 		private readonly AgentRelayService _agentRelayService;
 		private readonly JobService _jobService;
 		private readonly JobTaskSource _jobTaskSource;
-		private readonly ILogFileCollection _logFileCollection;
+		private readonly ILogCollection _logCollection;
 		private readonly IOptionsSnapshot<GlobalConfig> _globalConfig;
 		private readonly ILogger<SecureDebugController> _logger;
 
@@ -273,14 +273,14 @@ namespace Horde.Server.Server
 			AgentRelayService agentRelayService,
 			JobService jobService,
 			JobTaskSource jobTaskSource,
-			ILogFileCollection logFileCollection, IOptionsSnapshot<GlobalConfig> globalConfig, ILogger<SecureDebugController> logger)
+			ILogCollection logCollection, IOptionsSnapshot<GlobalConfig> globalConfig, ILogger<SecureDebugController> logger)
 		{
 			_mongoService = mongoService;
 			_configService = configService;
 			_jobService = jobService;
 			_agentRelayService = agentRelayService;
 			_jobTaskSource = jobTaskSource;
-			_logFileCollection = logFileCollection;
+			_logCollection = logCollection;
 			_globalConfig = globalConfig;
 			_logger = logger;
 		}
@@ -499,25 +499,25 @@ namespace Horde.Server.Server
 		/// <summary>
 		/// Retrieve metadata about a specific log file
 		/// </summary>
-		/// <param name="logFileId">Id of the log file to get information about</param>
+		/// <param name="logId">Id of the log file to get information about</param>
 		/// <param name="filter">Filter for the properties to return</param>
 		/// <returns>Information about the requested project</returns>
 		[HttpGet]
-		[Route("/api/v1/debug/logs/{LogFileId}")]
-		public async Task<ActionResult<object>> GetLogAsync(LogId logFileId, [FromQuery] PropertyFilter? filter = null)
+		[Route("/api/v1/debug/logs/{logId}")]
+		public async Task<ActionResult<object>> GetLogAsync(LogId logId, [FromQuery] PropertyFilter? filter = null)
 		{
 			if (!_globalConfig.Value.Authorize(ServerAclAction.Debug, User))
 			{
 				return Forbid(ServerAclAction.Debug);
 			}
 
-			ILogFile? logFile = await _logFileCollection.GetLogFileAsync(logFileId, CancellationToken.None);
-			if (logFile == null)
+			ILog? log = await _logCollection.GetLogAsync(logId, CancellationToken.None);
+			if (log == null)
 			{
 				return NotFound();
 			}
 
-			return logFile.ApplyFilter(filter);
+			return log.ApplyFilter(filter);
 		}
 
 		/// <summary>
diff --git a/Engine/Source/Programs/Horde/Horde.Server/Startup.cs b/Engine/Source/Programs/Horde/Horde.Server/Startup.cs
--- a/Engine/Source/Programs/Horde/Horde.Server/Startup.cs
+++ b/Engine/Source/Programs/Horde/Horde.Server/Startup.cs
@@ -500,7 +500,7 @@ namespace Horde.Server
 			services.AddSingleton<IJobTimingCollection, JobTimingCollection>();
 			services.AddSingleton<ILeaseCollection, LeaseCollection>();
 			services.AddSingleton<ILogEventCollection, LogEventCollection>();
-			services.AddSingleton<ILogFileCollection, LogFileCollection>();
+			services.AddSingleton<ILogCollection, LogCollection>();
 			services.AddSingleton<INotificationTriggerCollection, NotificationTriggerCollection>();
 			services.AddSingleton<IPoolCollection, PoolCollection>();
 			services.AddSingleton<IBisectTaskCollection, BisectTaskCollection>();
@@ -573,7 +573,7 @@ namespace Horde.Server
 			services.AddSingleton<EnrollmentService>();
 			services.AddSingleton<ServerStatusService>();
 			services.AddHostedService(provider => provider.GetRequiredService<ServerStatusService>());
-			services.AddSingleton<ILogFileService, LogFileService>();
+			services.AddSingleton<ILogService, LogService>();
 			services.AddSingleton<LogTailService>();
 			services.AddHostedService(provider => provider.GetRequiredService<LogTailService>()); // NB: Runs even on workers, to receive tail notifications.
 			services.AddSingleton<INotificationService, NotificationService>();
diff --git a/Engine/Source/Programs/Horde/Horde.Server/Tasks/ConformTaskSource.cs b/Engine/Source/Programs/Horde/Horde.Server/Tasks/ConformTaskSource.cs
--- a/Engine/Source/Programs/Horde/Horde.Server/Tasks/ConformTaskSource.cs
+++ b/Engine/Source/Programs/Horde/Horde.Server/Tasks/ConformTaskSource.cs
@@ -43,7 +43,7 @@ namespace Horde.Server.Tasks
 		readonly PoolService _poolService;
 		readonly SingletonDocument<ConformList> _conformList;
 		readonly PerforceLoadBalancer _perforceLoadBalancer;
-		readonly ILogFileService _logService;
+		readonly ILogService _logService;
 		readonly ILogger _logger;
 		readonly IOptionsMonitor<ServerSettings> _settings;
 		readonly IOptionsMonitor<GlobalConfig> _globalConfig;
@@ -52,7 +52,7 @@ namespace Horde.Server.Tasks
 		/// <summary>
 		/// Constructor
 		/// </summary>
-		public ConformTaskSource(MongoService mongoService, IAgentCollection agentCollection, PoolService poolService, ILogFileService logService, PerforceLoadBalancer perforceLoadBalancer, IClock clock, IOptionsMonitor<ServerSettings> settings, IOptionsMonitor<GlobalConfig> globalConfig, ILogger<ConformTaskSource> logger)
+		public ConformTaskSource(MongoService mongoService, IAgentCollection agentCollection, PoolService poolService, ILogService logService, PerforceLoadBalancer perforceLoadBalancer, IClock clock, IOptionsMonitor<ServerSettings> settings, IOptionsMonitor<GlobalConfig> globalConfig, ILogger<ConformTaskSource> logger)
 		{
 			_agentCollection = agentCollection;
 			_poolService = poolService;
@@ -179,7 +179,7 @@ namespace Horde.Server.Tasks
 					{
 						try
 						{
-							ILogFile log = await _logService.CreateLogFileAsync(JobId.Empty, leaseId, agent.SessionId, LogType.Json, cancellationToken: cancellationToken);
+							ILog log = await _logService.CreateLogAsync(JobId.Empty, leaseId, agent.SessionId, LogType.Json, cancellationToken: cancellationToken);
 							task.LogId = log.Id.ToString();
 							task.RemoveUntrackedFiles = agent.RequestFullConform;
 
diff --git a/Engine/Source/Programs/Horde/Horde.Server/Tasks/RestartTaskSource.cs b/Engine/Source/Programs/Horde/Horde.Server/Tasks/RestartTaskSource.cs
--- a/Engine/Source/Programs/Horde/Horde.Server/Tasks/RestartTaskSource.cs
+++ b/Engine/Source/Programs/Horde/Horde.Server/Tasks/RestartTaskSource.cs
@@ -20,9 +20,9 @@ namespace Horde.Server.Tasks
 
 		public override TaskSourceFlags Flags => TaskSourceFlags.AllowWhenDisabled | TaskSourceFlags.AllowDuringDowntime;
 
-		readonly ILogFileService _logService;
+		readonly ILogService _logService;
 
-		public RestartTaskSource(ILogFileService logService)
+		public RestartTaskSource(ILogService logService)
 		{
 			_logService = logService;
 
@@ -44,7 +44,7 @@ namespace Horde.Server.Tasks
 			}
 
 			LeaseId leaseId = new LeaseId(BinaryIdUtils.CreateNew());
-			ILogFile log = await _logService.CreateLogFileAsync(JobId.Empty, leaseId, agent.SessionId, LogType.Json, cancellationToken: cancellationToken);
+			ILog log = await _logService.CreateLogAsync(JobId.Empty, leaseId, agent.SessionId, LogType.Json, cancellationToken: cancellationToken);
 
 			RestartTask task = new RestartTask();
 			task.LogId = log.Id.ToString();
diff --git a/Engine/Source/Programs/Horde/Horde.Server/Tasks/ShutdownTaskSource.cs b/Engine/Source/Programs/Horde/Horde.Server/Tasks/ShutdownTaskSource.cs
--- a/Engine/Source/Programs/Horde/Horde.Server/Tasks/ShutdownTaskSource.cs
+++ b/Engine/Source/Programs/Horde/Horde.Server/Tasks/ShutdownTaskSource.cs
@@ -20,9 +20,9 @@ namespace Horde.Server.Tasks
 
 		public override TaskSourceFlags Flags => TaskSourceFlags.AllowWhenDisabled | TaskSourceFlags.AllowDuringDowntime;
 
-		readonly ILogFileService _logService;
+		readonly ILogService _logService;
 
-		public ShutdownTaskSource(ILogFileService logService)
+		public ShutdownTaskSource(ILogService logService)
 		{
 			_logService = logService;
 			OnLeaseStartedProperties.Add(nameof(ShutdownTask.LogId), x => LogId.Parse(x.LogId));
@@ -40,7 +40,7 @@ namespace Horde.Server.Tasks
 			}
 
 			LeaseId leaseId = new LeaseId(BinaryIdUtils.CreateNew());
-			ILogFile log = await _logService.CreateLogFileAsync(JobId.Empty, leaseId, agent.SessionId, LogType.Json, cancellationToken: cancellationToken);
+			ILog log = await _logService.CreateLogAsync(JobId.Empty, leaseId, agent.SessionId, LogType.Json, cancellationToken: cancellationToken);
 
 			ShutdownTask task = new ShutdownTask();
 			task.LogId = log.Id.ToString();
diff --git a/Engine/Source/Programs/Horde/Horde.Server/Tasks/UpgradeTaskSource.cs b/Engine/Source/Programs/Horde/Horde.Server/Tasks/UpgradeTaskSource.cs
--- a/Engine/Source/Programs/Horde/Horde.Server/Tasks/UpgradeTaskSource.cs
+++ b/Engine/Source/Programs/Horde/Horde.Server/Tasks/UpgradeTaskSource.cs
@@ -28,12 +28,12 @@ namespace Horde.Server.Tasks
 		public override TaskSourceFlags Flags => TaskSourceFlags.AllowWhenDisabled | TaskSourceFlags.AllowDuringDowntime;
 
 		readonly IToolCollection _toolCollection;
-		readonly ILogFileService _logService;
+		readonly ILogService _logService;
 		readonly IOptionsMonitor<GlobalConfig> _globalConfig;
 		readonly IOptions<ServerSettings> _serverSettings;
 		readonly IClock _clock;
 
-		public UpgradeTaskSource(IToolCollection toolCollection, ILogFileService logService, IOptionsMonitor<GlobalConfig> globalConfig, IOptions<ServerSettings> serverSettings, IClock clock)
+		public UpgradeTaskSource(IToolCollection toolCollection, ILogService logService, IOptionsMonitor<GlobalConfig> globalConfig, IOptions<ServerSettings> serverSettings, IClock clock)
 		{
 			_toolCollection = toolCollection;
 			_logService = logService;
@@ -69,14 +69,14 @@ namespace Horde.Server.Tasks
 			}
 
 			LeaseId leaseId = new LeaseId(BinaryIdUtils.CreateNew());
-			ILogFile logFile = await _logService.CreateLogFileAsync(JobId.Empty, leaseId, agent.SessionId, LogType.Json, cancellationToken: cancellationToken);
+			ILog log = await _logService.CreateLogAsync(JobId.Empty, leaseId, agent.SessionId, LogType.Json, cancellationToken: cancellationToken);
 
 			UpgradeTask task = new UpgradeTask();
 			task.SoftwareId = $"{tool.Id}:{deployment.Version}";
-			task.LogId = logFile.Id.ToString();
+			task.LogId = log.Id.ToString();
 
 			byte[] payload = Any.Pack(task).ToByteArray();
-			return LeaseAsync(new AgentLease(leaseId, null, $"Upgrade to {tool.Id} {deployment.Version}", null, null, logFile.Id, LeaseState.Pending, null, true, payload));
+			return LeaseAsync(new AgentLease(leaseId, null, $"Upgrade to {tool.Id} {deployment.Version}", null, null, log.Id, LeaseState.Pending, null, true, payload));
 		}
 
 		/// <summary>
diff --git a/Engine/Source/Programs/Horde/Horde.Server/Ugs/UgsController.cs b/Engine/Source/Programs/Horde/Horde.Server/Ugs/UgsController.cs
--- a/Engine/Source/Programs/Horde/Horde.Server/Ugs/UgsController.cs
+++ b/Engine/Source/Programs/Horde/Horde.Server/Ugs/UgsController.cs
@@ -44,7 +44,7 @@ namespace Horde.Server.Ugs
 		/// <summary>
 		/// The log file service
 		/// </summary>
-		private readonly ILogFileService _logFileService;
+		private readonly ILogService _logService;
 
 		/// <summary>
 		/// Server settings
@@ -59,12 +59,12 @@ namespace Horde.Server.Ugs
 		/// <summary>
 		/// Constructor
 		/// </summary>
-		public UgsController(IssueService issueService, IUgsMetadataCollection ugsMetadataCollection, IUserCollection userCollection, ILogFileService logFileService, IOptionsMonitor<ServerSettings> optionsMonitor, ILogger<UgsController> logger)
+		public UgsController(IssueService issueService, IUgsMetadataCollection ugsMetadataCollection, IUserCollection userCollection, ILogService logService, IOptionsMonitor<ServerSettings> optionsMonitor, ILogger<UgsController> logger)
 		{
 			_issueService = issueService;
 			_ugsMetadataCollection = ugsMetadataCollection;
 			_userCollection = userCollection;
-			_logFileService = logFileService;
+			_logService = logService;
 			_settings = optionsMonitor.CurrentValue;
 			_logger = logger;
 		}
@@ -270,22 +270,22 @@ namespace Horde.Server.Ugs
 		{
 			List<GetUgsIssueDiagnosticResponse> diagnostics = new List<GetUgsIssueDiagnosticResponse>();
 
-			Dictionary<LogId, ILogFile?> logFiles = new Dictionary<LogId, ILogFile?>();
+			Dictionary<LogId, ILog?> logs = new Dictionary<LogId, ILog?>();
 
 			IReadOnlyList<IIssueSpan> spans = await _issueService.Collection.FindSpansAsync(issueId, cancellationToken);
-			List<ILogEvent> events = await _logFileService.FindEventsForSpansAsync(spans.Select(x => x.Id), null, 0, count: 10, cancellationToken);
+			List<ILogEvent> events = await _logService.FindEventsForSpansAsync(spans.Select(x => x.Id), null, 0, count: 10, cancellationToken);
 
 			foreach (ILogEvent logEvent in events)
 			{
-				ILogFile? logFile;
-				if (!logFiles.TryGetValue(logEvent.LogId, out logFile))
+				ILog? log;
+				if (!logs.TryGetValue(logEvent.LogId, out log))
 				{
-					logFile = await _logFileService.GetLogFileAsync(logEvent.LogId, cancellationToken);
-					logFiles.Add(logEvent.LogId, logFile);
+					log = await _logService.GetLogAsync(logEvent.LogId, cancellationToken);
+					logs.Add(logEvent.LogId, log);
 				}
-				if (logFile != null)
+				if (log != null)
 				{
-					ILogEventData eventData = await _logFileService.GetEventDataAsync(logFile, logEvent.LineIndex, logEvent.LineCount, cancellationToken);
+					ILogEventData eventData = await _logService.GetEventDataAsync(log, logEvent.LineIndex, logEvent.LineCount, cancellationToken);
 					long buildId = logEvent.LogId.GetHashCode();
 					Uri url = new Uri(_settings.DashboardUrl, $"log/{logEvent.LogId}?lineindex={logEvent.LineIndex}");
 
diff --git a/Engine/Source/Programs/Shared/EpicGames.Horde/HordeHttpClient.cs b/Engine/Source/Programs/Shared/EpicGames.Horde/HordeHttpClient.cs
--- a/Engine/Source/Programs/Shared/EpicGames.Horde/HordeHttpClient.cs
+++ b/Engine/Source/Programs/Shared/EpicGames.Horde/HordeHttpClient.cs
@@ -452,27 +452,27 @@ namespace EpicGames.Horde
 		/// <summary>
 		/// Get the given log file 
 		/// </summary>
-		/// <param name="logFileId">Id of the log file to retrieve</param>
+		/// <param name="logId">Id of the log file to retrieve</param>
 		/// <param name="searchText">Text to search for in the log</param>
 		/// <param name="count">Number of lines to return (default 5)</param>
 		/// <param name="cancellationToken">Cancellation token for the operation</param>
 		/// <returns></returns>
-		public Task<SearchLogResponse> GetSearchLogAsync(LogId logFileId,string searchText, int count = 5, CancellationToken cancellationToken = default)
+		public Task<SearchLogResponse> GetSearchLogAsync(LogId logId, string searchText, int count = 5, CancellationToken cancellationToken = default)
 		{
-			return GetAsync<SearchLogResponse>(_httpClient, $"/api/v1/logs/{logFileId}/search?Text={Uri.EscapeDataString(searchText)}&count={count}", cancellationToken);
+			return GetAsync<SearchLogResponse>(_httpClient, $"/api/v1/logs/{logId}/search?Text={Uri.EscapeDataString(searchText)}&count={count}", cancellationToken);
 		}
 
 		/// <summary>
 		/// Get the requested number of lines from given logFileId, starting at index
 		/// </summary>
-		/// <param name="logFileId">Id of log file to retrieve lines from</param>
+		/// <param name="logId">Id of log file to retrieve lines from</param>
 		/// <param name="startIndex">Start index of lines to retrieve</param>
 		/// <param name="count">Number of lines to retrieve</param>
 		/// <param name="cancellationToken">Cancellation token for the operation</param>
 		/// <returns></returns>
-		public Task<LogLinesResponse> GetLogLinesAsync(LogId logFileId, int startIndex, int count, CancellationToken cancellationToken = default)
+		public Task<LogLinesResponse> GetLogLinesAsync(LogId logId, int startIndex, int count, CancellationToken cancellationToken = default)
 		{
-			return GetAsync<LogLinesResponse>(_httpClient, $"/api/v1/logs/{logFileId}/lines?index={startIndex}&count={count}", cancellationToken);
+			return GetAsync<LogLinesResponse>(_httpClient, $"/api/v1/logs/{logId}/lines?index={startIndex}&count={count}", cancellationToken);
 		}
 
 		#endregion
