Horde: Remove unused ILogCollection methods, and move update methods onto ILog itself.

#rnx

diff --git a/Engine/Source/Programs/Horde/Horde.Server.Tests/Jobs/JobExpirationTests.cs b/Engine/Source/Programs/Horde/Horde.Server.Tests/Jobs/JobExpirationTests.cs
--- a/Engine/Source/Programs/Horde/Horde.Server.Tests/Jobs/JobExpirationTests.cs
+++ b/Engine/Source/Programs/Horde/Horde.Server.Tests/Jobs/JobExpirationTests.cs
@@ -1,9 +1,6 @@
 // Copyright Epic Games, Inc. All Rights Reserved.
 
 using System;
-using System.Collections.Generic;
-using System.Linq;
-using System.Text;
 using System.Threading.Tasks;
 using EpicGames.Horde.Jobs.Templates;
 using EpicGames.Horde.Projects;
diff --git a/Engine/Source/Programs/Horde/Horde.Server.Tests/Logs/LogIndexingTests.cs b/Engine/Source/Programs/Horde/Horde.Server.Tests/Logs/LogIndexingTests.cs
--- a/Engine/Source/Programs/Horde/Horde.Server.Tests/Logs/LogIndexingTests.cs
+++ b/Engine/Source/Programs/Horde/Horde.Server.Tests/Logs/LogIndexingTests.cs
@@ -21,16 +21,14 @@ namespace Horde.Server.Tests.Logs
 	class TestLogWriter : IAsyncDisposable
 	{
 		ILog _logFile;
-		readonly ILogCollection _logCollection;
 		readonly IStorageClient _storageClient;
 		readonly IBlobWriter _blobWriter;
 		readonly LogBuilder _builder;
 		int _lineCount;
 
-		public TestLogWriter(ILog logFile, ILogCollection logCollection, StorageService storageService)
+		public TestLogWriter(ILog logFile, StorageService storageService)
 		{
 			_logFile = logFile;
-			_logCollection = logCollection;
 			_storageClient = storageService.CreateClient(Namespace.Logs);
 			_blobWriter = _storageClient.CreateBlobWriter(logFile.RefName);
 			_builder = new LogBuilder((logFile.Type == LogType.Text) ? LogFormat.Text : LogFormat.Json, NullLogger.Instance);
@@ -46,7 +44,7 @@ namespace Horde.Server.Tests.Logs
 		{
 			_builder.WriteData(data);
 			_lineCount += data.Span.Count((byte)'\n');
-			_logFile = await _logCollection.UpdateLineCountAsync(_logFile, _lineCount, false, CancellationToken.None);
+			_logFile = await _logFile.UpdateLineCountAsync(_lineCount, false, CancellationToken.None);
 			return await FlushAsync(false);
 		}
 
@@ -54,7 +52,7 @@ namespace Horde.Server.Tests.Logs
 		{
 			IBlobRef<LogNode> handle = await _builder.FlushAsync(_blobWriter, complete, CancellationToken.None);
 			await _storageClient.WriteRefAsync(_logFile.RefName, handle);
-			_logFile = await _logCollection.UpdateLineCountAsync(_logFile, _lineCount, complete, CancellationToken.None);
+			_logFile = await _logFile.UpdateLineCountAsync(_lineCount, complete, CancellationToken.None);
 			return _logFile;
 		}
 	}
@@ -71,7 +69,7 @@ namespace Horde.Server.Tests.Logs
 			ILog logFile = await LogService.CreateLogAsync(jobId, null, null, LogType.Text);
 
 			// Write the test data to the log file in blocks
-			await using (TestLogWriter logWriter = new TestLogWriter(logFile, LogCollection, StorageService))
+			await using (TestLogWriter logWriter = new TestLogWriter(logFile, StorageService))
 			{
 				int offset = 0;
 				int lineIndex = 0;
@@ -150,7 +148,7 @@ namespace Horde.Server.Tests.Logs
 			};
 
 			int length = 0;
-			await using (TestLogWriter logWriter = new TestLogWriter(logFile, LogCollection, StorageService))
+			await using (TestLogWriter logWriter = new TestLogWriter(logFile, StorageService))
 			{
 				for (int lineIdx = 0; lineIdx < lines.Length; lineIdx++)
 				{
@@ -187,7 +185,7 @@ namespace Horde.Server.Tests.Logs
 			JobId jobId = JobIdUtils.GenerateNewId();
 			ILog logFile = await LogService.CreateLogAsync(jobId, null, null, LogType.Text);
 
-			await using (TestLogWriter writer = new TestLogWriter(logFile, LogCollection, StorageService))
+			await using (TestLogWriter writer = new TestLogWriter(logFile, StorageService))
 			{
 				logFile = await writer.WriteDataAsync(Encoding.UTF8.GetBytes("abc\n"));
 				logFile = await writer.WriteDataAsync(Encoding.UTF8.GetBytes("def\n"));
diff --git a/Engine/Source/Programs/Horde/Horde.Server.Tests/Logs/LogServiceTests.cs b/Engine/Source/Programs/Horde/Horde.Server.Tests/Logs/LogServiceTests.cs
--- a/Engine/Source/Programs/Horde/Horde.Server.Tests/Logs/LogServiceTests.cs
+++ b/Engine/Source/Programs/Horde/Horde.Server.Tests/Logs/LogServiceTests.cs
@@ -44,7 +44,7 @@ namespace Horde.Server.Tests.Logs
 			JobId jobId = JobIdUtils.GenerateNewId();
 			ILog logFile = await LogService.CreateLogAsync(jobId, null, null, LogType.Text, logId: null, cancellationToken: CancellationToken.None);
 
-			await using (TestLogWriter writer = new TestLogWriter(logFile, LogCollection, StorageService))
+			await using (TestLogWriter writer = new TestLogWriter(logFile, StorageService))
 			{
 				await writer.WriteDataAsync(Encoding.ASCII.GetBytes("hello\n"));
 				await writer.FlushAsync(false);
@@ -104,7 +104,7 @@ namespace Horde.Server.Tests.Logs
 			int lineIndex = 0;
 			int offset = 0;
 
-			await using TestLogWriter logWriter = new TestLogWriter(log, LogCollection, StorageService);
+			await using TestLogWriter logWriter = new TestLogWriter(log, StorageService);
 
 			// First write with flush. Will become chunk #1
 			log = await logWriter.WriteDataAsync(Encoding.ASCII.GetBytes(str1));
@@ -175,9 +175,6 @@ namespace Horde.Server.Tests.Logs
 		[TestMethod]
 		public async Task GetLogFileTestAsync()
 		{
-			await GetMongoServiceSingleton().Database.DropCollectionAsync("LogFiles");
-			Assert.AreEqual(0, (await LogService.GetLogsAsync()).Count);
-
 			// Will implicitly test GetLogFileAsync(), AddCachedLogFile()
 			JobId jobId = JobIdUtils.GenerateNewId();
 			SessionId sessionId = SessionIdUtils.GenerateNewId();
@@ -189,9 +186,6 @@ namespace Horde.Server.Tests.Logs
 
 			ILog? notFound = await LogService.GetLogAsync(LogIdUtils.GenerateNewId(), CancellationToken.None);
 			Assert.IsNull(notFound);
-
-			await LogService.CreateLogAsync(JobIdUtils.GenerateNewId(), null, SessionIdUtils.GenerateNewId(), LogType.Text, logId: null, cancellationToken: CancellationToken.None);
-			Assert.AreEqual(2, (await LogService.GetLogsAsync()).Count);
 		}
 
 		[TestMethod]
@@ -221,7 +215,7 @@ namespace Horde.Server.Tests.Logs
 		{
 			JobId jobId = JobIdUtils.GenerateNewId();
 			ILog logFile = await LogService.CreateLogAsync(jobId, null, null, LogType.Text, logId: null, cancellationToken: CancellationToken.None);
-			await using TestLogWriter logWriter = new TestLogWriter(logFile, LogCollection, StorageService);
+			await using TestLogWriter logWriter = new TestLogWriter(logFile, StorageService);
 
 			byte[] line1 = Encoding.UTF8.GetBytes("hello world\n");
 			await logWriter.WriteDataAsync(line1);
diff --git a/Engine/Source/Programs/Horde/Horde.Server/Logs/ILog.cs b/Engine/Source/Programs/Horde/Horde.Server/Logs/ILog.cs
--- a/Engine/Source/Programs/Horde/Horde.Server/Logs/ILog.cs
+++ b/Engine/Source/Programs/Horde/Horde.Server/Logs/ILog.cs
@@ -2,6 +2,8 @@
 
 using System;
 using System.Collections.Generic;
+using System.Threading;
+using System.Threading.Tasks;
 using EpicGames.Core;
 using EpicGames.Horde.Agents.Leases;
 using EpicGames.Horde.Agents.Sessions;
@@ -106,6 +108,15 @@ namespace Horde.Server.Logs
 		/// Whether the log is complete (V2 storage backend)
 		/// </summary>
 		public bool Complete { get; }
+
+		/// <summary>
+		/// Updates the line count for a log file (v2 backend only)
+		/// </summary>
+		/// <param name="lineCount">New line count for the log file</param>
+		/// <param name="complete">Flag indicating whether the log is complete, or can still be tailed for additional data</param>
+		/// <param name="cancellationToken">Cancellation token for the call</param>
+		/// <returns>The updated log file document</returns>
+		Task<ILog> UpdateLineCountAsync(int lineCount, bool complete, CancellationToken cancellationToken = default);
 	}
 
 	/// <summary>
diff --git a/Engine/Source/Programs/Horde/Horde.Server/Logs/ILogCollection.cs b/Engine/Source/Programs/Horde/Horde.Server/Logs/ILogCollection.cs
--- a/Engine/Source/Programs/Horde/Horde.Server/Logs/ILogCollection.cs
+++ b/Engine/Source/Programs/Horde/Horde.Server/Logs/ILogCollection.cs
@@ -10,40 +10,6 @@ using EpicGames.Horde.Logs;
 
 namespace Horde.Server.Logs
 {
-	/// <summary>
-	/// Updates a log file chunk
-	/// </summary>
-	public class CompleteLogChunkUpdate
-	{
-		/// <summary>
-		/// Index of the chunk
-		/// </summary>
-		public int Index { get; set; }
-
-		/// <summary>
-		/// New length for the chunk
-		/// </summary>
-		public int Length { get; set; }
-
-		/// <summary>
-		/// Number of lines in the chunk
-		/// </summary>
-		public int LineCount { get; set; }
-
-		/// <summary>
-		/// Constructor
-		/// </summary>
-		/// <param name="index">Index of the chunk</param>
-		/// <param name="length">New length for the chunk</param>
-		/// <param name="lineCount">Number of lines in the chunk</param>
-		public CompleteLogChunkUpdate(int index, int length, int lineCount)
-		{
-			Index = index;
-			Length = length;
-			LineCount = lineCount;
-		}
-	}
-
 	/// <summary>
 	/// Wrapper around the jobs collection in a mongo DB
 	/// </summary>
@@ -61,44 +27,6 @@ namespace Horde.Server.Logs
 		/// <returns>The new log file document</returns>
 		Task<ILog> CreateLogAsync(JobId jobId, LeaseId? leaseId, SessionId? sessionId, LogType type, LogId? logId = null, CancellationToken cancellationToken = default);
 
-		/// <summary>
-		/// Updates the line count for a log file (v2 backend only)
-		/// </summary>
-		/// <param name="log">The current log file</param>
-		/// <param name="lineCount">New line count for the log file</param>
-		/// <param name="complete">Flag indicating whether the log is complete, or can still be tailed for additional data</param>
-		/// <param name="cancellationToken">Cancellation token for the call</param>
-		/// <returns>The updated log file document</returns>
-		Task<ILog> UpdateLineCountAsync(ILog log, int lineCount, bool complete, CancellationToken cancellationToken);
-
-		/// <summary>
-		/// Adds a new chunk
-		/// </summary>
-		/// <param name="log">The current log file</param>
-		/// <param name="offset">Offset of the new chunk</param>
-		/// <param name="lineIndex">Line index for the start of the chunk</param>
-		/// <param name="cancellationToken">Cancellation token for the call</param>
-		/// <returns>The updated log file document</returns>
-		Task<ILog?> TryAddChunkAsync(ILog log, long offset, int lineIndex, CancellationToken cancellationToken);
-
-		/// <summary>
-		/// Update the log file with final information about certain chunks
-		/// </summary>
-		/// <param name="log">The current log file</param>
-		/// <param name="chunks">Chunks to update. New chunks will be inserted</param>
-		/// <param name="cancellationToken">Cancellation token for the call</param>
-		/// <returns>The updated log file document</returns>
-		Task<ILog?> TryCompleteChunksAsync(ILog log, IEnumerable<CompleteLogChunkUpdate> chunks, CancellationToken cancellationToken);
-
-		/// <summary>
-		/// Update the log file with final information about the index
-		/// </summary>
-		/// <param name="log">The current log file</param>
-		/// <param name="newIndexLength">New length of the index</param>
-		/// <param name="cancellationToken">Cancellation token for the call</param>
-		/// <returns>The updated log file document</returns>
-		Task<ILog?> TryUpdateIndexAsync(ILog log, long newIndexLength, CancellationToken cancellationToken);
-
 		/// <summary>
 		/// Gets a log by ID
 		/// </summary>
@@ -106,14 +34,5 @@ namespace Horde.Server.Logs
 		/// <param name="cancellationToken">Cancellation token for the call</param>
 		/// <returns>The log instance</returns>
 		Task<ILog?> GetLogAsync(LogId logId, CancellationToken cancellationToken);
-
-		/// <summary>
-		/// Gets all the log files
-		/// </summary>
-		/// <param name="index">Index of the first result to return</param>
-		/// <param name="count">Number of results to return</param>
-		/// <param name="cancellationToken">Cancellation token for the call</param>
-		/// <returns>List of log files</returns>
-		Task<List<ILog>> GetLogsAsync(int? index = null, int? count = null, CancellationToken cancellationToken = default);
 	}
 }
diff --git a/Engine/Source/Programs/Horde/Horde.Server/Logs/LogCollection.cs b/Engine/Source/Programs/Horde/Horde.Server/Logs/LogCollection.cs
--- a/Engine/Source/Programs/Horde/Horde.Server/Logs/LogCollection.cs
+++ b/Engine/Source/Programs/Horde/Horde.Server/Logs/LogCollection.cs
@@ -2,7 +2,6 @@
 
 using System.Collections.Generic;
 using System.Linq;
-using System.Net;
 using System.Threading;
 using System.Threading.Tasks;
 using EpicGames.Horde.Agents.Leases;
@@ -50,7 +49,7 @@ namespace Horde.Server.Logs
 			}
 		}
 
-		class LogDocument : ILog
+		class LogDocument
 		{
 			[BsonRequired, BsonId]
 			public LogId Id { get; set; }
@@ -86,8 +85,6 @@ namespace Horde.Server.Logs
 			[BsonRequired]
 			public int UpdateIndex { get; set; }
 
-			IReadOnlyList<ILogChunk> ILog.Chunks => Chunks;
-
 			[BsonConstructor]
 			private LogDocument()
 			{
@@ -105,12 +102,37 @@ namespace Horde.Server.Logs
 				NamespaceId = namespaceId;
 				RefName = new RefName(Id.ToString());
 			}
+		}
 
-			public LogDocument Clone()
+		class Log : ILog
+		{
+			readonly LogCollection _collection;
+			readonly LogDocument _document;
+
+			public LogId Id => _document.Id;
+			public JobId JobId => _document.JobId;
+			public LeaseId? LeaseId => _document.LeaseId;
+			public SessionId? SessionId => _document.SessionId;
+			public bool UseNewStorageBackend => _document.UseNewStorageBackend;
+			public int? MaxLineIndex => _document.MaxLineIndex;
+			public long? IndexLength => _document.IndexLength;
+			public LogType Type => _document.Type;
+			public IReadOnlyList<ILogChunk> Chunks => _document.Chunks;
+			public NamespaceId NamespaceId => _document.NamespaceId;
+			public RefName RefName => _document.RefName;
+			public int LineCount => _document.LineCount;
+			public bool Complete => _document.Complete;
+			
+			public Log(LogCollection collection, LogDocument document)
 			{
-				LogDocument document = (LogDocument)MemberwiseClone();
-				document.Chunks = document.Chunks.ConvertAll(x => x.Clone());
-				return document;
+				_collection = collection;
+				_document = document;
+			}
+
+			public async Task<ILog> UpdateLineCountAsync(int lineCount, bool complete, CancellationToken cancellationToken = default)
+			{
+				LogDocument newDocument = await _collection.UpdateLineCountAsync(_document, lineCount, complete, cancellationToken);
+				return new Log(_collection, newDocument);
 			}
 		}
 
@@ -119,11 +141,6 @@ namespace Horde.Server.Logs
 		/// </summary>
 		readonly IMongoCollection<LogDocument> _logCollection;
 
-		/// <summary>
-		/// Hostname for the current server
-		/// </summary>
-		readonly string _hostName;
-
 		/// <summary>
 		/// Constructor
 		/// </summary>
@@ -131,7 +148,6 @@ namespace Horde.Server.Logs
 		public LogCollection(MongoService mongoService)
 		{
 			_logCollection = mongoService.GetCollection<LogDocument>("LogFiles");
-			_hostName = Dns.GetHostName();
 		}
 
 		/// <inheritdoc/>
@@ -139,121 +155,22 @@ namespace Horde.Server.Logs
 		{
 			LogDocument newLog = new LogDocument(jobId, leaseId, sessionId, type, logId, Namespace.Logs);
 			await _logCollection.InsertOneAsync(newLog, null, cancellationToken);
-			return newLog;
+			return new Log(this, newLog);
 		}
 
 		/// <inheritdoc/>
-		public async Task<ILog> UpdateLineCountAsync(ILog log, int lineCount, bool complete, CancellationToken cancellationToken)
+		async Task<LogDocument> UpdateLineCountAsync(LogDocument log, int lineCount, bool complete, CancellationToken cancellationToken)
 		{
 			FilterDefinition<LogDocument> filter = Builders<LogDocument>.Filter.Eq(x => x.Id, log.Id);
 			UpdateDefinition<LogDocument> update = Builders<LogDocument>.Update.Set(x => x.LineCount, lineCount).Set(x => x.Complete, complete).Inc(x => x.UpdateIndex, 1);
 			return await _logCollection.FindOneAndUpdateAsync(filter, update, new FindOneAndUpdateOptions<LogDocument, LogDocument> { ReturnDocument = ReturnDocument.After }, cancellationToken);
 		}
 
-		/// <inheritdoc/>
-		public async Task<ILog?> TryAddChunkAsync(ILog logInterface, long offset, int lineIndex, CancellationToken cancellationToken)
-		{
-			LogDocument log = ((LogDocument)logInterface).Clone();
-
-			int chunkIdx = log.Chunks.GetChunkForOffset(offset) + 1;
-
-			LogChunkDocument chunk = new LogChunkDocument();
-			chunk.Offset = offset;
-			chunk.LineIndex = lineIndex;
-			chunk.Server = _hostName;
-			log.Chunks.Insert(chunkIdx, chunk);
-
-			UpdateDefinition<LogDocument> update = Builders<LogDocument>.Update.Set(x => x.Chunks, log.Chunks);
-			if (chunkIdx == log.Chunks.Count - 1)
-			{
-				log.MaxLineIndex = null;
-				update = update.Unset(x => x.MaxLineIndex);
-			}
-
-			if (!await TryUpdateLogAsync(log, update, cancellationToken))
-			{
-				return null;
-			}
-
-			return log;
-		}
-
-		/// <inheritdoc/>
-		public async Task<ILog?> TryCompleteChunksAsync(ILog logInterface, IEnumerable<CompleteLogChunkUpdate> chunkUpdates, CancellationToken cancellationToken)
-		{
-			LogDocument log = ((LogDocument)logInterface).Clone();
-
-			// Update the length of any complete chunks
-			UpdateDefinitionBuilder<LogDocument> updateBuilder = Builders<LogDocument>.Update;
-			List<UpdateDefinition<LogDocument>> updates = new List<UpdateDefinition<LogDocument>>();
-			foreach (CompleteLogChunkUpdate chunkUpdate in chunkUpdates)
-			{
-				LogChunkDocument chunk = log.Chunks[chunkUpdate.Index];
-				chunk.Length = chunkUpdate.Length;
-				updates.Add(updateBuilder.Set(x => x.Chunks[chunkUpdate.Index].Length, chunkUpdate.Length));
-
-				if (chunkUpdate.Index == log.Chunks.Count - 1)
-				{
-					log.MaxLineIndex = chunk.LineIndex + chunkUpdate.LineCount;
-					updates.Add(updateBuilder.Set(x => x.MaxLineIndex, log.MaxLineIndex));
-				}
-			}
-
-			// Try to apply the updates
-			if (updates.Count > 0 && !await TryUpdateLogAsync(log, updateBuilder.Combine(updates), cancellationToken))
-			{
-				return null;
-			}
-
-			return log;
-		}
-
-		/// <inheritdoc/>
-		public async Task<ILog?> TryUpdateIndexAsync(ILog logInterface, long newIndexLength, CancellationToken cancellationToken)
-		{
-			LogDocument log = ((LogDocument)logInterface).Clone();
-
-			UpdateDefinition<LogDocument> update = Builders<LogDocument>.Update.Set(x => x.IndexLength, newIndexLength);
-			if (!await TryUpdateLogAsync(log, update, cancellationToken))
-			{
-				return null;
-			}
-
-			log.IndexLength = newIndexLength;
-			return log;
-		}
-
-		/// <inheritdoc/>
-		private async Task<bool> TryUpdateLogAsync(LogDocument current, UpdateDefinition<LogDocument> update, CancellationToken cancellationToken)
-		{
-			int prevUpdateIndex = current.UpdateIndex;
-			current.UpdateIndex++;
-			UpdateResult result = await _logCollection.UpdateOneAsync<LogDocument>(x => x.Id == current.Id && x.UpdateIndex == prevUpdateIndex, update.Set(x => x.UpdateIndex, current.UpdateIndex), cancellationToken: cancellationToken);
-			return result.ModifiedCount == 1;
-		}
-
 		/// <inheritdoc/>
 		public async Task<ILog?> GetLogAsync(LogId logId, CancellationToken cancellationToken)
 		{
 			LogDocument log = await _logCollection.Find<LogDocument>(x => x.Id == logId).FirstOrDefaultAsync(cancellationToken);
-			return log;
-		}
-
-		/// <inheritdoc/>
-		public async Task<List<ILog>> GetLogsAsync(int? index = null, int? count = null, CancellationToken cancellationToken = default)
-		{
-			IFindFluent<LogDocument, LogDocument> query = _logCollection.Find(FilterDefinition<LogDocument>.Empty);
-			if (index != null)
-			{
-				query = query.Skip(index.Value);
-			}
-			if (count != null)
-			{
-				query = query.Limit(count.Value);
-			}
-
-			List<LogDocument> results = await query.ToListAsync(cancellationToken);
-			return results.ConvertAll<ILog>(x => x);
+			return new Log(this, log);
 		}
 	}
 }
diff --git a/Engine/Source/Programs/Horde/Horde.Server/Logs/LogRpcService.cs b/Engine/Source/Programs/Horde/Horde.Server/Logs/LogRpcService.cs
--- a/Engine/Source/Programs/Horde/Horde.Server/Logs/LogRpcService.cs
+++ b/Engine/Source/Programs/Horde/Horde.Server/Logs/LogRpcService.cs
@@ -22,7 +22,6 @@ namespace Horde.Server.Logs
 	public class LogRpcService : LogRpc.LogRpcBase
 	{
 		readonly ILogService _logService;
-		readonly ILogCollection _logCollection;
 		readonly LogTailService _logTailService;
 		readonly StorageService _storageService;
 		readonly ILogger<LogRpcService> _logger;
@@ -30,10 +29,9 @@ namespace Horde.Server.Logs
 		/// <summary>
 		/// Constructor
 		/// </summary>
-		public LogRpcService(ILogService logService, ILogCollection logCollection, LogTailService logTailService, StorageService storageService, ILogger<LogRpcService> logger)
+		public LogRpcService(ILogService logService, LogTailService logTailService, StorageService storageService, ILogger<LogRpcService> logger)
 		{
 			_logService = logService;
-			_logCollection = logCollection;
 			_logTailService = logTailService;
 			_storageService = storageService;
 			_logger = logger;
@@ -65,7 +63,7 @@ namespace Horde.Server.Logs
 			IBlobRef target = store.CreateBlobRef(hash, new BlobLocator(request.TargetLocator));
 			await store.WriteRefAsync(new RefName(request.LogId), target);
 
-			await _logCollection.UpdateLineCountAsync(log, request.LineCount, request.Complete, CancellationToken.None);
+			await log.UpdateLineCountAsync(request.LineCount, request.Complete, context.CancellationToken);
 
 			await _logTailService.FlushAsync(log.Id, request.LineCount);
 
diff --git a/Engine/Source/Programs/Horde/Horde.Server/Logs/LogService.cs b/Engine/Source/Programs/Horde/Horde.Server/Logs/LogService.cs
--- a/Engine/Source/Programs/Horde/Horde.Server/Logs/LogService.cs
+++ b/Engine/Source/Programs/Horde/Horde.Server/Logs/LogService.cs
@@ -72,15 +72,6 @@ namespace Horde.Server.Logs
 		/// <returns>The log  document</returns>
 		Task<ILog?> GetLogAsync(LogId logId, CancellationToken cancellationToken);
 
-		/// <summary>
-		/// Returns a list of log files
-		/// </summary>
-		/// <param name="index">Index of the first result to return</param>
-		/// <param name="count">Number of results to return</param>
-		/// <param name="cancellationToken">Cancellation token for the call</param>
-		/// <returns>List of log documents</returns>
-		Task<List<ILog>> GetLogsAsync(int? index = null, int? count = null, CancellationToken cancellationToken = default);
-
 		/// <summary>
 		/// Read a set of lines from the given log file
 		/// </summary>
@@ -561,12 +552,6 @@ namespace Horde.Server.Logs
 			return await _logCollection.GetLogAsync(logId, cancellationToken);
 		}
 
-		/// <inheritdoc/>
-		public Task<List<ILog>> GetLogsAsync(int? index = null, int? count = null, CancellationToken cancellationToken = default)
-		{
-			return _logCollection.GetLogsAsync(index, count, cancellationToken);
-		}
-
 		/// <inheritdoc/>
 		public async Task<List<Utf8String>> ReadLinesAsync(ILog log, int index, int count, CancellationToken cancellationToken)
 		{
