Horde: Fix test failure (was wrapping a null log document in a Log object).

#rnx

diff --git a/Engine/Source/Programs/Horde/Horde.Server/Logs/LogCollection.cs b/Engine/Source/Programs/Horde/Horde.Server/Logs/LogCollection.cs
--- a/Engine/Source/Programs/Horde/Horde.Server/Logs/LogCollection.cs
+++ b/Engine/Source/Programs/Horde/Horde.Server/Logs/LogCollection.cs
@@ -169,7 +169,11 @@ namespace Horde.Server.Logs
 		/// <inheritdoc/>
 		public async Task<ILog?> GetLogAsync(LogId logId, CancellationToken cancellationToken)
 		{
-			LogDocument log = await _logCollection.Find<LogDocument>(x => x.Id == logId).FirstOrDefaultAsync(cancellationToken);
+			LogDocument? log = await _logCollection.Find<LogDocument>(x => x.Id == logId).FirstOrDefaultAsync(cancellationToken);
+			if (log == null)
+			{
+				return null;
+			}
 			return new Log(this, log);
 		}
 	}
