Horde: Fix static analysis warnings.

#rnx

diff --git a/Engine/Source/Programs/Horde/Horde.Server/Logs/ILogCollection.cs b/Engine/Source/Programs/Horde/Horde.Server/Logs/ILogCollection.cs
--- a/Engine/Source/Programs/Horde/Horde.Server/Logs/ILogCollection.cs
+++ b/Engine/Source/Programs/Horde/Horde.Server/Logs/ILogCollection.cs
@@ -1,6 +1,5 @@
 // Copyright Epic Games, Inc. All Rights Reserved.
 
-using System.Collections.Generic;
 using System.Threading;
 using System.Threading.Tasks;
 using EpicGames.Horde.Agents.Leases;
