Horde: Refactoring tools API.

- Operations are now performed directly on ITool and IToolDeployment instances rather than on IToolCollection.
- ToolCollection now enumerates available tools (rather than relying on the caller to supply a GlobalConfig instance), and captures the relevant ToolConfig instance in each ITool.

#rnx

diff --git a/Engine/Source/Programs/Horde/Horde.Server.Tests/Tools/ToolTests.cs b/Engine/Source/Programs/Horde/Horde.Server.Tests/Tools/ToolTests.cs
--- a/Engine/Source/Programs/Horde/Horde.Server.Tests/Tools/ToolTests.cs
+++ b/Engine/Source/Programs/Horde/Horde.Server.Tests/Tools/ToolTests.cs
@@ -38,15 +38,15 @@ namespace Horde.Server.Tests.Tools
 		{
 			IToolCollection collection = ServiceProvider.GetRequiredService<IToolCollection>();
 
-			ITool tool = Deref(await collection.GetAsync(_toolId, GlobalConfig.CurrentValue));
+			ITool tool = Deref(await collection.GetAsync(_toolId));
 			Assert.AreEqual(tool.Id, new ToolId("ugs"));
-			Assert.AreEqual(tool.Config.Name, "UnrealGameSync");
-			Assert.AreEqual(tool.Config.Description, "Tool for syncing content from source control");
+			Assert.AreEqual(tool.Name, "UnrealGameSync");
+			Assert.AreEqual(tool.Description, "Tool for syncing content from source control");
 
-			ITool tool2 = Deref(await collection.GetAsync(tool.Id, GlobalConfig.CurrentValue));
+			ITool tool2 = Deref(await collection.GetAsync(tool.Id));
 			Assert.AreEqual(tool.Id, tool2.Id);
-			Assert.AreEqual(tool.Config.Name, tool2.Config.Name);
-			Assert.AreEqual(tool.Config.Description, tool2.Config.Description);
+			Assert.AreEqual(tool.Name, tool2.Name);
+			Assert.AreEqual(tool.Description, tool2.Description);
 		}
 
 		[TestMethod]
@@ -54,10 +54,10 @@ namespace Horde.Server.Tests.Tools
 		{
 			IToolCollection collection = ServiceProvider.GetRequiredService<IToolCollection>();
 
-			ITool tool = Deref(await collection.GetAsync(_toolId, GlobalConfig.CurrentValue));
+			ITool tool = Deref(await collection.GetAsync(_toolId));
 			Assert.AreEqual(new ToolId("ugs"), tool.Id);
-			Assert.AreEqual("UnrealGameSync", tool.Config.Name);
-			Assert.AreEqual("Tool for syncing content from source control", tool.Config.Description);
+			Assert.AreEqual("UnrealGameSync", tool.Name);
+			Assert.AreEqual("Tool for syncing content from source control", tool.Description);
 			Assert.AreEqual(0, tool.Deployments.Count);
 
 			const string FileName = "test.txt";
@@ -80,7 +80,7 @@ namespace Horde.Server.Tests.Tools
 			ToolDeploymentId deploymentId;
 			using (MemoryStream stream = new MemoryStream(zipData))
 			{
-				tool = Deref(await collection.CreateDeploymentAsync(tool, new ToolDeploymentConfig { Version = "1.0", Duration = TimeSpan.FromMinutes(5.0), CreatePaused = true }, stream, GlobalConfig.CurrentValue, CancellationToken.None));
+				tool = Deref(await tool.CreateDeploymentAsync(new ToolDeploymentConfig { Version = "1.0", Duration = TimeSpan.FromMinutes(5.0), CreatePaused = true }, stream, CancellationToken.None));
 				Assert.AreEqual(1, tool.Deployments.Count);
 				Assert.IsNull(tool.Deployments[0].StartedAt);
 				deploymentId = tool.Deployments[^1].Id;
@@ -90,32 +90,35 @@ namespace Horde.Server.Tests.Tools
 			FakeClock clock = ServiceProvider.GetRequiredService<FakeClock>();
 			await clock.AdvanceAsync(TimeSpan.FromHours(1.0));
 
-			tool = Deref(await collection.GetAsync(tool.Id, GlobalConfig.CurrentValue));
+			tool = Deref(await collection.GetAsync(tool.Id));
 			Assert.AreEqual(1, tool.Deployments.Count);
 			Assert.IsNull(tool.Deployments[0].StartedAt);
 
 			// Start the deployment
-			tool = Deref(await collection.UpdateDeploymentAsync(tool, deploymentId, ToolDeploymentState.Active));
+			IToolDeployment deployment = tool.Deployments[0];
+			deployment = Deref(await deployment.UpdateAsync(ToolDeploymentState.Active));
+			tool = Deref(await collection.GetAsync(tool.Id));
 			Assert.AreEqual(1, tool.Deployments.Count);
 			Assert.IsNotNull(tool.Deployments[0].StartedAt);
 			Assert.IsTrue(Math.Abs((tool.Deployments[0].StartedAt!.Value - clock.UtcNow).TotalSeconds) < 1.0);
 
 			// Check it updates
 			await clock.AdvanceAsync(TimeSpan.FromMinutes(2.5));
-			tool = Deref(await collection.UpdateDeploymentAsync(tool, deploymentId, ToolDeploymentState.Paused));
+			deployment = Deref(await deployment.UpdateAsync(ToolDeploymentState.Paused));
+			tool = Deref(await collection.GetAsync(tool.Id));
 			Assert.AreEqual(1, tool.Deployments.Count);
-			Assert.IsNull(tool.Deployments[0].StartedAt);
-			Assert.IsTrue((tool.Deployments[0].Progress - 0.5) < 0.1);
+			Assert.IsNull(deployment.StartedAt);
+			Assert.IsTrue((deployment.Progress - 0.5) < 0.1);
 
 			// Check it stays paused
 			await clock.AdvanceAsync(TimeSpan.FromHours(1.0));
-			tool = Deref(await collection.GetAsync(tool.Id, GlobalConfig.CurrentValue));
+			tool = Deref(await collection.GetAsync(tool.Id));
 			Assert.AreEqual(1, tool.Deployments.Count);
 			Assert.IsNull(tool.Deployments[0].StartedAt);
 			Assert.IsTrue((tool.Deployments[0].Progress - 0.5) < 0.1);
 
 			// Get the deployment data
-			using Stream dataStream = await collection.GetDeploymentZipAsync(tool, tool.Deployments[0], CancellationToken.None);
+			using Stream dataStream = await tool.Deployments[0].OpenZipStreamAsync(CancellationToken.None);
 			using (ZipArchive archive = new ZipArchive(dataStream, ZipArchiveMode.Read))
 			{
 				ZipArchiveEntry entry = archive.Entries.First();
diff --git a/Engine/Source/Programs/Horde/Horde.Server.Tests/Tools/ToolsControllerTests.cs b/Engine/Source/Programs/Horde/Horde.Server.Tests/Tools/ToolsControllerTests.cs
--- a/Engine/Source/Programs/Horde/Horde.Server.Tests/Tools/ToolsControllerTests.cs
+++ b/Engine/Source/Programs/Horde/Horde.Server.Tests/Tools/ToolsControllerTests.cs
@@ -53,9 +53,9 @@ public class ToolsControllerTests
 		
 		// Create tool and deployment
 		using MemoryStream ms = new(await ToolTests.CreateZipFileDataAsync("foo.txt", "foo content"));
-		ITool? tool = await tools.GetAsync(toolId, globalConfig);
+		ITool? tool = await tools.GetAsync(toolId);
 		Assert.IsNotNull(tool);
-		tool = await tools.CreateDeploymentAsync(tool, new ToolDeploymentConfig() { Version = "1" }, ms, globalConfig);
+		tool = await tool.CreateDeploymentAsync(new ToolDeploymentConfig() { Version = "1" }, ms);
 
 		HttpClient client = app.CreateHttpClient();
 		using HttpRequestMessage req = new (HttpMethod.Get, $"/api/v1/tools/{toolId.Id}?action=download");
diff --git a/Engine/Source/Programs/Horde/Horde.Server/Agents/Software/AgentSoftwareController.cs b/Engine/Source/Programs/Horde/Horde.Server/Agents/Software/AgentSoftwareController.cs
--- a/Engine/Source/Programs/Horde/Horde.Server/Agents/Software/AgentSoftwareController.cs
+++ b/Engine/Source/Programs/Horde/Horde.Server/Agents/Software/AgentSoftwareController.cs
@@ -62,7 +62,7 @@ namespace Horde.Server.Agents.Software
 				return Forbid();
 			}
 
-			ITool? tool = await _toolCollection.GetAsync(AgentExtensions.AgentToolId, _globalConfig.Value);
+			ITool? tool = await _toolCollection.GetAsync(AgentExtensions.AgentToolId);
 			if (tool == null)
 			{
 				return NotFound("No agent software tool is currently registered");
@@ -92,7 +92,7 @@ namespace Horde.Server.Agents.Software
 				return Forbid();
 			}
 
-			ITool? tool = await _toolCollection.GetAsync(AgentExtensions.AgentToolId, _globalConfig.Value, cancellationToken);
+			ITool? tool = await _toolCollection.GetAsync(AgentExtensions.AgentToolId, cancellationToken);
 			if (tool == null)
 			{
 				return NotFound("No agent software tool is currently registered");
@@ -104,7 +104,7 @@ namespace Horde.Server.Agents.Software
 				return NotFound("No deployment currently set for agent software");
 			}
 
-			Stream stream = await _toolCollection.GetDeploymentZipAsync(tool, deployment, cancellationToken);
+			Stream stream = await deployment.OpenZipStreamAsync(cancellationToken);
 			return new FileStreamResult(stream, new MediaTypeHeaderValue("application/octet-stream")) { FileDownloadName = $"HordeAgent.zip" };
 		}
 	}
diff --git a/Engine/Source/Programs/Horde/Horde.Server/Server/RpcService.cs b/Engine/Source/Programs/Horde/Horde.Server/Server/RpcService.cs
--- a/Engine/Source/Programs/Horde/Horde.Server/Server/RpcService.cs
+++ b/Engine/Source/Programs/Horde/Horde.Server/Server/RpcService.cs
@@ -479,14 +479,13 @@ namespace Horde.Server.Server
 			ToolId toolId = new ToolId(request.Version.Substring(0, colonIdx));
 			string version = request.Version.Substring(colonIdx + 1);
 
-			ITool? tool = await _toolCollection.GetAsync(toolId, _globalConfig.Value);
+			ITool? tool = await _toolCollection.GetAsync(toolId, context.CancellationToken);
 			if (tool == null)
 			{
 				throw new StructuredRpcException(StatusCode.NotFound, $"Missing tool {toolId}");
 			}
 
-			ToolConfig toolConfig = tool.Config;
-			if (!toolConfig.Public && !toolConfig.Authorize(ToolAclAction.DownloadTool, context.GetHttpContext().User))
+			if (!tool.Public && !tool.Authorize(ToolAclAction.DownloadTool, context.GetHttpContext().User))
 			{
 				throw new StructuredRpcException(StatusCode.PermissionDenied, $"User does not have DownloadTool entitlement for {toolId}");
 			}
@@ -497,7 +496,7 @@ namespace Horde.Server.Server
 				throw new StructuredRpcException(StatusCode.NotFound, $"Missing tool version {version}");
 			}
 
-			await using Stream stream = await _toolCollection.GetDeploymentZipAsync(tool, deployment, context.CancellationToken);
+			await using Stream stream = await deployment.OpenZipStreamAsync(context.CancellationToken);
 			using (IMemoryOwner<byte> buffer = MemoryPool<byte>.Shared.Rent(128 * 1024))
 			{
 				long totalWritten = 0;
diff --git a/Engine/Source/Programs/Horde/Horde.Server/Server/ServerController.cs b/Engine/Source/Programs/Horde/Horde.Server/Server/ServerController.cs
--- a/Engine/Source/Programs/Horde/Horde.Server/Server/ServerController.cs
+++ b/Engine/Source/Programs/Horde/Horde.Server/Server/ServerController.cs
@@ -82,7 +82,7 @@ namespace Horde.Server.Server
 			response.ServerVersion = versionInfo.ProductVersion ?? String.Empty;
 			response.OsDescription = RuntimeInformation.OSDescription;
 
-			ITool? tool = await _toolCollection.GetAsync(AgentExtensions.AgentToolId, _globalConfig.Value);
+			ITool? tool = await _toolCollection.GetAsync(AgentExtensions.AgentToolId, HttpContext.RequestAborted);
 			if (tool != null)
 			{
 				IToolDeployment? deployment = tool.GetCurrentDeployment(1.0, _clock.UtcNow);
diff --git a/Engine/Source/Programs/Horde/Horde.Server/Tasks/UpgradeTaskSource.cs b/Engine/Source/Programs/Horde/Horde.Server/Tasks/UpgradeTaskSource.cs
--- a/Engine/Source/Programs/Horde/Horde.Server/Tasks/UpgradeTaskSource.cs
+++ b/Engine/Source/Programs/Horde/Horde.Server/Tasks/UpgradeTaskSource.cs
@@ -91,7 +91,7 @@ namespace Horde.Server.Tasks
 
 			ToolId toolId = agent.GetSoftwareToolId(globalConfig);
 
-			ITool? tool = await _toolCollection.GetAsync(toolId, globalConfig, cancellationToken);
+			ITool? tool = await _toolCollection.GetAsync(toolId, cancellationToken);
 			if (tool == null)
 			{
 				return null;
diff --git a/Engine/Source/Programs/Horde/Horde.Server/Tools/ITool.cs b/Engine/Source/Programs/Horde/Horde.Server/Tools/ITool.cs
--- a/Engine/Source/Programs/Horde/Horde.Server/Tools/ITool.cs
+++ b/Engine/Source/Programs/Horde/Horde.Server/Tools/ITool.cs
@@ -2,8 +2,14 @@
 
 using System;
 using System.Collections.Generic;
+using System.IO;
+using System.Security.Claims;
+using System.Threading;
+using System.Threading.Tasks;
+using EpicGames.Horde.Acls;
 using EpicGames.Horde.Storage;
 using EpicGames.Horde.Tools;
+using Horde.Server.Server;
 
 namespace Horde.Server.Tools
 {
@@ -18,15 +24,77 @@ namespace Horde.Server.Tools
 		/// </summary>
 		ToolId Id { get; }
 
+		/// <summary>
+		/// Name of the tool
+		/// </summary>
+		string Name { get; }
+
+		/// <summary>
+		/// Long-form description of the tool
+		/// </summary>
+		string Description { get; }
+
+		/// <summary>
+		/// Category for the tool on the dashboard
+		/// </summary>
+		string? Category { get; }
+
+		/// <summary>
+		/// Whether the tool is available to authenticated users
+		/// </summary>
+		bool Public { get; }
+
+		/// <summary>
+		/// Whether to show the tool for download in UGS
+		/// </summary>
+		bool ShowInUgs { get; }
+
+		/// <summary>
+		/// Whether to show the tool for download in the dashboard
+		/// </summary>
+		bool ShowInDashboard { get; }
+
 		/// <summary>
 		/// Current deployments of this tool, sorted by time.
 		/// </summary>
 		IReadOnlyList<IToolDeployment> Deployments { get; }
 
 		/// <summary>
-		/// Config object for this tool
+		/// Authorize a user to perform a particular action
 		/// </summary>
-		ToolConfig Config { get; }
+		/// <param name="action">Action the user is trying to perform</param>
+		/// <param name="principal">Identity of the user trying to perform the action</param>
+		bool Authorize(AclAction action, ClaimsPrincipal principal);
+
+		/// <summary>
+		/// Adds a new deployment to the given tool. The new deployment will replace the current active deployment.
+		/// </summary>
+		/// <param name="options">Options for the new deployment</param>
+		/// <param name="stream">Stream containing the tool data</param>
+		/// <param name="cancellationToken">Cancellation token for the operation</param>
+		/// <returns>Updated tool document, or null if it does not exist</returns>
+		Task<ITool?> CreateDeploymentAsync(ToolDeploymentConfig options, Stream stream, CancellationToken cancellationToken = default);
+
+		/// <summary>
+		/// Adds a new deployment to the given tool. The new deployment will replace the current active deployment.
+		/// </summary>
+		/// <param name="options">Options for the new deployment</param>
+		/// <param name="target">Path to the root node containing the tool data</param>
+		/// <param name="cancellationToken">Cancellation token for the operation</param>
+		/// <returns>Updated tool document, or null if it does not exist</returns>
+		Task<ITool?> CreateDeploymentAsync(ToolDeploymentConfig options, BlobRefValue target, CancellationToken cancellationToken = default);
+
+		/// <summary>
+		/// Gets the storage backend for a particular tool
+		/// </summary>
+		/// <returns>Instance of the backend client</returns>
+		IStorageBackend CreateStorageBackend();
+
+		/// <summary>
+		/// Gets the storage backend for a particular tool
+		/// </summary>
+		/// <returns>Instance of the storage client</returns>
+		IStorageClient CreateStorageClient();
 	}
 
 	/// <summary>
@@ -73,6 +141,21 @@ namespace Horde.Server.Tools
 		/// Reference to this tool in Horde Storage.
 		/// </summary>
 		RefName RefName { get; }
+
+		/// <summary>
+		/// Updates the state of the current deployment
+		/// </summary>
+		/// <param name="state">New state of the deployment</param>
+		/// <param name="cancellationToken">Cancellation token for the operation</param>
+		/// <returns></returns>
+		Task<IToolDeployment?> UpdateAsync(ToolDeploymentState state, CancellationToken cancellationToken = default);
+
+		/// <summary>
+		/// Opens a stream to the data for a particular deployment
+		/// </summary>
+		/// <param name="cancellationToken">Cancellation token for the operation</param>
+		/// <returns>Stream for the data</returns>
+		Task<Stream> OpenZipStreamAsync(CancellationToken cancellationToken = default);
 	}
 
 	/// <summary>
diff --git a/Engine/Source/Programs/Horde/Horde.Server/Tools/IToolCollection.cs b/Engine/Source/Programs/Horde/Horde.Server/Tools/IToolCollection.cs
--- a/Engine/Source/Programs/Horde/Horde.Server/Tools/IToolCollection.cs
+++ b/Engine/Source/Programs/Horde/Horde.Server/Tools/IToolCollection.cs
@@ -1,5 +1,6 @@
 // Copyright Epic Games, Inc. All Rights Reserved.
 
+using System.Collections.Generic;
 using System.IO;
 using System.Threading;
 using System.Threading.Tasks;
@@ -10,7 +11,7 @@ using Horde.Server.Server;
 namespace Horde.Server.Tools
 {
 	/// <summary>
-	/// Collection of tool documents
+	/// Collection of tools
 	/// </summary>
 	public interface IToolCollection
 	{
@@ -18,64 +19,15 @@ namespace Horde.Server.Tools
 		/// Gets a tool with the given identifier
 		/// </summary>
 		/// <param name="id">The tool identifier</param>
-		/// <param name="globalConfig">The current global configuration</param>
 		/// <param name="cancellationToken">Cancellation token for the operation</param>
-		/// <returns></returns>
-		Task<ITool?> GetAsync(ToolId id, GlobalConfig globalConfig, CancellationToken cancellationToken = default);
+		/// <returns>The requested tool, or null if it does not exist</returns>
+		Task<ITool?> GetAsync(ToolId id, CancellationToken cancellationToken = default);
 
 		/// <summary>
-		/// Adds a new deployment to the given tool. The new deployment will replace the current active deployment.
+		/// Gets all the available tools
 		/// </summary>
-		/// <param name="tool">The tool to update</param>
-		/// <param name="options">Options for the new deployment</param>
-		/// <param name="stream">Stream containing the tool data</param>
-		/// <param name="globalConfig">The current configuration</param>
-		/// <param name="cancellationToken">Cancellation token for the operation</param>
-		/// <returns>Updated tool document, or null if it does not exist</returns>
-		Task<ITool?> CreateDeploymentAsync(ITool tool, ToolDeploymentConfig options, Stream stream, GlobalConfig globalConfig, CancellationToken cancellationToken = default);
-
-		/// <summary>
-		/// Adds a new deployment to the given tool. The new deployment will replace the current active deployment.
-		/// </summary>
-		/// <param name="tool">The tool to update</param>
-		/// <param name="options">Options for the new deployment</param>
-		/// <param name="target">Path to the root node containing the tool data</param>
-		/// <param name="globalConfig">The current configuration</param>
-		/// <param name="cancellationToken">Cancellation token for the operation</param>
-		/// <returns>Updated tool document, or null if it does not exist</returns>
-		Task<ITool?> CreateDeploymentAsync(ITool tool, ToolDeploymentConfig options, BlobRefValue target, GlobalConfig globalConfig, CancellationToken cancellationToken = default);
-
-		/// <summary>
-		/// Updates the state of the current deployment
-		/// </summary>
-		/// <param name="tool">Tool to be updated</param>
-		/// <param name="deploymentId">Identifier for the deployment to modify</param>
-		/// <param name="action">New state of the deployment</param>
-		/// <param name="cancellationToken">Cancellation token for the operation</param>
-		/// <returns></returns>
-		Task<ITool?> UpdateDeploymentAsync(ITool tool, ToolDeploymentId deploymentId, ToolDeploymentState action, CancellationToken cancellationToken = default);
-
-		/// <summary>
-		/// Gets the storage backend for a particular tool
-		/// </summary>
-		/// <param name="tool">The tool to get a storage backend for</param>
-		/// <returns>Instance of the backend client</returns>
-		IStorageBackend CreateStorageBackend(ITool tool);
-
-		/// <summary>
-		/// Gets the storage backend for a particular tool
-		/// </summary>
-		/// <param name="tool">The tool to get a storage backend for</param>
-		/// <returns>Instance of the backend client</returns>
-		IStorageClient CreateStorageClient(ITool tool);
-
-		/// <summary>
-		/// Opens a stream to the data for a particular deployment
-		/// </summary>
-		/// <param name="tool">Identifier for the tool</param>
-		/// <param name="deployment">The deployment</param>
-		/// <param name="cancellationToken">Cancellation token for the operation</param>
-		/// <returns>Stream for the data</returns>
-		Task<Stream> GetDeploymentZipAsync(ITool tool, IToolDeployment deployment, CancellationToken cancellationToken);
+		/// <param name="cancellationToken"></param>
+		/// <returns>List of the available tools</returns>
+		Task<IReadOnlyList<ITool>> GetAllAsync(CancellationToken cancellationToken = default);
 	}
 }
diff --git a/Engine/Source/Programs/Horde/Horde.Server/Tools/ToolCollection.cs b/Engine/Source/Programs/Horde/Horde.Server/Tools/ToolCollection.cs
--- a/Engine/Source/Programs/Horde/Horde.Server/Tools/ToolCollection.cs
+++ b/Engine/Source/Programs/Horde/Horde.Server/Tools/ToolCollection.cs
@@ -5,11 +5,13 @@ using System.Buffers.Binary;
 using System.Collections.Generic;
 using System.IO;
 using System.Linq;
+using System.Security.Claims;
 using System.Text;
 using System.Threading;
 using System.Threading.Tasks;
 using EpicGames.Core;
 using EpicGames.Horde;
+using EpicGames.Horde.Acls;
 using EpicGames.Horde.Storage;
 using EpicGames.Horde.Storage.Backends;
 using EpicGames.Horde.Storage.Bundles;
@@ -21,6 +23,7 @@ using Horde.Server.Server;
 using Horde.Server.Storage;
 using Horde.Server.Utilities;
 using Microsoft.Extensions.Logging;
+using Microsoft.Extensions.Options;
 using MongoDB.Bson.Serialization.Attributes;
 using MongoDB.Driver;
 
@@ -31,58 +34,31 @@ namespace Horde.Server.Tools
 	/// </summary>
 	public class ToolCollection : IToolCollection
 	{
-		class Tool : ITool
+		class ToolDocument
 		{
 			public ToolId Id { get; set; }
 
-			[BsonIgnore]
-			public ToolConfig Config { get; set; } = null!;
-
 			[BsonElement("dep")]
-			public List<ToolDeployment> Deployments { get; set; } = new List<ToolDeployment>();
+			public List<ToolDeploymentDocument> Deployments { get; set; } = new List<ToolDeploymentDocument>();
 
 			// Last time that the document was updated. This field is checked and updated as part of updates to ensure atomicity.
 			[BsonElement("_u")]
 			public DateTime LastUpdateTime { get; set; }
 
-			// ITool interface
-			IReadOnlyList<IToolDeployment> ITool.Deployments => Deployments;
-
 			[BsonConstructor]
-			public Tool(ToolId id)
+			public ToolDocument(ToolId id)
 			{
 				Id = id;
-				Config = null!;
-			}
-
-			public Tool(ToolConfig config)
-			{
-				Id = config.Id;
-				Config = config;
-			}
-
-			public void UpdateTemporalState(DateTime utcNow)
-			{
-				foreach (ToolDeployment deployment in Deployments)
-				{
-					deployment.UpdateTemporalState(utcNow);
-				}
 			}
 		}
 
-		class ToolDeployment : IToolDeployment
+		class ToolDeploymentDocument
 		{
 			public ToolDeploymentId Id { get; set; }
 
 			[BsonElement("ver")]
 			public string Version { get; set; }
 
-			[BsonIgnore]
-			public ToolDeploymentState State { get; set; }
-
-			[BsonIgnore]
-			public double Progress { get; set; }
-
 			[BsonElement("bpr")]
 			public double BaseProgress { get; set; }
 
@@ -99,13 +75,13 @@ namespace Horde.Server.Tools
 			public RefName RefName { get; set; }
 
 			[BsonConstructor]
-			public ToolDeployment(ToolDeploymentId id)
+			public ToolDeploymentDocument(ToolDeploymentId id)
 			{
 				Id = id;
 				Version = String.Empty;
 			}
 
-			public ToolDeployment(ToolDeploymentId id, ToolDeploymentConfig options, NamespaceId namespaceId, RefName refName)
+			public ToolDeploymentDocument(ToolDeploymentId id, ToolDeploymentConfig options, NamespaceId namespaceId, RefName refName)
 			{
 				Id = id;
 				Version = options.Version;
@@ -113,10 +89,79 @@ namespace Horde.Server.Tools
 				NamespaceId = namespaceId;
 				RefName = refName;
 			}
+		}
 
-			public void UpdateTemporalState(DateTime utcNow)
+		class Tool : ITool
+		{
+			readonly ToolCollection _collection;
+			readonly ToolDocument _document;
+			readonly ToolConfig _config;
+			readonly List<ToolDeployment> _deployments;
+
+			public ToolConfig Config => _config;
+			public ToolDocument Document => _document;
+
+			public ToolId Id => _config.Id;
+			public string Name => _config.Name;
+			public string Description => _config.Description;
+			public string? Category => _config.Category;
+			public bool Public => _config.Public;
+			public bool ShowInUgs => _config.ShowInUgs;
+			public bool ShowInDashboard => _config.ShowInDashboard;
+			public IReadOnlyList<IToolDeployment> Deployments => _deployments;
+
+			public Tool(ToolCollection collection, ToolDocument document, ToolConfig config, DateTime utcNow)
+			{
+				_collection = collection;
+				_document = document;
+				_config = config;
+				_deployments = document.Deployments.ConvertAll(x => new ToolDeployment(this, _collection, x, utcNow));
+			}
+
+			public bool Authorize(AclAction action, ClaimsPrincipal principal)
+				=> _config.Authorize(action, principal);
+
+			public async Task<ITool?> CreateDeploymentAsync(ToolDeploymentConfig options, Stream stream, CancellationToken cancellationToken = default)
+			{
+				ToolDocument? document = await _collection.CreateDeploymentAsync(_document, _config, options, stream, cancellationToken);
+				return _collection.CreateToolObject(document);
+			}
+
+			public async Task<ITool?> CreateDeploymentAsync(ToolDeploymentConfig options, BlobRefValue target, CancellationToken cancellationToken = default)
 			{
-				if (BaseProgress >= 1.0)
+				ToolDocument? document = await _collection.CreateDeploymentAsync(_document, _config, options, target, cancellationToken);
+				return _collection.CreateToolObject(document);
+			}
+
+			public IStorageBackend CreateStorageBackend()
+				=> _collection.CreateStorageBackend(_config);
+
+			public IStorageClient CreateStorageClient()
+				=> _collection.CreateStorageClient(_config);
+		}
+
+		class ToolDeployment : IToolDeployment
+		{
+			readonly Tool _tool;
+			readonly ToolCollection _collection;
+			readonly ToolDeploymentDocument _document;
+
+			public ToolDeploymentId Id => _document.Id;
+			public string Version => _document.Version;
+			public ToolDeploymentState State { get; }
+			public double Progress { get; }
+			public DateTime? StartedAt => _document.StartedAt;
+			public TimeSpan Duration => _document.Duration;
+			public NamespaceId NamespaceId => _document.NamespaceId;
+			public RefName RefName => _document.RefName;
+
+			public ToolDeployment(Tool tool, ToolCollection collection, ToolDeploymentDocument document, DateTime utcNow)
+			{
+				_tool = tool;
+				_collection = collection;
+				_document = document;
+
+				if (document.BaseProgress >= 1.0)
 				{
 					State = ToolDeploymentState.Complete;
 					Progress = 1.0;
@@ -124,7 +169,7 @@ namespace Horde.Server.Tools
 				else if (StartedAt == null)
 				{
 					State = ToolDeploymentState.Paused;
-					Progress = BaseProgress;
+					Progress = document.BaseProgress;
 				}
 				else if (Duration > TimeSpan.Zero)
 				{
@@ -137,106 +182,132 @@ namespace Horde.Server.Tools
 					Progress = 1.0;
 				}
 			}
-		}
-
-		private class ToolDeploymentData
-		{
-			[CbField]
-			public ToolDeploymentId Id { get; set; }
-
-			[CbField("version")]
-			public string Version { get; set; } = String.Empty;
 
-			[CbField("data")]
-			public CbBinaryAttachment Data { get; set; }
-		}
+			public Task<Stream> OpenZipStreamAsync(CancellationToken cancellationToken = default)
+				=> _collection.GetDeploymentZipAsync(_tool.Config, _document, cancellationToken);
 
-		private class CachedIndex
-		{
-			[CbField("rev")]
-			public string Rev { get; set; } = String.Empty;
-
-			[CbField("empty")]
-			public bool Empty { get; set; }
-
-			[CbField("ids")]
-			public List<ToolId> Ids { get; set; } = new List<ToolId>();
+			public async Task<IToolDeployment?> UpdateAsync(ToolDeploymentState action, CancellationToken cancellationToken = default)
+			{
+				ToolDocument? newDocument = await _collection.UpdateDeploymentAsync(_tool.Document, Id, action, Progress, cancellationToken);
+				Tool? newTool = _collection.CreateToolObject(newDocument);
+				return newTool?.Deployments.FirstOrDefault(x => x.Id == Id);
+			}
 		}
 
-		private readonly IMongoCollection<Tool> _tools;
+		private readonly IMongoCollection<ToolDocument> _tools;
 		private readonly StorageService _storageService;
 		private readonly FileObjectStoreFactory _fileObjectStoreFactory;
 		private readonly IClock _clock;
 		private readonly BundleCache _cache;
+		private readonly IOptionsMonitor<GlobalConfig> _globalConfig;
 		private readonly ILogger _logger;
 
 		/// <summary>
 		/// Constructor
 		/// </summary>
-		public ToolCollection(MongoService mongoService, StorageService storageService, BundleCache cache, FileObjectStoreFactory fileObjectStoreFactory, IClock clock, ILogger<ToolCollection> logger)
+		public ToolCollection(MongoService mongoService, StorageService storageService, BundleCache cache, FileObjectStoreFactory fileObjectStoreFactory, IClock clock, IOptionsMonitor<GlobalConfig> globalConfig, ILogger<ToolCollection> logger)
 		{
-			_tools = mongoService.GetCollection<Tool>("Tools");
+			_tools = mongoService.GetCollection<ToolDocument>("Tools");
 			_storageService = storageService;
 			_fileObjectStoreFactory = fileObjectStoreFactory;
 			_clock = clock;
+			_globalConfig = globalConfig;
 			_cache = cache;
 			_logger = logger;
 		}
 
-		/// <inheritdoc/>
-		public async Task<ITool?> GetAsync(ToolId id, GlobalConfig globalConfig, CancellationToken cancellationToken)
-			=> await GetInternalAsync(id, globalConfig, cancellationToken);
-
-		/// <summary>
-		/// Gets a tool with the given identifier
-		/// </summary>
-		/// <param name="toolId">The tool identifier</param>
-		/// <param name="globalConfig">The current global configuration</param>
-		/// <param name="cancellationToken">Cancellation token for the operation</param>
-		/// <returns></returns>
-		async Task<Tool?> GetInternalAsync(ToolId toolId, GlobalConfig globalConfig, CancellationToken cancellationToken)
+		Tool? CreateToolObject(ToolDocument? document)
 		{
+			if (document == null)
+			{
+				return null;
+			}
+
 			ToolConfig? toolConfig;
-			if (globalConfig.TryGetTool(toolId, out toolConfig))
+			if (!_globalConfig.CurrentValue.TryGetTool(document.Id, out toolConfig))
 			{
-				Tool? tool;
-				for (; ; )
-				{
-					tool = await _tools.Find(x => x.Id == toolId).FirstOrDefaultAsync(cancellationToken);
-					if (tool != null)
-					{
-						break;
-					}
+				return null;
+			}
 
-					tool = new Tool(toolId);
-					if (await _tools.InsertOneIgnoreDuplicatesAsync(tool, cancellationToken))
-					{
-						break;
-					}
-				}
+			return new Tool(this, document, toolConfig, _clock.UtcNow);
+		}
+
+		/// <inheritdoc/>
+		public async Task<ITool?> GetAsync(ToolId id, CancellationToken cancellationToken)
+		{
+			GlobalConfig globalConfig = _globalConfig.CurrentValue;
 
-				tool.Config = toolConfig;
-				tool.UpdateTemporalState(_clock.UtcNow);
-				return tool;
+			ToolConfig? toolConfig;
+			if (globalConfig.TryGetTool(id, out toolConfig))
+			{
+				ToolDocument document = await FindOrAddDocument(toolConfig, cancellationToken);
+				return new Tool(this, document, toolConfig, _clock.UtcNow);
 			}
 
 			BundledToolConfig? bundledToolConfig;
-			if (globalConfig.ServerSettings.TryGetBundledTool(toolId, out bundledToolConfig))
+			if (globalConfig.ServerSettings.TryGetBundledTool(id, out bundledToolConfig))
 			{
-				Tool tool = new Tool(bundledToolConfig);
+				ToolDocument document = CreateBundledToolDocument(bundledToolConfig);
+				return new Tool(this, document, bundledToolConfig, _clock.UtcNow);
+			}
 
-				ToolDeploymentId deploymentId = GetDeploymentId(bundledToolConfig);
+			return null;
+		}
 
-				ToolDeployment deployment = new ToolDeployment(deploymentId);
-				deployment.Version = bundledToolConfig.Version;
-				deployment.State = ToolDeploymentState.Complete;
-				deployment.RefName = bundledToolConfig.RefName;
-				tool.Deployments.Add(deployment);
+		/// <inheritdoc/>
+		public async Task<IReadOnlyList<ITool>> GetAllAsync(CancellationToken cancellationToken)
+		{
+			DateTime utcNow = _clock.UtcNow;
+			GlobalConfig globalConfig = _globalConfig.CurrentValue;
 
-				return tool;
+			List<Tool> tools = new List<Tool>();
+			foreach (ToolConfig toolConfig in globalConfig.Tools)
+			{
+				ToolDocument document = await FindOrAddDocument(toolConfig, cancellationToken);
+				tools.Add(new Tool(this, document, toolConfig, utcNow));
+			}
+			foreach(BundledToolConfig bundledToolConfig in globalConfig.ServerSettings.BundledTools)
+			{
+				ToolDocument document = CreateBundledToolDocument(bundledToolConfig);
+				tools.Add(new Tool(this, document, bundledToolConfig, utcNow));
 			}
 
-			return null;
+			return tools;
+		}
+
+		async Task<ToolDocument> FindOrAddDocument(ToolConfig toolConfig, CancellationToken cancellationToken)
+		{
+			ToolDocument? tool;
+			for (; ; )
+			{
+				tool = await _tools.Find(x => x.Id == toolConfig.Id).FirstOrDefaultAsync(cancellationToken);
+				if (tool != null)
+				{
+					break;
+				}
+
+				tool = new ToolDocument(toolConfig.Id);
+				if (await _tools.InsertOneIgnoreDuplicatesAsync(tool, cancellationToken))
+				{
+					break;
+				}
+			}
+			return tool;
+		}
+
+		ToolDocument CreateBundledToolDocument(BundledToolConfig bundledToolConfig)
+		{
+			ToolDocument tool = new ToolDocument(bundledToolConfig.Id);
+
+			ToolDeploymentId deploymentId = GetDeploymentId(bundledToolConfig);
+
+			ToolDeploymentDocument deployment = new ToolDeploymentDocument(deploymentId);
+			deployment.StartedAt = DateTime.MinValue;
+			deployment.Version = bundledToolConfig.Version;
+			deployment.RefName = bundledToolConfig.RefName;
+			tool.Deployments.Add(deployment);
+
+			return tool;
 		}
 
 		static ToolDeploymentId GetDeploymentId(BundledToolConfig bundledToolConfig)
@@ -252,20 +323,11 @@ namespace Horde.Server.Tools
 			return new ToolDeploymentId(new BinaryId(bytes));
 		}
 
-		/// <summary>
-		/// Adds a new deployment to the given tool. The new deployment will replace the current active deployment.
-		/// </summary>
-		/// <param name="tool">The tool to update</param>
-		/// <param name="options">Options for the new deployment</param>
-		/// <param name="stream">Stream containing the tool data</param>
-		/// <param name="globalConfig">The current configuration</param>
-		/// <param name="cancellationToken">Cancellation token for the operation</param>
-		/// <returns>Updated tool document, or null if it does not exist</returns>
-		public async Task<ITool?> CreateDeploymentAsync(ITool tool, ToolDeploymentConfig options, Stream stream, GlobalConfig globalConfig, CancellationToken cancellationToken)
+		async Task<ToolDocument?> CreateDeploymentAsync(ToolDocument tool, ToolConfig toolConfig, ToolDeploymentConfig options, Stream stream, CancellationToken cancellationToken)
 		{
 			ToolDeploymentId deploymentId = new ToolDeploymentId(BinaryIdUtils.CreateNew());
 
-			using IStorageClient client = _storageService.CreateClient(tool.Config.NamespaceId);
+			using IStorageClient client = _storageService.CreateClient(toolConfig.NamespaceId);
 
 			IBlobRef<DirectoryNode> nodeRef;
 			await using (IBlobWriter writer = client.CreateBlobWriter($"{tool.Id}/{deploymentId}"))
@@ -275,29 +337,20 @@ namespace Horde.Server.Tools
 				nodeRef = await writer.WriteBlobAsync(directoryNode, cancellationToken: cancellationToken);
 			}
 
-			return await CreateDeploymentInternalAsync(tool, deploymentId, options, client, nodeRef, globalConfig, cancellationToken);
+			return await CreateDeploymentInternalAsync(tool, toolConfig, deploymentId, options, client, nodeRef, cancellationToken);
 		}
 
-		/// <summary>
-		/// Adds a new deployment to the given tool. The new deployment will replace the current active deployment.
-		/// </summary>
-		/// <param name="tool">The tool to update</param>
-		/// <param name="options">Options for the new deployment</param>
-		/// <param name="target">Path to the tool data</param>
-		/// <param name="globalConfig">The current configuration</param>
-		/// <param name="cancellationToken">Cancellation token for the operation</param>
-		/// <returns>Updated tool document, or null if it does not exist</returns>
-		public async Task<ITool?> CreateDeploymentAsync(ITool tool, ToolDeploymentConfig options, BlobRefValue target, GlobalConfig globalConfig, CancellationToken cancellationToken)
+		async Task<ToolDocument?> CreateDeploymentAsync(ToolDocument tool, ToolConfig toolConfig, ToolDeploymentConfig options, BlobRefValue target, CancellationToken cancellationToken)
 		{
 			ToolDeploymentId deploymentId = new ToolDeploymentId(BinaryIdUtils.CreateNew());
 
-			using IStorageClient client = _storageService.CreateClient(tool.Config.NamespaceId);
-			return await CreateDeploymentInternalAsync(tool, deploymentId, options, client, client.CreateBlobRef(target), globalConfig, cancellationToken);
+			using IStorageClient client = _storageService.CreateClient(toolConfig.NamespaceId);
+			return await CreateDeploymentInternalAsync(tool, toolConfig, deploymentId, options, client, client.CreateBlobRef(target), cancellationToken);
 		}
 
-		async Task<ITool?> CreateDeploymentInternalAsync(ITool tool, ToolDeploymentId deploymentId, ToolDeploymentConfig options, IStorageClient storageClient, IBlobRef content, GlobalConfig globalConfig, CancellationToken cancellationToken)
+		async Task<ToolDocument?> CreateDeploymentInternalAsync(ToolDocument tool, ToolConfig toolConfig, ToolDeploymentId deploymentId, ToolDeploymentConfig options, IStorageClient storageClient, IBlobRef content, CancellationToken cancellationToken)
 		{
-			if (tool.Config is BundledToolConfig)
+			if (toolConfig is BundledToolConfig)
 			{
 				throw new InvalidOperationException("Cannot update the state of bundled tools.");
 			}
@@ -307,7 +360,7 @@ namespace Horde.Server.Tools
 			await storageClient.WriteRefAsync(refName, content, cancellationToken: cancellationToken);
 
 			// Create the new deployment object
-			ToolDeployment deployment = new ToolDeployment(deploymentId, options, tool.Config.NamespaceId, refName);
+			ToolDeploymentDocument deployment = new ToolDeploymentDocument(deploymentId, options, toolConfig.NamespaceId, refName);
 
 			// Start the deployment
 			DateTime utcNow = _clock.UtcNow;
@@ -317,62 +370,47 @@ namespace Horde.Server.Tools
 			}
 
 			// Create the deployment
-			Tool? newTool = (Tool)tool;
+			ToolDocument? newTool = (ToolDocument)tool;
 			for (; ; )
 			{
 				newTool = await TryAddDeploymentAsync(newTool, deployment, cancellationToken);
 				if (newTool != null)
 				{
-					break;
+					return newTool;
 				}
 
-				newTool = await GetInternalAsync(tool.Id, globalConfig, cancellationToken);
+				newTool = await FindOrAddDocument(toolConfig, cancellationToken);
 				if (newTool == null)
 				{
 					return null;
 				}
 			}
-
-			// Return the new tool with updated deployment states
-			newTool.UpdateTemporalState(utcNow);
-			return newTool;
 		}
 
-		async ValueTask<Tool?> TryAddDeploymentAsync(Tool tool, ToolDeployment deployment, CancellationToken cancellationToken)
+		async ValueTask<ToolDocument?> TryAddDeploymentAsync(ToolDocument tool, ToolDeploymentDocument deployment, CancellationToken cancellationToken)
 		{
-			Tool? newTool = tool;
+			ToolDocument? newTool = tool;
 
 			// If there are already a maximum number of deployments, remove the oldest one
 			const int MaxDeploymentCount = 5;
 			while (newTool.Deployments.Count >= MaxDeploymentCount)
 			{
-				newTool = await UpdateAsync(newTool, Builders<Tool>.Update.PopFirst(x => x.Deployments), cancellationToken);
+				newTool = await UpdateAsync(newTool, Builders<ToolDocument>.Update.PopFirst(x => x.Deployments), cancellationToken);
 				if (newTool == null)
 				{
 					return null;
 				}
 
-				ToolDeployment removeDeployment = tool.Deployments[0];
+				ToolDeploymentDocument removeDeployment = tool.Deployments[0];
 				using IStorageClient client = _storageService.CreateClient(removeDeployment.NamespaceId);
 				await client.DeleteRefAsync(removeDeployment.RefName, cancellationToken);
 			}
 
 			// Add the new deployment
-			return await UpdateAsync(newTool, Builders<Tool>.Update.Push(x => x.Deployments, deployment), cancellationToken);
-		}
-
-		/// <inheritdoc/>
-		public async Task<ITool?> UpdateDeploymentAsync(ITool tool, ToolDeploymentId deploymentId, ToolDeploymentState action, CancellationToken cancellationToken)
-		{
-			if (tool.Config is BundledToolConfig)
-			{
-				throw new InvalidOperationException("Cannot update the state of bundled tools.");
-			}
-
-			return await UpdateDeploymentInternalAsync((Tool)tool, deploymentId, action, cancellationToken);
+			return await UpdateAsync(newTool, Builders<ToolDocument>.Update.Push(x => x.Deployments, deployment), cancellationToken);
 		}
 
-		async Task<Tool?> UpdateDeploymentInternalAsync(Tool tool, ToolDeploymentId deploymentId, ToolDeploymentState action, CancellationToken cancellationToken)
+		async Task<ToolDocument?> UpdateDeploymentAsync(ToolDocument tool, ToolDeploymentId deploymentId, ToolDeploymentState action, double currentProgress, CancellationToken cancellationToken)
 		{
 			int idx = tool.Deployments.FindIndex(x => x.Id == deploymentId);
 			if (idx == -1)
@@ -380,15 +418,15 @@ namespace Horde.Server.Tools
 				return null;
 			}
 
-			ToolDeployment deployment = tool.Deployments[idx];
+			ToolDeploymentDocument deployment = tool.Deployments[idx];
 			switch (action)
 			{
 				case ToolDeploymentState.Complete:
-					return await UpdateAsync(tool, Builders<Tool>.Update.Set(x => x.Deployments[idx].BaseProgress, 1.0).Unset(x => x.Deployments[idx].StartedAt), cancellationToken);
+					return await UpdateAsync(tool, Builders<ToolDocument>.Update.Set(x => x.Deployments[idx].BaseProgress, 1.0).Unset(x => x.Deployments[idx].StartedAt), cancellationToken);
 
 				case ToolDeploymentState.Cancelled:
-					List<ToolDeployment> newDeployments = tool.Deployments.Where(x => x != deployment).ToList();
-					return await UpdateAsync(tool, Builders<Tool>.Update.Set(x => x.Deployments, newDeployments), cancellationToken);
+					List<ToolDeploymentDocument> newDeployments = tool.Deployments.Where(x => x != deployment).ToList();
+					return await UpdateAsync(tool, Builders<ToolDocument>.Update.Set(x => x.Deployments, newDeployments), cancellationToken);
 
 				case ToolDeploymentState.Paused:
 					if (deployment.StartedAt == null)
@@ -397,7 +435,7 @@ namespace Horde.Server.Tools
 					}
 					else
 					{
-						return await UpdateAsync(tool, Builders<Tool>.Update.Set(x => x.Deployments[idx].BaseProgress, deployment.GetProgressValue(_clock.UtcNow)).Set(x => x.Deployments[idx].StartedAt, null), cancellationToken);
+						return await UpdateAsync(tool, Builders<ToolDocument>.Update.Set(x => x.Deployments[idx].BaseProgress, currentProgress).Set(x => x.Deployments[idx].StartedAt, null), cancellationToken);
 					}
 
 				case ToolDeploymentState.Active:
@@ -407,7 +445,7 @@ namespace Horde.Server.Tools
 					}
 					else
 					{
-						return await UpdateAsync(tool, Builders<Tool>.Update.Set(x => x.Deployments[idx].StartedAt, _clock.UtcNow), cancellationToken);
+						return await UpdateAsync(tool, Builders<ToolDocument>.Update.Set(x => x.Deployments[idx].StartedAt, _clock.UtcNow), cancellationToken);
 					}
 
 				default:
@@ -418,11 +456,11 @@ namespace Horde.Server.Tools
 		/// <summary>
 		/// Gets the storage client containing data for a particular tool
 		/// </summary>
-		/// <param name="tool">Identifier for the tool</param>
+		/// <param name="toolConfig">Identifier for the tool</param>
 		/// <returns>Storage client for the data</returns>
-		public IStorageClient CreateStorageClient(ITool tool)
+		IStorageClient CreateStorageClient(ToolConfig toolConfig)
 		{
-			if (tool.Config is BundledToolConfig bundledConfig)
+			if (toolConfig is BundledToolConfig bundledConfig)
 			{
 				return BundleStorageClient.CreateFromDirectory(DirectoryReference.Combine(ServerApp.AppDir, bundledConfig.DataDir ?? $"Tools"), _cache, _logger);
 			}
@@ -435,13 +473,13 @@ namespace Horde.Server.Tools
 		/// <summary>
 		/// Gets the storage client containing data for a particular tool
 		/// </summary>
-		/// <param name="tool">Identifier for the tool</param>
+		/// <param name="toolConfig">Identifier for the tool</param>
 		/// <returns>Storage client for the data</returns>
-		public IStorageBackend CreateStorageBackend(ITool tool)
+		IStorageBackend CreateStorageBackend(ToolConfig toolConfig)
 		{
-			if (tool.Config is BundledToolConfig bundledConfig)
+			if (toolConfig is BundledToolConfig bundledConfig)
 			{
-				return new FileStorageBackend(_fileObjectStoreFactory.CreateStore(DirectoryReference.Combine(ServerApp.AppDir, bundledConfig.DataDir ?? $"tools/{tool.Id}")), _logger);
+				return new FileStorageBackend(_fileObjectStoreFactory.CreateStore(DirectoryReference.Combine(ServerApp.AppDir, bundledConfig.DataDir ?? $"tools/{toolConfig.Id}")), _logger);
 			}
 			else
 			{
@@ -456,7 +494,7 @@ namespace Horde.Server.Tools
 		/// <param name="deployment">The deployment</param>
 		/// <param name="cancellationToken">Cancellation token for the operation</param>
 		/// <returns>Stream for the data</returns>
-		public async Task<Stream> GetDeploymentZipAsync(ITool tool, IToolDeployment deployment, CancellationToken cancellationToken)
+		async Task<Stream> GetDeploymentZipAsync(ToolConfig tool, ToolDeploymentDocument deployment, CancellationToken cancellationToken)
 		{
 #pragma warning disable CA2000
 			IStorageClient client = CreateStorageClient(tool);
@@ -473,12 +511,12 @@ namespace Horde.Server.Tools
 #pragma warning restore CA2000
 		}
 
-		async Task<Tool> UpdateAsync(Tool tool, UpdateDefinition<Tool> update, CancellationToken cancellationToken)
+		async Task<ToolDocument> UpdateAsync(ToolDocument tool, UpdateDefinition<ToolDocument> update, CancellationToken cancellationToken)
 		{
 			update = update.Set(x => x.LastUpdateTime, new DateTime(Math.Max(tool.LastUpdateTime.Ticks + 1, DateTime.UtcNow.Ticks)));
 
-			FilterDefinition<Tool> filter = Builders<Tool>.Filter.Eq(x => x.Id, tool.Id) & Builders<Tool>.Filter.Eq(x => x.LastUpdateTime, tool.LastUpdateTime);
-			return await _tools.FindOneAndUpdateAsync(filter, update, new FindOneAndUpdateOptions<Tool> { ReturnDocument = ReturnDocument.After }, cancellationToken);
+			FilterDefinition<ToolDocument> filter = Builders<ToolDocument>.Filter.Eq(x => x.Id, tool.Id) & Builders<ToolDocument>.Filter.Eq(x => x.LastUpdateTime, tool.LastUpdateTime);
+			return await _tools.FindOneAndUpdateAsync(filter, update, new FindOneAndUpdateOptions<ToolDocument> { ReturnDocument = ReturnDocument.After }, cancellationToken);
 		}
 	}
 }
diff --git a/Engine/Source/Programs/Horde/Horde.Server/Tools/ToolsController.cs b/Engine/Source/Programs/Horde/Horde.Server/Tools/ToolsController.cs
--- a/Engine/Source/Programs/Horde/Horde.Server/Tools/ToolsController.cs
+++ b/Engine/Source/Programs/Horde/Horde.Server/Tools/ToolsController.cs
@@ -52,12 +52,12 @@ namespace Horde.Server.Tools
 		[Route("/api/v1/tools/{id}/blobs")]
 		public async Task<ActionResult<WriteBlobResponse>> WriteBlobAsync(ToolId id, WriteBlobRequest request, CancellationToken cancellationToken = default)
 		{
-			ITool? tool = await _toolCollection.GetAsync(id, _globalConfig.Value, cancellationToken);
+			ITool? tool = await _toolCollection.GetAsync(id, cancellationToken);
 			if (tool == null)
 			{
 				return NotFound(id);
 			}
-			if (!tool.Config.Authorize(ToolAclAction.UploadTool, User))
+			if (!tool.Authorize(ToolAclAction.UploadTool, User))
 			{
 				return Forbid(ToolAclAction.UploadTool, id);
 			}
@@ -71,7 +71,7 @@ namespace Horde.Server.Tools
 				request.Prefix = $"{id}/{request.Prefix}";
 			}
 
-			IStorageBackend storageBackend = _toolCollection.CreateStorageBackend(tool);
+			IStorageBackend storageBackend = tool.CreateStorageBackend();
 			return await StorageController.WriteBlobAsync(storageBackend, request, cancellationToken);
 		}
 
@@ -83,20 +83,19 @@ namespace Horde.Server.Tools
 		[Route("/api/v1/tools/{id}/deployments")]
 		public async Task<ActionResult<CreateToolDeploymentResponse>> CreateDeploymentAsync(ToolId id, [FromForm] ToolDeploymentConfig options, [FromForm] IFormFile file, CancellationToken cancellationToken)
 		{
-			ITool? tool = await _toolCollection.GetAsync(id, _globalConfig.Value, cancellationToken);
-
+			ITool? tool = await _toolCollection.GetAsync(id, cancellationToken);
 			if (tool == null)
 			{
 				return NotFound(id);
 			}
-			if (!tool.Config.Authorize(ToolAclAction.UploadTool, User))
+			if (!tool.Authorize(ToolAclAction.UploadTool, User))
 			{
 				return Forbid(ToolAclAction.UploadTool, id);
 			}
 
 			using (Stream stream = file.OpenReadStream())
 			{
-				tool = await _toolCollection.CreateDeploymentAsync(tool, options, stream, _globalConfig.Value, cancellationToken);
+				tool = await tool.CreateDeploymentAsync(options, stream, cancellationToken);
 				if (tool == null)
 				{
 					return NotFound(id);
@@ -113,20 +112,20 @@ namespace Horde.Server.Tools
 		[Route("/api/v2/tools/{id}/deployments")]
 		public async Task<ActionResult<CreateToolDeploymentResponse>> CreateDeploymentAsync(ToolId id, CreateToolDeploymentRequest request, CancellationToken cancellationToken)
 		{
-			ITool? tool = await _toolCollection.GetAsync(id, _globalConfig.Value, cancellationToken);
+			ITool? tool = await _toolCollection.GetAsync(id, cancellationToken);
 
 			if (tool == null)
 			{
 				return NotFound(id);
 			}
-			if (!tool.Config.Authorize(ToolAclAction.UploadTool, User))
+			if (!tool.Authorize(ToolAclAction.UploadTool, User))
 			{
 				return Forbid(ToolAclAction.UploadTool, id);
 			}
 
 			ToolDeploymentConfig options = new ToolDeploymentConfig { Version = request.Version, Duration = TimeSpan.FromMinutes(request.Duration ?? 0.0), CreatePaused = request.CreatePaused ?? false };
 
-			tool = await _toolCollection.CreateDeploymentAsync(tool, options, request.Content, _globalConfig.Value, cancellationToken);
+			tool = await tool.CreateDeploymentAsync(options, request.Content, cancellationToken);
 			if (tool == null)
 			{
 				return NotFound(id);
@@ -142,20 +141,26 @@ namespace Horde.Server.Tools
 		[Route("/api/v1/tools/{id}/deployments/{deploymentId}")]
 		public async Task<ActionResult> UpdateDeploymentAsync(ToolId id, ToolDeploymentId deploymentId, [FromBody] UpdateDeploymentRequest request)
 		{
-			ITool? tool = await _toolCollection.GetAsync(id, _globalConfig.Value);
+			ITool? tool = await _toolCollection.GetAsync(id, HttpContext.RequestAborted);
 			if (tool == null)
 			{
 				return NotFound(id);
 			}
-			if (!tool.Config.Authorize(ToolAclAction.UploadTool, User))
+			if (!tool.Authorize(ToolAclAction.UploadTool, User))
 			{
 				return Forbid(ToolAclAction.UploadTool, id);
 			}
 
+			IToolDeployment? deployment = tool.Deployments.FirstOrDefault(x => x.Id == deploymentId);
+			if (deployment == null)
+			{
+				return NotFound(deploymentId);
+			}
+
 			if (request.State != null)
 			{
-				tool = await _toolCollection.UpdateDeploymentAsync(tool, deploymentId, request.State.Value);
-				if (tool == null)
+				deployment = await deployment.UpdateAsync(request.State.Value, HttpContext.RequestAborted);
+				if (deployment == null)
 				{
 					return NotFound(id, deploymentId);
 				}
@@ -194,28 +199,14 @@ namespace Horde.Server.Tools
 		[Route("/api/v1/tools")]
 		public async Task<ActionResult<GetToolsSummaryResponse>> GetToolsAsync()
 		{
-			GlobalConfig globalConfig = _globalConfig.Value;
-
-			Dictionary<ToolId, ToolConfig> tools = new Dictionary<ToolId, ToolConfig>();
-			foreach (BundledToolConfig bundledToolConfig in globalConfig.ServerSettings.BundledTools)
-			{
-				tools[bundledToolConfig.Id] = bundledToolConfig;
-			}
-			foreach (ToolConfig toolConfig in globalConfig.Tools)
-			{
-				tools[toolConfig.Id] = toolConfig;
-			}
+			IReadOnlyList<ITool> tools = await _toolCollection.GetAllAsync(HttpContext.RequestAborted);
 
 			List<GetToolSummaryResponse> toolSummaryList = new List<GetToolSummaryResponse>();
-			foreach (ToolConfig toolConfig in tools.Values.OrderBy(x => x.Name, StringComparer.Ordinal))
+			foreach (ITool tool in tools.OrderBy(x => x.Name, StringComparer.Ordinal))
 			{
-				if (AuthorizeDownload(toolConfig))
+				if (AuthorizeDownload(tool))
 				{
-					ITool? tool = await _toolCollection.GetAsync(toolConfig.Id, _globalConfig.Value);
-					if (tool != null)
-					{
-						toolSummaryList.Add(CreateGetToolSummaryResponse(tool));
-					}
+					toolSummaryList.Add(CreateGetToolSummaryResponse(tool));
 				}
 			}
 
@@ -225,7 +216,7 @@ namespace Horde.Server.Tools
 		static GetToolSummaryResponse CreateGetToolSummaryResponse(ITool tool)
 		{
 			IToolDeployment? deployment = (tool.Deployments.Count == 0) ? null : tool.Deployments[^1];
-			return new GetToolSummaryResponse(tool.Id, tool.Config.Name, tool.Config.Description, tool.Config.Category, deployment?.Version, deployment?.Id, tool.Config.ShowInUgs, tool.Config.ShowInDashboard);
+			return new GetToolSummaryResponse(tool.Id, tool.Name, tool.Description, tool.Category, deployment?.Version, deployment?.Id, tool.ShowInUgs, tool.ShowInDashboard);
 		}
 
 		/// <summary>
@@ -236,12 +227,12 @@ namespace Horde.Server.Tools
 		[Route("/api/v1/tools/{id}")]
 		public async Task<ActionResult> GetToolAsync(ToolId id, GetToolAction action = GetToolAction.Info, CancellationToken cancellationToken = default)
 		{
-			ITool? tool = await _toolCollection.GetAsync(id, _globalConfig.Value, cancellationToken);
+			ITool? tool = await _toolCollection.GetAsync(id, cancellationToken);
 			if (tool == null)
 			{
 				return NotFound(id);
 			}
-			if (!AuthorizeDownload(tool.Config))
+			if (!AuthorizeDownload(tool))
 			{
 				return Forbid(ToolAclAction.DownloadTool, id);
 			}
@@ -269,7 +260,7 @@ namespace Horde.Server.Tools
 
 		static GetToolResponse CreateGetToolResponse(ITool tool, List<GetToolDeploymentResponse> deployments)
 		{
-			return new GetToolResponse(tool.Id, tool.Config.Name, tool.Config.Description, tool.Config.Category, deployments, tool.Config.Public, tool.Config.ShowInUgs, tool.Config.ShowInDashboard);
+			return new GetToolResponse(tool.Id, tool.Name, tool.Description, tool.Category, deployments, tool.Public, tool.ShowInUgs, tool.ShowInDashboard);
 		}
 
 		/// <summary>
@@ -284,12 +275,12 @@ namespace Horde.Server.Tools
 		[Route("/api/v1/tools/{id}/deployments")]
 		public async Task<ActionResult> FindDeploymentAsync(ToolId id, [FromQuery] double phase = 0.0, [FromQuery] GetToolAction action = GetToolAction.Info, CancellationToken cancellationToken = default)
 		{
-			ITool? tool = await _toolCollection.GetAsync(id, _globalConfig.Value, cancellationToken);
+			ITool? tool = await _toolCollection.GetAsync(id, cancellationToken);
 			if (tool == null)
 			{
 				return NotFound(id);
 			}
-			if (!AuthorizeDownload(tool.Config))
+			if (!AuthorizeDownload(tool))
 			{
 				return Forbid(ToolAclAction.DownloadTool, id);
 			}
@@ -311,12 +302,12 @@ namespace Horde.Server.Tools
 		[Route("/api/v1/tools/{id}/deployments/{deploymentId}")]
 		public async Task<ActionResult> GetDeploymentAsync(ToolId id, ToolDeploymentId deploymentId, [FromQuery] GetToolAction action = GetToolAction.Info, CancellationToken cancellationToken = default)
 		{
-			ITool? tool = await _toolCollection.GetAsync(id, _globalConfig.Value, cancellationToken);
+			ITool? tool = await _toolCollection.GetAsync(id, cancellationToken);
 			if (tool == null)
 			{
 				return NotFound(id);
 			}
-			if (!AuthorizeDownload(tool.Config))
+			if (!AuthorizeDownload(tool))
 			{
 				return Forbid(ToolAclAction.DownloadTool, id);
 			}
@@ -338,7 +329,7 @@ namespace Horde.Server.Tools
 				return Ok(response);
 			}
 
-			IStorageClient client = _toolCollection.CreateStorageClient(tool);
+			IStorageClient client = tool.CreateStorageClient();
 			try
 			{
 				IBlobRef<DirectoryNode> nodeRef = await client.ReadRefAsync<DirectoryNode>(deployment.RefName, DateTime.UtcNow - TimeSpan.FromDays(2.0), cancellationToken: cancellationToken);
@@ -376,7 +367,7 @@ namespace Horde.Server.Tools
 
 		private async Task<GetToolDeploymentResponse> GetDeploymentInfoResponseAsync(ITool tool, IToolDeployment deployment, CancellationToken cancellationToken)
 		{
-			using IStorageClient client = _toolCollection.CreateStorageClient(tool);
+			using IStorageClient client = tool.CreateStorageClient();
 			IBlobHandle rootHandle = await client.ReadRefAsync(deployment.RefName, cancellationToken: cancellationToken);
 
 			return new GetToolDeploymentResponse(deployment.Id, deployment.Version, deployment.State, deployment.Progress, deployment.StartedAt, deployment.Duration, deployment.RefName, rootHandle.GetLocator());
@@ -393,12 +384,12 @@ namespace Horde.Server.Tools
 		[Route("/api/v1/tools/{id}/blobs/{*locator}")]
 		public async Task<ActionResult> ReadToolBlobAsync(ToolId id, BlobLocator locator, CancellationToken cancellationToken = default)
 		{
-			ITool? tool = await _toolCollection.GetAsync(id, _globalConfig.Value, cancellationToken);
+			ITool? tool = await _toolCollection.GetAsync(id, cancellationToken);
 			if (tool == null)
 			{
 				return NotFound(id);
 			}
-			if (!AuthorizeDownload(tool.Config))
+			if (!AuthorizeDownload(tool))
 			{
 				return Forbid(ToolAclAction.DownloadTool, id);
 			}
@@ -408,19 +399,19 @@ namespace Horde.Server.Tools
 				return BadRequest("Invalid blob id for tool");
 			}
 
-			IStorageBackend storageBackend = _toolCollection.CreateStorageBackend(tool);
+			IStorageBackend storageBackend = tool.CreateStorageBackend();
 			return await StorageController.ReadBlobInternalAsync(storageBackend, locator, Request.Headers, cancellationToken);
 		}
 
-		bool AuthorizeDownload(ToolConfig toolConfig)
+		bool AuthorizeDownload(ITool tool)
 		{
-			if (!toolConfig.Public)
+			if (!tool.Public)
 			{
 				if (User.Identity == null || !User.Identity.IsAuthenticated)
 				{
 					return false;
 				}
-				if (!toolConfig.Authorize(ToolAclAction.DownloadTool, User))
+				if (!tool.Authorize(ToolAclAction.DownloadTool, User))
 				{
 					return false;
 				}
