Horde: Reorder types in IJob.cs.

#rnx

diff --git a/Engine/Source/Programs/Horde/Horde.Server/Jobs/IJob.cs b/Engine/Source/Programs/Horde/Horde.Server/Jobs/IJob.cs
--- a/Engine/Source/Programs/Horde/Horde.Server/Jobs/IJob.cs
+++ b/Engine/Source/Programs/Horde/Horde.Server/Jobs/IJob.cs
@@ -32,510 +32,6 @@
 namespace Horde.Server.Jobs
 {
 	/// <summary>
-	/// Report for a job or jobstep
-	/// </summary>
-	public interface IReport
-	{
-		/// <summary>
-		/// Name of the report
-		/// </summary>
-		string Name { get; }
-
-		/// <summary>
-		/// Where to render the report
-		/// </summary>
-		ReportPlacement Placement { get; }
-
-		/// <summary>
-		/// The artifact id
-		/// </summary>
-		ObjectId? ArtifactId { get; }
-
-		/// <summary>
-		/// Inline data for the report
-		/// </summary>
-		string? Content { get; }
-	}
-
-	/// <summary>
-	/// Implementation of IReport
-	/// </summary>
-	public class Report : IReport
-	{
-		/// <inheritdoc/>
-		public string Name { get; set; } = String.Empty;
-
-		/// <inheritdoc/>
-		public ReportPlacement Placement { get; set; }
-
-		/// <inheritdoc/>
-		public ObjectId? ArtifactId { get; set; }
-
-		/// <inheritdoc/>
-		public string? Content { get; set; }
-	}
-
-	/// <summary>
-	/// Embedded jobstep document
-	/// </summary>
-	public interface IJobStep
-	{
-		/// <summary>
-		/// Unique ID assigned to this jobstep. A new id is generated whenever a jobstep's order is changed.
-		/// </summary>
-		public JobStepId Id { get; }
-
-		/// <summary>
-		/// Index of the node which this jobstep is to execute
-		/// </summary>
-		public int NodeIdx { get; }
-
-		/// <summary>
-		/// Current state of the job step. This is updated automatically when runs complete.
-		/// </summary>
-		public JobStepState State { get; }
-
-		/// <summary>
-		/// Current outcome of the jobstep
-		/// </summary>
-		public JobStepOutcome Outcome { get; }
-
-		/// <summary>
-		/// Error from executing this step
-		/// </summary>
-		public JobStepError Error { get; }
-
-		/// <summary>
-		/// The log id for this step
-		/// </summary>
-		public LogId? LogId { get; }
-
-		/// <summary>
-		/// Unique id for notifications
-		/// </summary>
-		public ObjectId? NotificationTriggerId { get; }
-
-		/// <summary>
-		/// Time at which the batch transitioned to the ready state (UTC).
-		/// </summary>
-		public DateTime? ReadyTimeUtc { get; }
-
-		/// <summary>
-		/// Time at which the batch transitioned to the executing state (UTC).
-		/// </summary>
-		public DateTime? StartTimeUtc { get; }
-
-		/// <summary>
-		/// Time at which the run finished (UTC)
-		/// </summary>
-		public DateTime? FinishTimeUtc { get; }
-
-		/// <summary>
-		/// Override for the priority of this step
-		/// </summary>
-		public Priority? Priority { get; }
-
-		/// <summary>
-		/// If a retry is requested, stores the name of the user that requested it
-		/// </summary>
-		public UserId? RetriedByUserId { get; }
-
-		/// <summary>
-		/// Signal if a step should be aborted
-		/// </summary>
-		public bool AbortRequested { get; }
-
-		/// <summary>
-		/// If an abort is requested, stores the id of the user that requested it
-		/// </summary>
-		public UserId? AbortedByUserId { get; }
-
-		/// <summary>
-		/// List of reports for this step
-		/// </summary>
-		public IReadOnlyList<IReport>? Reports { get; }
-
-		/// <summary>
-		/// Reports for this jobstep.
-		/// </summary>
-		public Dictionary<string, string>? Properties { get; }
-	}
-
-	/// <summary>
-	/// Extension methods for job steps
-	/// </summary>
-	public static class JobStepExtensions
-	{
-		/// <summary>
-		/// Determines if a jobstep has failed or is skipped. Can be used to determine whether dependent steps will be able to run.
-		/// </summary>
-		/// <returns>True if the step is failed or skipped</returns>
-		public static bool IsFailedOrSkipped(this IJobStep step)
-		{
-			return step.State == JobStepState.Skipped || step.Outcome == JobStepOutcome.Failure;
-		}
-
-		/// <summary>
-		/// Determines if a jobstep is done by checking to see if it is completed, skipped, or aborted.
-		/// </summary>
-		/// <returns>True if the step is completed, skipped, or aborted</returns>
-		public static bool IsPending(this IJobStep step)
-		{
-			return step.State != JobStepState.Aborted && step.State != JobStepState.Completed && step.State != JobStepState.Skipped;
-		}
-
-		/// <summary>
-		/// Determine if a step should be timed out
-		/// </summary>
-		/// <param name="step"></param>
-		/// <param name="utcNow"></param>
-		/// <returns></returns>
-		public static bool HasTimedOut(this IJobStep step, DateTime utcNow)
-		{
-			if (step.State == JobStepState.Running && step.StartTimeUtc != null)
-			{
-				TimeSpan elapsed = utcNow - step.StartTimeUtc.Value;
-				if (elapsed > TimeSpan.FromHours(24.0))
-				{
-					return true;
-				}
-			}
-			return false;
-		}
-	}
-
-	/// <summary>
-	/// Stores information about a batch of job steps
-	/// </summary>
-	public interface IJobStepBatch
-	{
-		/// <summary>
-		/// Unique id for this group
-		/// </summary>
-		public JobStepBatchId Id { get; }
-
-		/// <summary>
-		/// The log file id for this batch
-		/// </summary>
-		public LogId? LogId { get; }
-
-		/// <summary>
-		/// Index of the group being executed
-		/// </summary>
-		public int GroupIdx { get; }
-
-		/// <summary>
-		/// The state of this group
-		/// </summary>
-		public JobStepBatchState State { get; }
-
-		/// <summary>
-		/// Error associated with this group
-		/// </summary>
-		public JobStepBatchError Error { get; }
-
-		/// <summary>
-		/// Steps within this run
-		/// </summary>
-		public IReadOnlyList<IJobStep> Steps { get; }
-
-		/// <summary>
-		/// The pool that this agent was taken from
-		/// </summary>
-		public PoolId? PoolId { get; }
-
-		/// <summary>
-		/// The agent assigned to execute this group
-		/// </summary>
-		public AgentId? AgentId { get; }
-
-		/// <summary>
-		/// The agent session that is executing this group
-		/// </summary>
-		public SessionId? SessionId { get; }
-
-		/// <summary>
-		/// The lease that's executing this group
-		/// </summary>
-		public LeaseId? LeaseId { get; }
-
-		/// <summary>
-		/// The weighted priority of this batch for the scheduler
-		/// </summary>
-		public int SchedulePriority { get; }
-
-		/// <summary>
-		/// Time at which the group became ready (UTC).
-		/// </summary>
-		public DateTime? ReadyTimeUtc { get; }
-
-		/// <summary>
-		/// Time at which the group started (UTC).
-		/// </summary>
-		public DateTime? StartTimeUtc { get; }
-
-		/// <summary>
-		/// Time at which the group finished (UTC)
-		/// </summary>
-		public DateTime? FinishTimeUtc { get; }
-	}
-
-	/// <summary>
-	/// Extension methods for IJobStepBatch
-	/// </summary>
-	public static class JobStepBatchExtensions
-	{
-		/// <summary>
-		/// Attempts to get a step with the given id
-		/// </summary>
-		/// <param name="batch">The batch to search</param>
-		/// <param name="stepId">The step id</param>
-		/// <param name="step">On success, receives the step object</param>
-		/// <returns>True if the step was found</returns>
-		public static bool TryGetStep(this IJobStepBatch batch, JobStepId stepId, [NotNullWhen(true)] out IJobStep? step)
-		{
-			step = batch.Steps.FirstOrDefault(x => x.Id == stepId);
-			return step != null;
-		}
-
-		/// <summary>
-		/// Determines if new steps can be appended to this batch. We do not allow this after the last step has been completed, because the agent is shutting down.
-		/// </summary>
-		/// <param name="batch">The batch to search</param>
-		/// <returns>True if new steps can be appended to this batch</returns>
-		public static bool CanBeAppendedTo(this IJobStepBatch batch)
-		{
-			return batch.State <= JobStepBatchState.Running;
-		}
-
-		/// <summary>
-		/// Gets the wait time for this batch
-		/// </summary>
-		/// <param name="batch">The batch to search</param>
-		/// <returns>Wait time for the batch</returns>
-		public static TimeSpan? GetWaitTime(this IJobStepBatch batch)
-		{
-			if (batch.StartTimeUtc == null || batch.ReadyTimeUtc == null)
-			{
-				return null;
-			}
-			else
-			{
-				return batch.StartTimeUtc.Value - batch.ReadyTimeUtc.Value;
-			}
-		}
-
-		/// <summary>
-		/// Gets the initialization time for this batch
-		/// </summary>
-		/// <param name="batch">The batch to search</param>
-		/// <returns>Initialization time for this batch</returns>
-		public static TimeSpan? GetInitTime(this IJobStepBatch batch)
-		{
-			if (batch.StartTimeUtc != null)
-			{
-				foreach (IJobStep step in batch.Steps)
-				{
-					if (step.StartTimeUtc != null)
-					{
-						return step.StartTimeUtc - batch.StartTimeUtc.Value;
-					}
-				}
-			}
-			return null;
-		}
-
-		/// <summary>
-		/// Get the dependencies required for this batch to start, taking run-early nodes into account
-		/// </summary>
-		/// <param name="batch">The batch to search</param>
-		/// <param name="groups">List of node groups</param>
-		/// <returns>Set of nodes that must have completed for this batch to start</returns>
-		public static HashSet<INode> GetStartDependencies(this IJobStepBatch batch, IReadOnlyList<INodeGroup> groups)
-		{
-			// Find all the nodes that this group will start with.
-			List<INode> nodes = batch.Steps.ConvertAll(x => groups[batch.GroupIdx].Nodes[x.NodeIdx]);
-			if (nodes.Any(x => x.RunEarly))
-			{
-				nodes.RemoveAll(x => !x.RunEarly);
-			}
-
-			// Find all their dependencies
-			HashSet<INode> dependencies = new HashSet<INode>();
-			foreach (INode node in nodes)
-			{
-				dependencies.UnionWith(node.InputDependencies.Select(x => groups[x.GroupIdx].Nodes[x.NodeIdx]));
-				dependencies.UnionWith(node.OrderDependencies.Select(x => groups[x.GroupIdx].Nodes[x.NodeIdx]));
-			}
-
-			// Exclude all the dependencies within the same group
-			dependencies.ExceptWith(groups[batch.GroupIdx].Nodes);
-			return dependencies;
-		}
-	}
-
-	/// <summary>
-	/// Cumulative timing information to reach a certain point in a job
-	/// </summary>
-	public class TimingInfo
-	{
-		/// <summary>
-		/// Wait time on the critical path
-		/// </summary>
-		public TimeSpan? TotalWaitTime { get; set; }
-
-		/// <summary>
-		/// Sync time on the critical path
-		/// </summary>
-		public TimeSpan? TotalInitTime { get; set; }
-
-		/// <summary>
-		/// Duration to this point
-		/// </summary>
-		public TimeSpan? TotalTimeToComplete { get; set; }
-
-		/// <summary>
-		/// Average wait time to this point
-		/// </summary>
-		public TimeSpan? AverageTotalWaitTime { get; set; }
-
-		/// <summary>
-		/// Average sync time to this point
-		/// </summary>
-		public TimeSpan? AverageTotalInitTime { get; set; }
-
-		/// <summary>
-		/// Average duration to this point
-		/// </summary>
-		public TimeSpan? AverageTotalTimeToComplete { get; set; }
-
-		/// <summary>
-		/// Individual step timing information
-		/// </summary>
-		public IJobStepTiming? StepTiming { get; set; }
-
-		/// <summary>
-		/// Constructor
-		/// </summary>
-		public TimingInfo()
-		{
-			TotalWaitTime = TimeSpan.Zero;
-			TotalInitTime = TimeSpan.Zero;
-			TotalTimeToComplete = TimeSpan.Zero;
-
-			AverageTotalWaitTime = TimeSpan.Zero;
-			AverageTotalInitTime = TimeSpan.Zero;
-			AverageTotalTimeToComplete = TimeSpan.Zero;
-		}
-
-		/// <summary>
-		/// Copy constructor
-		/// </summary>
-		/// <param name="other">The timing info object to copy from</param>
-		public TimingInfo(TimingInfo other)
-		{
-			TotalWaitTime = other.TotalWaitTime;
-			TotalInitTime = other.TotalInitTime;
-			TotalTimeToComplete = other.TotalTimeToComplete;
-
-			AverageTotalWaitTime = other.AverageTotalWaitTime;
-			AverageTotalInitTime = other.AverageTotalInitTime;
-			AverageTotalTimeToComplete = other.AverageTotalTimeToComplete;
-		}
-
-		/// <summary>
-		/// Modifies this timing to wait for another timing
-		/// </summary>
-		/// <param name="other">The other node to wait for</param>
-		public void WaitFor(TimingInfo other)
-		{
-			if (TotalTimeToComplete != null)
-			{
-				if (other.TotalTimeToComplete == null || other.TotalTimeToComplete.Value > TotalTimeToComplete.Value)
-				{
-					TotalInitTime = other.TotalInitTime;
-					TotalWaitTime = other.TotalWaitTime;
-					TotalTimeToComplete = other.TotalTimeToComplete;
-				}
-			}
-
-			if (AverageTotalTimeToComplete != null)
-			{
-				if (other.AverageTotalTimeToComplete == null || other.AverageTotalTimeToComplete.Value > AverageTotalTimeToComplete.Value)
-				{
-					AverageTotalInitTime = other.AverageTotalInitTime;
-					AverageTotalWaitTime = other.AverageTotalWaitTime;
-					AverageTotalTimeToComplete = other.AverageTotalTimeToComplete;
-				}
-			}
-		}
-
-		/// <summary>
-		/// Waits for all the given timing info objects to complete
-		/// </summary>
-		/// <param name="others">Other timing info objects to wait for</param>
-		public void WaitForAll(IEnumerable<TimingInfo> others)
-		{
-			foreach (TimingInfo other in others)
-			{
-				WaitFor(other);
-			}
-		}
-
-		/// <summary>
-		/// Constructs a new TimingInfo object which represents the last TimingInfo to finish
-		/// </summary>
-		/// <param name="others">TimingInfo objects to wait for</param>
-		/// <returns>New TimingInfo instance</returns>
-		public static TimingInfo Max(IEnumerable<TimingInfo> others)
-		{
-			TimingInfo timingInfo = new TimingInfo();
-			timingInfo.WaitForAll(others);
-			return timingInfo;
-		}
-
-		/// <summary>
-		/// Copies this info to a repsonse object
-		/// </summary>
-		public void CopyToResponse(GetTimingInfoResponse response)
-		{
-			response.TotalWaitTime = (float?)TotalWaitTime?.TotalSeconds;
-			response.TotalInitTime = (float?)TotalInitTime?.TotalSeconds;
-			response.TotalTimeToComplete = (float?)TotalTimeToComplete?.TotalSeconds;
-
-			response.AverageTotalWaitTime = (float?)AverageTotalWaitTime?.TotalSeconds;
-			response.AverageTotalInitTime = (float?)AverageTotalInitTime?.TotalSeconds;
-			response.AverageTotalTimeToComplete = (float?)AverageTotalTimeToComplete?.TotalSeconds;
-		}
-	}
-
-	/// <summary>
-	/// Information about a chained job trigger
-	/// </summary>
-	public interface IChainedJob
-	{
-		/// <summary>
-		/// The target to monitor
-		/// </summary>
-		public string Target { get; }
-
-		/// <summary>
-		/// The template to trigger on success
-		/// </summary>
-		public TemplateId TemplateRefId { get; }
-
-		/// <summary>
-		/// The triggered job id
-		/// </summary>
-		public JobId? JobId { get; }
-
-		/// <summary>
-		/// Whether to run the latest change, or default change for the template, when starting the new job. Uses same change as the triggering job by default.
-		/// </summary>
-		public bool UseDefaultChangeForTemplate { get; }
-	}
-
-	/// <summary>
 	/// Document describing a job
 	/// </summary>
 	public interface IJob
@@ -671,11 +167,6 @@
 		public IReadOnlyList<AclClaimConfig> Claims { get; }
 
 		/// <summary>
-		/// Largest value of the CombinedPriority value for batches in the ready state.
-		/// </summary>
-		public int SchedulePriority { get; }
-
-		/// <summary>
 		/// Array of jobstep runs
 		/// </summary>
 		public IReadOnlyList<IJobStepBatch> Batches { get; }
@@ -741,11 +232,6 @@
 		public IReadOnlyList<IChainedJob> ChainedJobs { get; }
 
 		/// <summary>
-		/// Next id for batches or groups
-		/// </summary>
-		public SubResourceId NextSubResourceId { get; }
-
-		/// <summary>
 		/// The last update time
 		/// </summary>
 		public DateTime UpdateTimeUtc { get; }
@@ -1184,3 +670,720 @@
 					}
 				}
 			}
+			return null;
+		}
+
+		/// <summary>
+		/// Gets the estimated timing info for all nodes in the job
+		/// </summary>
+		/// <param name="job">The job document</param>
+		/// <param name="graph">Graph for this job</param>
+		/// <param name="jobTiming">Job timing information</param>
+		/// <returns>Map of node to expected timing info</returns>
+		public static Dictionary<INode, TimingInfo> GetTimingInfo(this IJob job, IGraph graph, IJobTiming jobTiming)
+		{
+#pragma warning disable IDE0054 // Use compound assignment
+			TimeSpan currentTime = DateTime.UtcNow - job.CreateTimeUtc;
+
+			Dictionary<INode, TimingInfo> nodeToTimingInfo = graph.Groups.SelectMany(x => x.Nodes).ToDictionary(x => x, x => new TimingInfo());
+			foreach (IJobStepBatch batch in job.Batches)
+			{
+				INodeGroup group = graph.Groups[batch.GroupIdx];
+
+				// Step through the batch, keeping track of the time that things finish.
+				TimingInfo timingInfo = new TimingInfo();
+
+				// Wait for the dependencies for the batch to start
+				HashSet<INode> dependencyNodes = batch.GetStartDependencies(graph.Groups);
+				timingInfo.WaitForAll(dependencyNodes.Select(x => nodeToTimingInfo[x]));
+
+				// If the batch has actually started, correct the expected time to use this instead
+				if (batch.StartTimeUtc != null)
+				{
+					timingInfo.TotalTimeToComplete = batch.StartTimeUtc - job.CreateTimeUtc;
+				}
+
+				// Get the average times for this batch
+				TimeSpan? averageWaitTime = GetAverageWaitTime(graph, batch, jobTiming);
+				TimeSpan? averageInitTime = GetAverageInitTime(graph, batch, jobTiming);
+
+				// Update the wait times and initialization times along this path
+				timingInfo.TotalWaitTime = timingInfo.TotalWaitTime + (batch.GetWaitTime() ?? averageWaitTime);
+				timingInfo.TotalInitTime = timingInfo.TotalInitTime + (batch.GetInitTime() ?? averageInitTime);
+
+				// Update the average wait and initialization times too
+				timingInfo.AverageTotalWaitTime = timingInfo.AverageTotalWaitTime + averageWaitTime;
+				timingInfo.AverageTotalInitTime = timingInfo.AverageTotalInitTime + averageInitTime;
+
+				// Step through the batch, updating the expected times as we go
+				foreach (IJobStep step in batch.Steps)
+				{
+					INode node = group.Nodes[step.NodeIdx];
+
+					// Get the timing for this step
+					IJobStepTiming? stepTimingInfo;
+					jobTiming.TryGetStepTiming(node.Name, out stepTimingInfo);
+
+					// If the step has already started, update the actual time to reach this point
+					if (step.StartTimeUtc != null)
+					{
+						timingInfo.TotalTimeToComplete = step.StartTimeUtc.Value - job.CreateTimeUtc;
+					}
+
+					// If the step hasn't started yet, make sure the start time is later than the current time
+					if (step.StartTimeUtc == null && currentTime > timingInfo.TotalTimeToComplete)
+					{
+						timingInfo.TotalTimeToComplete = currentTime;
+					}
+
+					// Wait for all the node dependencies to complete
+					timingInfo.WaitForAll(graph.GetDependencies(node).Select(x => nodeToTimingInfo[x]));
+
+					// If the step has actually finished, correct the time to use that instead
+					if (step.FinishTimeUtc != null)
+					{
+						timingInfo.TotalTimeToComplete = step.FinishTimeUtc.Value - job.CreateTimeUtc;
+					}
+					else
+					{
+						timingInfo.TotalTimeToComplete = timingInfo.TotalTimeToComplete + NullableTimeSpanFromSeconds(stepTimingInfo?.AverageDuration);
+					}
+
+					// If the step hasn't finished yet, make sure the start time is later than the current time
+					if (step.FinishTimeUtc == null && currentTime > timingInfo.TotalTimeToComplete)
+					{
+						timingInfo.TotalTimeToComplete = currentTime;
+					}
+
+					// Update the average time to complete
+					timingInfo.AverageTotalTimeToComplete = timingInfo.AverageTotalTimeToComplete + NullableTimeSpanFromSeconds(stepTimingInfo?.AverageDuration);
+
+					// Add it to the lookup
+					TimingInfo nodeTimingInfo = new TimingInfo(timingInfo);
+					nodeTimingInfo.StepTiming = stepTimingInfo;
+					nodeToTimingInfo[node] = nodeTimingInfo;
+				}
+			}
+			return nodeToTimingInfo;
+#pragma warning restore IDE0054 // Use compound assignment
+		}
+
+		/// <summary>
+		/// Gets the average wait time for this batch
+		/// </summary>
+		/// <param name="graph">Graph for the job</param>
+		/// <param name="batch">The batch to get timing info for</param>
+		/// <param name="jobTiming">The job timing information</param>
+		/// <returns>Wait time for the batch</returns>
+		public static TimeSpan? GetAverageWaitTime(IGraph graph, IJobStepBatch batch, IJobTiming jobTiming)
+		{
+			TimeSpan? waitTime = null;
+			foreach (IJobStep step in batch.Steps)
+			{
+				INode node = graph.Groups[batch.GroupIdx].Nodes[step.NodeIdx];
+				if (jobTiming.TryGetStepTiming(node.Name, out IJobStepTiming? timingInfo))
+				{
+					if (timingInfo.AverageWaitTime != null)
+					{
+						TimeSpan stepWaitTime = TimeSpan.FromSeconds(timingInfo.AverageWaitTime.Value);
+						if (waitTime == null || stepWaitTime > waitTime.Value)
+						{
+							waitTime = stepWaitTime;
+						}
+					}
+				}
+			}
+			return waitTime;
+		}
+
+		/// <summary>
+		/// Gets the average initialization time for this batch
+		/// </summary>
+		/// <param name="graph">Graph for the job</param>
+		/// <param name="batch">The batch to get timing info for</param>
+		/// <param name="jobTiming">The job timing information</param>
+		/// <returns>Initialization time for this batch</returns>
+		public static TimeSpan? GetAverageInitTime(IGraph graph, IJobStepBatch batch, IJobTiming jobTiming)
+		{
+			TimeSpan? initTime = null;
+			foreach (IJobStep step in batch.Steps)
+			{
+				INode node = graph.Groups[batch.GroupIdx].Nodes[step.NodeIdx];
+				if (jobTiming.TryGetStepTiming(node.Name, out IJobStepTiming? timingInfo))
+				{
+					if (timingInfo.AverageInitTime != null)
+					{
+						TimeSpan stepInitTime = TimeSpan.FromSeconds(timingInfo.AverageInitTime.Value);
+						if (initTime == null || stepInitTime > initTime.Value)
+						{
+							initTime = stepInitTime;
+						}
+					}
+				}
+			}
+			return initTime;
+		}
+
+		/// <summary>
+		/// Creates a nullable timespan from a nullable number of seconds
+		/// </summary>
+		/// <param name="seconds">The number of seconds to construct from</param>
+		/// <returns>TimeSpan object</returns>
+		static TimeSpan? NullableTimeSpanFromSeconds(float? seconds)
+		{
+			if (seconds == null)
+			{
+				return null;
+			}
+			else
+			{
+				return TimeSpan.FromSeconds(seconds.Value);
+			}
+		}
+
+		/// <summary>
+		/// Attempts to get a batch with the given id
+		/// </summary>
+		/// <param name="job">The job document</param>
+		/// <param name="batchId">The batch id</param>
+		/// <param name="batch">On success, receives the batch object</param>
+		/// <returns>True if the batch was found</returns>
+		public static bool TryGetBatch(this IJob job, JobStepBatchId batchId, [NotNullWhen(true)] out IJobStepBatch? batch)
+		{
+			batch = job.Batches.FirstOrDefault(x => x.Id == batchId);
+			return batch != null;
+		}
+
+		/// <summary>
+		/// Attempts to get a batch with the given id
+		/// </summary>
+		/// <param name="job">The job document</param>
+		/// <param name="batchId">The batch id</param>
+		/// <param name="stepId">The step id</param>
+		/// <param name="step">On success, receives the step object</param>
+		/// <returns>True if the batch was found</returns>
+		public static bool TryGetStep(this IJob job, JobStepBatchId batchId, JobStepId stepId, [NotNullWhen(true)] out IJobStep? step)
+		{
+			IJobStepBatch? batch;
+			if (!TryGetBatch(job, batchId, out batch))
+			{
+				step = null;
+				return false;
+			}
+			return batch.TryGetStep(stepId, out step);
+		}
+
+		/// <summary>
+		/// Finds the set of nodes affected by a label
+		/// </summary>
+		/// <param name="job">The job document</param>
+		/// <param name="graph">Graph definition for the job</param>
+		/// <param name="labelIdx">Index of the label. -1 or Graph.Labels.Count are treated as referring to the default lable.</param>
+		/// <returns>Set of nodes affected by the given label</returns>
+		public static HashSet<NodeRef> GetNodesForLabel(this IJob job, IGraph graph, int labelIdx)
+		{
+			if (labelIdx != -1 && labelIdx != graph.Labels.Count)
+			{
+				// Return all the nodes included by the label
+				return new HashSet<NodeRef>(graph.Labels[labelIdx].IncludedNodes);
+			}
+			else
+			{
+				// Set of nodes which are not covered by an existing label, initially containing everything
+				HashSet<NodeRef> unlabeledNodes = new HashSet<NodeRef>();
+				for (int groupIdx = 0; groupIdx < graph.Groups.Count; groupIdx++)
+				{
+					INodeGroup group = graph.Groups[groupIdx];
+					for (int nodeIdx = 0; nodeIdx < group.Nodes.Count; nodeIdx++)
+					{
+						unlabeledNodes.Add(new NodeRef(groupIdx, nodeIdx));
+					}
+				}
+
+				// Remove all the nodes that are part of an active label
+				IReadOnlyDictionary<NodeRef, IJobStep> stepForNode = job.GetStepForNodeMap();
+				foreach (ILabel label in graph.Labels)
+				{
+					if (label.RequiredNodes.Any(x => stepForNode.ContainsKey(x)))
+					{
+						unlabeledNodes.ExceptWith(label.IncludedNodes);
+					}
+				}
+				return unlabeledNodes;
+			}
+		}
+
+		/// <summary>
+		/// Create a list of aggregate responses, combining the graph definitions with the state of the job
+		/// </summary>
+		/// <param name="job">The job document</param>
+		/// <param name="graph">Graph definition for the job</param>
+		/// <param name="responses">List to receive all the responses</param>
+		/// <returns>The default label state</returns>
+		public static GetDefaultLabelStateResponse? GetLabelStateResponses(this IJob job, IGraph graph, List<GetLabelStateResponse> responses)
+		{
+			// Create a lookup from noderef to step information
+			IReadOnlyDictionary<NodeRef, IJobStep> stepForNode = job.GetStepForNodeMap();
+
+			// Set of nodes which are not covered by an existing label, initially containing everything
+			HashSet<NodeRef> unlabeledNodes = new HashSet<NodeRef>();
+			for (int groupIdx = 0; groupIdx < graph.Groups.Count; groupIdx++)
+			{
+				INodeGroup group = graph.Groups[groupIdx];
+				for (int nodeIdx = 0; nodeIdx < group.Nodes.Count; nodeIdx++)
+				{
+					unlabeledNodes.Add(new NodeRef(groupIdx, nodeIdx));
+				}
+			}
+
+			// Create the responses
+			foreach (ILabel label in graph.Labels)
+			{
+				// Refresh the state for this label
+				LabelState newState = LabelState.Unspecified;
+				foreach (NodeRef requiredNodeRef in label.RequiredNodes)
+				{
+					if (stepForNode.ContainsKey(requiredNodeRef))
+					{
+						newState = LabelState.Complete;
+						break;
+					}
+				}
+
+				// Refresh the outcome
+				LabelOutcome newOutcome = LabelOutcome.Success;
+				if (newState == LabelState.Complete)
+				{
+					GetLabelState(label.IncludedNodes, stepForNode, out newState, out newOutcome);
+					unlabeledNodes.ExceptWith(label.IncludedNodes);
+				}
+
+				// Create the response
+				responses.Add(new GetLabelStateResponse(newState, newOutcome));
+			}
+
+			// Remove all the nodes that don't have a step
+			unlabeledNodes.RemoveWhere(x => !stepForNode.ContainsKey(x));
+
+			// Remove successful "setup build" nodes from the list
+			if (graph.Groups.Count > 1 && graph.Groups[0].Nodes.Count > 0)
+			{
+				INode node = graph.Groups[0].Nodes[0];
+				if (node.Name == IJob.SetupNodeName)
+				{
+					NodeRef nodeRef = new NodeRef(0, 0);
+					if (unlabeledNodes.Contains(nodeRef))
+					{
+						IJobStep step = stepForNode[nodeRef];
+						if (step.State == JobStepState.Completed && step.Outcome == JobStepOutcome.Success && responses.Count > 0)
+						{
+							unlabeledNodes.Remove(nodeRef);
+						}
+					}
+				}
+			}
+
+			// Add a response for everything not included elsewhere.
+			GetLabelState(unlabeledNodes, stepForNode, out LabelState otherState, out LabelOutcome otherOutcome);
+			return new GetDefaultLabelStateResponse(otherState, otherOutcome, unlabeledNodes.Select(x => graph.GetNode(x).Name).ToList());
+		}
+
+		/// <summary>
+		/// Get the states of all labels for this job
+		/// </summary>
+		/// <param name="job">The job to get states for</param>
+		/// <param name="graph">The graph for this job</param>
+		/// <returns>Collection of label states by label index</returns>
+		public static IReadOnlyList<(LabelState, LabelOutcome)> GetLabelStates(this IJob job, IGraph graph)
+		{
+			IReadOnlyDictionary<NodeRef, IJobStep> stepForNodeRef = job.GetStepForNodeMap();
+
+			List<(LabelState, LabelOutcome)> states = new List<(LabelState, LabelOutcome)>();
+			for (int idx = 0; idx < graph.Labels.Count; idx++)
+			{
+				ILabel label = graph.Labels[idx];
+
+				// Default the label to the unspecified state
+				LabelState newState = LabelState.Unspecified;
+				LabelOutcome newOutcome = LabelOutcome.Unspecified;
+
+				// Check if the label should be included
+				if (label.RequiredNodes.Any(x => stepForNodeRef.ContainsKey(x)))
+				{
+					// Combine the state of the steps contributing towards this label
+					bool anySkipped = false;
+					bool anyWarnings = false;
+					bool anyFailed = false;
+					bool anyPending = false;
+					foreach (NodeRef includedNode in label.IncludedNodes)
+					{
+						IJobStep? step;
+						if (stepForNodeRef.TryGetValue(includedNode, out step))
+						{
+							anyPending |= step.IsPending();
+							anySkipped |= step.State == JobStepState.Aborted || step.State == JobStepState.Skipped;
+							anyFailed |= (step.Outcome == JobStepOutcome.Failure);
+							anyWarnings |= (step.Outcome == JobStepOutcome.Warnings);
+						}
+					}
+
+					// Figure out the overall label state
+					newState = anyPending ? LabelState.Running : LabelState.Complete;
+					newOutcome = anyFailed ? LabelOutcome.Failure : anyWarnings ? LabelOutcome.Warnings : anySkipped ? LabelOutcome.Unspecified : LabelOutcome.Success;
+				}
+
+				states.Add((newState, newOutcome));
+			}
+			return states;
+		}
+
+		/// <summary>
+		/// Get the states of all UGS badges for this job
+		/// </summary>
+		/// <param name="job">The job to get states for</param>
+		/// <param name="graph">The graph for this job</param>
+		/// <returns>List of badge states</returns>
+		public static Dictionary<int, UgsBadgeState> GetUgsBadgeStates(this IJob job, IGraph graph)
+		{
+			IReadOnlyList<(LabelState, LabelOutcome)> labelStates = GetLabelStates(job, graph);
+			return job.GetUgsBadgeStates(graph, labelStates);
+		}
+
+		/// <summary>
+		/// Get the states of all UGS badges for this job
+		/// </summary>
+		/// <param name="job">The job to get states for</param>
+		/// <param name="graph">The graph for this job</param>
+		/// <param name="labelStates">The existing label states to get the UGS badge states from</param>
+		/// <returns>List of badge states</returns>
+#pragma warning disable IDE0060 // Remove unused parameter
+		public static Dictionary<int, UgsBadgeState> GetUgsBadgeStates(this IJob job, IGraph graph, IReadOnlyList<(LabelState, LabelOutcome)> labelStates)
+#pragma warning restore IDE0060 // Remove unused parameter
+		{
+			Dictionary<int, UgsBadgeState> ugsBadgeStates = new Dictionary<int, UgsBadgeState>();
+			for (int labelIdx = 0; labelIdx < labelStates.Count; ++labelIdx)
+			{
+				if (graph.Labels[labelIdx].UgsName == null)
+				{
+					continue;
+				}
+
+				(LabelState state, LabelOutcome outcome) = labelStates[labelIdx];
+				switch (state)
+				{
+					case LabelState.Complete:
+						{
+							switch (outcome)
+							{
+								case LabelOutcome.Success:
+									{
+										ugsBadgeStates.Add(labelIdx, UgsBadgeState.Success);
+										break;
+									}
+
+								case LabelOutcome.Warnings:
+									{
+										ugsBadgeStates.Add(labelIdx, UgsBadgeState.Warning);
+										break;
+									}
+
+								case LabelOutcome.Failure:
+									{
+										ugsBadgeStates.Add(labelIdx, UgsBadgeState.Failure);
+										break;
+									}
+
+								case LabelOutcome.Unspecified:
+									{
+										ugsBadgeStates.Add(labelIdx, UgsBadgeState.Skipped);
+										break;
+									}
+							}
+							break;
+						}
+
+					case LabelState.Running:
+						{
+							ugsBadgeStates.Add(labelIdx, UgsBadgeState.Starting);
+							break;
+						}
+
+					case LabelState.Unspecified:
+						{
+							ugsBadgeStates.Add(labelIdx, UgsBadgeState.Skipped);
+							break;
+						}
+				}
+			}
+			return ugsBadgeStates;
+		}
+
+		/// <summary>
+		/// Gets the state of a job, as a label that includes all steps
+		/// </summary>
+		/// <param name="job">The job to query</param>
+		/// <param name="stepForNode">Map from node to step</param>
+		/// <param name="newState">Receives the state of the label</param>
+		/// <param name="newOutcome">Receives the outcome of the label</param>
+		public static void GetJobState(this IJob job, IReadOnlyDictionary<NodeRef, IJobStep> stepForNode, out LabelState newState, out LabelOutcome newOutcome)
+		{
+			List<NodeRef> nodes = new List<NodeRef>();
+			foreach (IJobStepBatch batch in job.Batches)
+			{
+				foreach (IJobStep step in batch.Steps)
+				{
+					nodes.Add(new NodeRef(batch.GroupIdx, step.NodeIdx));
+				}
+			}
+			GetLabelState(nodes, stepForNode, out newState, out newOutcome);
+		}
+
+		/// <summary>
+		/// Gets the state of a label
+		/// </summary>
+		/// <param name="includedNodes">Nodes to include in this label</param>
+		/// <param name="stepForNode">Map from node to step</param>
+		/// <param name="newState">Receives the state of the label</param>
+		/// <param name="newOutcome">Receives the outcome of the label</param>
+		public static void GetLabelState(IEnumerable<NodeRef> includedNodes, IReadOnlyDictionary<NodeRef, IJobStep> stepForNode, out LabelState newState, out LabelOutcome newOutcome)
+		{
+			newState = LabelState.Complete;
+			newOutcome = LabelOutcome.Success;
+			foreach (NodeRef includedNodeRef in includedNodes)
+			{
+				IJobStep? includedStep;
+				if (stepForNode.TryGetValue(includedNodeRef, out includedStep))
+				{
+					// Update the state
+					if (includedStep.State != JobStepState.Completed && includedStep.State != JobStepState.Skipped && includedStep.State != JobStepState.Aborted)
+					{
+						newState = LabelState.Running;
+					}
+
+					// Update the outcome
+					if (includedStep.State == JobStepState.Skipped || includedStep.State == JobStepState.Aborted || includedStep.Outcome == JobStepOutcome.Failure)
+					{
+						newOutcome = LabelOutcome.Failure;
+					}
+					else if (includedStep.Outcome == JobStepOutcome.Warnings && newOutcome == LabelOutcome.Success)
+					{
+						newOutcome = LabelOutcome.Warnings;
+					}
+				}
+			}
+		}
+
+		/// <summary>
+		/// Creates an RPC response object
+		/// </summary>
+		/// <param name="job">The job document</param>
+		/// <returns></returns>
+		public static RpcGetJobResponse ToRpcResponse(this IJob job)
+		{
+			RpcGetJobResponse response = new RpcGetJobResponse();
+			response.StreamId = job.StreamId.ToString();
+			response.Change = job.Change;
+			response.CodeChange = job.CodeChange;
+			response.PreflightChange = job.PreflightChange;
+			response.ClonedPreflightChange = job.ClonedPreflightChange;
+			response.Arguments.Add(job.Arguments);
+			return response;
+		}
+
+		/// <summary>
+		/// Gets a key attached to all artifacts produced for a job
+		/// </summary>
+		public static string GetArtifactKey(this IJob job)
+		{
+			return $"job:{job.Id}";
+		}
+
+		/// <summary>
+		/// Gets a key attached to all artifacts produced for a job step
+		/// </summary>
+		public static string GetArtifactKey(this IJob job, IJobStep jobStep)
+		{
+			return $"job:{job.Id}/step:{jobStep.Id}";
+		}
+
+		/// <inheritdoc cref="IJob.TrySkipBatchAsync(JobStepBatchId, JobStepBatchError, CancellationToken)"/>
+		public static async Task<IJob?> SkipBatchAsync(this IJob job, JobStepBatchId batchId, JobStepBatchError error, CancellationToken cancellationToken)
+		{
+			for (; ; )
+			{
+				IJob? newJob = await job.TrySkipBatchAsync(batchId, error, cancellationToken);
+				if (newJob != null)
+				{
+					return newJob;
+				}
+
+				newJob = await job.RefreshAsync(cancellationToken);
+				if (newJob == null)
+				{
+					return null;
+				}
+
+				job = newJob;
+			}
+		}
+
+		/// <inheritdoc cref="IJob.TrySkipAllBatchesAsync(JobStepBatchError, CancellationToken)"/>
+		public static async Task<IJob?> SkipAllBatchesAsync(this IJob job, JobStepBatchError reason, CancellationToken cancellationToken)
+		{
+			for (; ; )
+			{
+				IJob? newJob = await job.TrySkipAllBatchesAsync(reason, cancellationToken);
+				if (newJob != null)
+				{
+					return newJob;
+				}
+
+				newJob = await job.RefreshAsync(cancellationToken);
+				if (newJob == null)
+				{
+					return null;
+				}
+
+				job = newJob;
+			}
+		}
+	}
+
+	/// <summary>
+	/// Stores information about a batch of job steps
+	/// </summary>
+	public interface IJobStepBatch
+	{
+		/// <summary>
+		/// Unique id for this group
+		/// </summary>
+		public JobStepBatchId Id { get; }
+
+		/// <summary>
+		/// The log file id for this batch
+		/// </summary>
+		public LogId? LogId { get; }
+
+		/// <summary>
+		/// Index of the group being executed
+		/// </summary>
+		public int GroupIdx { get; }
+
+		/// <summary>
+		/// The state of this group
+		/// </summary>
+		public JobStepBatchState State { get; }
+
+		/// <summary>
+		/// Error associated with this group
+		/// </summary>
+		public JobStepBatchError Error { get; }
+
+		/// <summary>
+		/// Steps within this run
+		/// </summary>
+		public IReadOnlyList<IJobStep> Steps { get; }
+
+		/// <summary>
+		/// The pool that this agent was taken from
+		/// </summary>
+		public PoolId? PoolId { get; }
+
+		/// <summary>
+		/// The agent assigned to execute this group
+		/// </summary>
+		public AgentId? AgentId { get; }
+
+		/// <summary>
+		/// The agent session that is executing this group
+		/// </summary>
+		public SessionId? SessionId { get; }
+
+		/// <summary>
+		/// The lease that's executing this group
+		/// </summary>
+		public LeaseId? LeaseId { get; }
+
+		/// <summary>
+		/// The weighted priority of this batch for the scheduler
+		/// </summary>
+		public int SchedulePriority { get; }
+
+		/// <summary>
+		/// Time at which the group became ready (UTC).
+		/// </summary>
+		public DateTime? ReadyTimeUtc { get; }
+
+		/// <summary>
+		/// Time at which the group started (UTC).
+		/// </summary>
+		public DateTime? StartTimeUtc { get; }
+
+		/// <summary>
+		/// Time at which the group finished (UTC)
+		/// </summary>
+		public DateTime? FinishTimeUtc { get; }
+	}
+
+	/// <summary>
+	/// Extension methods for IJobStepBatch
+	/// </summary>
+	public static class JobStepBatchExtensions
+	{
+		/// <summary>
+		/// Attempts to get a step with the given id
+		/// </summary>
+		/// <param name="batch">The batch to search</param>
+		/// <param name="stepId">The step id</param>
+		/// <param name="step">On success, receives the step object</param>
+		/// <returns>True if the step was found</returns>
+		public static bool TryGetStep(this IJobStepBatch batch, JobStepId stepId, [NotNullWhen(true)] out IJobStep? step)
+		{
+			step = batch.Steps.FirstOrDefault(x => x.Id == stepId);
+			return step != null;
+		}
+
+		/// <summary>
+		/// Determines if new steps can be appended to this batch. We do not allow this after the last step has been completed, because the agent is shutting down.
+		/// </summary>
+		/// <param name="batch">The batch to search</param>
+		/// <returns>True if new steps can be appended to this batch</returns>
+		public static bool CanBeAppendedTo(this IJobStepBatch batch)
+		{
+			return batch.State <= JobStepBatchState.Running;
+		}
+
+		/// <summary>
+		/// Gets the wait time for this batch
+		/// </summary>
+		/// <param name="batch">The batch to search</param>
+		/// <returns>Wait time for the batch</returns>
+		public static TimeSpan? GetWaitTime(this IJobStepBatch batch)
+		{
+			if (batch.StartTimeUtc == null || batch.ReadyTimeUtc == null)
+			{
+				return null;
+			}
+			else
+			{
+				return batch.StartTimeUtc.Value - batch.ReadyTimeUtc.Value;
+			}
+		}
+
+		/// <summary>
+		/// Gets the initialization time for this batch
+		/// </summary>
+		/// <param name="batch">The batch to search</param>
+		/// <returns>Initialization time for this batch</returns>
+		public static TimeSpan? GetInitTime(this IJobStepBatch batch)
+		{
+			if (batch.StartTimeUtc != null)
+			{
+				foreach (IJobStep step in batch.Steps)
+				{
+					if (step.StartTimeUtc != null)
+					{
+						return step.StartTimeUtc - batch.StartTimeUtc.Value;
+					}
+				}
+			}
@@ -1188,578 +1391,365 @@
 		}
 
 		/// <summary>
-		/// Gets the estimated timing info for all nodes in the job
+		/// Get the dependencies required for this batch to start, taking run-early nodes into account
 		/// </summary>
-		/// <param name="job">The job document</param>
-		/// <param name="graph">Graph for this job</param>
-		/// <param name="jobTiming">Job timing information</param>
-		/// <returns>Map of node to expected timing info</returns>
-		public static Dictionary<INode, TimingInfo> GetTimingInfo(this IJob job, IGraph graph, IJobTiming jobTiming)
+		/// <param name="batch">The batch to search</param>
+		/// <param name="groups">List of node groups</param>
+		/// <returns>Set of nodes that must have completed for this batch to start</returns>
+		public static HashSet<INode> GetStartDependencies(this IJobStepBatch batch, IReadOnlyList<INodeGroup> groups)
 		{
-#pragma warning disable IDE0054 // Use compound assignment
-			TimeSpan currentTime = DateTime.UtcNow - job.CreateTimeUtc;
+			// Find all the nodes that this group will start with.
+			List<INode> nodes = batch.Steps.ConvertAll(x => groups[batch.GroupIdx].Nodes[x.NodeIdx]);
+			if (nodes.Any(x => x.RunEarly))
+			{
+				nodes.RemoveAll(x => !x.RunEarly);
+			}
 
-			Dictionary<INode, TimingInfo> nodeToTimingInfo = graph.Groups.SelectMany(x => x.Nodes).ToDictionary(x => x, x => new TimingInfo());
-			foreach (IJobStepBatch batch in job.Batches)
+			// Find all their dependencies
+			HashSet<INode> dependencies = new HashSet<INode>();
+			foreach (INode node in nodes)
 			{
-				INodeGroup group = graph.Groups[batch.GroupIdx];
+				dependencies.UnionWith(node.InputDependencies.Select(x => groups[x.GroupIdx].Nodes[x.NodeIdx]));
+				dependencies.UnionWith(node.OrderDependencies.Select(x => groups[x.GroupIdx].Nodes[x.NodeIdx]));
+			}
 
-				// Step through the batch, keeping track of the time that things finish.
-				TimingInfo timingInfo = new TimingInfo();
+			// Exclude all the dependencies within the same group
+			dependencies.ExceptWith(groups[batch.GroupIdx].Nodes);
+			return dependencies;
+		}
+	}
 
-				// Wait for the dependencies for the batch to start
-				HashSet<INode> dependencyNodes = batch.GetStartDependencies(graph.Groups);
-				timingInfo.WaitForAll(dependencyNodes.Select(x => nodeToTimingInfo[x]));
+	/// <summary>
+	/// Embedded jobstep document
+	/// </summary>
+	public interface IJobStep
+	{
+		/// <summary>
+		/// Unique ID assigned to this jobstep. A new id is generated whenever a jobstep's order is changed.
+		/// </summary>
+		public JobStepId Id { get; }
 
-				// If the batch has actually started, correct the expected time to use this instead
-				if (batch.StartTimeUtc != null)
-				{
-					timingInfo.TotalTimeToComplete = batch.StartTimeUtc - job.CreateTimeUtc;
-				}
+		/// <summary>
+		/// Index of the node which this jobstep is to execute
+		/// </summary>
+		public int NodeIdx { get; }
 
-				// Get the average times for this batch
-				TimeSpan? averageWaitTime = GetAverageWaitTime(graph, batch, jobTiming);
-				TimeSpan? averageInitTime = GetAverageInitTime(graph, batch, jobTiming);
+		/// <summary>
+		/// Current state of the job step. This is updated automatically when runs complete.
+		/// </summary>
+		public JobStepState State { get; }
 
-				// Update the wait times and initialization times along this path
-				timingInfo.TotalWaitTime = timingInfo.TotalWaitTime + (batch.GetWaitTime() ?? averageWaitTime);
-				timingInfo.TotalInitTime = timingInfo.TotalInitTime + (batch.GetInitTime() ?? averageInitTime);
+		/// <summary>
+		/// Current outcome of the jobstep
+		/// </summary>
+		public JobStepOutcome Outcome { get; }
 
-				// Update the average wait and initialization times too
-				timingInfo.AverageTotalWaitTime = timingInfo.AverageTotalWaitTime + averageWaitTime;
-				timingInfo.AverageTotalInitTime = timingInfo.AverageTotalInitTime + averageInitTime;
+		/// <summary>
+		/// Error from executing this step
+		/// </summary>
+		public JobStepError Error { get; }
 
-				// Step through the batch, updating the expected times as we go
-				foreach (IJobStep step in batch.Steps)
-				{
-					INode node = group.Nodes[step.NodeIdx];
+		/// <summary>
+		/// The log id for this step
+		/// </summary>
+		public LogId? LogId { get; }
 
-					// Get the timing for this step
-					IJobStepTiming? stepTimingInfo;
-					jobTiming.TryGetStepTiming(node.Name, out stepTimingInfo);
+		/// <summary>
+		/// Unique id for notifications
+		/// </summary>
+		public ObjectId? NotificationTriggerId { get; }
 
-					// If the step has already started, update the actual time to reach this point
-					if (step.StartTimeUtc != null)
-					{
-						timingInfo.TotalTimeToComplete = step.StartTimeUtc.Value - job.CreateTimeUtc;
-					}
+		/// <summary>
+		/// Time at which the batch transitioned to the ready state (UTC).
+		/// </summary>
+		public DateTime? ReadyTimeUtc { get; }
 
-					// If the step hasn't started yet, make sure the start time is later than the current time
-					if (step.StartTimeUtc == null && currentTime > timingInfo.TotalTimeToComplete)
-					{
-						timingInfo.TotalTimeToComplete = currentTime;
-					}
+		/// <summary>
+		/// Time at which the batch transitioned to the executing state (UTC).
+		/// </summary>
+		public DateTime? StartTimeUtc { get; }
 
-					// Wait for all the node dependencies to complete
-					timingInfo.WaitForAll(graph.GetDependencies(node).Select(x => nodeToTimingInfo[x]));
+		/// <summary>
+		/// Time at which the run finished (UTC)
+		/// </summary>
+		public DateTime? FinishTimeUtc { get; }
 
-					// If the step has actually finished, correct the time to use that instead
-					if (step.FinishTimeUtc != null)
-					{
-						timingInfo.TotalTimeToComplete = step.FinishTimeUtc.Value - job.CreateTimeUtc;
-					}
-					else
-					{
-						timingInfo.TotalTimeToComplete = timingInfo.TotalTimeToComplete + NullableTimeSpanFromSeconds(stepTimingInfo?.AverageDuration);
-					}
+		/// <summary>
+		/// Override for the priority of this step
+		/// </summary>
+		public Priority? Priority { get; }
 
-					// If the step hasn't finished yet, make sure the start time is later than the current time
-					if (step.FinishTimeUtc == null && currentTime > timingInfo.TotalTimeToComplete)
-					{
-						timingInfo.TotalTimeToComplete = currentTime;
-					}
+		/// <summary>
+		/// If a retry is requested, stores the name of the user that requested it
+		/// </summary>
+		public UserId? RetriedByUserId { get; }
 
-					// Update the average time to complete
-					timingInfo.AverageTotalTimeToComplete = timingInfo.AverageTotalTimeToComplete + NullableTimeSpanFromSeconds(stepTimingInfo?.AverageDuration);
+		/// <summary>
+		/// Signal if a step should be aborted
+		/// </summary>
+		public bool AbortRequested { get; }
 
-					// Add it to the lookup
-					TimingInfo nodeTimingInfo = new TimingInfo(timingInfo);
-					nodeTimingInfo.StepTiming = stepTimingInfo;
-					nodeToTimingInfo[node] = nodeTimingInfo;
-				}
-			}
-			return nodeToTimingInfo;
-#pragma warning restore IDE0054 // Use compound assignment
+		/// <summary>
+		/// If an abort is requested, stores the id of the user that requested it
+		/// </summary>
+		public UserId? AbortedByUserId { get; }
+
+		/// <summary>
+		/// List of reports for this step
+		/// </summary>
+		public IReadOnlyList<IReport>? Reports { get; }
+
+		/// <summary>
+		/// Reports for this jobstep.
+		/// </summary>
+		public Dictionary<string, string>? Properties { get; }
+	}
+
+	/// <summary>
+	/// Extension methods for job steps
+	/// </summary>
+	public static class JobStepExtensions
+	{
+		/// <summary>
+		/// Determines if a jobstep has failed or is skipped. Can be used to determine whether dependent steps will be able to run.
+		/// </summary>
+		/// <returns>True if the step is failed or skipped</returns>
+		public static bool IsFailedOrSkipped(this IJobStep step)
+		{
+			return step.State == JobStepState.Skipped || step.Outcome == JobStepOutcome.Failure;
 		}
 
 		/// <summary>
-		/// Gets the average wait time for this batch
+		/// Determines if a jobstep is done by checking to see if it is completed, skipped, or aborted.
 		/// </summary>
-		/// <param name="graph">Graph for the job</param>
-		/// <param name="batch">The batch to get timing info for</param>
-		/// <param name="jobTiming">The job timing information</param>
-		/// <returns>Wait time for the batch</returns>
-		public static TimeSpan? GetAverageWaitTime(IGraph graph, IJobStepBatch batch, IJobTiming jobTiming)
+		/// <returns>True if the step is completed, skipped, or aborted</returns>
+		public static bool IsPending(this IJobStep step)
 		{
-			TimeSpan? waitTime = null;
-			foreach (IJobStep step in batch.Steps)
-			{
-				INode node = graph.Groups[batch.GroupIdx].Nodes[step.NodeIdx];
-				if (jobTiming.TryGetStepTiming(node.Name, out IJobStepTiming? timingInfo))
-				{
-					if (timingInfo.AverageWaitTime != null)
-					{
-						TimeSpan stepWaitTime = TimeSpan.FromSeconds(timingInfo.AverageWaitTime.Value);
-						if (waitTime == null || stepWaitTime > waitTime.Value)
-						{
-							waitTime = stepWaitTime;
-						}
-					}
-				}
-			}
-			return waitTime;
+			return step.State != JobStepState.Aborted && step.State != JobStepState.Completed && step.State != JobStepState.Skipped;
 		}
 
 		/// <summary>
-		/// Gets the average initialization time for this batch
+		/// Determine if a step should be timed out
 		/// </summary>
-		/// <param name="graph">Graph for the job</param>
-		/// <param name="batch">The batch to get timing info for</param>
-		/// <param name="jobTiming">The job timing information</param>
-		/// <returns>Initialization time for this batch</returns>
-		public static TimeSpan? GetAverageInitTime(IGraph graph, IJobStepBatch batch, IJobTiming jobTiming)
-		{
-			TimeSpan? initTime = null;
-			foreach (IJobStep step in batch.Steps)
-			{
-				INode node = graph.Groups[batch.GroupIdx].Nodes[step.NodeIdx];
-				if (jobTiming.TryGetStepTiming(node.Name, out IJobStepTiming? timingInfo))
-				{
-					if (timingInfo.AverageInitTime != null)
-					{
-						TimeSpan stepInitTime = TimeSpan.FromSeconds(timingInfo.AverageInitTime.Value);
-						if (initTime == null || stepInitTime > initTime.Value)
-						{
-							initTime = stepInitTime;
-						}
-					}
-				}
-			}
-			return initTime;
-		}
-
-		/// <summary>
-		/// Creates a nullable timespan from a nullable number of seconds
-		/// </summary>
-		/// <param name="seconds">The number of seconds to construct from</param>
-		/// <returns>TimeSpan object</returns>
-		static TimeSpan? NullableTimeSpanFromSeconds(float? seconds)
-		{
-			if (seconds == null)
-			{
-				return null;
-			}
-			else
-			{
-				return TimeSpan.FromSeconds(seconds.Value);
-			}
-		}
-
-		/// <summary>
-		/// Attempts to get a batch with the given id
-		/// </summary>
-		/// <param name="job">The job document</param>
-		/// <param name="batchId">The batch id</param>
-		/// <param name="batch">On success, receives the batch object</param>
-		/// <returns>True if the batch was found</returns>
-		public static bool TryGetBatch(this IJob job, JobStepBatchId batchId, [NotNullWhen(true)] out IJobStepBatch? batch)
-		{
-			batch = job.Batches.FirstOrDefault(x => x.Id == batchId);
-			return batch != null;
-		}
-
-		/// <summary>
-		/// Attempts to get a batch with the given id
-		/// </summary>
-		/// <param name="job">The job document</param>
-		/// <param name="batchId">The batch id</param>
-		/// <param name="stepId">The step id</param>
-		/// <param name="step">On success, receives the step object</param>
-		/// <returns>True if the batch was found</returns>
-		public static bool TryGetStep(this IJob job, JobStepBatchId batchId, JobStepId stepId, [NotNullWhen(true)] out IJobStep? step)
-		{
-			IJobStepBatch? batch;
-			if (!TryGetBatch(job, batchId, out batch))
-			{
-				step = null;
-				return false;
-			}
-			return batch.TryGetStep(stepId, out step);
-		}
-
-		/// <summary>
-		/// Finds the set of nodes affected by a label
-		/// </summary>
-		/// <param name="job">The job document</param>
-		/// <param name="graph">Graph definition for the job</param>
-		/// <param name="labelIdx">Index of the label. -1 or Graph.Labels.Count are treated as referring to the default lable.</param>
-		/// <returns>Set of nodes affected by the given label</returns>
-		public static HashSet<NodeRef> GetNodesForLabel(this IJob job, IGraph graph, int labelIdx)
-		{
-			if (labelIdx != -1 && labelIdx != graph.Labels.Count)
-			{
-				// Return all the nodes included by the label
-				return new HashSet<NodeRef>(graph.Labels[labelIdx].IncludedNodes);
-			}
-			else
-			{
-				// Set of nodes which are not covered by an existing label, initially containing everything
-				HashSet<NodeRef> unlabeledNodes = new HashSet<NodeRef>();
-				for (int groupIdx = 0; groupIdx < graph.Groups.Count; groupIdx++)
-				{
-					INodeGroup group = graph.Groups[groupIdx];
-					for (int nodeIdx = 0; nodeIdx < group.Nodes.Count; nodeIdx++)
-					{
-						unlabeledNodes.Add(new NodeRef(groupIdx, nodeIdx));
-					}
-				}
-
-				// Remove all the nodes that are part of an active label
-				IReadOnlyDictionary<NodeRef, IJobStep> stepForNode = job.GetStepForNodeMap();
-				foreach (ILabel label in graph.Labels)
-				{
-					if (label.RequiredNodes.Any(x => stepForNode.ContainsKey(x)))
-					{
-						unlabeledNodes.ExceptWith(label.IncludedNodes);
-					}
-				}
-				return unlabeledNodes;
-			}
-		}
-
-		/// <summary>
-		/// Create a list of aggregate responses, combining the graph definitions with the state of the job
-		/// </summary>
-		/// <param name="job">The job document</param>
-		/// <param name="graph">Graph definition for the job</param>
-		/// <param name="responses">List to receive all the responses</param>
-		/// <returns>The default label state</returns>
-		public static GetDefaultLabelStateResponse? GetLabelStateResponses(this IJob job, IGraph graph, List<GetLabelStateResponse> responses)
-		{
-			// Create a lookup from noderef to step information
-			IReadOnlyDictionary<NodeRef, IJobStep> stepForNode = job.GetStepForNodeMap();
-
-			// Set of nodes which are not covered by an existing label, initially containing everything
-			HashSet<NodeRef> unlabeledNodes = new HashSet<NodeRef>();
-			for (int groupIdx = 0; groupIdx < graph.Groups.Count; groupIdx++)
-			{
-				INodeGroup group = graph.Groups[groupIdx];
-				for (int nodeIdx = 0; nodeIdx < group.Nodes.Count; nodeIdx++)
-				{
-					unlabeledNodes.Add(new NodeRef(groupIdx, nodeIdx));
-				}
-			}
-
-			// Create the responses
-			foreach (ILabel label in graph.Labels)
-			{
-				// Refresh the state for this label
-				LabelState newState = LabelState.Unspecified;
-				foreach (NodeRef requiredNodeRef in label.RequiredNodes)
-				{
-					if (stepForNode.ContainsKey(requiredNodeRef))
-					{
-						newState = LabelState.Complete;
-						break;
-					}
-				}
-
-				// Refresh the outcome
-				LabelOutcome newOutcome = LabelOutcome.Success;
-				if (newState == LabelState.Complete)
-				{
-					GetLabelState(label.IncludedNodes, stepForNode, out newState, out newOutcome);
-					unlabeledNodes.ExceptWith(label.IncludedNodes);
-				}
-
-				// Create the response
-				responses.Add(new GetLabelStateResponse(newState, newOutcome));
-			}
-
-			// Remove all the nodes that don't have a step
-			unlabeledNodes.RemoveWhere(x => !stepForNode.ContainsKey(x));
-
-			// Remove successful "setup build" nodes from the list
-			if (graph.Groups.Count > 1 && graph.Groups[0].Nodes.Count > 0)
-			{
-				INode node = graph.Groups[0].Nodes[0];
-				if (node.Name == IJob.SetupNodeName)
-				{
-					NodeRef nodeRef = new NodeRef(0, 0);
-					if (unlabeledNodes.Contains(nodeRef))
-					{
-						IJobStep step = stepForNode[nodeRef];
-						if (step.State == JobStepState.Completed && step.Outcome == JobStepOutcome.Success && responses.Count > 0)
-						{
-							unlabeledNodes.Remove(nodeRef);
-						}
-					}
-				}
-			}
-
-			// Add a response for everything not included elsewhere.
-			GetLabelState(unlabeledNodes, stepForNode, out LabelState otherState, out LabelOutcome otherOutcome);
-			return new GetDefaultLabelStateResponse(otherState, otherOutcome, unlabeledNodes.Select(x => graph.GetNode(x).Name).ToList());
-		}
-
-		/// <summary>
-		/// Get the states of all labels for this job
-		/// </summary>
-		/// <param name="job">The job to get states for</param>
-		/// <param name="graph">The graph for this job</param>
-		/// <returns>Collection of label states by label index</returns>
-		public static IReadOnlyList<(LabelState, LabelOutcome)> GetLabelStates(this IJob job, IGraph graph)
-		{
-			IReadOnlyDictionary<NodeRef, IJobStep> stepForNodeRef = job.GetStepForNodeMap();
-
-			List<(LabelState, LabelOutcome)> states = new List<(LabelState, LabelOutcome)>();
-			for (int idx = 0; idx < graph.Labels.Count; idx++)
-			{
-				ILabel label = graph.Labels[idx];
-
-				// Default the label to the unspecified state
-				LabelState newState = LabelState.Unspecified;
-				LabelOutcome newOutcome = LabelOutcome.Unspecified;
-
-				// Check if the label should be included
-				if (label.RequiredNodes.Any(x => stepForNodeRef.ContainsKey(x)))
-				{
-					// Combine the state of the steps contributing towards this label
-					bool anySkipped = false;
-					bool anyWarnings = false;
-					bool anyFailed = false;
-					bool anyPending = false;
-					foreach (NodeRef includedNode in label.IncludedNodes)
-					{
-						IJobStep? step;
-						if (stepForNodeRef.TryGetValue(includedNode, out step))
-						{
-							anyPending |= step.IsPending();
-							anySkipped |= step.State == JobStepState.Aborted || step.State == JobStepState.Skipped;
-							anyFailed |= (step.Outcome == JobStepOutcome.Failure);
-							anyWarnings |= (step.Outcome == JobStepOutcome.Warnings);
-						}
-					}
-
-					// Figure out the overall label state
-					newState = anyPending ? LabelState.Running : LabelState.Complete;
-					newOutcome = anyFailed ? LabelOutcome.Failure : anyWarnings ? LabelOutcome.Warnings : anySkipped ? LabelOutcome.Unspecified : LabelOutcome.Success;
-				}
-
-				states.Add((newState, newOutcome));
-			}
-			return states;
-		}
-
-		/// <summary>
-		/// Get the states of all UGS badges for this job
-		/// </summary>
-		/// <param name="job">The job to get states for</param>
-		/// <param name="graph">The graph for this job</param>
-		/// <returns>List of badge states</returns>
-		public static Dictionary<int, UgsBadgeState> GetUgsBadgeStates(this IJob job, IGraph graph)
-		{
-			IReadOnlyList<(LabelState, LabelOutcome)> labelStates = GetLabelStates(job, graph);
-			return job.GetUgsBadgeStates(graph, labelStates);
-		}
-
-		/// <summary>
-		/// Get the states of all UGS badges for this job
-		/// </summary>
-		/// <param name="job">The job to get states for</param>
-		/// <param name="graph">The graph for this job</param>
-		/// <param name="labelStates">The existing label states to get the UGS badge states from</param>
-		/// <returns>List of badge states</returns>
-#pragma warning disable IDE0060 // Remove unused parameter
-		public static Dictionary<int, UgsBadgeState> GetUgsBadgeStates(this IJob job, IGraph graph, IReadOnlyList<(LabelState, LabelOutcome)> labelStates)
-#pragma warning restore IDE0060 // Remove unused parameter
-		{
-			Dictionary<int, UgsBadgeState> ugsBadgeStates = new Dictionary<int, UgsBadgeState>();
-			for (int labelIdx = 0; labelIdx < labelStates.Count; ++labelIdx)
-			{
-				if (graph.Labels[labelIdx].UgsName == null)
-				{
-					continue;
-				}
-
-				(LabelState state, LabelOutcome outcome) = labelStates[labelIdx];
-				switch (state)
-				{
-					case LabelState.Complete:
-						{
-							switch (outcome)
-							{
-								case LabelOutcome.Success:
-									{
-										ugsBadgeStates.Add(labelIdx, UgsBadgeState.Success);
-										break;
-									}
-
-								case LabelOutcome.Warnings:
-									{
-										ugsBadgeStates.Add(labelIdx, UgsBadgeState.Warning);
-										break;
-									}
-
-								case LabelOutcome.Failure:
-									{
-										ugsBadgeStates.Add(labelIdx, UgsBadgeState.Failure);
-										break;
-									}
-
-								case LabelOutcome.Unspecified:
-									{
-										ugsBadgeStates.Add(labelIdx, UgsBadgeState.Skipped);
-										break;
-									}
-							}
-							break;
-						}
-
-					case LabelState.Running:
-						{
-							ugsBadgeStates.Add(labelIdx, UgsBadgeState.Starting);
-							break;
-						}
-
-					case LabelState.Unspecified:
-						{
-							ugsBadgeStates.Add(labelIdx, UgsBadgeState.Skipped);
-							break;
-						}
-				}
-			}
-			return ugsBadgeStates;
-		}
-
-		/// <summary>
-		/// Gets the state of a job, as a label that includes all steps
-		/// </summary>
-		/// <param name="job">The job to query</param>
-		/// <param name="stepForNode">Map from node to step</param>
-		/// <param name="newState">Receives the state of the label</param>
-		/// <param name="newOutcome">Receives the outcome of the label</param>
-		public static void GetJobState(this IJob job, IReadOnlyDictionary<NodeRef, IJobStep> stepForNode, out LabelState newState, out LabelOutcome newOutcome)
-		{
-			List<NodeRef> nodes = new List<NodeRef>();
-			foreach (IJobStepBatch batch in job.Batches)
-			{
-				foreach (IJobStep step in batch.Steps)
-				{
-					nodes.Add(new NodeRef(batch.GroupIdx, step.NodeIdx));
-				}
-			}
-			GetLabelState(nodes, stepForNode, out newState, out newOutcome);
-		}
-
-		/// <summary>
-		/// Gets the state of a label
-		/// </summary>
-		/// <param name="includedNodes">Nodes to include in this label</param>
-		/// <param name="stepForNode">Map from node to step</param>
-		/// <param name="newState">Receives the state of the label</param>
-		/// <param name="newOutcome">Receives the outcome of the label</param>
-		public static void GetLabelState(IEnumerable<NodeRef> includedNodes, IReadOnlyDictionary<NodeRef, IJobStep> stepForNode, out LabelState newState, out LabelOutcome newOutcome)
-		{
-			newState = LabelState.Complete;
-			newOutcome = LabelOutcome.Success;
-			foreach (NodeRef includedNodeRef in includedNodes)
-			{
-				IJobStep? includedStep;
-				if (stepForNode.TryGetValue(includedNodeRef, out includedStep))
-				{
-					// Update the state
-					if (includedStep.State != JobStepState.Completed && includedStep.State != JobStepState.Skipped && includedStep.State != JobStepState.Aborted)
-					{
-						newState = LabelState.Running;
-					}
-
-					// Update the outcome
-					if (includedStep.State == JobStepState.Skipped || includedStep.State == JobStepState.Aborted || includedStep.Outcome == JobStepOutcome.Failure)
-					{
-						newOutcome = LabelOutcome.Failure;
-					}
-					else if (includedStep.Outcome == JobStepOutcome.Warnings && newOutcome == LabelOutcome.Success)
-					{
-						newOutcome = LabelOutcome.Warnings;
-					}
-				}
-			}
-		}
-
-		/// <summary>
-		/// Creates an RPC response object
-		/// </summary>
-		/// <param name="job">The job document</param>
+		/// <param name="step"></param>
+		/// <param name="utcNow"></param>
 		/// <returns></returns>
-		public static RpcGetJobResponse ToRpcResponse(this IJob job)
+		public static bool HasTimedOut(this IJobStep step, DateTime utcNow)
 		{
-			RpcGetJobResponse response = new RpcGetJobResponse();
-			response.StreamId = job.StreamId.ToString();
-			response.Change = job.Change;
-			response.CodeChange = job.CodeChange;
-			response.PreflightChange = job.PreflightChange;
-			response.ClonedPreflightChange = job.ClonedPreflightChange;
-			response.Arguments.Add(job.Arguments);
-			return response;
+			if (step.State == JobStepState.Running && step.StartTimeUtc != null)
+			{
+				TimeSpan elapsed = utcNow - step.StartTimeUtc.Value;
+				if (elapsed > TimeSpan.FromHours(24.0))
+				{
+					return true;
+				}
+			}
+			return false;
+		}
+	}
+
+	/// <summary>
+	/// Cumulative timing information to reach a certain point in a job
+	/// </summary>
+	public class TimingInfo
+	{
+		/// <summary>
+		/// Wait time on the critical path
+		/// </summary>
+		public TimeSpan? TotalWaitTime { get; set; }
+
+		/// <summary>
+		/// Sync time on the critical path
+		/// </summary>
+		public TimeSpan? TotalInitTime { get; set; }
+
+		/// <summary>
+		/// Duration to this point
+		/// </summary>
+		public TimeSpan? TotalTimeToComplete { get; set; }
+
+		/// <summary>
+		/// Average wait time to this point
+		/// </summary>
+		public TimeSpan? AverageTotalWaitTime { get; set; }
+
+		/// <summary>
+		/// Average sync time to this point
+		/// </summary>
+		public TimeSpan? AverageTotalInitTime { get; set; }
+
+		/// <summary>
+		/// Average duration to this point
+		/// </summary>
+		public TimeSpan? AverageTotalTimeToComplete { get; set; }
+
+		/// <summary>
+		/// Individual step timing information
+		/// </summary>
+		public IJobStepTiming? StepTiming { get; set; }
+
+		/// <summary>
+		/// Constructor
+		/// </summary>
+		public TimingInfo()
+		{
+			TotalWaitTime = TimeSpan.Zero;
+			TotalInitTime = TimeSpan.Zero;
+			TotalTimeToComplete = TimeSpan.Zero;
+
+			AverageTotalWaitTime = TimeSpan.Zero;
+			AverageTotalInitTime = TimeSpan.Zero;
+			AverageTotalTimeToComplete = TimeSpan.Zero;
 		}
 
 		/// <summary>
-		/// Gets a key attached to all artifacts produced for a job
+		/// Copy constructor
 		/// </summary>
-		public static string GetArtifactKey(this IJob job)
+		/// <param name="other">The timing info object to copy from</param>
+		public TimingInfo(TimingInfo other)
 		{
-			return $"job:{job.Id}";
+			TotalWaitTime = other.TotalWaitTime;
+			TotalInitTime = other.TotalInitTime;
+			TotalTimeToComplete = other.TotalTimeToComplete;
+
+			AverageTotalWaitTime = other.AverageTotalWaitTime;
+			AverageTotalInitTime = other.AverageTotalInitTime;
+			AverageTotalTimeToComplete = other.AverageTotalTimeToComplete;
 		}
 
 		/// <summary>
-		/// Gets a key attached to all artifacts produced for a job step
+		/// Modifies this timing to wait for another timing
 		/// </summary>
-		public static string GetArtifactKey(this IJob job, IJobStep jobStep)
+		/// <param name="other">The other node to wait for</param>
+		public void WaitFor(TimingInfo other)
 		{
-			return $"job:{job.Id}/step:{jobStep.Id}";
-		}
+			if (TotalTimeToComplete != null)
+			{
+				if (other.TotalTimeToComplete == null || other.TotalTimeToComplete.Value > TotalTimeToComplete.Value)
+				{
+					TotalInitTime = other.TotalInitTime;
+					TotalWaitTime = other.TotalWaitTime;
+					TotalTimeToComplete = other.TotalTimeToComplete;
+				}
+			}
 
-		/// <inheritdoc cref="IJob.TrySkipBatchAsync(JobStepBatchId, JobStepBatchError, CancellationToken)"/>
-		public static async Task<IJob?> SkipBatchAsync(this IJob job, JobStepBatchId batchId, JobStepBatchError error, CancellationToken cancellationToken)
-		{
-			for (; ; )
+			if (AverageTotalTimeToComplete != null)
 			{
-				IJob? newJob = await job.TrySkipBatchAsync(batchId, error, cancellationToken);
-				if (newJob != null)
+				if (other.AverageTotalTimeToComplete == null || other.AverageTotalTimeToComplete.Value > AverageTotalTimeToComplete.Value)
 				{
-					return newJob;
+					AverageTotalInitTime = other.AverageTotalInitTime;
+					AverageTotalWaitTime = other.AverageTotalWaitTime;
+					AverageTotalTimeToComplete = other.AverageTotalTimeToComplete;
 				}
-
-				newJob = await job.RefreshAsync(cancellationToken);
-				if (newJob == null)
-				{
-					return null;
-				}
-
-				job = newJob;
 			}
 		}
 
-		/// <inheritdoc cref="IJob.TrySkipAllBatchesAsync(JobStepBatchError, CancellationToken)"/>
-		public static async Task<IJob?> SkipAllBatchesAsync(this IJob job, JobStepBatchError reason, CancellationToken cancellationToken)
+		/// <summary>
+		/// Waits for all the given timing info objects to complete
+		/// </summary>
+		/// <param name="others">Other timing info objects to wait for</param>
+		public void WaitForAll(IEnumerable<TimingInfo> others)
 		{
-			for(; ;)
+			foreach (TimingInfo other in others)
 			{
-				IJob? newJob = await job.TrySkipAllBatchesAsync(reason, cancellationToken);
-				if (newJob != null)
-				{
-					return newJob;
-				}
-
-				newJob = await job.RefreshAsync(cancellationToken);
-				if (newJob == null)
-				{
-					return null;
-				}
-
-				job = newJob;
+				WaitFor(other);
 			}
 		}
+
+		/// <summary>
+		/// Constructs a new TimingInfo object which represents the last TimingInfo to finish
+		/// </summary>
+		/// <param name="others">TimingInfo objects to wait for</param>
+		/// <returns>New TimingInfo instance</returns>
+		public static TimingInfo Max(IEnumerable<TimingInfo> others)
+		{
+			TimingInfo timingInfo = new TimingInfo();
+			timingInfo.WaitForAll(others);
+			return timingInfo;
+		}
+
+		/// <summary>
+		/// Copies this info to a repsonse object
+		/// </summary>
+		public void CopyToResponse(GetTimingInfoResponse response)
+		{
+			response.TotalWaitTime = (float?)TotalWaitTime?.TotalSeconds;
+			response.TotalInitTime = (float?)TotalInitTime?.TotalSeconds;
+			response.TotalTimeToComplete = (float?)TotalTimeToComplete?.TotalSeconds;
+
+			response.AverageTotalWaitTime = (float?)AverageTotalWaitTime?.TotalSeconds;
+			response.AverageTotalInitTime = (float?)AverageTotalInitTime?.TotalSeconds;
+			response.AverageTotalTimeToComplete = (float?)AverageTotalTimeToComplete?.TotalSeconds;
+		}
 	}
+
+	/// <summary>
+	/// Information about a chained job trigger
+	/// </summary>
+	public interface IChainedJob
+	{
+		/// <summary>
+		/// The target to monitor
+		/// </summary>
+		public string Target { get; }
+
+		/// <summary>
+		/// The template to trigger on success
+		/// </summary>
+		public TemplateId TemplateRefId { get; }
+
+		/// <summary>
+		/// The triggered job id
+		/// </summary>
+		public JobId? JobId { get; }
+
+		/// <summary>
+		/// Whether to run the latest change, or default change for the template, when starting the new job. Uses same change as the triggering job by default.
+		/// </summary>
+		public bool UseDefaultChangeForTemplate { get; }
+	}
+
+	/// <summary>
+	/// Report for a job or jobstep
+	/// </summary>
+	public interface IReport
+	{
+		/// <summary>
+		/// Name of the report
+		/// </summary>
+		string Name { get; }
+
+		/// <summary>
+		/// Where to render the report
+		/// </summary>
+		ReportPlacement Placement { get; }
+
+		/// <summary>
+		/// The artifact id
+		/// </summary>
+		ObjectId? ArtifactId { get; }
+
+		/// <summary>
+		/// Inline data for the report
+		/// </summary>
+		string? Content { get; }
+	}
+
+	/// <summary>
+	/// Implementation of IReport
+	/// </summary>
+	public class Report : IReport
+	{
+		/// <inheritdoc/>
+		public string Name { get; set; } = String.Empty;
+
+		/// <inheritdoc/>
+		public ReportPlacement Placement { get; set; }
+
+		/// <inheritdoc/>
+		public ObjectId? ArtifactId { get; set; }
+
+		/// <inheritdoc/>
+		public string? Content { get; set; }
+	}
 }
