Horde: Move methods for querying/manipulating log objects from collection to ILog interface. Allows separating V1/V2 log implementation.

#rnx

diff --git a/Engine/Source/Programs/Horde/Horde.Server.Tests/Logs/LogIndexingTests.cs b/Engine/Source/Programs/Horde/Horde.Server.Tests/Logs/LogIndexingTests.cs
--- a/Engine/Source/Programs/Horde/Horde.Server.Tests/Logs/LogIndexingTests.cs
+++ b/Engine/Source/Programs/Horde/Horde.Server.Tests/Logs/LogIndexingTests.cs
@@ -98,7 +98,7 @@ namespace Horde.Server.Tests.Logs
 
 			// Read the data back out and check it's the same
 			byte[] readData = new byte[_data.Length];
-			using (Stream stream = await LogService.OpenRawStreamAsync(logFile))
+			using (Stream stream = await logFile.OpenRawStreamAsync())
 			{
 				int readSize = await stream.ReadAsync(readData, 0, readData.Length);
 				Assert.AreEqual(readData.Length, readSize);
@@ -167,7 +167,7 @@ namespace Horde.Server.Tests.Logs
 						string str = lines[lineIdx].Substring(strOfs, strLen);
 
 						SearchStats stats = new SearchStats();
-						List<int> results = await LogService.SearchLogDataAsync(logFile, str, 0, 5, stats, CancellationToken.None);
+						List<int> results = await logFile.SearchLogDataAsync(str, 0, 5, stats, CancellationToken.None);
 						Assert.AreEqual(1, results.Count);
 						Assert.AreEqual(lineIdx, results[0]);
 
@@ -195,7 +195,7 @@ namespace Horde.Server.Tests.Logs
 
 			{
 				SearchStats stats = new SearchStats();
-				List<int> results = await LogService.SearchLogDataAsync(logFile, "abc", 0, 5, stats, CancellationToken.None);
+				List<int> results = await logFile.SearchLogDataAsync("abc", 0, 5, stats, CancellationToken.None);
 				Assert.AreEqual(1, results.Count);
 				Assert.AreEqual(0, results[0]);
 
@@ -205,7 +205,7 @@ namespace Horde.Server.Tests.Logs
 			}
 			{
 				SearchStats stats = new SearchStats();
-				List<int> results = await LogService.SearchLogDataAsync(logFile, "def", 0, 5, stats, CancellationToken.None);
+				List<int> results = await logFile.SearchLogDataAsync("def", 0, 5, stats, CancellationToken.None);
 				Assert.AreEqual(1, results.Count);
 				Assert.AreEqual(1, results[0]);
 
@@ -215,7 +215,7 @@ namespace Horde.Server.Tests.Logs
 			}
 			{
 				SearchStats stats = new SearchStats();
-				List<int> results = await LogService.SearchLogDataAsync(logFile, "ghi", 0, 5, stats, CancellationToken.None);
+				List<int> results = await logFile.SearchLogDataAsync("ghi", 0, 5, stats, CancellationToken.None);
 				Assert.AreEqual(1, results.Count);
 				Assert.AreEqual(2, results[0]);
 
@@ -237,10 +237,10 @@ namespace Horde.Server.Tests.Logs
 			await SearchLogDataTestAsync(logFile, "NEWSLETTER", 0, 100, new[] { 7886 });
 		}
 
-		async Task SearchLogDataTestAsync(ILog logFile, string text, int firstLine, int count, int[] expectedLines)
+		static async Task SearchLogDataTestAsync(ILog logFile, string text, int firstLine, int count, int[] expectedLines)
 		{
 			SearchStats stats = new SearchStats();
-			List<int> lines = await LogService.SearchLogDataAsync(logFile, text, firstLine, count, stats, CancellationToken.None);
+			List<int> lines = await logFile.SearchLogDataAsync(text, firstLine, count, stats, CancellationToken.None);
 			Assert.IsTrue(lines.SequenceEqual(expectedLines));
 		}
 	}
diff --git a/Engine/Source/Programs/Horde/Horde.Server.Tests/Logs/LogServiceTests.cs b/Engine/Source/Programs/Horde/Horde.Server.Tests/Logs/LogServiceTests.cs
--- a/Engine/Source/Programs/Horde/Horde.Server.Tests/Logs/LogServiceTests.cs
+++ b/Engine/Source/Programs/Horde/Horde.Server.Tests/Logs/LogServiceTests.cs
@@ -58,10 +58,10 @@ namespace Horde.Server.Tests.Logs
 
 			logFile = (await LogService.GetLogAsync(logFile.Id, CancellationToken.None))!;
 
-			Assert.AreEqual("hello", await ReadLogAsync(LogService, logFile, 0, 5));
-			Assert.AreEqual("foo\nbar\nbaz\n", await ReadLogAsync(LogService, logFile, 6, 12));
+			Assert.AreEqual("hello", await ReadLogAsync(logFile, 0, 5));
+			Assert.AreEqual("foo\nbar\nbaz\n", await ReadLogAsync(logFile, 6, 12));
 
-			List<Utf8String> lines = await LogService.ReadLinesAsync(logFile, 0, 100, CancellationToken.None);
+			List<Utf8String> lines = await logFile.ReadLinesAsync(0, 100, CancellationToken.None);
 			Assert.AreEqual("hello", lines[0].ToString());
 			Assert.AreEqual("foo", lines[1].ToString());
 			Assert.AreEqual("bar", lines[2].ToString());
@@ -71,15 +71,16 @@ namespace Horde.Server.Tests.Logs
 		private static async Task AssertLineOffsetAsync(ILogService logService, LogId logId, int lineIndex, int clampedLineIndex, long offset)
 		{
 			ILog? logFile = await logService.GetLogAsync(logId, CancellationToken.None);
+			Assert.IsNotNull(logFile);
 
-			List<Utf8String> lines = await logService.ReadLinesAsync(logFile!, 0, lineIndex + 1, CancellationToken.None);
+			List<Utf8String> lines = await logFile.ReadLinesAsync(0, lineIndex + 1, CancellationToken.None);
 			Assert.AreEqual(clampedLineIndex, Math.Min(lineIndex, lines.Count));
 			Assert.AreEqual(offset, lines.Take(lineIndex).Sum(x => x.Length + 1));
 		}
 
-		static async Task<string> ReadLogAsync(ILogService logService, ILog log, long offset, long length)
+		static async Task<string> ReadLogAsync(ILog log, long offset, long length)
 		{
-			using Stream stream = await logService.OpenRawStreamAsync(log);
+			using Stream stream = await log.OpenRawStreamAsync();
 
 			byte[] prefix = new byte[(int)offset];
 			await stream.ReadGreedyAsync(prefix);
@@ -108,8 +109,8 @@ namespace Horde.Server.Tests.Logs
 
 			// First write with flush. Will become chunk #1
 			log = await logWriter.WriteDataAsync(Encoding.ASCII.GetBytes(str1));
-			Assert.AreEqual(str1, await ReadLogAsync(LogService, log, 0, str1.Length));
-			Assert.AreEqual(str1, await ReadLogAsync(LogService, log, 0, str1.Length + 100)); // Reading too far is valid?
+			Assert.AreEqual(str1, await ReadLogAsync(log, 0, str1.Length));
+			Assert.AreEqual(str1, await ReadLogAsync(log, 0, str1.Length + 100)); // Reading too far is valid?
 			await AssertLineOffsetAsync(LogService, log.Id, 0, 0, 0);
 			await AssertLineOffsetAsync(LogService, log.Id, 1, 1, 6);
 			await AssertLineOffsetAsync(LogService, log.Id, 1, 1, 6);
@@ -120,7 +121,7 @@ namespace Horde.Server.Tests.Logs
 			lineIndex += str1.Count(f => f == '\n');
 			log = await logWriter.WriteDataAsync(Encoding.ASCII.GetBytes(str2));
 
-			Assert.AreEqual(str1 + str2, await ReadLogAsync(LogService, log, 0, str1.Length + str2.Length));
+			Assert.AreEqual(str1 + str2, await ReadLogAsync(log, 0, str1.Length + str2.Length));
 			await AssertLineOffsetAsync(LogService, log.Id, 0, 0, 0);
 			await AssertLineOffsetAsync(LogService, log.Id, 1, 1, 6);
 			await AssertLineOffsetAsync(LogService, log.Id, 2, 2, 10);
@@ -130,7 +131,7 @@ namespace Horde.Server.Tests.Logs
 			offset += str2.Length;
 			lineIndex += str3.Count(f => f == '\n');
 			log = await logWriter.WriteDataAsync(Encoding.ASCII.GetBytes(str3));
-			Assert.AreEqual(str1 + str2 + str3, await ReadLogAsync(LogService, log, 0, str1.Length + str2.Length + str3.Length));
+			Assert.AreEqual(str1 + str2 + str3, await ReadLogAsync(log, 0, str1.Length + str2.Length + str3.Length));
 			await AssertLineOffsetAsync(LogService, log.Id, 0, 0, 0);
 			await AssertLineOffsetAsync(LogService, log.Id, 1, 1, 6);
 			await AssertLineOffsetAsync(LogService, log.Id, 2, 2, 10);
@@ -142,9 +143,9 @@ namespace Horde.Server.Tests.Logs
 			lineIndex += str4.Count(f => f == '\n');
 			log = await logWriter.WriteDataAsync(Encoding.ASCII.GetBytes(str4));
 			Assert.AreEqual(str1 + str2 + str3 + str4,
-			await ReadLogAsync(LogService, log, 0, str1.Length + str2.Length + str3.Length + str4.Length));
+			await ReadLogAsync(log, 0, str1.Length + str2.Length + str3.Length + str4.Length));
 			Assert.AreEqual(str3 + str4,
-			await ReadLogAsync(LogService, log, str1.Length + str2.Length, str3.Length + str4.Length));
+			await ReadLogAsync(log, str1.Length + str2.Length, str3.Length + str4.Length));
 			await AssertLineOffsetAsync(LogService, log.Id, 0, 0, 0);
 			await AssertLineOffsetAsync(LogService, log.Id, 1, 1, 6);
 			await AssertLineOffsetAsync(LogService, log.Id, 2, 2, 10);
@@ -169,7 +170,7 @@ namespace Horde.Server.Tests.Logs
 			log = await logWriter.WriteDataAsync(Encoding.ASCII.GetBytes(a));
 			log = await logWriter.WriteDataAsync(Encoding.ASCII.GetBytes(b));
 
-			Assert.AreEqual(str5, await ReadLogAsync(LogService, log, offset, str5.Length));
+			Assert.AreEqual(str5, await ReadLogAsync(log, offset, str5.Length));
 		}
 
 		[TestMethod]
@@ -233,7 +234,8 @@ namespace Horde.Server.Tests.Logs
 			await logWriter.WriteDataAsync(line5);
 
 			logFile = await logWriter.FlushAsync();
-			Assert.AreEqual(6, logFile.LineCount);
+			LogMetadata metadata = await logFile.GetMetadataAsync(CancellationToken.None);
+			Assert.AreEqual(6, metadata.MaxLineIndex);
 		}
 	}
 }
\ No newline at end of file
diff --git a/Engine/Source/Programs/Horde/Horde.Server/Logs/ILog.cs b/Engine/Source/Programs/Horde/Horde.Server/Logs/ILog.cs
--- a/Engine/Source/Programs/Horde/Horde.Server/Logs/ILog.cs
+++ b/Engine/Source/Programs/Horde/Horde.Server/Logs/ILog.cs
@@ -1,7 +1,7 @@
 // Copyright Epic Games, Inc. All Rights Reserved.
 
-using System;
 using System.Collections.Generic;
+using System.IO;
 using System.Threading;
 using System.Threading.Tasks;
 using EpicGames.Core;
@@ -13,32 +13,6 @@ using EpicGames.Horde.Storage;
 
 namespace Horde.Server.Logs
 {
-	/// <summary>
-	/// Information about a log file chunk
-	/// </summary>
-	public interface ILogChunk
-	{
-		/// <summary>
-		/// Offset of the chunk within the log
-		/// </summary>
-		long Offset { get; }
-
-		/// <summary>
-		/// Length of this chunk. If zero, the chunk is still being written to.
-		/// </summary>
-		int Length { get; }
-
-		/// <summary>
-		/// Index of the first line within this chunk. If a line straddles two chunks, this is the index of the split line.
-		/// </summary>
-		int LineIndex { get; }
-
-		/// <summary>
-		/// If the chunk has yet to be pushed to persistent storage, includes the name of the server that is currently storing it.
-		/// </summary>
-		string? Server { get; }
-	}
-
 	/// <summary>
 	/// Information about a log file
 	/// </summary>
@@ -65,95 +39,123 @@ namespace Horde.Server.Logs
 		public SessionId? SessionId { get; }
 
 		/// <summary>
-		/// Whether to use the new storage backend for log data
+		/// Type of data stored in this log 
 		/// </summary>
-		public bool UseNewStorageBackend { get; }
+		public LogType Type { get; }
 
 		/// <summary>
-		/// Maximum line index in the file
+		/// Namespace containing the log data
 		/// </summary>
-		public int? MaxLineIndex { get; }
+		public NamespaceId NamespaceId { get; }
 
 		/// <summary>
-		/// Length of the file which is indexed
+		/// Name of the ref used to store data for this log
 		/// </summary>
-		public long? IndexLength { get; }
+		public RefName RefName { get; }
+
+		//		/// <summary>
+		//		/// Number of lines (V2 storage backend)
+		//		/// </summary>
+		//		public int LineCount { get; }
+
+		//		/// <summary>
+		//		/// Whether the log is complete (V2 storage backend)
+		//		/// </summary>
+		//		public bool Complete { get; }
 
 		/// <summary>
-		/// Type of data stored in this log 
+		/// Read a set of lines from the given log file
 		/// </summary>
-		public LogType Type { get; }
+		/// <param name="index">Index of the first line to read</param>
+		/// <param name="count">Maximum number of lines to return</param>
+		/// <param name="cancellationToken">Cancellation token for the operation</param>
+		/// <returns>List of lines</returns>
+		Task<List<Utf8String>> ReadLinesAsync(int index, int count, CancellationToken cancellationToken = default);
 
 		/// <summary>
-		/// Chunks within this file
+		/// Gets metadata about the log file
 		/// </summary>
-		public IReadOnlyList<ILogChunk> Chunks { get; }
+		/// <param name="cancellationToken">Cancellation token for the call</param>
+		/// <returns>Metadata about the log file</returns>
+		Task<LogMetadata> GetMetadataAsync(CancellationToken cancellationToken);
 
 		/// <summary>
-		/// Namespace containing the log data
+		/// Updates the line count for a log file (v2 backend only)
 		/// </summary>
-		public NamespaceId NamespaceId { get; }
+		/// <param name="lineCount">New line count for the log file</param>
+		/// <param name="complete">Flag indicating whether the log is complete, or can still be tailed for additional data</param>
+		/// <param name="cancellationToken">Cancellation token for the call</param>
+		/// <returns>The updated log file document</returns>
+		Task<ILog> UpdateLineCountAsync(int lineCount, bool complete, CancellationToken cancellationToken = default);
 
 		/// <summary>
-		/// Name of the ref used to store data for this log
+		/// Gets lines from the given log 
 		/// </summary>
-		public RefName RefName { get; }
+		/// <param name="cancellationToken">Cancellation token for the call</param>
+		/// <returns>Data for the requested range</returns>
+		Task<Stream> OpenRawStreamAsync(CancellationToken cancellationToken = default);
 
 		/// <summary>
-		/// Number of lines (V2 storage backend)
+		/// Gets lines from the given log 
 		/// </summary>
-		public int LineCount { get; }
+		/// <param name="offset"></param>
+		/// <param name="length"></param>
+		/// <param name="cancellationToken">Cancellation token for the call</param>
+		/// <returns>Data for the requested range</returns>
+		Task<Stream> OpenRawStreamAsync(long offset, long length, CancellationToken cancellationToken);
 
 		/// <summary>
-		/// Whether the log is complete (V2 storage backend)
+		/// Parses a stream of json text and outputs plain text
 		/// </summary>
-		public bool Complete { get; }
+		/// <param name="outputStream">Output stream to receive the text data</param>
+		/// <param name="cancellationToken">Cancellation token for the call</param>
+		/// <returns>Async text</returns>
+		Task CopyPlainTextStreamAsync(Stream outputStream, CancellationToken cancellationToken = default);
 
 		/// <summary>
-		/// Updates the line count for a log file (v2 backend only)
+		/// Search for the specified text in a log file
 		/// </summary>
-		/// <param name="lineCount">New line count for the log file</param>
-		/// <param name="complete">Flag indicating whether the log is complete, or can still be tailed for additional data</param>
+		/// <param name="text">Text to search for</param>
+		/// <param name="firstLine">Line to start search from</param>
+		/// <param name="count">Number of results to return</param>
+		/// <param name="stats">Receives stats for the search</param>
 		/// <param name="cancellationToken">Cancellation token for the call</param>
-		/// <returns>The updated log file document</returns>
-		Task<ILog> UpdateLineCountAsync(int lineCount, bool complete, CancellationToken cancellationToken = default);
+		/// <returns>List of line numbers containing the given term</returns>
+		Task<List<int>> SearchLogDataAsync(string text, int firstLine, int count, SearchStats stats, CancellationToken cancellationToken);
 	}
 
 	/// <summary>
-	/// Extension methods for log files
+	/// Metadata about a log file
 	/// </summary>
-	public static class LogExtensions
+	public class LogMetadata
 	{
 		/// <summary>
-		/// Gets the chunk index containing the given offset.
+		/// Length of the log file
 		/// </summary>
-		/// <param name="chunks">The chunks to search</param>
-		/// <param name="offset">The offset to search for</param>
-		/// <returns>The chunk index containing the given offset</returns>
-		public static int GetChunkForOffset(this IReadOnlyList<ILogChunk> chunks, long offset)
-		{
-			int chunkIndex = chunks.BinarySearch(x => x.Offset, offset);
-			if (chunkIndex < 0)
-			{
-				chunkIndex = ~chunkIndex - 1;
-			}
-			return chunkIndex;
-		}
+		public long Length { get; set; }
+
+		/// <summary>
+		/// Number of lines in the log file
+		/// </summary>
+		public int MaxLineIndex { get; set; }
+	}
 
+	/// <summary>
+	/// Extension methods for log files
+	/// </summary>
+	public static class LogExtensions
+	{
 		/// <summary>
-		/// Gets the starting chunk index for the given line
+		/// Parses a stream of json text and outputs plain text
 		/// </summary>
-		/// <param name="chunks">The chunks to search</param>
-		/// <param name="lineIndex">Index of the line to query</param>
-		/// <returns>Index of the chunk to fetch</returns>
-		public static int GetChunkForLine(this IReadOnlyList<ILogChunk> chunks, int lineIndex)
+		/// <param name="log">The log file to query</param>
+		/// <param name="outputStream">Output stream to receive the text data</param>
+		/// <param name="cancellationToken">Cancellation token for the call</param>
+		/// <returns>Async text</returns>
+		public static async Task CopyRawStreamAsync(this ILog log, Stream outputStream, CancellationToken cancellationToken)
 		{
-			int chunkIndex = chunks.BinarySearch(x => x.LineIndex, lineIndex);
-			if (chunkIndex < 0)
-			{
-				chunkIndex = ~chunkIndex - 1;
-			}
-			return chunkIndex;
+			await using Stream stream = await log.OpenRawStreamAsync(cancellationToken);
+			await stream.CopyToAsync(outputStream, cancellationToken);
 		}
 	}
 }
diff --git a/Engine/Source/Programs/Horde/Horde.Server/Logs/ILogCollection.cs b/Engine/Source/Programs/Horde/Horde.Server/Logs/ILogCollection.cs
--- a/Engine/Source/Programs/Horde/Horde.Server/Logs/ILogCollection.cs
+++ b/Engine/Source/Programs/Horde/Horde.Server/Logs/ILogCollection.cs
@@ -15,7 +15,7 @@ namespace Horde.Server.Logs
 	public interface ILogCollection
 	{
 		/// <summary>
-		/// Creates a new log
+		/// Adds a new log
 		/// </summary>
 		/// <param name="jobId">Unique id of the job that owns this log file</param>
 		/// <param name="leaseId">Agent lease allowed to update the log</param>
@@ -24,7 +24,7 @@ namespace Horde.Server.Logs
 		/// <param name="logId">ID of the log file (optional)</param>
 		/// <param name="cancellationToken">Cancellation token for the call</param>
 		/// <returns>The new log file document</returns>
-		Task<ILog> CreateLogAsync(JobId jobId, LeaseId? leaseId, SessionId? sessionId, LogType type, LogId? logId = null, CancellationToken cancellationToken = default);
+		Task<ILog> AddAsync(JobId jobId, LeaseId? leaseId, SessionId? sessionId, LogType type, LogId? logId = null, CancellationToken cancellationToken = default);
 
 		/// <summary>
 		/// Gets a log by ID
@@ -32,6 +32,6 @@ namespace Horde.Server.Logs
 		/// <param name="logId">Unique id of the log file</param>
 		/// <param name="cancellationToken">Cancellation token for the call</param>
 		/// <returns>The log instance</returns>
-		Task<ILog?> GetLogAsync(LogId logId, CancellationToken cancellationToken);
+		Task<ILog?> GetAsync(LogId logId, CancellationToken cancellationToken);
 	}
 }
diff --git a/Engine/Source/Programs/Horde/Horde.Server/Logs/LogCollection.cs b/Engine/Source/Programs/Horde/Horde.Server/Logs/LogCollection.cs
--- a/Engine/Source/Programs/Horde/Horde.Server/Logs/LogCollection.cs
+++ b/Engine/Source/Programs/Horde/Horde.Server/Logs/LogCollection.cs
@@ -1,52 +1,79 @@
 // Copyright Epic Games, Inc. All Rights Reserved.
 
+using System;
 using System.Collections.Generic;
+using System.Diagnostics;
+using System.IO;
 using System.Linq;
+using System.Runtime.CompilerServices;
+using System.Text;
 using System.Threading;
 using System.Threading.Tasks;
+using EpicGames.Core;
 using EpicGames.Horde.Agents.Leases;
 using EpicGames.Horde.Agents.Sessions;
 using EpicGames.Horde.Jobs;
 using EpicGames.Horde.Logs;
 using EpicGames.Horde.Storage;
+using Horde.Server.Logs.Data;
 using Horde.Server.Server;
 using Horde.Server.Storage;
+using Horde.Server.Utilities;
+using Microsoft.Extensions.Caching.Memory;
+using Microsoft.Extensions.Logging;
 using MongoDB.Bson.Serialization.Attributes;
 using MongoDB.Driver;
+using OpenTelemetry.Trace;
 
 namespace Horde.Server.Logs
 {
 	/// <summary>
 	/// Wrapper around the jobs collection in a mongo DB
 	/// </summary>
-	public class LogCollection : ILogCollection
+	public sealed class LogCollection : ILogCollection, IDisposable
 	{
-		class LogChunkDocument : ILogChunk
+		class Log : ILog
 		{
-			public long Offset { get; set; }
-			public int Length { get; set; }
-			public int LineIndex { get; set; }
+			readonly LogCollection _collection;
+			readonly LogDocument _document;
 
-			[BsonIgnoreIfNull]
-			public string? Server { get; set; }
+			LogId ILog.Id => _document.Id;
+			JobId ILog.JobId => _document.JobId;
+			LeaseId? ILog.LeaseId => _document.LeaseId;
+			SessionId? ILog.SessionId => _document.SessionId;
+			LogType ILog.Type => _document.Type;
+			NamespaceId ILog.NamespaceId => _document.NamespaceId;
+			RefName ILog.RefName => _document.RefName;
 
-			[BsonConstructor]
-			public LogChunkDocument()
+			public Log(LogCollection collection, LogDocument document)
 			{
+				_collection = collection;
+				_document = document;
 			}
 
-			public LogChunkDocument(LogChunkDocument other)
+			public async Task<ILog> UpdateLineCountAsync(int lineCount, bool complete, CancellationToken cancellationToken = default)
 			{
-				Offset = other.Offset;
-				Length = other.Length;
-				LineIndex = other.LineIndex;
-				Server = other.Server;
+				LogDocument newDocument = await _collection.UpdateLineCountAsync(_document, lineCount, complete, cancellationToken);
+				return new Log(_collection, newDocument);
 			}
 
-			public LogChunkDocument Clone()
-			{
-				return (LogChunkDocument)MemberwiseClone();
-			}
+			public Task<List<Utf8String>> ReadLinesAsync(int index, int count, CancellationToken cancellationToken = default)
+				=> _collection.ReadLinesAsync(_document, index, count, cancellationToken);
+
+			public Task<LogMetadata> GetMetadataAsync(CancellationToken cancellationToken)
+				=> _collection.GetMetadataAsync(_document, cancellationToken);
+
+			public Task<Stream> OpenRawStreamAsync(CancellationToken cancellationToken = default)
+				=> _collection.OpenRawStreamAsync(_document, 0, Int64.MaxValue, cancellationToken);
+
+			public Task<Stream> OpenRawStreamAsync(long offset, long length, CancellationToken cancellationToken)
+				=> _collection.OpenRawStreamAsync(_document, offset, length, cancellationToken);
+
+			public Task CopyPlainTextStreamAsync(Stream outputStream, CancellationToken cancellationToken = default)
+				=> _collection.CopyPlainTextStreamAsync(this, outputStream, cancellationToken);
+
+			public Task<List<int>> SearchLogDataAsync(string text, int firstLine, int count, SearchStats stats, CancellationToken cancellationToken)
+				=> _collection.SearchLogDataAsync(_document, text, firstLine, count, stats, cancellationToken);
 		}
 
 		class LogDocument
@@ -104,60 +131,148 @@
 			}
 		}
 
-		class Log : ILog
+		class LogChunkDocument
 		{
-			readonly LogCollection _collection;
-			readonly LogDocument _document;
+			public long Offset { get; set; }
+			public int Length { get; set; }
+			public int LineIndex { get; set; }
 
-			public LogId Id => _document.Id;
-			public JobId JobId => _document.JobId;
-			public LeaseId? LeaseId => _document.LeaseId;
-			public SessionId? SessionId => _document.SessionId;
-			public bool UseNewStorageBackend => _document.UseNewStorageBackend;
-			public int? MaxLineIndex => _document.MaxLineIndex;
-			public long? IndexLength => _document.IndexLength;
-			public LogType Type => _document.Type;
-			public IReadOnlyList<ILogChunk> Chunks => _document.Chunks;
-			public NamespaceId NamespaceId => _document.NamespaceId;
-			public RefName RefName => _document.RefName;
-			public int LineCount => _document.LineCount;
-			public bool Complete => _document.Complete;
-			
-			public Log(LogCollection collection, LogDocument document)
+			[BsonIgnoreIfNull]
+			public string? Server { get; set; }
+
+			[BsonConstructor]
+			public LogChunkDocument()
 			{
-				_collection = collection;
-				_document = document;
 			}
 
-			public async Task<ILog> UpdateLineCountAsync(int lineCount, bool complete, CancellationToken cancellationToken = default)
+			public LogChunkDocument(LogChunkDocument other)
 			{
-				LogDocument newDocument = await _collection.UpdateLineCountAsync(_document, lineCount, complete, cancellationToken);
-				return new Log(_collection, newDocument);
+				Offset = other.Offset;
+				Length = other.Length;
+				LineIndex = other.LineIndex;
+				Server = other.Server;
+			}
+
+			public LogChunkDocument Clone()
+			{
+				return (LogChunkDocument)MemberwiseClone();
 			}
 		}
 
-		/// <summary>
-		/// The jobs collection
-		/// </summary>
 		readonly IMongoCollection<LogDocument> _logCollection;
+		readonly ILogStorage _storage;
+		readonly LogTailService _logTailService;
+		readonly StorageService _storageService;
+		readonly Tracer _tracer;
+		readonly ILogger _logger;
+		readonly IMemoryCache _logCache;
 
 		/// <summary>
 		/// Constructor
 		/// </summary>
-		/// <param name="mongoService">The database service singleton</param>
-		public LogCollection(MongoService mongoService)
+		public LogCollection(MongoService mongoService, ILogStorage storage, LogTailService logTailService, StorageService storageService, Tracer tracer, ILogger<LogCollection> logger)
 		{
 			_logCollection = mongoService.GetCollection<LogDocument>("LogFiles");
+			_storage = storage;
+			_logTailService = logTailService;
+			_storageService = storageService;
+			_tracer = tracer;
+			_logger = logger;
+			_logCache = new MemoryCache(new MemoryCacheOptions());
 		}
 
 		/// <inheritdoc/>
-		public async Task<ILog> CreateLogAsync(JobId jobId, LeaseId? leaseId, SessionId? sessionId, LogType type, LogId? logId, CancellationToken cancellationToken)
+		public void Dispose()
 		{
+			_logCache.Dispose();
+		}
+
+		/// <inheritdoc/>
+		public async Task<ILog> AddAsync(JobId jobId, LeaseId? leaseId, SessionId? sessionId, LogType type, LogId? logId, CancellationToken cancellationToken)
+		{
 			LogDocument newLog = new LogDocument(jobId, leaseId, sessionId, type, logId, Namespace.Logs);
 			await _logCollection.InsertOneAsync(newLog, null, cancellationToken);
 			return new Log(this, newLog);
 		}
 
 		/// <inheritdoc/>
+		public async Task<ILog?> GetAsync(LogId logId, CancellationToken cancellationToken)
+		{
+			LogDocument? logDocument = await _logCollection.Find<LogDocument>(x => x.Id == logId).FirstOrDefaultAsync(cancellationToken);
+			if (logDocument == null)
+			{
+				return null;
+			}
+
+			if (logDocument.UseNewStorageBackend)
+			{
+				return new Log(this, logDocument);
+			}
+			else
+			{
+				return new LogV1(this, logDocument);
+			}
+		}
+
+		/// <inheritdoc/>
+		async Task<List<Utf8String>> ReadLinesAsync(LogDocument log, int index, int count, CancellationToken cancellationToken)
+		{
+			List<Utf8String> lines = new List<Utf8String>();
+
+			using IStorageClient storageClient = _storageService.CreateClient(log.NamespaceId);
+
+			int maxIndex = index + count;
+			bool complete = log.Complete;
+
+			LogNode? root = await storageClient.TryReadRefTargetAsync<LogNode>(log.RefName, cancellationToken: cancellationToken);
+			if (root != null)
+			{
+				int chunkIdx = root.TextChunkRefs.GetChunkForLine(index);
+				for (; index < maxIndex && chunkIdx < root.TextChunkRefs.Count; chunkIdx++)
+				{
+					LogChunkRef chunk = root.TextChunkRefs[chunkIdx];
+					LogChunkNode chunkData = await chunk.Target.ReadBlobAsync(cancellationToken: cancellationToken);
+
+					for (; index < maxIndex && index < chunk.LineIndex; index++)
+					{
+						lines.Add(new Utf8String($"Internal error; missing data for line {index}\n"));
+					}
+
+					for (; index < maxIndex && index < chunk.LineIndex + chunk.LineCount; index++)
+					{
+						lines.Add(chunkData.GetLine(index - chunk.LineIndex));
+					}
+				}
+				complete |= root.Complete;
+			}
+
+			if (!complete)
+			{
+				await _logTailService.EnableTailingAsync(log.Id, root?.LineCount ?? 0, cancellationToken);
+				if (index < maxIndex)
+				{
+					await _logTailService.ReadAsync(log.Id, index, maxIndex - index, lines);
+				}
+			}
+
+			return lines;
+		}
+
+		/// <inheritdoc/>
+		async Task<LogMetadata> GetMetadataAsync(LogDocument log, CancellationToken cancellationToken)
+		{
+			LogMetadata metadata = new LogMetadata();
+			if (log.Complete)
+			{
+				metadata.MaxLineIndex = log.LineCount;
+			}
+			else
+			{
+				metadata.MaxLineIndex = await _logTailService.GetFullLineCountAsync(log.Id, log.LineCount, cancellationToken);
+			}
+			return metadata;
+		}
+
+		/// <inheritdoc/>
 		async Task<LogDocument> UpdateLineCountAsync(LogDocument log, int lineCount, bool complete, CancellationToken cancellationToken)
 		{
@@ -167,14 +282,918 @@
 		}
 
 		/// <inheritdoc/>
-		public async Task<ILog?> GetLogAsync(LogId logId, CancellationToken cancellationToken)
+		async Task<(int, long)> GetLineOffsetAsync(LogDocument log, int lineIdx, CancellationToken cancellationToken)
 		{
-			LogDocument? log = await _logCollection.Find<LogDocument>(x => x.Id == logId).FirstOrDefaultAsync(cancellationToken);
-			if (log == null)
+			using IStorageClient storageClient = _storageService.CreateClient(log.NamespaceId);
+
+			LogNode? root = await storageClient.TryReadRefTargetAsync<LogNode>(log.RefName, cancellationToken: cancellationToken);
+			if (root == null)
 			{
+				return (0, 0);
+			}
+
+			int chunkIdx = root.TextChunkRefs.GetChunkForLine(lineIdx);
+			LogChunkRef chunk = root.TextChunkRefs[chunkIdx];
+			LogChunkNode chunkData = await chunk.Target.ReadBlobAsync(cancellationToken: cancellationToken);
+
+			if (lineIdx < chunk.LineIndex)
+			{
+				lineIdx = chunk.LineIndex;
+			}
+
+			int maxLineIndex = chunk.LineIndex + chunkData.LineCount;
+			if (lineIdx >= maxLineIndex)
+			{
+				lineIdx = maxLineIndex;
+			}
+
+			long offset = chunk.Offset + chunkData.LineOffsets[lineIdx - chunk.LineIndex];
+			return (lineIdx, offset);
+		}
+
+		/// <inheritdoc/>
+		async Task<Stream> OpenRawStreamAsync(LogDocument log, long offset, long length, CancellationToken cancellationToken)
+		{
+			using IStorageClient storageClient = _storageService.CreateClient(log.NamespaceId);
+
+			LogNode? root = await storageClient.TryReadRefTargetAsync<LogNode>(log.RefName, cancellationToken: cancellationToken);
+			if (root == null || root.TextChunkRefs.Count == 0)
+			{
+				return new MemoryStream(Array.Empty<byte>(), false);
+			}
+			else
+			{
+				int lastChunkIdx = root.TextChunkRefs.Count - 1;
+
+				// Clamp the length of the request
+				LogChunkRef lastChunk = root.TextChunkRefs[lastChunkIdx];
+				if (length > lastChunk.Offset)
+				{
+					long lastChunkLength = lastChunk.Length;
+					if (lastChunkLength <= 0)
+					{
+						LogChunkNode lastChunkNode = await lastChunk.Target.ReadBlobAsync(cancellationToken: cancellationToken);
+						lastChunkLength = lastChunkNode.Length;
+					}
+					length = Math.Min(length, (lastChunk.Offset + lastChunkLength) - offset);
+				}
+
+				// Create the new stream
+				return new ResponseStream(root, offset, length);
+			}
+		}
+
+		/// <inheritdoc/>
+		async Task<List<int>> SearchLogDataAsync(LogDocument log, string text, int firstLine, int count, SearchStats searchStats, CancellationToken cancellationToken)
+		{
+			Stopwatch timer = Stopwatch.StartNew();
+
+			using TelemetrySpan span = _tracer.StartActiveSpan($"{nameof(LogService)}.{nameof(SearchLogDataAsync)}");
+			span.SetAttribute("logId", log.Id.ToString());
+			span.SetAttribute("text", text);
+			span.SetAttribute("count", count);
+
+			List<int> results = new List<int>();
+			if (count > 0)
+			{
+				IAsyncEnumerable<int> enumerable = SearchLogDataInternalNewAsync(log, text, firstLine, searchStats, cancellationToken);
+
+				await using IAsyncEnumerator<int> enumerator = enumerable.GetAsyncEnumerator(cancellationToken);
+				while (await enumerator.MoveNextAsync() && results.Count < count)
+				{
+					results.Add(enumerator.Current);
+				}
+			}
+
+			_logger.LogDebug("Search for \"{SearchText}\" in log {LogId} found {NumResults}/{MaxResults} results, took {Time}ms ({@Stats})", text, log.Id, results.Count, count, timer.ElapsedMilliseconds, searchStats);
+			return results;
+		}
+
+		async IAsyncEnumerable<int> SearchLogDataInternalNewAsync(LogDocument log, string text, int firstLine, SearchStats searchStats, [EnumeratorCancellation] CancellationToken cancellationToken)
+		{
+			SearchTerm searchText = new SearchTerm(text);
+			using IStorageClient storageClient = _storageService.CreateClient(log.NamespaceId);
+
+			// Search the index
+			if (log.LineCount > 0)
+			{
+				LogNode? root = await storageClient.ReadRefTargetAsync<LogNode>(log.RefName, cancellationToken: cancellationToken);
+				if (root != null)
+				{
+					LogIndexNode index = await root.IndexRef.ReadBlobAsync(cancellationToken: cancellationToken);
+					await foreach (int lineIdx in index.SearchAsync(firstLine, searchText, searchStats, cancellationToken: cancellationToken))
+					{
+						yield return lineIdx;
+					}
+					if (root.Complete)
+					{
+						yield break;
+					}
+					firstLine = root.LineCount;
+				}
+			}
+
+			// Search any tail data we have
+			if (!log.Complete)
+			{
+				for (; ; )
+				{
+					Utf8String[] lines = await ReadTailAsync(log, firstLine, cancellationToken);
+					if (lines.Length == 0)
+					{
+						break;
+					}
+
+					for (int idx = 0; idx < lines.Length; idx++)
+					{
+						if (SearchTerm.FindNextOcurrence(lines[idx].Span, 0, searchText) != -1)
+						{
+							yield return firstLine + idx;
+						}
+					}
+
+					firstLine += lines.Length;
+				}
+			}
+		}
+
+		async Task<Utf8String[]> ReadTailAsync(LogDocument log, int index, CancellationToken cancellationToken)
+		{
+			_ = cancellationToken;
+
+			const int BatchSize = 128;
+
+			if (log.Complete)
+			{
+				return Array.Empty<Utf8String>();
+			}
+
+			string cacheKey = $"{log.Id}@{index}";
+			if (_logCache.TryGetValue(cacheKey, out Utf8String[]? lines))
+			{
+				return lines!;
+			}
+
+			lines = (await _logTailService.ReadAsync(log.Id, index, BatchSize)).ToArray();
+			if (log.Type == LogType.Json)
+			{
+				LogChunkBuilder builder = new LogChunkBuilder(lines.Sum(x => x.Length));
+				foreach (Utf8String line in lines)
+				{
+					builder.AppendJsonAsPlainText(line.Span, _logger);
+				}
+				lines = lines.ToArray();
+			}
+
+			if (lines.Length == BatchSize)
+			{
+				int length = lines.Sum(x => x.Length);
+				using (ICacheEntry entry = _logCache.CreateEntry(cacheKey))
+				{
+					entry.SetSlidingExpiration(TimeSpan.FromMinutes(1.0));
+					entry.SetSize(length);
+					entry.SetValue(lines);
+				}
+			}
+
+			return lines;
+		}
+
+		/// <inheritdoc/>
+		async Task CopyPlainTextStreamAsync(ILog log, Stream outputStream, CancellationToken cancellationToken)
+		{
+			long offset = 0;
+			long length = Int64.MaxValue;
+
+			using (Stream stream = await log.OpenRawStreamAsync(0, Int64.MaxValue, cancellationToken))
+			{
+				byte[] readBuffer = new byte[4096];
+				int readBufferLength = 0;
+
+				byte[] writeBuffer = new byte[4096];
+				int writeBufferLength = 0;
+
+				while (length > 0)
+				{
+					// Add more data to the buffer
+					int readBytes = await stream.ReadAsync(readBuffer.AsMemory(readBufferLength, readBuffer.Length - readBufferLength), cancellationToken);
+					readBufferLength += readBytes;
+
+					// Copy as many lines as possible to the output
+					int convertedBytes = 0;
+					for (int endIdx = 1; endIdx < readBufferLength; endIdx++)
+					{
+						if (readBuffer[endIdx] == '\n')
+						{
+							writeBufferLength = GuardedConvertToPlainText(readBuffer.AsSpan(convertedBytes, endIdx - convertedBytes), writeBuffer, writeBufferLength, _logger);
+							convertedBytes = endIdx + 1;
+						}
+					}
+
+					// If there's anything in the write buffer, write it out
+					if (writeBufferLength > 0)
+					{
+						if (offset < writeBufferLength)
+						{
+							int writeLength = (int)Math.Min((long)writeBufferLength - offset, length);
+							await outputStream.WriteAsync(writeBuffer.AsMemory((int)offset, writeLength), cancellationToken);
+							length -= writeLength;
+						}
+						offset = Math.Max(offset - writeBufferLength, 0);
+						writeBufferLength = 0;
+					}
+
+					// If we were able to read something, shuffle down the rest of the buffer. Otherwise expand the read buffer.
+					if (convertedBytes > 0)
+					{
+						Buffer.BlockCopy(readBuffer, convertedBytes, readBuffer, 0, readBufferLength - convertedBytes);
+						readBufferLength -= convertedBytes;
+					}
+					else if (readBufferLength > 0)
+					{
+						Array.Resize(ref readBuffer, readBuffer.Length + 128);
+						writeBuffer = new byte[readBuffer.Length];
+					}
+
+					// Exit if we didn't read anything in this iteration
+					if (readBytes == 0)
+					{
+						break;
+					}
+				}
+			}
+		}
+
+		/// <summary>
+		/// Helper method for catching exceptions in <see cref="LogText.ConvertToPlainText(ReadOnlySpan{Byte}, Byte[], Int32)"/>
+		/// </summary>
+		public static int GuardedConvertToPlainText(ReadOnlySpan<byte> input, byte[] output, int outputOffset, ILogger logger)
+		{
+			try
+			{
+				return LogText.ConvertToPlainText(input, output, outputOffset);
+			}
+			catch (Exception ex)
+			{
+				logger.LogWarning(ex, "Unable to convert log line to plain text: {Line}", Encoding.UTF8.GetString(input));
+				output[outputOffset] = (byte)'\n';
+				return outputOffset + 1;
+			}
+		}
+
+		/// <summary>
+		/// Streams log data to a caller
+		/// </summary>
+		class ResponseStream : Stream
+		{
+			readonly LogNode _rootNode;
+
+			/// <summary>
+			/// Starting offset within the file of the data to return 
+			/// </summary>
+			readonly long _responseOffset;
+
+			/// <summary>
+			/// Length of data to return
+			/// </summary>
+			readonly long _responseLength;
+
+			/// <summary>
+			/// Current offset within the stream
+			/// </summary>
+			long _currentOffset;
+
+			/// <summary>
+			/// The current chunk index
+			/// </summary>
+			int _chunkIdx;
+
+			/// <summary>
+			/// Buffer containing a message for missing data
+			/// </summary>
+			ReadOnlyMemory<byte> _sourceBuffer;
+
+			/// <summary>
+			/// Offset within the source buffer
+			/// </summary>
+			int _sourcePos;
+
+			/// <summary>
+			/// Length of the source buffer being copied from
+			/// </summary>
+			int _sourceEnd;
+
+			/// <summary>
+			/// Constructor
+			/// </summary>
+			public ResponseStream(LogNode rootNode, long offset, long length)
+			{
+				_rootNode = rootNode;
+
+				_responseOffset = offset;
+				_responseLength = length;
+
+				_currentOffset = offset;
+
+				_chunkIdx = rootNode.TextChunkRefs.GetChunkForOffset(offset);
+				_sourceBuffer = null!;
+			}
+
+			/// <inheritdoc/>
+			public override bool CanRead => true;
+
+			/// <inheritdoc/>
+			public override bool CanSeek => false;
+
+			/// <inheritdoc/>
+			public override bool CanWrite => false;
+
+			/// <inheritdoc/>
+			public override long Length => _responseLength;
+
+			/// <inheritdoc/>
+			public override long Position
+			{
+				get => _currentOffset - _responseOffset;
+				set => throw new NotImplementedException();
+			}
+
+			/// <inheritdoc/>
+			public override void Flush()
+			{
+			}
+
+			/// <inheritdoc/>
+			public override int Read(byte[] buffer, int offset, int count)
+			{
+#pragma warning disable VSTHRD002
+				return ReadAsync(buffer, offset, count, CancellationToken.None).Result;
+#pragma warning restore VSTHRD002
+			}
+
+			/// <inheritdoc/>
+			public override async Task<int> ReadAsync(byte[] buffer, int offset, int length, CancellationToken cancellationToken)
+			{
+				return await ReadAsync(buffer.AsMemory(offset, length), cancellationToken);
+			}
+
+			/// <inheritdoc/>
+			public override async ValueTask<int> ReadAsync(Memory<byte> buffer, CancellationToken cancellationToken)
+			{
+				int readBytes = 0;
+				while (readBytes < buffer.Length)
+				{
+					if (_sourcePos < _sourceEnd)
+					{
+						// Try to copy from the current buffer
+						int blockSize = Math.Min(_sourceEnd - _sourcePos, buffer.Length - readBytes);
+						_sourceBuffer.Slice(_sourcePos, blockSize).Span.CopyTo(buffer.Slice(readBytes).Span);
+						_currentOffset += blockSize;
+						readBytes += blockSize;
+						_sourcePos += blockSize;
+					}
+					else if (_currentOffset < _responseOffset + _responseLength)
+					{
+						// Move to the right chunk
+						while (_chunkIdx + 1 < _rootNode.TextChunkRefs.Count && _currentOffset >= _rootNode.TextChunkRefs[_chunkIdx + 1].Offset)
+						{
+							_chunkIdx++;
+						}
+
+						// Get the chunk data
+						LogChunkRef chunk = _rootNode.TextChunkRefs[_chunkIdx];
+						LogChunkNode chunkNode = await chunk.Target.ReadBlobAsync(cancellationToken: cancellationToken);
+
+						// Get the source data
+						_sourceBuffer = chunkNode.Data;
+						_sourcePos = (int)(_currentOffset - chunk.Offset);
+						_sourceEnd = (int)Math.Min(_sourceBuffer.Length, (_responseOffset + _responseLength) - chunk.Offset);
+					}
+					else
+					{
+						// End of the log
+						break;
+					}
+				}
+				return readBytes;
+			}
+
+			/// <inheritdoc/>
+			public override long Seek(long offset, SeekOrigin origin) => throw new NotImplementedException();
+
+			/// <inheritdoc/>
+			public override void SetLength(long value) => throw new NotImplementedException();
+
+			/// <inheritdoc/>
+			public override void Write(byte[] buffer, int offset, int count) => throw new NotImplementedException();
+		}
+
+		#region V1
+
+		class LogV1 : ILog
+		{
+			readonly LogCollection _collection;
+			readonly LogDocument _document;
+
+			LogId ILog.Id => _document.Id;
+			JobId ILog.JobId => _document.JobId;
+			LeaseId? ILog.LeaseId => _document.LeaseId;
+			SessionId? ILog.SessionId => _document.SessionId;
+			LogType ILog.Type => _document.Type;
+			NamespaceId ILog.NamespaceId => _document.NamespaceId;
+			RefName ILog.RefName => _document.RefName;
+
+			public LogV1(LogCollection collection, LogDocument document)
+			{
+				_collection = collection;
+				_document = document;
+			}
+
+			public async Task<ILog> UpdateLineCountAsync(int lineCount, bool complete, CancellationToken cancellationToken = default)
+			{
+				LogDocument newDocument = await _collection.UpdateLineCountAsync(_document, lineCount, complete, cancellationToken);
+				return new Log(_collection, newDocument);
+			}
+
+			public Task<List<Utf8String>> ReadLinesAsync(int index, int count, CancellationToken cancellationToken = default)
+				=> _collection.ReadLinesV1Async(_document, index, count, cancellationToken);
+
+			public Task<LogMetadata> GetMetadataAsync(CancellationToken cancellationToken)
+				=> _collection.GetMetadataV1Async(_document);
+
+			public Task<Stream> OpenRawStreamAsync(CancellationToken cancellationToken = default)
+				=> _collection.OpenRawStreamV1Async(_document, 0, Int64.MaxValue);
+
+			public Task<Stream> OpenRawStreamAsync(long offset, long length, CancellationToken cancellationToken)
+				=> _collection.OpenRawStreamV1Async(_document, offset, length);
+
+			public Task CopyPlainTextStreamAsync(Stream outputStream, CancellationToken cancellationToken = default)
+				=> _collection.CopyPlainTextStreamAsync(this, outputStream, cancellationToken);
+
+			public Task<List<int>> SearchLogDataAsync(string text, int firstLine, int count, SearchStats stats, CancellationToken cancellationToken)
+				=> _collection.SearchLogDataV1Async(_document, text, firstLine, count, stats, cancellationToken);
+		}
+
+		async Task<List<Utf8String>> ReadLinesV1Async(LogDocument log, int index, int count, CancellationToken cancellationToken)
+		{
+			List<Utf8String> lines = new List<Utf8String>();
+
+			(_, long minOffset) = await GetLineOffsetAsync(log, index, cancellationToken);
+			(_, long maxOffset) = await GetLineOffsetAsync(log, index + Math.Min(count, Int32.MaxValue - index), cancellationToken);
+
+			byte[] result;
+			using (System.IO.Stream stream = await OpenRawStreamAsync(log, minOffset, maxOffset - minOffset, cancellationToken))
+			{
+				result = new byte[stream.Length];
+				await stream.ReadFixedSizeDataAsync(result, 0, result.Length, cancellationToken);
+			}
+
+			int offset = 0;
+			for (int idx = 0; idx < result.Length; idx++)
+			{
+				if (result[idx] == (byte)'\n')
+				{
+					lines.Add(new Utf8String(result.AsMemory(offset, idx - offset)));
+					offset = idx + 1;
+				}
+			}
+
+			return lines;
+		}
+
+		/// <inheritdoc/>
+		async Task<LogMetadata> GetMetadataV1Async(LogDocument log)
+		{
+			LogMetadata metadata = new LogMetadata();
+			if (log.Chunks.Count > 0)
+			{
+				LogChunkDocument chunk = log.Chunks[log.Chunks.Count - 1];
+				if (log.MaxLineIndex == null || chunk.Length == 0)
+				{
+					LogChunkData chunkData = await ReadChunkV1Async(log, log.Chunks.Count - 1);
+					metadata.Length = chunk.Offset + chunkData.Length;
+					metadata.MaxLineIndex = chunk.LineIndex + chunkData.LineCount;
+				}
+				else
+				{
+					metadata.Length = chunk.Offset + chunk.Length;
+					metadata.MaxLineIndex = log.MaxLineIndex.Value;
+				}
+			}
+			return metadata;
+		}
+
+		/// <inheritdoc/>
+		async Task<Stream> OpenRawStreamV1Async(LogDocument log, long offset, long length)
+		{
+			if (log.Chunks.Count == 0)
+			{
+				return new MemoryStream(Array.Empty<byte>(), false);
+			}
+			else
+			{
+				int lastChunkIdx = log.Chunks.Count - 1;
+
+				// Clamp the length of the request
+				LogChunkDocument lastChunk = log.Chunks[lastChunkIdx];
+				if (length > lastChunk.Offset)
+				{
+					long lastChunkLength = lastChunk.Length;
+					if (lastChunkLength <= 0)
+					{
+						LogChunkData lastChunkData = await ReadChunkV1Async(log, lastChunkIdx);
+						lastChunkLength = lastChunkData.Length;
+					}
+					length = Math.Min(length, (lastChunk.Offset + lastChunkLength) - offset);
+				}
+
+				// Create the new stream
+				return new ResponseStreamV1(this, log, offset, length);
+			}
+		}
+
+		/// <inheritdoc/>
+		async Task<List<int>> SearchLogDataV1Async(LogDocument log, string text, int firstLine, int count, SearchStats searchStats, CancellationToken cancellationToken)
+		{
+			Stopwatch timer = Stopwatch.StartNew();
+
+			using TelemetrySpan span = _tracer.StartActiveSpan($"{nameof(LogService)}.{nameof(SearchLogDataAsync)}");
+			span.SetAttribute("logId", log.Id.ToString());
+			span.SetAttribute("text", text);
+			span.SetAttribute("count", count);
+
+			List<int> results = new List<int>();
+			if (count > 0)
+			{
+				IAsyncEnumerable<int> enumerable = SearchLogDataInternalV1Async(log, text, firstLine, searchStats);
+
+				await using IAsyncEnumerator<int> enumerator = enumerable.GetAsyncEnumerator(cancellationToken);
+				while (await enumerator.MoveNextAsync() && results.Count < count)
+				{
+					results.Add(enumerator.Current);
+				}
+			}
+
+			_logger.LogDebug("Search for \"{SearchText}\" in log {LogId} found {NumResults}/{MaxResults} results, took {Time}ms ({@Stats})", text, log.Id, results.Count, count, timer.ElapsedMilliseconds, searchStats);
+			return results;
+		}
+
+		/// <summary>
+		/// Reads a chunk from storage
+		/// </summary>
+		/// <param name="log">Log file to read from</param>
+		/// <param name="chunkIdx">The chunk to read</param>
+		/// <returns>Chunk data</returns>
+		private async Task<LogChunkData> ReadChunkV1Async(LogDocument log, int chunkIdx)
+		{
+			LogChunkDocument chunk = log.Chunks[chunkIdx];
+
+			// Try to read the chunk data from storage
+			LogChunkData? chunkData = null;
+			try
+			{
+				chunkData = await _storage.ReadChunkAsync(log.Id, chunk.Offset, chunk.LineIndex);
+			}
+			catch (Exception ex)
+			{
+				_logger.LogError(ex, "Unable to read log {LogId} at offset {Offset}", log.Id, chunk.Offset);
+			}
+
+			// Get the minimum length and line count for the chunk
+			if (chunkIdx + 1 < log.Chunks.Count)
+			{
+				LogChunkDocument nextChunk = log.Chunks[chunkIdx + 1];
+				chunkData = RepairChunkDataV1(log, chunkIdx, chunkData, (int)(nextChunk.Offset - chunk.Offset), nextChunk.LineIndex - chunk.LineIndex, $"before next");
+			}
+			else
+			{
+				if (log.MaxLineIndex != null && chunk.Length != 0)
+				{
+					chunkData = RepairChunkDataV1(log, chunkIdx, chunkData, chunk.Length, log.MaxLineIndex.Value - chunk.LineIndex, $"last chunk (max line index = {log.MaxLineIndex})");
+				}
+				else
+				{
+					chunkData ??= RepairChunkDataV1(log, chunkIdx, chunkData, 1024, 1, "default");
+				}
+			}
+
+			return chunkData;
+		}
+
+		/// <summary>
+		/// Validates the given chunk data, and fix it up if necessary
+		/// </summary>
+		/// <param name="log">The log file instance</param>
+		/// <param name="chunkIdx">Index of the chunk within the log</param>
+		/// <param name="chunkData">The chunk data that was read</param>
+		/// <param name="length">Expected length of the data</param>
+		/// <param name="lineCount">Expected number of lines in the data</param>
+		/// <param name="context">Context string for diagnostic output</param>
+		/// <returns>Repaired chunk data</returns>
+		LogChunkData RepairChunkDataV1(LogDocument log, int chunkIdx, LogChunkData? chunkData, int length, int lineCount, string context)
+		{
+			int currentLength = 0;
+			int currentLineCount = 0;
+			if (chunkData != null)
+			{
+				currentLength = chunkData.Length;
+				currentLineCount = chunkData.LineCount;
+			}
+
+			if (chunkData == null || currentLength < length || currentLineCount < lineCount)
+			{
+				_logger.LogWarning("Creating placeholder subchunk for log {LogId} chunk {ChunkIdx} (length {Length} vs expected {ExpLength}, lines {LineCount} vs expected {ExpLineCount}, context {Context})", log.Id, chunkIdx, currentLength, length, currentLineCount, lineCount, context);
+
+				List<LogSubChunkData> subChunks = new List<LogSubChunkData>();
+				if (chunkData != null && chunkData.Length < length && chunkData.LineCount < lineCount)
+				{
+					subChunks.AddRange(chunkData.SubChunks);
+				}
+
+				LogText text = new LogText();
+				text.AppendMissingDataInfo(chunkIdx, log.Chunks[chunkIdx].Server, length - currentLength, lineCount - currentLineCount);
+				subChunks.Add(new LogSubChunkData(log.Type, currentLength, currentLineCount, text));
+
+				LogChunkDocument chunk = log.Chunks[chunkIdx];
+				chunkData = new LogChunkData(chunk.Offset, chunk.LineIndex, subChunks);
+			}
+			return chunkData;
+		}
+
+		async IAsyncEnumerable<int> SearchLogDataInternalV1Async(LogDocument log, string text, int firstLine, SearchStats searchStats)
+		{
+			SearchText searchText = new SearchText(text);
+
+			// Read the index for this log file
+			if (log.IndexLength != null)
+			{
+				LogIndexData? indexData = await ReadIndexV1Async(log, log.IndexLength.Value);
+				if (indexData != null && firstLine < indexData.LineCount)
+				{
+					using TelemetrySpan span = _tracer.StartActiveSpan($"{nameof(LogService)}.{nameof(SearchLogDataInternalV1Async)}.Indexed");
+					span.SetAttribute("lineCount", indexData.LineCount);
+
+					foreach (int lineIndex in indexData.Search(firstLine, searchText, searchStats))
+					{
+						yield return lineIndex;
+					}
+
+					firstLine = indexData.LineCount;
+				}
+			}
+
+			// Manually search through the rest of the log
+			int chunkIdx = GetChunkForLineV1(log.Chunks, firstLine);
+			for (; chunkIdx < log.Chunks.Count; chunkIdx++)
+			{
+				LogChunkDocument chunk = log.Chunks[chunkIdx];
+
+				// Read the chunk data
+				LogChunkData chunkData = await ReadChunkV1Async(log, chunkIdx);
+				if (firstLine < chunkData.LineIndex + chunkData.LineCount)
+				{
+					// Find the first sub-chunk we're looking for
+					int subChunkIdx = 0;
+					if (firstLine > chunk.LineIndex)
+					{
+						subChunkIdx = chunkData.GetSubChunkForLine(firstLine - chunk.LineIndex);
+					}
+
+					// Search through the sub-chunks
+					for (; subChunkIdx < chunkData.SubChunks.Count; subChunkIdx++)
+					{
+						LogSubChunkData subChunkData = chunkData.SubChunks[subChunkIdx];
+						if (firstLine < subChunkData.LineIndex + subChunkData.LineCount)
+						{
+							// Create an index containing just this sub-chunk
+							LogIndexData index = subChunkData.BuildIndex(_logger);
+							foreach (int lineIndex in index.Search(firstLine, searchText, searchStats))
+							{
+								yield return lineIndex;
+							}
+						}
+					}
+				}
+			}
+		}
+
+		/// <summary>
+		/// Reads a chunk from storage
+		/// </summary>
+		/// <param name="log">Log file to read from</param>
+		/// <param name="length">Length of the log covered by the index</param>
+		/// <returns>Chunk data</returns>
+		private async Task<LogIndexData?> ReadIndexV1Async(LogDocument log, long length)
+		{
+			try
+			{
+				LogIndexData? index = await _storage.ReadIndexAsync(log.Id, length);
+				return index;
+			}
+			catch (Exception ex)
+			{
+				_logger.LogError(ex, "Unable to read log {LogId} index at length {Length}", log.Id, length);
 				return null;
 			}
-			return new Log(this, log);
 		}
+
+		/// <summary>
+		/// Gets the chunk index containing the given offset.
+		/// </summary>
+		/// <param name="chunks">The chunks to search</param>
+		/// <param name="offset">The offset to search for</param>
+		/// <returns>The chunk index containing the given offset</returns>
+		static int GetChunkForOffsetV1(IReadOnlyList<LogChunkDocument> chunks, long offset)
+		{
+			int chunkIndex = chunks.BinarySearch(x => x.Offset, offset);
+			if (chunkIndex < 0)
+			{
+				chunkIndex = ~chunkIndex - 1;
+			}
+			return chunkIndex;
+		}
+
+		/// <summary>
+		/// Gets the starting chunk index for the given line
+		/// </summary>
+		/// <param name="chunks">The chunks to search</param>
+		/// <param name="lineIndex">Index of the line to query</param>
+		/// <returns>Index of the chunk to fetch</returns>
+		static int GetChunkForLineV1(IReadOnlyList<LogChunkDocument> chunks, int lineIndex)
+		{
+			int chunkIndex = chunks.BinarySearch(x => x.LineIndex, lineIndex);
+			if (chunkIndex < 0)
+			{
+				chunkIndex = ~chunkIndex - 1;
+			}
+			return chunkIndex;
+		}
+
+		/// <summary>
+		/// Streams log data to a caller
+		/// </summary>
+		class ResponseStreamV1 : Stream
+		{
+			/// <summary>
+			/// The log file service that created this stream
+			/// </summary>
+			readonly LogCollection _logCollection;
+
+			/// <summary>
+			/// The log file being read
+			/// </summary>
+			readonly LogDocument _log;
+
+			/// <summary>
+			/// Starting offset within the file of the data to return 
+			/// </summary>
+			readonly long _responseOffset;
+
+			/// <summary>
+			/// Length of data to return
+			/// </summary>
+			readonly long _responseLength;
+
+			/// <summary>
+			/// Current offset within the stream
+			/// </summary>
+			long _currentOffset;
+
+			/// <summary>
+			/// The current chunk index
+			/// </summary>
+			int _chunkIdx;
+
+			/// <summary>
+			/// Buffer containing a message for missing data
+			/// </summary>
+			ReadOnlyMemory<byte> _sourceBuffer;
+
+			/// <summary>
+			/// Offset within the source buffer
+			/// </summary>
+			int _sourcePos;
+
+			/// <summary>
+			/// Length of the source buffer being copied from
+			/// </summary>
+			int _sourceEnd;
+
+			/// <summary>
+			/// Constructor
+			/// </summary>
+			/// <param name="logService">The log file service, for q</param>
+			/// <param name="log"></param>
+			/// <param name="offset"></param>
+			/// <param name="length"></param>
+			public ResponseStreamV1(LogCollection logService, LogDocument log, long offset, long length)
+			{
+				_logCollection = logService;
+				_log = log;
+
+				_responseOffset = offset;
+				_responseLength = length;
+
+				_currentOffset = offset;
+
+				_chunkIdx = GetChunkForOffsetV1(log.Chunks, offset);
+				_sourceBuffer = null!;
+			}
+
+			/// <inheritdoc/>
+			public override bool CanRead => true;
+
+			/// <inheritdoc/>
+			public override bool CanSeek => false;
+
+			/// <inheritdoc/>
+			public override bool CanWrite => false;
+
+			/// <inheritdoc/>
+			public override long Length => _responseLength;
+
+			/// <inheritdoc/>
+			public override long Position
+			{
+				get => _currentOffset - _responseOffset;
+				set => throw new NotImplementedException();
+			}
+
+			/// <inheritdoc/>
+			public override void Flush()
+			{
+			}
+
+			/// <inheritdoc/>
+			public override int Read(byte[] buffer, int offset, int count)
+			{
+#pragma warning disable VSTHRD002
+				return ReadAsync(buffer, offset, count, CancellationToken.None).Result;
+#pragma warning restore VSTHRD002
+			}
+
+			/// <inheritdoc/>
+			public override async Task<int> ReadAsync(byte[] buffer, int offset, int length, CancellationToken cancellationToken)
+			{
+				return await ReadAsync(buffer.AsMemory(offset, length), cancellationToken);
+			}
+
+			/// <inheritdoc/>
+			public override async ValueTask<int> ReadAsync(Memory<byte> buffer, CancellationToken cancellationToken)
+			{
+				int readBytes = 0;
+				while (readBytes < buffer.Length)
+				{
+					if (_sourcePos < _sourceEnd)
+					{
+						// Try to copy from the current buffer
+						int blockSize = Math.Min(_sourceEnd - _sourcePos, buffer.Length - readBytes);
+						_sourceBuffer.Slice(_sourcePos, blockSize).Span.CopyTo(buffer.Slice(readBytes).Span);
+						_currentOffset += blockSize;
+						readBytes += blockSize;
+						_sourcePos += blockSize;
+					}
+					else if (_currentOffset < _responseOffset + _responseLength)
+					{
+						// Move to the right chunk
+						while (_chunkIdx + 1 < _log.Chunks.Count && _currentOffset >= _log.Chunks[_chunkIdx + 1].Offset)
+						{
+							_chunkIdx++;
+						}
+
+						// Get the chunk data
+						LogChunkDocument chunk = _log.Chunks[_chunkIdx];
+						LogChunkData chunkData = await _logCollection.ReadChunkV1Async(_log, _chunkIdx);
+
+						// Figure out which sub-chunk to use
+						int subChunkIdx = chunkData.GetSubChunkForOffsetWithinChunk((int)(_currentOffset - chunk.Offset));
+						LogSubChunkData subChunkData = chunkData.SubChunks[subChunkIdx];
+
+						// Get the source data
+						long subChunkOffset = chunk.Offset + chunkData.SubChunkOffset[subChunkIdx];
+						_sourceBuffer = subChunkData.InflateText().Data;
+						_sourcePos = (int)(_currentOffset - subChunkOffset);
+						_sourceEnd = (int)Math.Min(_sourceBuffer.Length, (_responseOffset + _responseLength) - subChunkOffset);
+					}
+					else
+					{
+						// End of the log
+						break;
+					}
+				}
+				return readBytes;
+			}
+
+			/// <inheritdoc/>
+			public override long Seek(long offset, SeekOrigin origin) => throw new NotImplementedException();
+
+			/// <inheritdoc/>
+			public override void SetLength(long value) => throw new NotImplementedException();
+
+			/// <inheritdoc/>
+			public override void Write(byte[] buffer, int offset, int count) => throw new NotImplementedException();
+		}
+
+		#endregion
 	}
 }
diff --git a/Engine/Source/Programs/Horde/Horde.Server/Logs/LogService.cs b/Engine/Source/Programs/Horde/Horde.Server/Logs/LogService.cs
--- a/Engine/Source/Programs/Horde/Horde.Server/Logs/LogService.cs
+++ b/Engine/Source/Programs/Horde/Horde.Server/Logs/LogService.cs
@@ -2,12 +2,8 @@
 
 using System;
 using System.Collections.Generic;
-using System.Diagnostics;
-using System.IO;
 using System.Linq;
-using System.Runtime.CompilerServices;
 using System.Security.Claims;
-using System.Text;
 using System.Text.Json;
 using System.Threading;
 using System.Threading.Tasks;
@@ -16,12 +12,8 @@ using EpicGames.Horde.Agents.Leases;
 using EpicGames.Horde.Agents.Sessions;
 using EpicGames.Horde.Jobs;
 using EpicGames.Horde.Logs;
-using EpicGames.Horde.Storage;
 using Horde.Server.Acls;
-using Horde.Server.Logs.Data;
 using Horde.Server.Storage;
-using Horde.Server.Utilities;
-using Microsoft.Extensions.Caching.Memory;
 using Microsoft.Extensions.Logging;
 using Microsoft.Extensions.Options;
 using MongoDB.Bson;
@@ -29,24 +21,6 @@ using OpenTelemetry.Trace;
 
 namespace Horde.Server.Logs
 {
-	using Stream = System.IO.Stream;
-
-	/// <summary>
-	/// Metadata about a log file
-	/// </summary>
-	public class LogMetadata
-	{
-		/// <summary>
-		/// Length of the log file
-		/// </summary>
-		public long Length { get; set; }
-
-		/// <summary>
-		/// Number of lines in the log file
-		/// </summary>
-		public int MaxLineIndex { get; set; }
-	}
-
 	/// <summary>
 	/// Interface for the log file service
 	/// </summary>
@@ -72,24 +46,6 @@ namespace Horde.Server.Logs
 		/// <returns>The log  document</returns>
 		Task<ILog?> GetLogAsync(LogId logId, CancellationToken cancellationToken);
 
-		/// <summary>
-		/// Read a set of lines from the given log file
-		/// </summary>
-		/// <param name="log">Log file to read</param>
-		/// <param name="index">Index of the first line to read</param>
-		/// <param name="count">Maximum number of lines to return</param>
-		/// <param name="cancellationToken">Cancellation token for the operation</param>
-		/// <returns>List of lines</returns>
-		Task<List<Utf8String>> ReadLinesAsync(ILog log, int index, int count, CancellationToken cancellationToken = default);
-
-		/// <summary>
-		/// Gets metadata about the log file
-		/// </summary>
-		/// <param name="log">The log file to query</param>
-		/// <param name="cancellationToken">Cancellation token for the call</param>
-		/// <returns>Metadata about the log file</returns>
-		Task<LogMetadata> GetMetadataAsync(ILog log, CancellationToken cancellationToken);
-
 		/// <summary>
 		/// Creates new log events
 		/// </summary>
@@ -138,55 +94,6 @@ namespace Horde.Server.Logs
 		/// <param name="cancellationToken">Cancellation token for the call</param>
 		/// <returns>New event data instance</returns>
 		Task<ILogEventData> GetEventDataAsync(ILog log, int lineIndex, int lineCount, CancellationToken cancellationToken = default);
-
-		/// <summary>
-		/// Gets lines from the given log 
-		/// </summary>
-		/// <param name="log">The log file</param>
-		/// <param name="cancellationToken">Cancellation token for the call</param>
-		/// <returns>Data for the requested range</returns>
-		Task<Stream> OpenRawStreamAsync(ILog log, CancellationToken cancellationToken = default);
-
-		/// <summary>
-		/// Parses a stream of json text and outputs plain text
-		/// </summary>
-		/// <param name="log">The log file to query</param>
-		/// <param name="outputStream">Output stream to receive the text data</param>
-		/// <param name="cancellationToken">Cancellation token for the call</param>
-		/// <returns>Async text</returns>
-		Task CopyPlainTextStreamAsync(ILog log, Stream outputStream, CancellationToken cancellationToken = default);
-
-		/// <summary>
-		/// Search for the specified text in a log file
-		/// </summary>
-		/// <param name="log">The log file to search</param>
-		/// <param name="text">Text to search for</param>
-		/// <param name="firstLine">Line to start search from</param>
-		/// <param name="count">Number of results to return</param>
-		/// <param name="stats">Receives stats for the search</param>
-		/// <param name="cancellationToken">Cancellation token for the call</param>
-		/// <returns>List of line numbers containing the given term</returns>
-		Task<List<int>> SearchLogDataAsync(ILog log, string text, int firstLine, int count, SearchStats stats, CancellationToken cancellationToken);
-	}
-
-	/// <summary>
-	/// Extension methods for dealing with log files
-	/// </summary>
-	public static class LogServiceExtensions
-	{
-		/// <summary>
-		/// Parses a stream of json text and outputs plain text
-		/// </summary>
-		/// <param name="logService">The log file service</param>
-		/// <param name="log">The log file to query</param>
-		/// <param name="outputStream">Output stream to receive the text data</param>
-		/// <param name="cancellationToken">Cancellation token for the call</param>
-		/// <returns>Async text</returns>
-		public static async Task CopyRawStreamAsync(this ILogService logService, ILog log, Stream outputStream, CancellationToken cancellationToken)
-		{
-			await using Stream stream = await logService.OpenRawStreamAsync(log, cancellationToken);
-			await stream.CopyToAsync(outputStream, cancellationToken);
-		}
 	}
 
 	/// <summary>
@@ -201,319 +108,6 @@ namespace Horde.Server.Logs
 		private readonly ILogStorage _storage;
 		private readonly StorageService _storageService;
 		private readonly IOptions<ServerSettings> _settings;
-		private readonly IMemoryCache _logCache;
-
-		/// <summary>
-		/// Streams log data to a caller
-		/// </summary>
-		class ResponseStream : Stream
-		{
-			/// <summary>
-			/// The log file service that created this stream
-			/// </summary>
-			readonly LogService _logService;
-
-			/// <summary>
-			/// The log file being read
-			/// </summary>
-			readonly ILog _log;
-
-			/// <summary>
-			/// Starting offset within the file of the data to return 
-			/// </summary>
-			readonly long _responseOffset;
-
-			/// <summary>
-			/// Length of data to return
-			/// </summary>
-			readonly long _responseLength;
-
-			/// <summary>
-			/// Current offset within the stream
-			/// </summary>
-			long _currentOffset;
-
-			/// <summary>
-			/// The current chunk index
-			/// </summary>
-			int _chunkIdx;
-
-			/// <summary>
-			/// Buffer containing a message for missing data
-			/// </summary>
-			ReadOnlyMemory<byte> _sourceBuffer;
-
-			/// <summary>
-			/// Offset within the source buffer
-			/// </summary>
-			int _sourcePos;
-
-			/// <summary>
-			/// Length of the source buffer being copied from
-			/// </summary>
-			int _sourceEnd;
-
-			/// <summary>
-			/// Constructor
-			/// </summary>
-			/// <param name="logService">The log file service, for q</param>
-			/// <param name="log"></param>
-			/// <param name="offset"></param>
-			/// <param name="length"></param>
-			public ResponseStream(LogService logService, ILog log, long offset, long length)
-			{
-				_logService = logService;
-				_log = log;
-
-				_responseOffset = offset;
-				_responseLength = length;
-
-				_currentOffset = offset;
-
-				_chunkIdx = log.Chunks.GetChunkForOffset(offset);
-				_sourceBuffer = null!;
-			}
-
-			/// <inheritdoc/>
-			public override bool CanRead => true;
-
-			/// <inheritdoc/>
-			public override bool CanSeek => false;
-
-			/// <inheritdoc/>
-			public override bool CanWrite => false;
-
-			/// <inheritdoc/>
-			public override long Length => _responseLength;
-
-			/// <inheritdoc/>
-			public override long Position
-			{
-				get => _currentOffset - _responseOffset;
-				set => throw new NotImplementedException();
-			}
-
-			/// <inheritdoc/>
-			public override void Flush()
-			{
-			}
-
-			/// <inheritdoc/>
-			public override int Read(byte[] buffer, int offset, int count)
-			{
-#pragma warning disable VSTHRD002
-				return ReadAsync(buffer, offset, count, CancellationToken.None).Result;
-#pragma warning restore VSTHRD002
-			}
-
-			/// <inheritdoc/>
-			public override async Task<int> ReadAsync(byte[] buffer, int offset, int length, CancellationToken cancellationToken)
-			{
-				return await ReadAsync(buffer.AsMemory(offset, length), cancellationToken);
-			}
-
-			/// <inheritdoc/>
-			public override async ValueTask<int> ReadAsync(Memory<byte> buffer, CancellationToken cancellationToken)
-			{
-				int readBytes = 0;
-				while (readBytes < buffer.Length)
-				{
-					if (_sourcePos < _sourceEnd)
-					{
-						// Try to copy from the current buffer
-						int blockSize = Math.Min(_sourceEnd - _sourcePos, buffer.Length - readBytes);
-						_sourceBuffer.Slice(_sourcePos, blockSize).Span.CopyTo(buffer.Slice(readBytes).Span);
-						_currentOffset += blockSize;
-						readBytes += blockSize;
-						_sourcePos += blockSize;
-					}
-					else if (_currentOffset < _responseOffset + _responseLength)
-					{
-						// Move to the right chunk
-						while (_chunkIdx + 1 < _log.Chunks.Count && _currentOffset >= _log.Chunks[_chunkIdx + 1].Offset)
-						{
-							_chunkIdx++;
-						}
-
-						// Get the chunk data
-						ILogChunk chunk = _log.Chunks[_chunkIdx];
-						LogChunkData chunkData = await _logService.ReadChunkAsync(_log, _chunkIdx);
-
-						// Figure out which sub-chunk to use
-						int subChunkIdx = chunkData.GetSubChunkForOffsetWithinChunk((int)(_currentOffset - chunk.Offset));
-						LogSubChunkData subChunkData = chunkData.SubChunks[subChunkIdx];
-
-						// Get the source data
-						long subChunkOffset = chunk.Offset + chunkData.SubChunkOffset[subChunkIdx];
-						_sourceBuffer = subChunkData.InflateText().Data;
-						_sourcePos = (int)(_currentOffset - subChunkOffset);
-						_sourceEnd = (int)Math.Min(_sourceBuffer.Length, (_responseOffset + _responseLength) - subChunkOffset);
-					}
-					else
-					{
-						// End of the log
-						break;
-					}
-				}
-				return readBytes;
-			}
-
-			/// <inheritdoc/>
-			public override long Seek(long offset, SeekOrigin origin) => throw new NotImplementedException();
-
-			/// <inheritdoc/>
-			public override void SetLength(long value) => throw new NotImplementedException();
-
-			/// <inheritdoc/>
-			public override void Write(byte[] buffer, int offset, int count) => throw new NotImplementedException();
-		}
-
-		/// <summary>
-		/// Streams log data to a caller
-		/// </summary>
-		class NewLoggerResponseStream : Stream
-		{
-			readonly LogNode _rootNode;
-
-			/// <summary>
-			/// Starting offset within the file of the data to return 
-			/// </summary>
-			readonly long _responseOffset;
-
-			/// <summary>
-			/// Length of data to return
-			/// </summary>
-			readonly long _responseLength;
-
-			/// <summary>
-			/// Current offset within the stream
-			/// </summary>
-			long _currentOffset;
-
-			/// <summary>
-			/// The current chunk index
-			/// </summary>
-			int _chunkIdx;
-
-			/// <summary>
-			/// Buffer containing a message for missing data
-			/// </summary>
-			ReadOnlyMemory<byte> _sourceBuffer;
-
-			/// <summary>
-			/// Offset within the source buffer
-			/// </summary>
-			int _sourcePos;
-
-			/// <summary>
-			/// Length of the source buffer being copied from
-			/// </summary>
-			int _sourceEnd;
-
-			/// <summary>
-			/// Constructor
-			/// </summary>
-			public NewLoggerResponseStream(LogNode rootNode, long offset, long length)
-			{
-				_rootNode = rootNode;
-
-				_responseOffset = offset;
-				_responseLength = length;
-
-				_currentOffset = offset;
-
-				_chunkIdx = rootNode.TextChunkRefs.GetChunkForOffset(offset);
-				_sourceBuffer = null!;
-			}
-
-			/// <inheritdoc/>
-			public override bool CanRead => true;
-
-			/// <inheritdoc/>
-			public override bool CanSeek => false;
-
-			/// <inheritdoc/>
-			public override bool CanWrite => false;
-
-			/// <inheritdoc/>
-			public override long Length => _responseLength;
-
-			/// <inheritdoc/>
-			public override long Position
-			{
-				get => _currentOffset - _responseOffset;
-				set => throw new NotImplementedException();
-			}
-
-			/// <inheritdoc/>
-			public override void Flush()
-			{
-			}
-
-			/// <inheritdoc/>
-			public override int Read(byte[] buffer, int offset, int count)
-			{
-#pragma warning disable VSTHRD002
-				return ReadAsync(buffer, offset, count, CancellationToken.None).Result;
-#pragma warning restore VSTHRD002
-			}
-
-			/// <inheritdoc/>
-			public override async Task<int> ReadAsync(byte[] buffer, int offset, int length, CancellationToken cancellationToken)
-			{
-				return await ReadAsync(buffer.AsMemory(offset, length), cancellationToken);
-			}
-
-			/// <inheritdoc/>
-			public override async ValueTask<int> ReadAsync(Memory<byte> buffer, CancellationToken cancellationToken)
-			{
-				int readBytes = 0;
-				while (readBytes < buffer.Length)
-				{
-					if (_sourcePos < _sourceEnd)
-					{
-						// Try to copy from the current buffer
-						int blockSize = Math.Min(_sourceEnd - _sourcePos, buffer.Length - readBytes);
-						_sourceBuffer.Slice(_sourcePos, blockSize).Span.CopyTo(buffer.Slice(readBytes).Span);
-						_currentOffset += blockSize;
-						readBytes += blockSize;
-						_sourcePos += blockSize;
-					}
-					else if (_currentOffset < _responseOffset + _responseLength)
-					{
-						// Move to the right chunk
-						while (_chunkIdx + 1 < _rootNode.TextChunkRefs.Count && _currentOffset >= _rootNode.TextChunkRefs[_chunkIdx + 1].Offset)
-						{
-							_chunkIdx++;
-						}
-
-						// Get the chunk data
-						LogChunkRef chunk = _rootNode.TextChunkRefs[_chunkIdx];
-						LogChunkNode chunkNode = await chunk.Target.ReadBlobAsync(cancellationToken: cancellationToken);
-
-						// Get the source data
-						_sourceBuffer = chunkNode.Data;
-						_sourcePos = (int)(_currentOffset - chunk.Offset);
-						_sourceEnd = (int)Math.Min(_sourceBuffer.Length, (_responseOffset + _responseLength) - chunk.Offset);
-					}
-					else
-					{
-						// End of the log
-						break;
-					}
-				}
-				return readBytes;
-			}
-
-			/// <inheritdoc/>
-			public override long Seek(long offset, SeekOrigin origin) => throw new NotImplementedException();
-
-			/// <inheritdoc/>
-			public override void SetLength(long value) => throw new NotImplementedException();
-
-			/// <inheritdoc/>
-			public override void Write(byte[] buffer, int offset, int count) => throw new NotImplementedException();
-		}
 
 		readonly LogTailService _logTailService;
 
@@ -524,7 +118,6 @@ namespace Horde.Server.Logs
 		{
 			_logCollection = logCollection;
 			_logEventCollection = logEventCollection;
-			_logCache = new MemoryCache(new MemoryCacheOptions());
 			_storage = storage;
 			_logTailService = logTailService;
 			_storageService = storageService;
@@ -536,125 +129,19 @@ namespace Horde.Server.Logs
 		/// <inheritdoc/>
 		public void Dispose()
 		{
-			_logCache.Dispose();
 			_storage.Dispose();
 		}
 
 		/// <inheritdoc/>
 		public Task<ILog> CreateLogAsync(JobId jobId, LeaseId? leaseId, SessionId? sessionId, LogType type, LogId? logId, CancellationToken cancellationToken)
 		{
-			return _logCollection.CreateLogAsync(jobId, leaseId, sessionId, type, logId, cancellationToken);
+			return _logCollection.AddAsync(jobId, leaseId, sessionId, type, logId, cancellationToken);
 		}
 
 		/// <inheritdoc/>
 		public async Task<ILog?> GetLogAsync(LogId logId, CancellationToken cancellationToken)
 		{
-			return await _logCollection.GetLogAsync(logId, cancellationToken);
-		}
-
-		/// <inheritdoc/>
-		public async Task<List<Utf8String>> ReadLinesAsync(ILog log, int index, int count, CancellationToken cancellationToken)
-		{
-			List<Utf8String> lines = new List<Utf8String>();
-
-			if (log.UseNewStorageBackend)
-			{
-				using IStorageClient storageClient = _storageService.CreateClient(log.NamespaceId);
-
-				int maxIndex = index + count;
-				bool complete = log.Complete;
-
-				LogNode? root = await storageClient.TryReadRefTargetAsync<LogNode>(log.RefName, cancellationToken: cancellationToken);
-				if (root != null)
-				{
-					int chunkIdx = root.TextChunkRefs.GetChunkForLine(index);
-					for (; index < maxIndex && chunkIdx < root.TextChunkRefs.Count; chunkIdx++)
-					{
-						LogChunkRef chunk = root.TextChunkRefs[chunkIdx];
-						LogChunkNode chunkData = await chunk.Target.ReadBlobAsync(cancellationToken: cancellationToken);
-
-						for (; index < maxIndex && index < chunk.LineIndex; index++)
-						{
-							lines.Add(new Utf8String($"Internal error; missing data for line {index}\n"));
-						}
-
-						for (; index < maxIndex && index < chunk.LineIndex + chunk.LineCount; index++)
-						{
-							lines.Add(chunkData.GetLine(index - chunk.LineIndex));
-						}
-					}
-					complete |= root.Complete;
-				}
-
-				if (!complete)
-				{
-					await _logTailService.EnableTailingAsync(log.Id, root?.LineCount ?? 0, cancellationToken);
-					if (index < maxIndex)
-					{
-						await _logTailService.ReadAsync(log.Id, index, maxIndex - index, lines);
-					}
-				}
-			}
-			else
-			{
-				(_, long minOffset) = await GetLineOffsetAsync(log, index, cancellationToken);
-				(_, long maxOffset) = await GetLineOffsetAsync(log, index + Math.Min(count, Int32.MaxValue - index), cancellationToken);
-
-				byte[] result;
-				using (System.IO.Stream stream = await OpenRawStreamAsync(log, minOffset, maxOffset - minOffset, cancellationToken))
-				{
-					result = new byte[stream.Length];
-					await stream.ReadFixedSizeDataAsync(result, 0, result.Length, cancellationToken);
-				}
-
-				int offset = 0;
-				for (int idx = 0; idx < result.Length; idx++)
-				{
-					if (result[idx] == (byte)'\n')
-					{
-						lines.Add(new Utf8String(result.AsMemory(offset, idx - offset)));
-						offset = idx + 1;
-					}
-				}
-			}
-
-			return lines;
-		}
-
-		/// <inheritdoc/>
-		public async Task<LogMetadata> GetMetadataAsync(ILog log, CancellationToken cancellationToken)
-		{
-			LogMetadata metadata = new LogMetadata();
-			if (log.UseNewStorageBackend)
-			{
-				if (log.Complete)
-				{
-					metadata.MaxLineIndex = log.LineCount;
-				}
-				else
-				{
-					metadata.MaxLineIndex = await _logTailService.GetFullLineCountAsync(log.Id, log.LineCount, cancellationToken);
-				}
-			}
-			else
-			{
-				if (log.Chunks.Count > 0)
-				{
-					ILogChunk chunk = log.Chunks[log.Chunks.Count - 1];
-					if (log.MaxLineIndex == null || chunk.Length == 0)
-					{
-						LogChunkData chunkData = await ReadChunkAsync(log, log.Chunks.Count - 1);
-						metadata.Length = chunk.Offset + chunkData.Length;
-						metadata.MaxLineIndex = chunk.LineIndex + chunkData.LineCount;
-					}
-					else
-					{
-						metadata.Length = chunk.Offset + chunk.Length;
-						metadata.MaxLineIndex = log.MaxLineIndex.Value;
-					}
-				}
-			}
-			return metadata;
+			return await _logCollection.GetAsync(logId, cancellationToken);
 		}
 
 		/// <inheritdoc/>
@@ -712,7 +199,7 @@ namespace Horde.Server.Logs
 			span.SetAttribute("lineIndex", lineIndex);
 			span.SetAttribute("lineCount", lineCount);
 
-			List<Utf8String> lines = await ReadLinesAsync(log, lineIndex, lineCount, cancellationToken);
+			List<Utf8String> lines = await log.ReadLinesAsync(lineIndex, lineCount, cancellationToken);
 			List<JsonLogEvent> jsonLines = new List<JsonLogEvent>(lines.Count);
 
 			foreach (Utf8String line in lines)
@@ -730,312 +217,6 @@ namespace Horde.Server.Logs
 			return new LogEventData(jsonLines);
 		}
 
-		/// <inheritdoc/>
-		public Task<Stream> OpenRawStreamAsync(ILog log, CancellationToken cancellationToken)
-		{
-			return OpenRawStreamAsync(log, 0, Int64.MaxValue, cancellationToken);
-		}
-
-		/// <inheritdoc/>
-		public async Task<Stream> OpenRawStreamAsync(ILog log, long offset, long length, CancellationToken cancellationToken)
-		{
-			if (log.UseNewStorageBackend)
-			{
-				using IStorageClient storageClient = _storageService.CreateClient(log.NamespaceId);
-
-				LogNode? root = await storageClient.TryReadRefTargetAsync<LogNode>(log.RefName, cancellationToken: cancellationToken);
-				if (root == null || root.TextChunkRefs.Count == 0)
-				{
-					return new MemoryStream(Array.Empty<byte>(), false);
-				}
-				else
-				{
-					int lastChunkIdx = root.TextChunkRefs.Count - 1;
-
-					// Clamp the length of the request
-					LogChunkRef lastChunk = root.TextChunkRefs[lastChunkIdx];
-					if (length > lastChunk.Offset)
-					{
-						long lastChunkLength = lastChunk.Length;
-						if (lastChunkLength <= 0)
-						{
-							LogChunkNode lastChunkNode = await lastChunk.Target.ReadBlobAsync(cancellationToken: cancellationToken);
-							lastChunkLength = lastChunkNode.Length;
-						}
-						length = Math.Min(length, (lastChunk.Offset + lastChunkLength) - offset);
-					}
-
-					// Create the new stream
-					return new NewLoggerResponseStream(root, offset, length);
-				}
-			}
-			else
-			{
-				if (log.Chunks.Count == 0)
-				{
-					return new MemoryStream(Array.Empty<byte>(), false);
-				}
-				else
-				{
-					int lastChunkIdx = log.Chunks.Count - 1;
-
-					// Clamp the length of the request
-					ILogChunk lastChunk = log.Chunks[lastChunkIdx];
-					if (length > lastChunk.Offset)
-					{
-						long lastChunkLength = lastChunk.Length;
-						if (lastChunkLength <= 0)
-						{
-							LogChunkData lastChunkData = await ReadChunkAsync(log, lastChunkIdx);
-							lastChunkLength = lastChunkData.Length;
-						}
-						length = Math.Min(length, (lastChunk.Offset + lastChunkLength) - offset);
-					}
-
-					// Create the new stream
-					return new ResponseStream(this, log, offset, length);
-				}
-			}
-		}
-
-		/// <summary>
-		/// Helper method for catching exceptions in <see cref="LogText.ConvertToPlainText(ReadOnlySpan{Byte}, Byte[], Int32)"/>
-		/// </summary>
-		public static int GuardedConvertToPlainText(ReadOnlySpan<byte> input, byte[] output, int outputOffset, ILogger logger)
-		{
-			try
-			{
-				return LogText.ConvertToPlainText(input, output, outputOffset);
-			}
-			catch (Exception ex)
-			{
-				logger.LogWarning(ex, "Unable to convert log line to plain text: {Line}", Encoding.UTF8.GetString(input));
-				output[outputOffset] = (byte)'\n';
-				return outputOffset + 1;
-			}
-		}
-
-		/// <inheritdoc/>
-		public async Task CopyPlainTextStreamAsync(ILog log, Stream outputStream, CancellationToken cancellationToken)
-		{
-			long offset = 0;
-			long length = Int64.MaxValue;
-
-			using (Stream stream = await OpenRawStreamAsync(log, 0, Int64.MaxValue, cancellationToken))
-			{
-				byte[] readBuffer = new byte[4096];
-				int readBufferLength = 0;
-
-				byte[] writeBuffer = new byte[4096];
-				int writeBufferLength = 0;
-
-				while (length > 0)
-				{
-					// Add more data to the buffer
-					int readBytes = await stream.ReadAsync(readBuffer.AsMemory(readBufferLength, readBuffer.Length - readBufferLength), cancellationToken);
-					readBufferLength += readBytes;
-
-					// Copy as many lines as possible to the output
-					int convertedBytes = 0;
-					for (int endIdx = 1; endIdx < readBufferLength; endIdx++)
-					{
-						if (readBuffer[endIdx] == '\n')
-						{
-							writeBufferLength = GuardedConvertToPlainText(readBuffer.AsSpan(convertedBytes, endIdx - convertedBytes), writeBuffer, writeBufferLength, _logger);
-							convertedBytes = endIdx + 1;
-						}
-					}
-
-					// If there's anything in the write buffer, write it out
-					if (writeBufferLength > 0)
-					{
-						if (offset < writeBufferLength)
-						{
-							int writeLength = (int)Math.Min((long)writeBufferLength - offset, length);
-							await outputStream.WriteAsync(writeBuffer.AsMemory((int)offset, writeLength), cancellationToken);
-							length -= writeLength;
-						}
-						offset = Math.Max(offset - writeBufferLength, 0);
-						writeBufferLength = 0;
-					}
-
-					// If we were able to read something, shuffle down the rest of the buffer. Otherwise expand the read buffer.
-					if (convertedBytes > 0)
-					{
-						Buffer.BlockCopy(readBuffer, convertedBytes, readBuffer, 0, readBufferLength - convertedBytes);
-						readBufferLength -= convertedBytes;
-					}
-					else if (readBufferLength > 0)
-					{
-						Array.Resize(ref readBuffer, readBuffer.Length + 128);
-						writeBuffer = new byte[readBuffer.Length];
-					}
-
-					// Exit if we didn't read anything in this iteration
-					if (readBytes == 0)
-					{
-						break;
-					}
-				}
-			}
-		}
-
-		/// <inheritdoc/>
-		public async Task<(int, long)> GetLineOffsetAsync(ILog log, int lineIdx, CancellationToken cancellationToken)
-		{
-			if (log.UseNewStorageBackend)
-			{
-				using IStorageClient storageClient = _storageService.CreateClient(log.NamespaceId);
-
-				LogNode? root = await storageClient.TryReadRefTargetAsync<LogNode>(log.RefName, cancellationToken: cancellationToken);
-				if (root == null)
-				{
-					return (0, 0);
-				}
-
-				int chunkIdx = root.TextChunkRefs.GetChunkForLine(lineIdx);
-				LogChunkRef chunk = root.TextChunkRefs[chunkIdx];
-				LogChunkNode chunkData = await chunk.Target.ReadBlobAsync(cancellationToken: cancellationToken);
-
-				if (lineIdx < chunk.LineIndex)
-				{
-					lineIdx = chunk.LineIndex;
-				}
-
-				int maxLineIndex = chunk.LineIndex + chunkData.LineCount;
-				if (lineIdx >= maxLineIndex)
-				{
-					lineIdx = maxLineIndex;
-				}
-
-				long offset = chunk.Offset + chunkData.LineOffsets[lineIdx - chunk.LineIndex];
-				return (lineIdx, offset);
-			}
-			else
-			{
-				int chunkIdx = log.Chunks.GetChunkForLine(lineIdx);
-
-				ILogChunk chunk = log.Chunks[chunkIdx];
-				LogChunkData chunkData = await ReadChunkAsync(log, chunkIdx);
-
-				if (lineIdx < chunk.LineIndex)
-				{
-					lineIdx = chunk.LineIndex;
-				}
-
-				int maxLineIndex = chunk.LineIndex + chunkData.LineCount;
-				if (lineIdx >= maxLineIndex)
-				{
-					lineIdx = maxLineIndex;
-				}
-
-				long offset = chunk.Offset + chunkData.GetLineOffsetWithinChunk(lineIdx - chunk.LineIndex);
-				return (lineIdx, offset);
-			}
-		}
-
-		/// <summary>
-		/// Reads a chunk from storage
-		/// </summary>
-		/// <param name="log">Log file to read from</param>
-		/// <param name="chunkIdx">The chunk to read</param>
-		/// <returns>Chunk data</returns>
-		private async Task<LogChunkData> ReadChunkAsync(ILog log, int chunkIdx)
-		{
-			ILogChunk chunk = log.Chunks[chunkIdx];
-
-			// Try to read the chunk data from storage
-			LogChunkData? chunkData = null;
-			try
-			{
-				chunkData = await _storage.ReadChunkAsync(log.Id, chunk.Offset, chunk.LineIndex);
-			}
-			catch (Exception ex)
-			{
-				_logger.LogError(ex, "Unable to read log {LogId} at offset {Offset}", log.Id, chunk.Offset);
-			}
-
-			// Get the minimum length and line count for the chunk
-			if (chunkIdx + 1 < log.Chunks.Count)
-			{
-				ILogChunk nextChunk = log.Chunks[chunkIdx + 1];
-				chunkData = RepairChunkData(log, chunkIdx, chunkData, (int)(nextChunk.Offset - chunk.Offset), nextChunk.LineIndex - chunk.LineIndex, $"before next");
-			}
-			else
-			{
-				if (log.MaxLineIndex != null && chunk.Length != 0)
-				{
-					chunkData = RepairChunkData(log, chunkIdx, chunkData, chunk.Length, log.MaxLineIndex.Value - chunk.LineIndex, $"last chunk (max line index = {log.MaxLineIndex})");
-				}
-				else
-				{
-					chunkData ??= RepairChunkData(log, chunkIdx, chunkData, 1024, 1, "default");
-				}
-			}
-
-			return chunkData;
-		}
-
-		/// <summary>
-		/// Validates the given chunk data, and fix it up if necessary
-		/// </summary>
-		/// <param name="log">The log file instance</param>
-		/// <param name="chunkIdx">Index of the chunk within the log</param>
-		/// <param name="chunkData">The chunk data that was read</param>
-		/// <param name="length">Expected length of the data</param>
-		/// <param name="lineCount">Expected number of lines in the data</param>
-		/// <param name="context">Context string for diagnostic output</param>
-		/// <returns>Repaired chunk data</returns>
-		LogChunkData RepairChunkData(ILog log, int chunkIdx, LogChunkData? chunkData, int length, int lineCount, string context)
-		{
-			int currentLength = 0;
-			int currentLineCount = 0;
-			if (chunkData != null)
-			{
-				currentLength = chunkData.Length;
-				currentLineCount = chunkData.LineCount;
-			}
-
-			if (chunkData == null || currentLength < length || currentLineCount < lineCount)
-			{
-				_logger.LogWarning("Creating placeholder subchunk for log {LogId} chunk {ChunkIdx} (length {Length} vs expected {ExpLength}, lines {LineCount} vs expected {ExpLineCount}, context {Context})", log.Id, chunkIdx, currentLength, length, currentLineCount, lineCount, context);
-
-				List<LogSubChunkData> subChunks = new List<LogSubChunkData>();
-				if (chunkData != null && chunkData.Length < length && chunkData.LineCount < lineCount)
-				{
-					subChunks.AddRange(chunkData.SubChunks);
-				}
-
-				LogText text = new LogText();
-				text.AppendMissingDataInfo(chunkIdx, log.Chunks[chunkIdx].Server, length - currentLength, lineCount - currentLineCount);
-				subChunks.Add(new LogSubChunkData(log.Type, currentLength, currentLineCount, text));
-
-				ILogChunk chunk = log.Chunks[chunkIdx];
-				chunkData = new LogChunkData(chunk.Offset, chunk.LineIndex, subChunks);
-			}
-			return chunkData;
-		}
-
-		/// <summary>
-		/// Reads a chunk from storage
-		/// </summary>
-		/// <param name="log">Log file to read from</param>
-		/// <param name="length">Length of the log covered by the index</param>
-		/// <returns>Chunk data</returns>
-		private async Task<LogIndexData?> ReadIndexAsync(ILog log, long length)
-		{
-			try
-			{
-				LogIndexData? index = await _storage.ReadIndexAsync(log.Id, length);
-				return index;
-			}
-			catch (Exception ex)
-			{
-				_logger.LogError(ex, "Unable to read log {LogId} index at length {Length}", log.Id, length);
-				return null;
-			}
-		}
-
 		/// <summary>
 		/// Determines if the user is authorized to perform an action on a particular template
 		/// </summary>
@@ -1054,181 +235,6 @@ namespace Horde.Server.Logs
 			}
 			return false;
 		}
-
-		/// <inheritdoc/>
-		public async Task<List<int>> SearchLogDataAsync(ILog log, string text, int firstLine, int count, SearchStats searchStats, CancellationToken cancellationToken)
-		{
-			Stopwatch timer = Stopwatch.StartNew();
-
-			using TelemetrySpan span = _tracer.StartActiveSpan($"{nameof(LogService)}.{nameof(SearchLogDataAsync)}");
-			span.SetAttribute("logId", log.Id.ToString());
-			span.SetAttribute("text", text);
-			span.SetAttribute("count", count);
-
-			List<int> results = new List<int>();
-			if (count > 0)
-			{
-				IAsyncEnumerable<int> enumerable = (log.UseNewStorageBackend) ?
-						SearchLogDataInternalNewAsync(log, text, firstLine, searchStats, cancellationToken) :
-						SearchLogDataInternalAsync(log, text, firstLine, searchStats);
-
-				await using IAsyncEnumerator<int> enumerator = enumerable.GetAsyncEnumerator(cancellationToken);
-				while (await enumerator.MoveNextAsync() && results.Count < count)
-				{
-					results.Add(enumerator.Current);
-				}
-			}
-
-			_logger.LogDebug("Search for \"{SearchText}\" in log {LogId} found {NumResults}/{MaxResults} results, took {Time}ms ({@Stats})", text, log.Id, results.Count, count, timer.ElapsedMilliseconds, searchStats);
-			return results;
-		}
-
-		async IAsyncEnumerable<int> SearchLogDataInternalNewAsync(ILog log, string text, int firstLine, SearchStats searchStats, [EnumeratorCancellation] CancellationToken cancellationToken)
-		{
-			SearchTerm searchText = new SearchTerm(text);
-			using IStorageClient storageClient = _storageService.CreateClient(log.NamespaceId);
-
-			// Search the index
-			if (log.LineCount > 0)
-			{
-				LogNode? root = await storageClient.ReadRefTargetAsync<LogNode>(log.RefName, cancellationToken: cancellationToken);
-				if (root != null)
-				{
-					LogIndexNode index = await root.IndexRef.ReadBlobAsync(cancellationToken: cancellationToken);
-					await foreach (int lineIdx in index.SearchAsync(firstLine, searchText, searchStats, cancellationToken: cancellationToken))
-					{
-						yield return lineIdx;
-					}
-					if (root.Complete)
-					{
-						yield break;
-					}
-					firstLine = root.LineCount;
-				}
-			}
-
-			// Search any tail data we have
-			if (!log.Complete)
-			{
-				for (; ; )
-				{
-					Utf8String[] lines = await ReadTailAsync(log, firstLine, cancellationToken);
-					if (lines.Length == 0)
-					{
-						break;
-					}
-
-					for (int idx = 0; idx < lines.Length; idx++)
-					{
-						if (SearchTerm.FindNextOcurrence(lines[idx].Span, 0, searchText) != -1)
-						{
-							yield return firstLine + idx;
-						}
-					}
-
-					firstLine += lines.Length;
-				}
-			}
-		}
-
-		async Task<Utf8String[]> ReadTailAsync(ILog log, int index, CancellationToken cancellationToken)
-		{
-			_ = cancellationToken;
-
-			const int BatchSize = 128;
-
-			if (log.Complete)
-			{
-				return Array.Empty<Utf8String>();
-			}
-
-			string cacheKey = $"{log.Id}@{index}";
-			if (_logCache.TryGetValue(cacheKey, out Utf8String[]? lines))
-			{
-				return lines!;
-			}
-
-			lines = (await _logTailService.ReadAsync(log.Id, index, BatchSize)).ToArray();
-			if (log.Type == LogType.Json)
-			{
-				LogChunkBuilder builder = new LogChunkBuilder(lines.Sum(x => x.Length));
-				foreach (Utf8String line in lines)
-				{
-					builder.AppendJsonAsPlainText(line.Span, _logger);
-				}
-				lines = lines.ToArray();
-			}
-
-			if (lines.Length == BatchSize)
-			{
-				int length = lines.Sum(x => x.Length);
-				using (ICacheEntry entry = _logCache.CreateEntry(cacheKey))
-				{
-					entry.SetSlidingExpiration(TimeSpan.FromMinutes(1.0));
-					entry.SetSize(length);
-					entry.SetValue(lines);
-				}
-			}
-
-			return lines;
-		}
-
-		async IAsyncEnumerable<int> SearchLogDataInternalAsync(ILog log, string text, int firstLine, SearchStats searchStats)
-		{
-			SearchText searchText = new SearchText(text);
-
-			// Read the index for this log file
-			if (log.IndexLength != null)
-			{
-				LogIndexData? indexData = await ReadIndexAsync(log, log.IndexLength.Value);
-				if (indexData != null && firstLine < indexData.LineCount)
-				{
-					using TelemetrySpan span = _tracer.StartActiveSpan($"{nameof(LogService)}.{nameof(SearchLogDataInternalAsync)}.Indexed");
-					span.SetAttribute("lineCount", indexData.LineCount);
-
-					foreach (int lineIndex in indexData.Search(firstLine, searchText, searchStats))
-					{
-						yield return lineIndex;
-					}
-
-					firstLine = indexData.LineCount;
-				}
-			}
-
-			// Manually search through the rest of the log
-			int chunkIdx = log.Chunks.GetChunkForLine(firstLine);
-			for (; chunkIdx < log.Chunks.Count; chunkIdx++)
-			{
-				ILogChunk chunk = log.Chunks[chunkIdx];
-
-				// Read the chunk data
-				LogChunkData chunkData = await ReadChunkAsync(log, chunkIdx);
-				if (firstLine < chunkData.LineIndex + chunkData.LineCount)
-				{
-					// Find the first sub-chunk we're looking for
-					int subChunkIdx = 0;
-					if (firstLine > chunk.LineIndex)
-					{
-						subChunkIdx = chunkData.GetSubChunkForLine(firstLine - chunk.LineIndex);
-					}
-
-					// Search through the sub-chunks
-					for (; subChunkIdx < chunkData.SubChunks.Count; subChunkIdx++)
-					{
-						LogSubChunkData subChunkData = chunkData.SubChunks[subChunkIdx];
-						if (firstLine < subChunkData.LineIndex + subChunkData.LineCount)
-						{
-							// Create an index containing just this sub-chunk
-							LogIndexData index = subChunkData.BuildIndex(_logger);
-							foreach (int lineIndex in index.Search(firstLine, searchText, searchStats))
-							{
-								yield return lineIndex;
-							}
-						}
-					}
-				}
-			}
-		}
 	}
 }
 
diff --git a/Engine/Source/Programs/Horde/Horde.Server/Logs/LogsController.cs b/Engine/Source/Programs/Horde/Horde.Server/Logs/LogsController.cs
--- a/Engine/Source/Programs/Horde/Horde.Server/Logs/LogsController.cs
+++ b/Engine/Source/Programs/Horde/Horde.Server/Logs/LogsController.cs
@@ -92,7 +92,7 @@ namespace Horde.Server.Logs
 				return Forbid();
 			}
 
-			LogMetadata metadata = await _logService.GetMetadataAsync(log, cancellationToken);
+			LogMetadata metadata = await log.GetMetadataAsync(cancellationToken);
 			return CreateGetLogResponse(log, metadata).ApplyFilter(filter);
 		}
 
@@ -136,7 +136,7 @@ namespace Horde.Server.Logs
 
 			request.Prefix = $"{log.RefName}";
 
-			IStorageBackend storageBackend = _storageService.CreateBackend(Namespace.Logs);
+			IStorageBackend storageBackend = _storageService.CreateBackend(log.NamespaceId);
 			return await StorageController.WriteBlobAsync(storageBackend, request, cancellationToken);
 		}
 
@@ -171,11 +171,11 @@ namespace Horde.Server.Logs
 			Func<Stream, ActionContext, Task> copyTask;
 			if (format == LogOutputFormat.Text && log.Type == LogType.Json)
 			{
-				copyTask = (outputStream, context) => _logService.CopyPlainTextStreamAsync(log, outputStream, cancellationToken);
+				copyTask = (outputStream, context) => log.CopyPlainTextStreamAsync(outputStream, cancellationToken);
 			}
 			else
 			{
-				copyTask = (outputStream, context) => _logService.CopyRawStreamAsync(log, outputStream, cancellationToken);
+				copyTask = (outputStream, context) => log.CopyRawStreamAsync(outputStream, cancellationToken);
 			}
 
 			return new CustomFileCallbackResult(fileName ?? $"log-{logId}.txt", "text/plain", !download, copyTask);
@@ -203,9 +203,9 @@ namespace Horde.Server.Logs
 				return Forbid();
 			}
 
-			LogMetadata metadata = await _logService.GetMetadataAsync(log, cancellationToken);
+			LogMetadata metadata = await log.GetMetadataAsync(cancellationToken);
 
-			List<Utf8String> lines = await _logService.ReadLinesAsync(log, index, count, cancellationToken);
+			List<Utf8String> lines = await log.ReadLinesAsync(index, count, cancellationToken);
 			using (MemoryStream stream = new MemoryStream(lines.Sum(x => x.Length) + (lines.Count * 20)))
 			{
 				stream.WriteByte((byte)'{');
@@ -300,7 +300,7 @@ namespace Horde.Server.Logs
 
 			SearchLogResponse response = new SearchLogResponse();
 			response.Stats = new SearchStats();
-			response.Lines = await _logService.SearchLogDataAsync(log, text, firstLine, count, response.Stats, cancellationToken);
+			response.Lines = await log.SearchLogDataAsync(text, firstLine, count, response.Stats, cancellationToken);
 			return response;
 		}
 
diff --git a/Engine/Source/Programs/Horde/Horde.Server/Server/DebugController.cs b/Engine/Source/Programs/Horde/Horde.Server/Server/DebugController.cs
--- a/Engine/Source/Programs/Horde/Horde.Server/Server/DebugController.cs
+++ b/Engine/Source/Programs/Horde/Horde.Server/Server/DebugController.cs
@@ -511,7 +511,7 @@ namespace Horde.Server.Server
 				return Forbid(ServerAclAction.Debug);
 			}
 
-			ILog? log = await _logCollection.GetLogAsync(logId, CancellationToken.None);
+			ILog? log = await _logCollection.GetAsync(logId, CancellationToken.None);
 			if (log == null)
 			{
 				return NotFound();
diff --git a/Engine/Source/Programs/Horde/Horde.Server/Tools/ITool.cs b/Engine/Source/Programs/Horde/Horde.Server/Tools/ITool.cs
--- a/Engine/Source/Programs/Horde/Horde.Server/Tools/ITool.cs
+++ b/Engine/Source/Programs/Horde/Horde.Server/Tools/ITool.cs
@@ -9,7 +9,6 @@ using System.Threading.Tasks;
 using EpicGames.Horde.Acls;
 using EpicGames.Horde.Storage;
 using EpicGames.Horde.Tools;
-using Horde.Server.Server;
 
 namespace Horde.Server.Tools
 {
diff --git a/Engine/Source/Programs/Horde/Horde.Server/Tools/IToolCollection.cs b/Engine/Source/Programs/Horde/Horde.Server/Tools/IToolCollection.cs
--- a/Engine/Source/Programs/Horde/Horde.Server/Tools/IToolCollection.cs
+++ b/Engine/Source/Programs/Horde/Horde.Server/Tools/IToolCollection.cs
@@ -1,12 +1,9 @@
 // Copyright Epic Games, Inc. All Rights Reserved.
 
 using System.Collections.Generic;
-using System.IO;
 using System.Threading;
 using System.Threading.Tasks;
-using EpicGames.Horde.Storage;
 using EpicGames.Horde.Tools;
-using Horde.Server.Server;
 
 namespace Horde.Server.Tools
 {
diff --git a/Engine/Source/Programs/Horde/Horde.Server/Tools/ToolCollection.cs b/Engine/Source/Programs/Horde/Horde.Server/Tools/ToolCollection.cs
--- a/Engine/Source/Programs/Horde/Horde.Server/Tools/ToolCollection.cs
+++ b/Engine/Source/Programs/Horde/Horde.Server/Tools/ToolCollection.cs
@@ -18,7 +18,6 @@ using EpicGames.Horde.Storage.Bundles;
 using EpicGames.Horde.Storage.Nodes;
 using EpicGames.Horde.Storage.ObjectStores;
 using EpicGames.Horde.Tools;
-using EpicGames.Serialization;
 using Horde.Server.Server;
 using Horde.Server.Storage;
 using Horde.Server.Utilities;
@@ -240,7 +239,7 @@ namespace Horde.Server.Tools
 			ToolConfig? toolConfig;
 			if (globalConfig.TryGetTool(id, out toolConfig))
 			{
-				ToolDocument document = await FindOrAddDocument(toolConfig, cancellationToken);
+				ToolDocument document = await FindOrAddDocumentAsync(toolConfig, cancellationToken);
 				return new Tool(this, document, toolConfig, _clock.UtcNow);
 			}
 
@@ -263,10 +262,10 @@ namespace Horde.Server.Tools
 			List<Tool> tools = new List<Tool>();
 			foreach (ToolConfig toolConfig in globalConfig.Tools)
 			{
-				ToolDocument document = await FindOrAddDocument(toolConfig, cancellationToken);
+				ToolDocument document = await FindOrAddDocumentAsync(toolConfig, cancellationToken);
 				tools.Add(new Tool(this, document, toolConfig, utcNow));
 			}
-			foreach(BundledToolConfig bundledToolConfig in globalConfig.ServerSettings.BundledTools)
+			foreach (BundledToolConfig bundledToolConfig in globalConfig.ServerSettings.BundledTools)
 			{
 				ToolDocument document = CreateBundledToolDocument(bundledToolConfig);
 				tools.Add(new Tool(this, document, bundledToolConfig, utcNow));
@@ -275,7 +274,7 @@ namespace Horde.Server.Tools
 			return tools;
 		}
 
-		async Task<ToolDocument> FindOrAddDocument(ToolConfig toolConfig, CancellationToken cancellationToken)
+		async Task<ToolDocument> FindOrAddDocumentAsync(ToolConfig toolConfig, CancellationToken cancellationToken)
 		{
 			ToolDocument? tool;
 			for (; ; )
@@ -295,7 +294,7 @@ namespace Horde.Server.Tools
 			return tool;
 		}
 
-		ToolDocument CreateBundledToolDocument(BundledToolConfig bundledToolConfig)
+		static ToolDocument CreateBundledToolDocument(BundledToolConfig bundledToolConfig)
 		{
 			ToolDocument tool = new ToolDocument(bundledToolConfig.Id);
 
@@ -379,7 +378,7 @@ namespace Horde.Server.Tools
 					return newTool;
 				}
 
-				newTool = await FindOrAddDocument(toolConfig, cancellationToken);
+				newTool = await FindOrAddDocumentAsync(toolConfig, cancellationToken);
 				if (newTool == null)
 				{
 					return null;
diff --git a/Engine/Source/Programs/Horde/Horde.Server/Tools/ToolsController.cs b/Engine/Source/Programs/Horde/Horde.Server/Tools/ToolsController.cs
--- a/Engine/Source/Programs/Horde/Horde.Server/Tools/ToolsController.cs
+++ b/Engine/Source/Programs/Horde/Horde.Server/Tools/ToolsController.cs
@@ -10,7 +10,6 @@ using EpicGames.Core;
 using EpicGames.Horde.Storage;
 using EpicGames.Horde.Storage.Nodes;
 using EpicGames.Horde.Tools;
-using Horde.Server.Server;
 using Horde.Server.Storage;
 using Horde.Server.Utilities;
 using Microsoft.AspNetCore.Authorization;
@@ -18,7 +17,6 @@ using Microsoft.AspNetCore.Http;
 using Microsoft.AspNetCore.Mvc;
 using Microsoft.AspNetCore.StaticFiles;
 using Microsoft.Extensions.Logging;
-using Microsoft.Extensions.Options;
 
 namespace Horde.Server.Tools
 {
@@ -31,15 +29,13 @@ namespace Horde.Server.Tools
 	public class ToolsController : HordeControllerBase
 	{
 		readonly IToolCollection _toolCollection;
-		readonly IOptionsSnapshot<GlobalConfig> _globalConfig;
 
 		/// <summary>
 		/// Constructor
 		/// </summary>
-		public ToolsController(IToolCollection toolCollection, IOptionsSnapshot<GlobalConfig> globalConfig)
+		public ToolsController(IToolCollection toolCollection)
 		{
 			_toolCollection = toolCollection;
-			_globalConfig = globalConfig;
 		}
 
 		/// <summary>
@@ -178,17 +174,15 @@ namespace Horde.Server.Tools
 	public class PublicToolsController : HordeControllerBase
 	{
 		readonly IToolCollection _toolCollection;
-		readonly IOptionsSnapshot<GlobalConfig> _globalConfig;
 		readonly IClock _clock;
 
 		/// <summary>
 		/// Constructor
 		/// </summary>
-		public PublicToolsController(IToolCollection toolCollection, IClock clock, IOptionsSnapshot<GlobalConfig> globalConfig)
+		public PublicToolsController(IToolCollection toolCollection, IClock clock)
 		{
 			_toolCollection = toolCollection;
 			_clock = clock;
-			_globalConfig = globalConfig;
 		}
 
 		/// <summary>
@@ -365,7 +359,7 @@ namespace Horde.Server.Tools
 			}
 		}
 
-		private async Task<GetToolDeploymentResponse> GetDeploymentInfoResponseAsync(ITool tool, IToolDeployment deployment, CancellationToken cancellationToken)
+		static async Task<GetToolDeploymentResponse> GetDeploymentInfoResponseAsync(ITool tool, IToolDeployment deployment, CancellationToken cancellationToken)
 		{
 			using IStorageClient client = tool.CreateStorageClient();
 			IBlobHandle rootHandle = await client.ReadRefAsync(deployment.RefName, cancellationToken: cancellationToken);
