Horde: Merge logs and log event collections. Add support for deleting logs, along with any events and data in the underlying storage system.

#rnx

diff --git a/Engine/Source/Programs/Horde/Horde.Server.Tests/Compute/ComputeServiceTest.cs b/Engine/Source/Programs/Horde/Horde.Server.Tests/Compute/ComputeServiceTest.cs
--- a/Engine/Source/Programs/Horde/Horde.Server.Tests/Compute/ComputeServiceTest.cs
+++ b/Engine/Source/Programs/Horde/Horde.Server.Tests/Compute/ComputeServiceTest.cs
@@ -287,7 +287,7 @@ namespace Horde.Server.Tests.Compute
 		private async Task<ComputeService> CreateComputeServiceAsync(string? tunnelAddress)
 		{
 			ServerSettings ss = new() { ComputeTunnelAddress = tunnelAddress };
-			ComputeService cs = new(AgentCollection, LogService, AgentService, AgentRelayService, GetRedisServiceSingleton(),
+			ComputeService cs = new(AgentCollection, LogCollection, AgentService, AgentRelayService, GetRedisServiceSingleton(),
 				new TestOptionsMonitor<ServerSettings>(ss), GlobalConfig, Clock, Tracer, Meter,
 				NullLogger<ComputeService>.Instance);
 			List<string> props = new() { "ComputeIp=11.0.0.1", "ComputePort=5000" };
diff --git a/Engine/Source/Programs/Horde/Horde.Server.Tests/Issues/IssueServiceTests.cs b/Engine/Source/Programs/Horde/Horde.Server.Tests/Issues/IssueServiceTests.cs
--- a/Engine/Source/Programs/Horde/Horde.Server.Tests/Issues/IssueServiceTests.cs
+++ b/Engine/Source/Programs/Horde/Horde.Server.Tests/Issues/IssueServiceTests.cs
@@ -44,18 +44,16 @@ namespace Horde.Server.Tests.Issues
 	{
 		class TestJsonLogger : ILogger, IAsyncDisposable
 		{
-			readonly ILogService _logFileService;
-			readonly LogId _logId;
+			readonly ILog _log;
 			readonly LogBuilder _builder;
 			readonly List<(LogLevel, ReadOnlyMemory<byte>)> _events = new List<(LogLevel, ReadOnlyMemory<byte>)>();
 			readonly IStorageClient _storageClient;
 
 			int _lineIndex;
 
-			public TestJsonLogger(ILogService logFileService, LogId logId, IStorageClient storageClient)
+			public TestJsonLogger(ILog log, IStorageClient storageClient)
 			{
-				_logFileService = logFileService;
-				_logId = logId;
+				_log = log;
 				_builder = new LogBuilder(LogFormat.Json, NullLogger.Instance);
 				_storageClient = storageClient;
 			}
@@ -73,8 +71,7 @@ namespace Horde.Server.Tests.Issues
 				await using (IBlobWriter writer = _storageClient.CreateBlobWriter())
 				{
 					IBlobRef<LogNode> handle = await _builder.FlushAsync(writer, true, CancellationToken.None);
-					ILog? logFile = await _logFileService.GetLogAsync(_logId, CancellationToken.None);
-					await _storageClient.WriteRefAsync(logFile!.RefName, handle);
+					await _storageClient.WriteRefAsync(_log.RefName, handle);
 				}
 
 				_storageClient.Dispose();
@@ -100,7 +97,7 @@ namespace Horde.Server.Tests.Issues
 					if (ev.LineIndex == 0)
 					{
 						LogEventSeverity severity = (level == LogLevel.Warning) ? LogEventSeverity.Warning : LogEventSeverity.Error;
-						await _logFileService.CreateEventsAsync(new List<NewLogEventData> { new NewLogEventData { LogId = _logId, LineIndex = _lineIndex, LineCount = ev.LineCount, Severity = severity } }, CancellationToken.None);
+						await _log.AddEventsAsync(new List<NewLogEventData> { new NewLogEventData { LineIndex = _lineIndex, LineCount = ev.LineCount, Severity = severity } }, CancellationToken.None);
 					}
 				}
 
@@ -246,7 +243,7 @@ namespace Horde.Server.Tests.Issues
 				{
 					JobStepId stepId = new JobStepId((ushort)((groupIdx * 100) + nodeIdx));
 
-					ILog logFile = LogService.CreateLogAsync(jobId, null, null, LogType.Json).Result;
+					ILog logFile = LogCollection.AddAsync(jobId, null, null, LogType.Json).Result;
 
 					Mock<IJobStep> step = new Mock<IJobStep>(MockBehavior.Strict);
 					step.SetupGet(x => x.Id).Returns(stepId);
@@ -323,7 +320,7 @@ namespace Horde.Server.Tests.Issues
 		{
 			LogId logId = job.Batches[batchIdx].Steps[stepIdx].LogId!.Value;
 
-			ILog logFile = (await LogService.GetLogAsync(logId, CancellationToken.None))!;
+			ILog log = (await LogCollection.GetAsync(logId, CancellationToken.None))!;
 
 			using (IStorageClient storageClient = StorageService.CreateClient(Namespace.Logs))
 			{
@@ -333,26 +330,29 @@ namespace Horde.Server.Tests.Issues
 					builder.WriteData(data);
 
 					IBlobRef<LogNode> handle = await builder.FlushAsync(writer, true, CancellationToken.None);
-					await storageClient.WriteRefAsync(logFile!.RefName, handle);
+					await storageClient.WriteRefAsync(log!.RefName, handle);
 				}
 			}
 
-			await LogService.CreateEventsAsync(new List<NewLogEventData> { new NewLogEventData { LogId = logId, LineIndex = 0, LineCount = 1, Severity = severity } }, CancellationToken.None);
+			await log.AddEventsAsync(new List<NewLogEventData> { new NewLogEventData { LineIndex = 0, LineCount = 1, Severity = severity } }, CancellationToken.None);
 		}
 
-		private TestJsonLogger CreateLogger(IJob job, int batchIdx, int stepIdx)
+		private async Task<TestJsonLogger> CreateLoggerAsync(IJob job, int batchIdx, int stepIdx)
 		{
 			LogId logId = job.Batches[batchIdx].Steps[stepIdx].LogId!.Value;
 			IStorageClient storageClient = StorageService.CreateClient(Namespace.Logs);
-			return new TestJsonLogger(LogService, logId, storageClient);
+			ILog? log = await LogCollection.GetAsync(logId)!;
+			Assert.IsNotNull(log);
+			return new TestJsonLogger(log, storageClient);
 		}
 
 		private async Task ParseEventsAsync(IJob job, int batchIdx, int stepIdx, string[] lines)
 		{
 			LogId logId = job.Batches[batchIdx].Steps[stepIdx].LogId!.Value;
+			ILog log = (await LogCollection.GetAsync(logId))!;
 
 			IStorageClient storageClient = StorageService.CreateClient(Namespace.Logs);
-			await using (TestJsonLogger logger = new TestJsonLogger(LogService, logId, storageClient))
+			await using (TestJsonLogger logger = new TestJsonLogger(log, storageClient))
 			{
 				PerforceMetadataLogger perforceLogger = new PerforceMetadataLogger(logger);
 				perforceLogger.AddClientView(_autoSdkDir, "//depot/CarefullyRedist/...", 12345);
@@ -571,7 +571,7 @@ namespace Horde.Server.Tests.Issues
 				_perforce.AddChange(_mainStreamId, 160, john, "Description", new string[] { "Engine/Foo/Baz.txt" });
 
 				IJob job = CreateJob(_mainStreamId, 170, "Test Build", _graph);
-				await using (TestJsonLogger logger = CreateLogger(job, 0, 0))
+				await using (TestJsonLogger logger = await CreateLoggerAsync(job, 0, 0))
 				{
 					logger.LogWarning(KnownLogEvents.AutomationTool_PerforceCase, "    {DepotFile}", new LogValue(LogValueType.DepotPath, "//UE5/Main/Engine/Foo/Bar.txt"));
 				}
@@ -669,8 +669,9 @@ namespace Horde.Server.Tests.Issues
 				await ParseEventsAsync(job, 0, 0, lines);
 				await UpdateCompleteStepAsync(job, 0, 0, JobStepOutcome.Warnings);
 
-				ILog? log = await LogService.GetLogAsync(job.Batches[0].Steps[0].LogId!.Value, CancellationToken.None);
-				List<ILogEvent> events = await LogService.FindEventsAsync(log!);
+				ILog? log = await LogCollection.GetAsync(job.Batches[0].Steps[0].LogId!.Value, CancellationToken.None);
+				Assert.IsNotNull(log);
+				List<ILogEvent> events = await log.GetEventsAsync();
 				Assert.AreEqual(1, events.Count);
 				Assert.AreEqual(2, events[0].LineCount);
 
@@ -711,8 +712,10 @@ namespace Horde.Server.Tests.Issues
 				await AddEventAsync(job, 0, 0, LogLevel.Warning, id: KnownLogEvents.Gauntlet_TestEvent.Id);
 				await UpdateCompleteStepAsync(job, 0, 0, JobStepOutcome.Warnings);
 
-				ILog? log = await LogService.GetLogAsync(job.Batches[0].Steps[0].LogId!.Value, CancellationToken.None);
-				List<ILogEvent> events = await LogService.FindEventsAsync(log!);
+				ILog? log = await LogCollection.GetAsync(job.Batches[0].Steps[0].LogId!.Value, CancellationToken.None);
+				Assert.IsNotNull(log);
+
+				List<ILogEvent> events = await log.GetEventsAsync();
 				Assert.AreEqual(1, events.Count);
 				Assert.AreEqual(1, events[0].LineCount);
 
@@ -1726,7 +1729,7 @@ namespace Horde.Server.Tests.Issues
 			// Expected: Gauntlet fingerprint using Name prefix
 			{
 				IJob job = CreateJob(_mainStreamId, 110, "Test Build", _graph);
-				await using (TestJsonLogger logger = CreateLogger(job, 0, 0))
+				await using (TestJsonLogger logger = await CreateLoggerAsync(job, 0, 0))
 				{
 					logger.LogError(KnownLogEvents.Gauntlet_TestEvent, "    Test {Name} failed", "Bar.Foo.Test");
 				}
@@ -1743,7 +1746,7 @@ namespace Horde.Server.Tests.Issues
 			// Expected: Gauntlet fingerprint using Device prefix
 			{
 				IJob job = CreateJob(_mainStreamId, 120, "Test Build", _graph);
-				await using (TestJsonLogger logger = CreateLogger(job, 0, 0))
+				await using (TestJsonLogger logger = await CreateLoggerAsync(job, 0, 0))
 				{
 					logger.LogWarning(KnownLogEvents.Gauntlet_DeviceEvent, "    Device {Name} reported an issue", "Foo");
 				}
@@ -1760,7 +1763,7 @@ namespace Horde.Server.Tests.Issues
 			// Expected: Gauntlet fingerprint using Access prefix
 			{
 				IJob job = CreateJob(_mainStreamId, 130, "Test Build", _graph);
-				await using (TestJsonLogger logger = CreateLogger(job, 0, 0))
+				await using (TestJsonLogger logger = await CreateLoggerAsync(job, 0, 0))
 				{
 					logger.LogError(KnownLogEvents.Gauntlet_BuildDropEvent, "    File {File} reported an issue", "/Bar/Foo.txt");
 					logger.LogError(KnownLogEvents.Gauntlet_BuildDropEvent, "    Folder {Directory} reported an issue", "/Bar/Foo");
@@ -1805,7 +1808,7 @@ namespace Horde.Server.Tests.Issues
 				+ @"    Test did not run until completion. The test exited prematurely.";
 
 				IJob job = CreateJob(_mainStreamId, 140, "Test Build", _graph);
-				await using (TestJsonLogger logger = CreateLogger(job, 0, 0))
+				await using (TestJsonLogger logger = await CreateLoggerAsync(job, 0, 0))
 				{
 					logger.LogError(KnownLogEvents.Gauntlet_FatalEvent, "{Message}", logMessage);
 				}
@@ -1836,7 +1839,7 @@ namespace Horde.Server.Tests.Issues
 				};
 
 				IJob job = CreateJob(_mainStreamId, 150, "Test Build", _graph);
-				await using (TestJsonLogger logger = CreateLogger(job, 0, 0))
+				await using (TestJsonLogger logger = await CreateLoggerAsync(job, 0, 0))
 				{
 					foreach (string error in logErrors)
 					{
diff --git a/Engine/Source/Programs/Horde/Horde.Server.Tests/Jobs/TestData/TestDataTests.cs b/Engine/Source/Programs/Horde/Horde.Server.Tests/Jobs/TestData/TestDataTests.cs
--- a/Engine/Source/Programs/Horde/Horde.Server.Tests/Jobs/TestData/TestDataTests.cs
+++ b/Engine/Source/Programs/Horde/Horde.Server.Tests/Jobs/TestData/TestDataTests.cs
@@ -123,7 +123,7 @@ namespace Horde.Server.Tests.Jobs.TestData
 				{
 					JobStepId stepId = new JobStepId((ushort)((groupIdx * 100) + nodeIdx));
 
-					ILog logFile = LogService.CreateLogAsync(jobId, null, null, LogType.Json).Result;
+					ILog logFile = LogCollection.AddAsync(jobId, null, null, LogType.Json).Result;
 
 					Mock<IJobStep> step = new Mock<IJobStep>(MockBehavior.Strict);
 					step.SetupGet(x => x.Id).Returns(stepId);
diff --git a/Engine/Source/Programs/Horde/Horde.Server.Tests/Logs/LogIndexingTests.cs b/Engine/Source/Programs/Horde/Horde.Server.Tests/Logs/LogIndexingTests.cs
--- a/Engine/Source/Programs/Horde/Horde.Server.Tests/Logs/LogIndexingTests.cs
+++ b/Engine/Source/Programs/Horde/Horde.Server.Tests/Logs/LogIndexingTests.cs
@@ -66,7 +66,7 @@ namespace Horde.Server.Tests.Logs
 		public async Task IndexTestsAsync()
 		{
 			JobId jobId = JobIdUtils.GenerateNewId();
-			ILog logFile = await LogService.CreateLogAsync(jobId, null, null, LogType.Text);
+			ILog logFile = await LogCollection.AddAsync(jobId, null, null, LogType.Text);
 
 			// Write the test data to the log file in blocks
 			await using (TestLogWriter logWriter = new TestLogWriter(logFile, StorageService))
@@ -137,7 +137,7 @@ namespace Horde.Server.Tests.Logs
 		public async Task PartialTokenTestsAsync()
 		{
 			JobId jobId = JobIdUtils.GenerateNewId();
-			ILog logFile = await LogService.CreateLogAsync(jobId, null, null, LogType.Text);
+			ILog logFile = await LogCollection.AddAsync(jobId, null, null, LogType.Text);
 
 			string[] lines =
 			{
@@ -183,7 +183,7 @@ namespace Horde.Server.Tests.Logs
 		public async Task AppendIndexTestsAsync()
 		{
 			JobId jobId = JobIdUtils.GenerateNewId();
-			ILog logFile = await LogService.CreateLogAsync(jobId, null, null, LogType.Text);
+			ILog logFile = await LogCollection.AddAsync(jobId, null, null, LogType.Text);
 
 			await using (TestLogWriter writer = new TestLogWriter(logFile, StorageService))
 			{
@@ -227,7 +227,7 @@ namespace Horde.Server.Tests.Logs
 
 		async Task SearchLogDataTestAsync(LogId logId)
 		{
-			ILog? logFile = await LogService.GetLogAsync(logId, CancellationToken.None);
+			ILog? logFile = await LogCollection.GetAsync(logId, CancellationToken.None);
 			Assert.IsNotNull(logFile);
 
 			await SearchLogDataTestAsync(logFile, "HISPANIOLA", 0, 4, new[] { 1503, 1520, 1525, 1595 });
diff --git a/Engine/Source/Programs/Horde/Horde.Server.Tests/Logs/LogServiceTests.cs b/Engine/Source/Programs/Horde/Horde.Server.Tests/Logs/LogServiceTests.cs
--- a/Engine/Source/Programs/Horde/Horde.Server.Tests/Logs/LogServiceTests.cs
+++ b/Engine/Source/Programs/Horde/Horde.Server.Tests/Logs/LogServiceTests.cs
@@ -42,7 +42,7 @@ namespace Horde.Server.Tests.Logs
 		public async Task WriteLogLifecycleOldTestAsync()
 		{
 			JobId jobId = JobIdUtils.GenerateNewId();
-			ILog logFile = await LogService.CreateLogAsync(jobId, null, null, LogType.Text, logId: null, cancellationToken: CancellationToken.None);
+			ILog logFile = await LogCollection.AddAsync(jobId, null, null, LogType.Text, logId: null, cancellationToken: CancellationToken.None);
 
 			await using (TestLogWriter writer = new TestLogWriter(logFile, StorageService))
 			{
@@ -56,7 +56,7 @@ namespace Horde.Server.Tests.Logs
 				await writer.FlushAsync(true);
 			}
 
-			logFile = (await LogService.GetLogAsync(logFile.Id, CancellationToken.None))!;
+			logFile = (await LogCollection.GetAsync(logFile.Id, CancellationToken.None))!;
 
 			Assert.AreEqual("hello", await ReadLogAsync(logFile, 0, 5));
 			Assert.AreEqual("foo\nbar\nbaz\n", await ReadLogAsync(logFile, 6, 12));
@@ -68,9 +68,9 @@ namespace Horde.Server.Tests.Logs
 			Assert.AreEqual("baz", lines[3].ToString());
 		}
 
-		private static async Task AssertLineOffsetAsync(ILogService logService, LogId logId, int lineIndex, int clampedLineIndex, long offset)
+		private static async Task AssertLineOffsetAsync(ILogCollection logCollection, LogId logId, int lineIndex, int clampedLineIndex, long offset)
 		{
-			ILog? logFile = await logService.GetLogAsync(logId, CancellationToken.None);
+			ILog? logFile = await logCollection.GetAsync(logId, CancellationToken.None);
 			Assert.IsNotNull(logFile);
 
 			List<Utf8String> lines = await logFile.ReadLinesAsync(0, lineIndex + 1, CancellationToken.None);
@@ -95,7 +95,7 @@ namespace Horde.Server.Tests.Logs
 		public async Task WriteLogLifecycleAsync()
 		{
 			JobId jobId = JobIdUtils.GenerateNewId();
-			ILog log = await LogService.CreateLogAsync(jobId, null, null, LogType.Text);
+			ILog log = await LogCollection.AddAsync(jobId, null, null, LogType.Text);
 
 			string str1 = "hello\n";
 			string str2 = "foo\nbar\n";
@@ -111,10 +111,10 @@ namespace Horde.Server.Tests.Logs
 			log = await logWriter.WriteDataAsync(Encoding.ASCII.GetBytes(str1));
 			Assert.AreEqual(str1, await ReadLogAsync(log, 0, str1.Length));
 			Assert.AreEqual(str1, await ReadLogAsync(log, 0, str1.Length + 100)); // Reading too far is valid?
-			await AssertLineOffsetAsync(LogService, log.Id, 0, 0, 0);
-			await AssertLineOffsetAsync(LogService, log.Id, 1, 1, 6);
-			await AssertLineOffsetAsync(LogService, log.Id, 1, 1, 6);
-			await AssertLineOffsetAsync(LogService, log.Id, 1, 1, 6);
+			await AssertLineOffsetAsync(LogCollection, log.Id, 0, 0, 0);
+			await AssertLineOffsetAsync(LogCollection, log.Id, 1, 1, 6);
+			await AssertLineOffsetAsync(LogCollection, log.Id, 1, 1, 6);
+			await AssertLineOffsetAsync(LogCollection, log.Id, 1, 1, 6);
 
 			// Second write without flushing. Will become chunk #2
 			offset += str1.Length;
@@ -122,21 +122,21 @@ namespace Horde.Server.Tests.Logs
 			log = await logWriter.WriteDataAsync(Encoding.ASCII.GetBytes(str2));
 
 			Assert.AreEqual(str1 + str2, await ReadLogAsync(log, 0, str1.Length + str2.Length));
-			await AssertLineOffsetAsync(LogService, log.Id, 0, 0, 0);
-			await AssertLineOffsetAsync(LogService, log.Id, 1, 1, 6);
-			await AssertLineOffsetAsync(LogService, log.Id, 2, 2, 10);
-			await AssertLineOffsetAsync(LogService, log.Id, 3, 3, 14);
+			await AssertLineOffsetAsync(LogCollection, log.Id, 0, 0, 0);
+			await AssertLineOffsetAsync(LogCollection, log.Id, 1, 1, 6);
+			await AssertLineOffsetAsync(LogCollection, log.Id, 2, 2, 10);
+			await AssertLineOffsetAsync(LogCollection, log.Id, 3, 3, 14);
 
 			// Third write without flushing. Will become chunk #2
 			offset += str2.Length;
 			lineIndex += str3.Count(f => f == '\n');
 			log = await logWriter.WriteDataAsync(Encoding.ASCII.GetBytes(str3));
 			Assert.AreEqual(str1 + str2 + str3, await ReadLogAsync(log, 0, str1.Length + str2.Length + str3.Length));
-			await AssertLineOffsetAsync(LogService, log.Id, 0, 0, 0);
-			await AssertLineOffsetAsync(LogService, log.Id, 1, 1, 6);
-			await AssertLineOffsetAsync(LogService, log.Id, 2, 2, 10);
-			await AssertLineOffsetAsync(LogService, log.Id, 3, 3, 14);
-			await AssertLineOffsetAsync(LogService, log.Id, 4, 4, 18);
+			await AssertLineOffsetAsync(LogCollection, log.Id, 0, 0, 0);
+			await AssertLineOffsetAsync(LogCollection, log.Id, 1, 1, 6);
+			await AssertLineOffsetAsync(LogCollection, log.Id, 2, 2, 10);
+			await AssertLineOffsetAsync(LogCollection, log.Id, 3, 3, 14);
+			await AssertLineOffsetAsync(LogCollection, log.Id, 4, 4, 18);
 
 			// Fourth write with flush. Will become chunk #2
 			offset += str3.Length;
@@ -146,14 +146,14 @@ namespace Horde.Server.Tests.Logs
 			await ReadLogAsync(log, 0, str1.Length + str2.Length + str3.Length + str4.Length));
 			Assert.AreEqual(str3 + str4,
 			await ReadLogAsync(log, str1.Length + str2.Length, str3.Length + str4.Length));
-			await AssertLineOffsetAsync(LogService, log.Id, 0, 0, 0);
-			await AssertLineOffsetAsync(LogService, log.Id, 1, 1, 6);
-			await AssertLineOffsetAsync(LogService, log.Id, 2, 2, 10);
-			await AssertLineOffsetAsync(LogService, log.Id, 3, 3, 14);
-			await AssertLineOffsetAsync(LogService, log.Id, 4, 4, 18);
-			await AssertLineOffsetAsync(LogService, log.Id, 5, 5, 22);
-			await AssertLineOffsetAsync(LogService, log.Id, 6, 6, 27);
-			await AssertLineOffsetAsync(LogService, log.Id, 7, 7, 32);
+			await AssertLineOffsetAsync(LogCollection, log.Id, 0, 0, 0);
+			await AssertLineOffsetAsync(LogCollection, log.Id, 1, 1, 6);
+			await AssertLineOffsetAsync(LogCollection, log.Id, 2, 2, 10);
+			await AssertLineOffsetAsync(LogCollection, log.Id, 3, 3, 14);
+			await AssertLineOffsetAsync(LogCollection, log.Id, 4, 4, 18);
+			await AssertLineOffsetAsync(LogCollection, log.Id, 5, 5, 22);
+			await AssertLineOffsetAsync(LogCollection, log.Id, 6, 6, 27);
+			await AssertLineOffsetAsync(LogCollection, log.Id, 7, 7, 32);
 
 			// Fifth write with flush and data that will span more than chunk. Will become chunk #3
 			string a = "Lorem ipsum dolor sit amet\n";
@@ -179,13 +179,13 @@ namespace Horde.Server.Tests.Logs
 			// Will implicitly test GetLogFileAsync(), AddCachedLogFile()
 			JobId jobId = JobIdUtils.GenerateNewId();
 			SessionId sessionId = SessionIdUtils.GenerateNewId();
-			ILog a = await LogService.CreateLogAsync(jobId, null, sessionId, LogType.Text, logId: null, CancellationToken.None);
-			ILog b = (await LogService.GetLogAsync(a.Id, CancellationToken.None))!;
+			ILog a = await LogCollection.AddAsync(jobId, null, sessionId, LogType.Text, logId: null, CancellationToken.None);
+			ILog b = (await LogCollection.GetAsync(a.Id, CancellationToken.None))!;
 			Assert.AreEqual(a.JobId, b.JobId);
 			Assert.AreEqual(a.SessionId, b.SessionId);
 			Assert.AreEqual(a.Type, b.Type);
 
-			ILog? notFound = await LogService.GetLogAsync(LogIdUtils.GenerateNewId(), CancellationToken.None);
+			ILog? notFound = await LogCollection.GetAsync(LogIdUtils.GenerateNewId(), CancellationToken.None);
 			Assert.IsNull(notFound);
 		}
 
@@ -194,8 +194,8 @@ namespace Horde.Server.Tests.Logs
 		{
 			JobId jobId = JobIdUtils.GenerateNewId();
 			SessionId sessionId = SessionIdUtils.GenerateNewId();
-			ILog log = await LogService.CreateLogAsync(jobId, null, sessionId, LogType.Text, logId: null, cancellationToken: CancellationToken.None);
-			ILog logNoSession = await LogService.CreateLogAsync(jobId, null, null, LogType.Text, logId: null, cancellationToken: CancellationToken.None);
+			ILog log = await LogCollection.AddAsync(jobId, null, sessionId, LogType.Text, logId: null, cancellationToken: CancellationToken.None);
+			ILog logNoSession = await LogCollection.AddAsync(jobId, null, null, LogType.Text, logId: null, cancellationToken: CancellationToken.None);
 
 			ClaimsPrincipal hasClaim = new ClaimsPrincipal(new ClaimsIdentity(new List<Claim>
 			{
@@ -206,16 +206,16 @@ namespace Horde.Server.Tests.Logs
 				new Claim(HordeClaimTypes.AgentSessionId, "invalid-session-id"),
 			}, "TestAuthType"));
 
-			Assert.IsTrue(Horde.Server.Logs.LogService.AuthorizeForSession(log, hasClaim));
-			Assert.IsFalse(Horde.Server.Logs.LogService.AuthorizeForSession(log, hasNoClaim));
-			Assert.IsFalse(Horde.Server.Logs.LogService.AuthorizeForSession(logNoSession, hasClaim));
+			Assert.IsTrue(log.AuthorizeForSession(hasClaim));
+			Assert.IsFalse(log.AuthorizeForSession(hasNoClaim));
+			Assert.IsFalse(logNoSession.AuthorizeForSession(hasClaim));
 		}
 
 		[TestMethod]
 		public async Task ChunkSplittingAsync()
 		{
 			JobId jobId = JobIdUtils.GenerateNewId();
-			ILog logFile = await LogService.CreateLogAsync(jobId, null, null, LogType.Text, logId: null, cancellationToken: CancellationToken.None);
+			ILog logFile = await LogCollection.AddAsync(jobId, null, null, LogType.Text, logId: null, cancellationToken: CancellationToken.None);
 			await using TestLogWriter logWriter = new TestLogWriter(logFile, StorageService);
 
 			byte[] line1 = Encoding.UTF8.GetBytes("hello world\n");
diff --git a/Engine/Source/Programs/Horde/Horde.Server.Tests/TestSetup.cs b/Engine/Source/Programs/Horde/Horde.Server.Tests/TestSetup.cs
--- a/Engine/Source/Programs/Horde/Horde.Server.Tests/TestSetup.cs
+++ b/Engine/Source/Programs/Horde/Horde.Server.Tests/TestSetup.cs
@@ -112,7 +112,6 @@ namespace Horde.Server.Tests
 		public MongoService MongoService => ServiceProvider.GetRequiredService<MongoService>();
 		public ITemplateCollection TemplateCollection => ServiceProvider.GetRequiredService<ITemplateCollection>();
 		internal PerforceServiceStub PerforceService => (PerforceServiceStub)ServiceProvider.GetRequiredService<IPerforceService>();
-		public ILogService LogService => ServiceProvider.GetRequiredService<ILogService>();
 		public ISubscriptionCollection SubscriptionCollection => ServiceProvider.GetRequiredService<ISubscriptionCollection>();
 		public INotificationService NotificationService => ServiceProvider.GetRequiredService<INotificationService>();
 		public IssueService IssueService => ServiceProvider.GetRequiredService<IssueService>();
@@ -222,7 +221,6 @@ namespace Horde.Server.Tests
 			services.AddSingleton<IJobStepRefCollection, JobStepRefCollection>();
 			services.AddSingleton<IJobTimingCollection, JobTimingCollection>();
 			services.AddSingleton<ILeaseCollection, LeaseCollection>();
-			services.AddSingleton<ILogEventCollection, LogEventCollection>();
 			services.AddSingleton<ILogCollection, LogCollection>();
 			services.AddSingleton<INotificationTriggerCollection, NotificationTriggerCollection>();
 			services.AddSingleton<IPoolCollection, PoolCollection>();
@@ -267,7 +265,6 @@ namespace Horde.Server.Tests
 			services.AddSingleton<JobExpirationService>();
 			services.AddSingleton<LifetimeService>();
 			services.AddSingleton<ILogStorage, NullLogStorage>();
-			services.AddSingleton<ILogService, LogService>();
 			services.AddSingleton<LogTailService>();
 			services.AddSingleton<INotificationService, NotificationService>();
 			services.AddSingleton<IPerforceService, PerforceServiceStub>();
diff --git a/Engine/Source/Programs/Horde/Horde.Server/Compute/ComputeService.cs b/Engine/Source/Programs/Horde/Horde.Server/Compute/ComputeService.cs
--- a/Engine/Source/Programs/Horde/Horde.Server/Compute/ComputeService.cs
+++ b/Engine/Source/Programs/Horde/Horde.Server/Compute/ComputeService.cs
@@ -167,7 +167,7 @@ namespace Horde.Server.Compute
 		public event ResourceNeedEvent? OnResourceNeedsUpdated;
 
 		readonly IAgentCollection _agentCollection;
-		readonly ILogService _logService;
+		readonly ILogCollection _logCollection;
 		readonly AgentService _agentService;
 		readonly AgentRelayService _agentRelayService;
 		readonly RedisService _redisService;
@@ -189,7 +189,7 @@ namespace Horde.Server.Compute
 		/// </summary>
 		public ComputeService(
 			IAgentCollection agentCollection,
-			ILogService logService,
+			ILogCollection logCollection,
 			AgentService agentService,
 			AgentRelayService agentRelayService,
 			RedisService redisService,
@@ -201,7 +201,7 @@ namespace Horde.Server.Compute
 			ILogger<ComputeService> logger)
 		{
 			_agentCollection = agentCollection;
-			_logService = logService;
+			_logCollection = logCollection;
 			_agentService = agentService;
 			_agentRelayService = agentRelayService;
 			_redisService = redisService;
@@ -404,7 +404,7 @@ namespace Horde.Server.Compute
 						}
 
 						LeaseId leaseId = new LeaseId(BinaryIdUtils.CreateNew());
-						ILog? log = await _logService.CreateLogAsync(JobId.Empty, leaseId, agent.SessionId, LogType.Json, cancellationToken: cancellationToken);
+						ILog? log = await _logCollection.AddAsync(JobId.Empty, leaseId, agent.SessionId, LogType.Json, cancellationToken: cancellationToken);
 
 						using TelemetrySpan createTaskSpan = _tracer.StartActiveSpan("CreateComputeTask");
 
diff --git a/Engine/Source/Programs/Horde/Horde.Server/Issues/IssueService.cs b/Engine/Source/Programs/Horde/Horde.Server/Issues/IssueService.cs
--- a/Engine/Source/Programs/Horde/Horde.Server/Issues/IssueService.cs
+++ b/Engine/Source/Programs/Horde/Horde.Server/Issues/IssueService.cs
@@ -203,7 +203,7 @@ namespace Horde.Server.Issues
 		readonly ICommitService _commitService;
 		readonly IStreamCollection _streams;
 		readonly IUserCollection _userCollection;
-		readonly ILogService _logService;
+		readonly ILogCollection _logCollection;
 		readonly IClock _clock;
 		readonly ITicker _ticker;
 
@@ -255,7 +255,7 @@ namespace Horde.Server.Issues
 		/// <summary>
 		/// Constructor
 		/// </summary>
-		public IssueService(IIssueCollection issueCollection, ICommitService commitService, IJobStepRefCollection jobStepRefs, IStreamCollection streams, IUserCollection userCollection, ILogService logService, IClock clock, IOptionsMonitor<GlobalConfig> globalConfig, Tracer tracer, ILogger<IssueService> logger)
+		public IssueService(IIssueCollection issueCollection, ICommitService commitService, IJobStepRefCollection jobStepRefs, IStreamCollection streams, IUserCollection userCollection, ILogCollection logCollection, IClock clock, IOptionsMonitor<GlobalConfig> globalConfig, Tracer tracer, ILogger<IssueService> logger)
 		{
 			Type[] issueTypes = Assembly.GetExecutingAssembly().GetTypes().Where(x => !x.IsAbstract && typeof(IIssue).IsAssignableFrom(x)).ToArray();
 			foreach (Type issueType in issueTypes)
@@ -269,7 +269,7 @@ namespace Horde.Server.Issues
 			_jobStepRefs = jobStepRefs;
 			_streams = streams;
 			_userCollection = userCollection;
-			_logService = logService;
+			_logCollection = logCollection;
 			_clock = clock;
 			_ticker = clock.AddTicker<IssueService>(TimeSpan.FromMinutes(1.0), TickAsync, logger);
 			_globalConfig = globalConfig;
@@ -679,7 +679,7 @@ namespace Horde.Server.Issues
 			}
 
 			// Find the log file for this step
-			ILog? log = await _logService.GetLogAsync(step.LogId.Value, CancellationToken.None);
+			ILog? log = await _logCollection.GetAsync(step.LogId.Value, CancellationToken.None);
 			if (log == null)
 			{
 				throw new ArgumentException($"Unable to retrieve log {step.LogId}");
@@ -702,10 +702,10 @@ namespace Horde.Server.Issues
 			}
 
 			// Create all the issue definitions by passing each log event to the handlers in order until one attaches it to an issue
-			List<ILogEvent> stepEvents = await _logService.FindEventsAsync(log, cancellationToken: cancellationToken);
+			List<ILogEvent> stepEvents = await log.GetEventsAsync(cancellationToken: cancellationToken);
 			foreach (ILogEvent stepEvent in stepEvents)
 			{
-				ILogEventData stepEventData = await _logService.GetEventDataAsync(log, stepEvent.LineIndex, stepEvent.LineCount, cancellationToken);
+				ILogEventData stepEventData = await stepEvent.GetDataAsync(cancellationToken);
 
 				IssueEventInternal issueEvent = new IssueEventInternal(stepEvent, stepEventData);
 				foreach (IssueHandler handler in handlers)
@@ -809,7 +809,7 @@ namespace Horde.Server.Issues
 						}
 
 						// Assign all the events to the span
-						await _logService.AddSpanToEventsAsync(matchEventGroups.SelectMany(x => x.Events.Select(x => x.Event)), newSpan.Id, cancellationToken);
+						await _logCollection.AddSpanToEventsAsync(matchEventGroups.SelectMany(x => x.Events.Select(x => x.Event)), newSpan.Id, cancellationToken);
 
 						// Remove the matches from the set of events
 						newEventGroups.ExceptWith(matchEventGroups);
@@ -907,7 +907,7 @@ namespace Horde.Server.Issues
 
 				// Update the log events
 				_logger.LogDebug("Created new span {SpanId} from event group {Group}", newSpan.Id, eventGroup.Id.ToString());
-				await _logService.AddSpanToEventsAsync(eventGroup.Events.Select(x => x.Event), newSpan.Id, cancellationToken);
+				await _logCollection.AddSpanToEventsAsync(eventGroup.Events.Select(x => x.Event), newSpan.Id, cancellationToken);
 
 				// Remove the events from the remaining list of events to match
 				newEventGroups.ExceptWith(sourceEventGroups);
diff --git a/Engine/Source/Programs/Horde/Horde.Server/Issues/IssuesController.cs b/Engine/Source/Programs/Horde/Horde.Server/Issues/IssuesController.cs
--- a/Engine/Source/Programs/Horde/Horde.Server/Issues/IssuesController.cs
+++ b/Engine/Source/Programs/Horde/Horde.Server/Issues/IssuesController.cs
@@ -42,20 +42,20 @@ namespace Horde.Server.Issues
 		private readonly IExternalIssueService _externalIssueService;
 		private readonly JobService _jobService;
 		private readonly IUserCollection _userCollection;
-		private readonly ILogService _logService;
+		private readonly ILogCollection _logCollection;
 		private readonly IOptionsSnapshot<GlobalConfig> _globalConfig;
 		private readonly ILogger<IssuesController> _logger;
 
 		/// <summary>
 		/// Constructor
 		/// </summary>
-		public IssuesController(IIssueCollection issueCollection, IssueService issueService, JobService jobService, IUserCollection userCollection, ILogService logService, IExternalIssueService externalIssueService, IOptionsSnapshot<GlobalConfig> globalConfig, ILogger<IssuesController> logger)
+		public IssuesController(IIssueCollection issueCollection, IssueService issueService, JobService jobService, IUserCollection userCollection, ILogCollection logCollection, IExternalIssueService externalIssueService, IOptionsSnapshot<GlobalConfig> globalConfig, ILogger<IssuesController> logger)
 		{
 			_issueCollection = issueCollection;
 			_issueService = issueService;
 			_jobService = jobService;
 			_userCollection = userCollection;
-			_logService = logService;
+			_logCollection = logCollection;
 			_externalIssueService = externalIssueService;
 			_globalConfig = globalConfig;
 			_logger = logger;
@@ -728,7 +728,7 @@ namespace Horde.Server.Issues
 			}
 
 			IReadOnlyList<IIssueSpan> spans = await _issueCollection.FindSpansAsync(issueId, cancellationToken);
-			List<ILogEvent> events = await _logService.FindEventsForSpansAsync(spans.Select(x => x.Id), logIdValues.ToArray(), index, count, cancellationToken);
+			IReadOnlyList<ILogEvent> events = await _logCollection.FindEventsForSpansAsync(spans.Select(x => x.Id), logIdValues.ToArray(), index, count, cancellationToken);
 
 			JobPermissionsCache permissionsCache = new JobPermissionsCache();
 			Dictionary<LogId, ILog?> logs = new Dictionary<LogId, ILog?>();
@@ -739,12 +739,12 @@ namespace Horde.Server.Issues
 				ILog? log;
 				if (!logs.TryGetValue(logEvent.LogId, out log))
 				{
-					log = await _logService.GetLogAsync(logEvent.LogId, cancellationToken);
+					log = await _logCollection.GetAsync(logEvent.LogId, cancellationToken);
 					logs[logEvent.LogId] = log;
 				}
 				if (log != null && await _jobService.AuthorizeAsync(log.JobId, LogAclAction.ViewLog, User, _globalConfig.Value, cancellationToken))
 				{
-					ILogEventData data = await _logService.GetEventDataAsync(log, logEvent.LineIndex, logEvent.LineCount, cancellationToken);
+					ILogEventData data = await logEvent.GetDataAsync(cancellationToken);
 					GetLogEventResponse response = LogsController.CreateGetLogEventResponse(logEvent, data, issueId);
 					responses.Add(PropertyFilter.Apply(response, filter));
 				}
diff --git a/Engine/Source/Programs/Horde/Horde.Server/Jobs/JobRpcService.cs b/Engine/Source/Programs/Horde/Horde.Server/Jobs/JobRpcService.cs
--- a/Engine/Source/Programs/Horde/Horde.Server/Jobs/JobRpcService.cs
+++ b/Engine/Source/Programs/Horde/Horde.Server/Jobs/JobRpcService.cs
@@ -245,7 +245,7 @@ namespace Horde.Server.Jobs
 	{
 		readonly JobService _jobService;
 		readonly IArtifactCollectionV1 _artifactCollection;
-		readonly ILogService _logService;
+		readonly ILogCollection _logCollection;
 		readonly IGraphCollection _graphs;
 		readonly ITestDataCollection _testData;
 		readonly IJobStepRefCollection _jobStepRefCollection;
@@ -258,11 +258,11 @@ namespace Horde.Server.Jobs
 		/// <summary>
 		/// Constructor
 		/// </summary>
-		public JobRpcCommon(JobService jobService, IArtifactCollectionV1 artifactCollection, ILogService logService, IGraphCollection graphs, ITestDataCollection testData, IJobStepRefCollection jobStepRefCollection, ITemplateCollection templateCollection, HttpClient httpClient, IClock clock, IOptionsSnapshot<GlobalConfig> globalConfig, ILogger<JobRpcCommon> logger)
+		public JobRpcCommon(JobService jobService, IArtifactCollectionV1 artifactCollection, ILogCollection logCollection, IGraphCollection graphs, ITestDataCollection testData, IJobStepRefCollection jobStepRefCollection, ITemplateCollection templateCollection, HttpClient httpClient, IClock clock, IOptionsSnapshot<GlobalConfig> globalConfig, ILogger<JobRpcCommon> logger)
 		{
 			_jobService = jobService;
 			_artifactCollection = artifactCollection;
-			_logService = logService;
+			_logCollection = logCollection;
 			_graphs = graphs;
 			_testData = testData;
 			_jobStepRefCollection = jobStepRefCollection;
@@ -487,7 +487,7 @@ namespace Horde.Server.Jobs
 			}
 
 			// Create a log file if necessary
-			log.Value ??= await _logService.CreateLogAsync(job.Id, batch.LeaseId, batch.SessionId, LogType.Json);
+			log.Value ??= await _logCollection.AddAsync(job.Id, batch.LeaseId, batch.SessionId, LogType.Json);
 
 			// Get the node for this step
 			IGraph graph = await _jobService.GetGraphAsync(job);
@@ -885,20 +885,30 @@ namespace Horde.Server.Jobs
 		{
 			if (!_globalConfig.Value.Authorize(LogAclAction.CreateEvent, context.GetHttpContext().User))
 			{
-				throw new StructuredRpcException(StatusCode.NotFound, "Access denied");
+				throw new StructuredRpcException(StatusCode.PermissionDenied, "Access denied");
 			}
 
-			List<NewLogEventData> newEvents = new List<NewLogEventData>();
-			foreach (RpcCreateEventRequest createEvent in request.Events)
+			foreach (IGrouping<string, RpcCreateEventRequest> createEventGroup in request.Events.GroupBy(x => x.LogId))
 			{
-				NewLogEventData newEvent = new NewLogEventData();
-				newEvent.LogId = LogId.Parse(createEvent.LogId);
-				newEvent.Severity = (LogEventSeverity)createEvent.Severity;
-				newEvent.LineIndex = createEvent.LineIndex;
-				newEvent.LineCount = createEvent.LineCount;
-				newEvents.Add(newEvent);
+				ILog? log = await _logCollection.GetAsync(LogId.Parse(createEventGroup.Key), context.CancellationToken);
+				if (log == null)
+				{
+					throw new StructuredRpcException(StatusCode.NotFound, "Log not found");
+				}
+
+				List<NewLogEventData> newEvents = new List<NewLogEventData>();
+				foreach (RpcCreateEventRequest createEvent in createEventGroup)
+				{
+					NewLogEventData newEvent = new NewLogEventData();
+					newEvent.Severity = (LogEventSeverity)createEvent.Severity;
+					newEvent.LineIndex = createEvent.LineIndex;
+					newEvent.LineCount = createEvent.LineCount;
+					newEvents.Add(newEvent);
+				}
+
+				await log.AddEventsAsync(newEvents, context.CancellationToken);
 			}
-			await _logService.CreateEventsAsync(newEvents, context.CancellationToken);
+
 			return new Empty();
 		}
 
diff --git a/Engine/Source/Programs/Horde/Horde.Server/Jobs/JobTaskSource.cs b/Engine/Source/Programs/Horde/Horde.Server/Jobs/JobTaskSource.cs
--- a/Engine/Source/Programs/Horde/Horde.Server/Jobs/JobTaskSource.cs
+++ b/Engine/Source/Programs/Horde/Horde.Server/Jobs/JobTaskSource.cs
@@ -190,7 +190,7 @@ namespace Horde.Server.Jobs
 
 		readonly AclService _aclService;
 		readonly IStreamCollection _streamCollection;
-		readonly ILogService _logService;
+		readonly ILogCollection _logCollection;
 		readonly IAgentCollection _agentsCollection;
 		readonly IJobCollection _jobs;
 		readonly IJobStepRefCollection _jobStepRefs;
@@ -238,7 +238,7 @@ namespace Horde.Server.Jobs
 		/// <summary>
 		/// Constructor
 		/// </summary>
-		public JobTaskSource(AclService aclService, IAgentCollection agents, IJobCollection jobs, IJobStepRefCollection jobStepRefs, IBisectTaskCollection bisectTasks, IGraphCollection graphs, IPoolCollection pools, PoolService poolService, IUgsMetadataCollection ugsMetadataCollection, IStreamCollection streamCollection, ILogService logService, PerforceLoadBalancer perforceLoadBalancer, IClock clock, IOptionsMonitor<ServerSettings> settings, IOptionsMonitor<GlobalConfig> globalConfig, Tracer tracer, ILogger<JobTaskSource> logger)
+		public JobTaskSource(AclService aclService, IAgentCollection agents, IJobCollection jobs, IJobStepRefCollection jobStepRefs, IBisectTaskCollection bisectTasks, IGraphCollection graphs, IPoolCollection pools, PoolService poolService, IUgsMetadataCollection ugsMetadataCollection, IStreamCollection streamCollection, ILogCollection logCollection, PerforceLoadBalancer perforceLoadBalancer, IClock clock, IOptionsMonitor<ServerSettings> settings, IOptionsMonitor<GlobalConfig> globalConfig, Tracer tracer, ILogger<JobTaskSource> logger)
 		{
 			_aclService = aclService;
 			_agentsCollection = agents;
@@ -250,7 +250,7 @@ namespace Horde.Server.Jobs
 			_poolService = poolService;
 			_ugsMetadataCollection = ugsMetadataCollection;
 			_streamCollection = streamCollection;
-			_logService = logService;
+			_logCollection = logCollection;
 			_perforceLoadBalancer = perforceLoadBalancer;
 			_clock = clock;
 			_ticker = clock.AddTicker<JobTaskSource>(s_refreshInterval, TickAsync, logger);
@@ -827,7 +827,7 @@ namespace Horde.Server.Jobs
 					if (waiter.LeaseSource.TrySetResult(lease))
 					{
 						_logger.LogInformation("Assigned lease to agent");
-						await _logService.CreateLogAsync(job.Id, leaseId, agent.SessionId, LogType.Json, logId, cancellationToken);
+						await _logCollection.AddAsync(job.Id, leaseId, agent.SessionId, LogType.Json, logId, cancellationToken);
 						return lease;
 					}
 				}
diff --git a/Engine/Source/Programs/Horde/Horde.Server/Logs/ILog.cs b/Engine/Source/Programs/Horde/Horde.Server/Logs/ILog.cs
--- a/Engine/Source/Programs/Horde/Horde.Server/Logs/ILog.cs
+++ b/Engine/Source/Programs/Horde/Horde.Server/Logs/ILog.cs
@@ -2,6 +2,9 @@
 
 using System.Collections.Generic;
 using System.IO;
+using System.Linq;
+using System.Security.Claims;
+using System.Text.Json;
 using System.Threading;
 using System.Threading.Tasks;
 using EpicGames.Core;
@@ -10,6 +13,9 @@ using EpicGames.Horde.Agents.Sessions;
 using EpicGames.Horde.Jobs;
 using EpicGames.Horde.Logs;
 using EpicGames.Horde.Storage;
+using Horde.Server.Acls;
+using Microsoft.Extensions.Logging;
+using MongoDB.Bson;
 
 namespace Horde.Server.Logs
 {
@@ -63,6 +69,12 @@ namespace Horde.Server.Logs
 		//		/// </summary>
 		//		public bool Complete { get; }
 
+		/// <summary>
+		/// Delete this log
+		/// </summary>
+		/// <param name="cancellationToken">Cancellation token for the operation</param>
+		Task DeleteAsync(CancellationToken cancellationToken = default);
+
 		/// <summary>
 		/// Read a set of lines from the given log file
 		/// </summary>
@@ -122,6 +134,133 @@ namespace Horde.Server.Logs
 		/// <param name="cancellationToken">Cancellation token for the call</param>
 		/// <returns>List of line numbers containing the given term</returns>
 		Task<List<int>> SearchLogDataAsync(string text, int firstLine, int count, SearchStats stats, CancellationToken cancellationToken);
+
+		#region Events
+
+		/// <summary>
+		/// Creates a new event
+		/// </summary>
+		/// <param name="newEvents">List of events to create</param>
+		/// <param name="cancellationToken">Cancellation token for the operation</param>
+		Task AddEventsAsync(List<NewLogEventData> newEvents, CancellationToken cancellationToken = default);
+
+		/// <summary>
+		/// Finds events within a log file
+		/// </summary>
+		/// <param name="spanId">Optional span to filter events by</param>
+		/// <param name="index">Start index within the matching results</param>
+		/// <param name="count">Maximum number of results to return</param>
+		/// <param name="cancellationToken">Cancellation token for the operation</param>
+		/// <returns>List of events matching the query</returns>
+		Task<List<ILogEvent>> GetEventsAsync(ObjectId? spanId = null, int? index = null, int? count = null, CancellationToken cancellationToken = default);
+
+		#endregion
+	}
+
+	/// <summary>
+	/// Represents a node in the graph
+	/// </summary>
+	public interface ILogEvent
+	{
+		/// <summary>
+		/// Unique id of the log containing this event
+		/// </summary>
+		public LogId LogId { get; }
+
+		/// <summary>
+		/// Severity of the event
+		/// </summary>
+		public LogEventSeverity Severity { get; }
+
+		/// <summary>
+		/// Index of the first line for this event
+		/// </summary>
+		public int LineIndex { get; }
+
+		/// <summary>
+		/// Number of lines in the event
+		/// </summary>
+		public int LineCount { get; }
+
+		/// <summary>
+		/// Span id for this log event
+		/// </summary>
+		public ObjectId? SpanId { get; }
+
+		/// <summary>
+		/// Gets the data for this log event
+		/// </summary>
+		/// <param name="cancellationToken">Cancellation token for the operation</param>
+		public Task<ILogEventData> GetDataAsync(CancellationToken cancellationToken = default);
+	}
+
+	/// <summary>
+	/// Interface for event data
+	/// </summary>
+	public interface ILogEventData
+	{
+		/// <summary>
+		/// The log level
+		/// </summary>
+		LogEventSeverity Severity { get; }
+
+		/// <summary>
+		/// The type of event
+		/// </summary>
+		EventId? EventId { get; }
+
+		/// <summary>
+		/// The complete rendered message, in plaintext
+		/// </summary>
+		string Message { get; }
+
+		/// <summary>
+		/// Gets this event data as a JSON objects
+		/// </summary>
+		IReadOnlyList<JsonLogEvent> Lines { get; }
+	}
+
+	/// <summary>
+	/// Represents a node in the graph
+	/// </summary>
+	public class NewLogEventData
+	{
+		/// <summary>
+		/// Severity of the event
+		/// </summary>
+		public LogEventSeverity Severity { get; set; }
+
+		/// <summary>
+		/// Index of the first line for this event
+		/// </summary>
+		public int LineIndex { get; set; }
+
+		/// <summary>
+		/// Number of lines in the event
+		/// </summary>
+		public int LineCount { get; set; }
+
+		/// <summary>
+		/// The span this this event belongs to
+		/// </summary>
+		public ObjectId? SpanId { get; set; }
+	}
+
+	/// <summary>
+	/// Extensions for parsing log event properties
+	/// </summary>
+	public static class LogEventExtensions
+	{
+		/// <summary>
+		/// Find all properties of the given type in a particular log line
+		/// </summary>
+		/// <param name="data">Line data</param>
+		/// <param name="type">Type of property to return</param>
+		/// <returns></returns>
+		public static IEnumerable<JsonProperty> FindPropertiesOfType(this ILogEventData data, Utf8String type)
+		{
+			return data.Lines.SelectMany(x => x.FindPropertiesOfType(type));
+		}
 	}
 
 	/// <summary>
@@ -145,6 +284,36 @@ namespace Horde.Server.Logs
 	/// </summary>
 	public static class LogExtensions
 	{
+		/// <summary>
+		/// Creates a new event
+		/// </summary>
+		/// <param name="log">Log to add an event to</param>
+		/// <param name="newEvent">The new event to vreate</param>
+		/// <param name="cancellationToken">Cancellation token for the operation</param>
+		public static Task AddEventAsync(this ILog log, NewLogEventData newEvent, CancellationToken cancellationToken = default)
+		{
+			return log.AddEventsAsync(new List<NewLogEventData> { newEvent }, cancellationToken);
+		}
+
+		/// <summary>
+		/// Determines if the user is authorized to perform an action on a particular template
+		/// </summary>
+		/// <param name="log">The template to check</param>
+		/// <param name="user">The principal to authorize</param>
+		/// <returns>True if the action is authorized</returns>
+		public static bool AuthorizeForSession(this ILog log, ClaimsPrincipal user)
+		{
+			if (log.SessionId != null && user.HasSessionClaim(log.SessionId.Value))
+			{
+				return true;
+			}
+			if (log.LeaseId != null && user.HasLeaseClaim(log.LeaseId.Value))
+			{
+				return true;
+			}
+			return false;
+		}
+
 		/// <summary>
 		/// Parses a stream of json text and outputs plain text
 		/// </summary>
diff --git a/Engine/Source/Programs/Horde/Horde.Server/Logs/ILogCollection.cs b/Engine/Source/Programs/Horde/Horde.Server/Logs/ILogCollection.cs
--- a/Engine/Source/Programs/Horde/Horde.Server/Logs/ILogCollection.cs
+++ b/Engine/Source/Programs/Horde/Horde.Server/Logs/ILogCollection.cs
@@ -1,11 +1,13 @@
 // Copyright Epic Games, Inc. All Rights Reserved.
 
+using System.Collections.Generic;
 using System.Threading;
 using System.Threading.Tasks;
 using EpicGames.Horde.Agents.Leases;
 using EpicGames.Horde.Agents.Sessions;
 using EpicGames.Horde.Jobs;
 using EpicGames.Horde.Logs;
+using MongoDB.Bson;
 
 namespace Horde.Server.Logs
 {
@@ -32,6 +34,30 @@ namespace Horde.Server.Logs
 		/// <param name="logId">Unique id of the log file</param>
 		/// <param name="cancellationToken">Cancellation token for the call</param>
 		/// <returns>The log instance</returns>
-		Task<ILog?> GetAsync(LogId logId, CancellationToken cancellationToken);
+		Task<ILog?> GetAsync(LogId logId, CancellationToken cancellationToken = default);
+
+		#region Events
+
+		/// <summary>
+		/// Finds a list of events for a set of spans
+		/// </summary>
+		/// <param name="spanIds">The span ids</param>
+		/// <param name="logIds">List of log ids to query</param>
+		/// <param name="index">Index of the first result to return</param>
+		/// <param name="count">Number of results to return</param>
+		/// <param name="cancellationToken">Cancellation token for the call</param>
+		/// <returns>List of events for this issue</returns>
+		Task<IReadOnlyList<ILogEvent>> FindEventsForSpansAsync(IEnumerable<ObjectId> spanIds, LogId[]? logIds = null, int index = 0, int count = 10, CancellationToken cancellationToken = default);
+
+		/// <summary>
+		/// Update the span for an event
+		/// </summary>
+		/// <param name="events">The events to update</param>
+		/// <param name="spanId">New span id</param>
+		/// <param name="cancellationToken">Cancellation token for the call</param>
+		/// <returns>Async task</returns>
+		Task AddSpanToEventsAsync(IEnumerable<ILogEvent> events, ObjectId spanId, CancellationToken cancellationToken = default);
+
+		#endregion
 	}
 }
diff --git a/Engine/Source/Programs/Horde/Horde.Server/Logs/ILogEvent.cs b/Engine/Source/Programs/Horde/Horde.Server/Logs/ILogEvent.cs
deleted file mode 100644
diff --git a/Engine/Source/Programs/Horde/Horde.Server/Logs/ILogEventCollection.cs b/Engine/Source/Programs/Horde/Horde.Server/Logs/ILogEventCollection.cs
deleted file mode 100644
diff --git a/Engine/Source/Programs/Horde/Horde.Server/Logs/LogCollection.cs b/Engine/Source/Programs/Horde/Horde.Server/Logs/LogCollection.cs
--- a/Engine/Source/Programs/Horde/Horde.Server/Logs/LogCollection.cs
+++ b/Engine/Source/Programs/Horde/Horde.Server/Logs/LogCollection.cs
@@ -7,6 +7,7 @@ using System.IO;
 using System.Linq;
 using System.Runtime.CompilerServices;
 using System.Text;
+using System.Text.Json;
 using System.Threading;
 using System.Threading.Tasks;
 using EpicGames.Core;
@@ -21,6 +22,7 @@ using Horde.Server.Storage;
 using Horde.Server.Utilities;
 using Microsoft.Extensions.Caching.Memory;
 using Microsoft.Extensions.Logging;
+using MongoDB.Bson;
 using MongoDB.Bson.Serialization.Attributes;
 using MongoDB.Driver;
 using OpenTelemetry.Trace;
@@ -51,6 +53,9 @@ namespace Horde.Server.Logs
 				_document = document;
 			}
 
+			public Task DeleteAsync(CancellationToken cancellationToken = default)
+				=> _collection.DeleteAsync(_document, cancellationToken);
+
 			public async Task<ILog> UpdateLineCountAsync(int lineCount, bool complete, CancellationToken cancellationToken = default)
 			{
 				LogDocument newDocument = await _collection.UpdateLineCountAsync(_document, lineCount, complete, cancellationToken);
@@ -74,6 +79,12 @@ namespace Horde.Server.Logs
 
 			public Task<List<int>> SearchLogDataAsync(string text, int firstLine, int count, SearchStats stats, CancellationToken cancellationToken)
 				=> _collection.SearchLogDataAsync(_document, text, firstLine, count, stats, cancellationToken);
+
+			public Task AddEventsAsync(List<NewLogEventData> newEvents, CancellationToken cancellationToken = default)
+				=> _collection.AddEventsAsync(_document.Id, newEvents, cancellationToken);
+
+			public Task<List<ILogEvent>> GetEventsAsync(ObjectId? spanId = null, int? index = null, int? count = null, CancellationToken cancellationToken = default)
+				=> _collection.GetEventsAsync(this, spanId, index, count, cancellationToken);
 		}
 
 		class LogDocument
@@ -159,7 +170,94 @@ namespace Horde.Server.Logs
 			}
 		}
 
+		class LogEvent : ILogEvent
+		{
+			readonly ILog _log;
+			readonly LogCollection _collection;
+			readonly LogEventDocument _document;
+
+			LogId ILogEvent.LogId => _document.Id.LogId;
+			LogEventSeverity ILogEvent.Severity => _document.IsWarning ? LogEventSeverity.Warning : LogEventSeverity.Error;
+			int ILogEvent.LineIndex => _document.Id.LineIndex;
+			int ILogEvent.LineCount => _document.LineCount ?? 1;
+			ObjectId? ILogEvent.SpanId => _document.SpanId;
+
+			public LogEvent(ILog log, LogCollection collection, LogEventDocument document)
+			{
+				_log = log;
+				_collection = collection;
+				_document = document;
+			}
+
+			public async Task<ILogEventData> GetDataAsync(CancellationToken cancellationToken = default)
+				=> await _collection.GetEventDataAsync(_log, _document.Id.LineIndex, _document.LineCount ?? 1, cancellationToken);
+		}
+
+		class LogEventId
+		{
+			[BsonElement("l")]
+			public LogId LogId { get; set; }
+
+			[BsonElement("n")]
+			public int LineIndex { get; set; }
+		}
+
+		class LogEventDocument
+		{
+			[BsonId]
+			public LogEventId Id { get; set; }
+
+			[BsonElement("w"), BsonIgnoreIfDefault, BsonDefaultValue(false)]
+			public bool IsWarning { get; set; }
+
+			[BsonElement("c"), BsonIgnoreIfNull]
+			public int? LineCount { get; set; }
+
+			[BsonElement("s")]
+			public ObjectId? SpanId { get; set; }
+
+			[BsonConstructor]
+			public LogEventDocument()
+			{
+				Id = new LogEventId();
+			}
+
+			public LogEventDocument(LogId logId, LogEventSeverity severity, int lineIndex, int lineCount, ObjectId? spanId)
+			{
+				Id = new LogEventId { LogId = logId, LineIndex = lineIndex };
+				IsWarning = severity == LogEventSeverity.Warning;
+				LineCount = (lineCount > 1) ? (int?)lineCount : null;
+				SpanId = spanId;
+			}
+
+			public LogEventDocument(LogId logId, NewLogEventData data)
+				: this(logId, data.Severity, data.LineIndex, data.LineCount, data.SpanId)
+			{
+			}
+		}
+
+		class LegacyLogEventDocument
+		{
+			public ObjectId Id { get; set; }
+			public DateTime Time { get; set; }
+			public LogEventSeverity Severity { get; set; }
+			public LogId LogId { get; set; }
+			public int LineIndex { get; set; }
+			public int LineCount { get; set; }
+
+			public string? Message { get; set; }
+
+			[BsonIgnoreIfNull, BsonElement("IssueId2")]
+			public int? IssueId { get; set; }
+
+			public BsonDocument? Data { get; set; }
+
+			public int UpgradeVersion { get; set; }
+		}
+
 		readonly IMongoCollection<LogDocument> _logCollection;
+		readonly IMongoCollection<LogEventDocument> _logEvents;
+		readonly IMongoCollection<LegacyLogEventDocument> _legacyEvents;
 		readonly ILogStorage _storage;
 		readonly LogTailService _logTailService;
 		readonly StorageService _storageService;
@@ -179,6 +277,13 @@ namespace Horde.Server.Logs
 			_tracer = tracer;
 			_logger = logger;
 			_logCache = new MemoryCache(new MemoryCacheOptions());
+
+			List<MongoIndex<LogEventDocument>> logEventIndexes = new List<MongoIndex<LogEventDocument>>();
+			logEventIndexes.Add(keys => keys.Ascending(x => x.Id.LogId));
+			logEventIndexes.Add(keys => keys.Ascending(x => x.SpanId).Ascending(x => x.Id));
+			_logEvents = mongoService.GetCollection<LogEventDocument>("LogEvents", logEventIndexes);
+
+			_legacyEvents = mongoService.GetCollection<LegacyLogEventDocument>("Events", keys => keys.Ascending(x => x.LogId));
 		}
 
 		/// <inheritdoc/>
@@ -195,6 +300,20 @@ namespace Horde.Server.Logs
 			return new Log(this, newLog);
 		}
 
+		/// <inheritdoc/>
+		async Task DeleteAsync(LogDocument logDocument, CancellationToken cancellationToken)
+		{
+			await _logEvents.DeleteManyAsync(x => x.Id.LogId == logDocument.Id, cancellationToken);
+			await _legacyEvents.DeleteManyAsync(x => x.LogId == logDocument.Id, cancellationToken);
+			await _logCollection.DeleteOneAsync(x => x.Id == logDocument.Id, cancellationToken);
+
+			if (logDocument.UseNewStorageBackend)
+			{
+				using IStorageClient storageClient = _storageService.CreateClient(logDocument.NamespaceId);
+				await storageClient.DeleteRefAsync(logDocument.RefName, cancellationToken);
+			}
+		}
+
 		/// <inheritdoc/>
 		public async Task<ILog?> GetAsync(LogId logId, CancellationToken cancellationToken)
 		{
@@ -348,7 +467,7 @@ namespace Horde.Server.Logs
 		{
 			Stopwatch timer = Stopwatch.StartNew();
 
-			using TelemetrySpan span = _tracer.StartActiveSpan($"{nameof(LogService)}.{nameof(SearchLogDataAsync)}");
+			using TelemetrySpan span = _tracer.StartActiveSpan($"{nameof(LogCollection)}.{nameof(SearchLogDataAsync)}");
 			span.SetAttribute("logId", log.Id.ToString());
 			span.SetAttribute("text", text);
 			span.SetAttribute("count", count);
@@ -541,6 +660,123 @@ namespace Horde.Server.Logs
 			}
 		}
 
+		#region Log events
+
+		/// <inheritdoc/>
+		Task AddEventsAsync(LogId logId, List<NewLogEventData> newEvents, CancellationToken cancellationToken)
+		{
+			return _logEvents.InsertManyAsync(newEvents.ConvertAll(x => new LogEventDocument(logId, x)), cancellationToken: cancellationToken);
+		}
+
+		/// <inheritdoc/>
+		async Task<List<ILogEvent>> GetEventsAsync(ILog log, ObjectId? spanId = null, int? index = null, int? count = null, CancellationToken cancellationToken = default)
+		{
+			_logger.LogInformation("Querying for log events for log {LogId}", log.Id);
+
+			FilterDefinitionBuilder<LogEventDocument> builder = Builders<LogEventDocument>.Filter;
+
+			FilterDefinition<LogEventDocument> filter = builder.Eq(x => x.Id.LogId, log.Id);
+			if (spanId != null)
+			{
+				filter &= builder.Eq(x => x.SpanId, spanId.Value);
+			}
+
+			IFindFluent<LogEventDocument, LogEventDocument> results = _logEvents.Find(filter).SortBy(x => x.Id);
+			if (index != null)
+			{
+				results = results.Skip(index.Value);
+			}
+			if (count != null)
+			{
+				results = results.Limit(count.Value);
+			}
+
+			List<LogEventDocument> logEventDocuments = await results.ToListAsync(cancellationToken);
+			return logEventDocuments.ConvertAll<ILogEvent>(x => new LogEvent(log, this, x));
+		}
+
+		/// <inheritdoc/>
+		public async Task<IReadOnlyList<ILogEvent>> FindEventsForSpansAsync(IEnumerable<ObjectId> spanIds, LogId[]? logIds, int index, int count, CancellationToken cancellationToken)
+		{
+			FilterDefinition<LogEventDocument> filter = Builders<LogEventDocument>.Filter.In(x => x.SpanId, spanIds.Select<ObjectId, ObjectId?>(x => x));
+			if (logIds != null && logIds.Length > 0)
+			{
+				filter &= Builders<LogEventDocument>.Filter.In(x => x.Id.LogId, logIds);
+			}
+
+			List<LogEvent> logEvents = new List<LogEvent>();
+
+			List<LogEventDocument> logEventDocuments = await _logEvents.Find(filter).Skip(index).Limit(count).ToListAsync(cancellationToken);
+			foreach (IGrouping<LogId, LogEventDocument> logEventGroup in logEventDocuments.GroupBy(x => x.Id.LogId))
+			{
+				ILog? log = await GetAsync(logEventGroup.Key, cancellationToken);
+				if (log != null)
+				{
+					logEvents.AddRange(logEventGroup.Select(x => new LogEvent(log, this, x)));
+				}
+			}
+
+			return logEvents;
+		}
+
+		/// <inheritdoc/>
+		public async Task AddSpanToEventsAsync(IEnumerable<ILogEvent> events, ObjectId spanId, CancellationToken cancellationToken)
+		{
+			FilterDefinition<LogEventDocument> eventFilter = Builders<LogEventDocument>.Filter.In(x => x.Id, events.OfType<LogEventDocument>().Select(x => x.Id));
+			UpdateDefinition<LogEventDocument> eventUpdate = Builders<LogEventDocument>.Update.Set(x => x.SpanId, spanId);
+			await _logEvents.UpdateManyAsync(eventFilter, eventUpdate, cancellationToken: cancellationToken);
+		}
+
+		class LogEventData : ILogEventData
+		{
+			public string? _message;
+			public IReadOnlyList<JsonLogEvent> Lines { get; }
+
+			EventId? ILogEventData.EventId => (Lines.Count > 0) ? Lines[0].EventId : null;
+			LogEventSeverity ILogEventData.Severity => (Lines.Count == 0) ? LogEventSeverity.Information : (Lines[0].Level == LogLevel.Warning) ? LogEventSeverity.Warning : LogEventSeverity.Error;
+
+			public LogEventData(IReadOnlyList<JsonLogEvent> lines)
+			{
+				Lines = lines;
+			}
+
+			string ILogEventData.Message
+			{
+				get
+				{
+					_message ??= String.Join("\n", Lines.Select(x => x.GetRenderedMessage().ToString()));
+					return _message;
+				}
+			}
+		}
+
+		async Task<LogEventData> GetEventDataAsync(ILog log, int lineIndex, int lineCount, CancellationToken cancellationToken)
+		{
+			using TelemetrySpan span = _tracer.StartActiveSpan($"{nameof(LogCollection)}.{nameof(GetEventDataAsync)}");
+			span.SetAttribute("logId", log.Id.ToString());
+			span.SetAttribute("lineIndex", lineIndex);
+			span.SetAttribute("lineCount", lineCount);
+
+			List<Utf8String> lines = await log.ReadLinesAsync(lineIndex, lineCount, cancellationToken);
+			List<JsonLogEvent> jsonLines = new List<JsonLogEvent>(lines.Count);
+
+			foreach (Utf8String line in lines)
+			{
+				try
+				{
+					jsonLines.Add(JsonLogEvent.Parse(line.Memory));
+				}
+				catch (JsonException ex)
+				{
+					_logger.LogWarning(ex, "Unable to parse line from log file: {Line}", line);
+				}
+			}
+
+			return new LogEventData(jsonLines);
+		}
+
+		#endregion
+
 		/// <summary>
 		/// Streams log data to a caller
 		/// </summary>
@@ -715,6 +951,9 @@ namespace Horde.Server.Logs
 				return new Log(_collection, newDocument);
 			}
 
+			public Task DeleteAsync(CancellationToken cancellationToken = default)
+				=> _collection.DeleteAsync(_document, cancellationToken);
+
 			public Task<List<Utf8String>> ReadLinesAsync(int index, int count, CancellationToken cancellationToken = default)
 				=> _collection.ReadLinesV1Async(_document, index, count, cancellationToken);
 
@@ -732,6 +971,12 @@ namespace Horde.Server.Logs
 
 			public Task<List<int>> SearchLogDataAsync(string text, int firstLine, int count, SearchStats stats, CancellationToken cancellationToken)
 				=> _collection.SearchLogDataV1Async(_document, text, firstLine, count, stats, cancellationToken);
+
+			public Task AddEventsAsync(List<NewLogEventData> newEvents, CancellationToken cancellationToken = default)
+				=> _collection.AddEventsAsync(_document.Id, newEvents, cancellationToken);
+
+			public Task<List<ILogEvent>> GetEventsAsync(ObjectId? spanId = null, int? index = null, int? count = null, CancellationToken cancellationToken = default)
+				=> _collection.GetEventsAsync(this, spanId, index, count, cancellationToken);
 		}
 
 		async Task<List<Utf8String>> ReadLinesV1Async(LogDocument log, int index, int count, CancellationToken cancellationToken)
@@ -817,7 +1062,7 @@ namespace Horde.Server.Logs
 		{
 			Stopwatch timer = Stopwatch.StartNew();
 
-			using TelemetrySpan span = _tracer.StartActiveSpan($"{nameof(LogService)}.{nameof(SearchLogDataAsync)}");
+			using TelemetrySpan span = _tracer.StartActiveSpan($"{nameof(LogCollection)}.{nameof(SearchLogDataAsync)}");
 			span.SetAttribute("logId", log.Id.ToString());
 			span.SetAttribute("text", text);
 			span.SetAttribute("count", count);
@@ -930,7 +1175,7 @@ namespace Horde.Server.Logs
 				LogIndexData? indexData = await ReadIndexV1Async(log, log.IndexLength.Value);
 				if (indexData != null && firstLine < indexData.LineCount)
 				{
-					using TelemetrySpan span = _tracer.StartActiveSpan($"{nameof(LogService)}.{nameof(SearchLogDataInternalV1Async)}.Indexed");
+					using TelemetrySpan span = _tracer.StartActiveSpan($"{nameof(LogCollection)}.{nameof(SearchLogDataInternalV1Async)}.Indexed");
 					span.SetAttribute("lineCount", indexData.LineCount);
 
 					foreach (int lineIndex in indexData.Search(firstLine, searchText, searchStats))
diff --git a/Engine/Source/Programs/Horde/Horde.Server/Logs/LogEventCollection.cs b/Engine/Source/Programs/Horde/Horde.Server/Logs/LogEventCollection.cs
deleted file mode 100644
diff --git a/Engine/Source/Programs/Horde/Horde.Server/Logs/LogRpcService.cs b/Engine/Source/Programs/Horde/Horde.Server/Logs/LogRpcService.cs
--- a/Engine/Source/Programs/Horde/Horde.Server/Logs/LogRpcService.cs
+++ b/Engine/Source/Programs/Horde/Horde.Server/Logs/LogRpcService.cs
@@ -21,7 +21,7 @@ namespace Horde.Server.Logs
 	[Authorize]
 	public class LogRpcService : LogRpc.LogRpcBase
 	{
-		readonly ILogService _logService;
+		readonly ILogCollection _logCollection;
 		readonly LogTailService _logTailService;
 		readonly StorageService _storageService;
 		readonly ILogger<LogRpcService> _logger;
@@ -29,9 +29,9 @@ namespace Horde.Server.Logs
 		/// <summary>
 		/// Constructor
 		/// </summary>
-		public LogRpcService(ILogService logService, LogTailService logTailService, StorageService storageService, ILogger<LogRpcService> logger)
+		public LogRpcService(ILogCollection logCollection, LogTailService logTailService, StorageService storageService, ILogger<LogRpcService> logger)
 		{
-			_logService = logService;
+			_logCollection = logCollection;
 			_logTailService = logTailService;
 			_storageService = storageService;
 			_logger = logger;
@@ -40,12 +40,12 @@ namespace Horde.Server.Logs
 		/// <inheritdoc/>
 		public override async Task<UpdateLogResponse> UpdateLog(UpdateLogRequest request, ServerCallContext context)
 		{
-			ILog? log = await _logService.GetLogAsync(LogId.Parse(request.LogId), context.CancellationToken);
+			ILog? log = await _logCollection.GetAsync(LogId.Parse(request.LogId), context.CancellationToken);
 			if (log == null)
 			{
 				throw new StructuredRpcException(StatusCode.NotFound, "Resource not found");
 			}
-			if (!LogService.AuthorizeForSession(log, context.GetHttpContext().User))
+			if (!log.AuthorizeForSession(context.GetHttpContext().User))
 			{
 				throw new StructuredRpcException(StatusCode.PermissionDenied, "Access denied");
 			}
diff --git a/Engine/Source/Programs/Horde/Horde.Server/Logs/LogService.cs b/Engine/Source/Programs/Horde/Horde.Server/Logs/LogService.cs
deleted file mode 100644
diff --git a/Engine/Source/Programs/Horde/Horde.Server/Logs/LogsController.cs b/Engine/Source/Programs/Horde/Horde.Server/Logs/LogsController.cs
--- a/Engine/Source/Programs/Horde/Horde.Server/Logs/LogsController.cs
+++ b/Engine/Source/Programs/Horde/Horde.Server/Logs/LogsController.cs
@@ -52,7 +52,7 @@ namespace Horde.Server.Logs
 	[Route("[controller]")]
 	public class LogsController : ControllerBase
 	{
-		private readonly ILogService _logService;
+		private readonly ILogCollection _logCollection;
 		private readonly IIssueCollection _issueCollection;
 		private readonly JobService _jobService;
 		private readonly StorageService _storageService;
@@ -61,9 +61,9 @@ namespace Horde.Server.Logs
 		/// <summary>
 		/// Constructor
 		/// </summary>
-		public LogsController(ILogService logService, IIssueCollection issueCollection, JobService jobService, StorageService storageService, IOptionsSnapshot<GlobalConfig> globalConfig)
+		public LogsController(ILogCollection logCollection, IIssueCollection issueCollection, JobService jobService, StorageService storageService, IOptionsSnapshot<GlobalConfig> globalConfig)
 		{
-			_logService = logService;
+			_logCollection = logCollection;
 			_issueCollection = issueCollection;
 			_jobService = jobService;
 			_storageService = storageService;
@@ -82,7 +82,7 @@ namespace Horde.Server.Logs
 		[ProducesResponseType(typeof(GetLogResponse), 200)]
 		public async Task<ActionResult<object>> GetLogAsync(LogId logId, [FromQuery] PropertyFilter? filter = null, CancellationToken cancellationToken = default)
 		{
-			ILog? log = await _logService.GetLogAsync(logId, cancellationToken);
+			ILog? log = await _logCollection.GetAsync(logId, cancellationToken);
 			if (log == null)
 			{
 				return NotFound();
@@ -120,7 +120,7 @@ namespace Horde.Server.Logs
 		[ProducesResponseType(typeof(WriteBlobResponse), 200)]
 		public async Task<ActionResult<WriteBlobResponse>> WriteLogBlobAsync(LogId logId, WriteBlobRequest request, CancellationToken cancellationToken = default)
 		{
-			ILog? log = await _logService.GetLogAsync(logId, cancellationToken);
+			ILog? log = await _logCollection.GetAsync(logId, cancellationToken);
 			if (log == null)
 			{
 				return NotFound();
@@ -158,7 +158,7 @@ namespace Horde.Server.Logs
 			[FromQuery] bool download = false,
 			CancellationToken cancellationToken = default)
 		{
-			ILog? log = await _logService.GetLogAsync(logId, cancellationToken);
+			ILog? log = await _logCollection.GetAsync(logId, cancellationToken);
 			if (log == null)
 			{
 				return NotFound();
@@ -193,7 +193,7 @@ namespace Horde.Server.Logs
 		[Route("/api/v1/logs/{logId}/lines")]
 		public async Task<ActionResult> GetLogLinesAsync(LogId logId, [FromQuery] int index = 0, [FromQuery] int count = 100, CancellationToken cancellationToken = default)
 		{
-			ILog? log = await _logService.GetLogAsync(logId, cancellationToken);
+			ILog? log = await _logCollection.GetAsync(logId, cancellationToken);
 			if (log == null)
 			{
 				return NotFound();
@@ -288,7 +288,7 @@ namespace Horde.Server.Logs
 			[FromQuery] int count = 5,
 			CancellationToken cancellationToken = default)
 		{
-			ILog? log = await _logService.GetLogAsync(logId, cancellationToken);
+			ILog? log = await _logCollection.GetAsync(logId, cancellationToken);
 			if (log == null)
 			{
 				return NotFound();
@@ -317,7 +317,7 @@ namespace Horde.Server.Logs
 		[ProducesResponseType(typeof(List<GetLogEventResponse>), 200)]
 		public async Task<ActionResult<List<GetLogEventResponse>>> GetEventsAsync(LogId logId, [FromQuery] int? index = null, [FromQuery] int? count = null, CancellationToken cancellationToken = default)
 		{
-			ILog? log = await _logService.GetLogAsync(logId, cancellationToken);
+			ILog? log = await _logCollection.GetAsync(logId, cancellationToken);
 			if (log == null)
 			{
 				return NotFound();
@@ -327,14 +327,14 @@ namespace Horde.Server.Logs
 				return Forbid();
 			}
 
-			List<ILogEvent> logEvents = await _logService.FindEventsAsync(log, null, index, count, cancellationToken);
+			List<ILogEvent> logEvents = await log.GetEventsAsync(null, index, count, cancellationToken);
 
 			Dictionary<ObjectId, int?> spanIdToIssueId = new Dictionary<ObjectId, int?>();
 
 			List<GetLogEventResponse> responses = new List<GetLogEventResponse>();
 			foreach (ILogEvent logEvent in logEvents)
 			{
-				ILogEventData logEventData = await _logService.GetEventDataAsync(log, logEvent.LineIndex, logEvent.LineCount, cancellationToken);
+				ILogEventData logEventData = await logEvent.GetDataAsync(cancellationToken);
 
 				int? issueId = null;
 				if (logEvent.SpanId != null && !spanIdToIssueId.TryGetValue(logEvent.SpanId.Value, out issueId))
@@ -386,7 +386,7 @@ namespace Horde.Server.Logs
 			{
 				return true;
 			}
-			if (log.SessionId != null && LogService.AuthorizeForSession(log, user))
+			if (log.SessionId != null && log.AuthorizeForSession(user))
 			{
 				return true;
 			}
diff --git a/Engine/Source/Programs/Horde/Horde.Server/Notifications/NotificationService.cs b/Engine/Source/Programs/Horde/Horde.Server/Notifications/NotificationService.cs
--- a/Engine/Source/Programs/Horde/Horde.Server/Notifications/NotificationService.cs
+++ b/Engine/Source/Programs/Horde/Horde.Server/Notifications/NotificationService.cs
@@ -76,9 +76,9 @@ namespace Horde.Server.Notifications
 		private readonly IssueService _issueService;
 
 		/// <summary>
-		/// Instance of the <see cref="_logService"/>.
+		/// Instance of the <see cref="ILogCollection"/>.
 		/// </summary>
-		private readonly ILogService _logService;
+		private readonly ILogCollection _logCollection;
 
 		/// <summary>
 		/// Cache for de-duplicating queued notifications
@@ -145,7 +145,7 @@ namespace Horde.Server.Notifications
 			IUserCollection userCollection,
 			JobService jobService,
 			IssueService issueService,
-			ILogService logService,
+			ILogCollection logCollection,
 			Meter meter,
 			IMemoryCache cache,
 			RedisService redisService,
@@ -161,7 +161,7 @@ namespace Horde.Server.Notifications
 			_userCollection = userCollection;
 			_jobService = jobService;
 			_issueService = issueService;
-			_logService = logService;
+			_logCollection = logCollection;
 			_cache = cache;
 			_redisConnectionPool = redisService.ConnectionPool;
 			_backgroundTask = new BackgroundTask(ExecuteAsync);
@@ -676,18 +676,18 @@ namespace Horde.Server.Notifications
 				return;
 			}
 
-			ILog? log = await _logService.GetLogAsync(step.LogId.Value, cancellationToken);
+			ILog? log = await _logCollection.GetAsync(step.LogId.Value, cancellationToken);
 			if (log == null)
 			{
 				_logger.LogError("Step does not have a log file");
 				return;
 			}
 
-			List<ILogEvent> jobStepEvents = await _logService.FindEventsAsync(log, cancellationToken: cancellationToken);
+			List<ILogEvent> jobStepEvents = await log.GetEventsAsync(cancellationToken: cancellationToken);
 			List<ILogEventData> jobStepEventData = new List<ILogEventData>();
 			foreach (ILogEvent logEvent in jobStepEvents)
 			{
-				ILogEventData eventData = await _logService.GetEventDataAsync(log, logEvent.LineIndex, logEvent.LineCount, cancellationToken);
+				ILogEventData eventData = await logEvent.GetDataAsync(cancellationToken);
 				jobStepEventData.Add(eventData);
 			}
 
diff --git a/Engine/Source/Programs/Horde/Horde.Server/Notifications/Sinks/SlackNotificationSink.cs b/Engine/Source/Programs/Horde/Horde.Server/Notifications/Sinks/SlackNotificationSink.cs
--- a/Engine/Source/Programs/Horde/Horde.Server/Notifications/Sinks/SlackNotificationSink.cs
+++ b/Engine/Source/Programs/Horde/Horde.Server/Notifications/Sinks/SlackNotificationSink.cs
@@ -214,7 +214,7 @@ namespace Horde.Server.Notifications.Sinks
 		readonly RedisService _redisService;
 		readonly IssueService _issueService;
 		readonly IUserCollection _userCollection;
-		readonly ILogService _logService;
+		readonly ILogCollection _logCollection;
 		readonly IWebHostEnvironment _environment;
 		readonly ServerSettings _settings;
 		readonly IMongoCollection<MessageStateDocument> _messageStates;
@@ -246,12 +246,12 @@ namespace Horde.Server.Notifications.Sinks
 		/// <summary>
 		/// Constructor
 		/// </summary>
-		public SlackNotificationSink(MongoService mongoService, RedisService redisService, IssueService issueService, IUserCollection userCollection, ILogService logService, IExternalIssueService externalIssueService, IWebHostEnvironment environment, IOptions<ServerSettings> settings, IClock clock, IOptionsMonitor<GlobalConfig> globalConfig, ILogger<SlackNotificationSink> logger)
+		public SlackNotificationSink(MongoService mongoService, RedisService redisService, IssueService issueService, IUserCollection userCollection, ILogCollection logCollection, IExternalIssueService externalIssueService, IWebHostEnvironment environment, IOptions<ServerSettings> settings, IClock clock, IOptionsMonitor<GlobalConfig> globalConfig, ILogger<SlackNotificationSink> logger)
 		{
 			_redisService = redisService;
 			_issueService = issueService;
 			_userCollection = userCollection;
-			_logService = logService;
+			_logCollection = logCollection;
 			_externalIssueService = externalIssueService;
 			_environment = environment;
 			_settings = settings.Value;
@@ -1014,10 +1014,10 @@ namespace Horde.Server.Notifications.Sinks
 				if (span.FirstFailure.LogId != null)
 				{
 					LogId logId = span.FirstFailure.LogId.Value;
-					ILog? log = await _logService.GetLogAsync(logId, cancellationToken);
+					ILog? log = await _logCollection.GetAsync(logId, cancellationToken);
 					if (log != null)
 					{
-						events = await _logService.FindEventsAsync(log, span.Id, 0, 50, cancellationToken);
+						events = await log.GetEventsAsync(span.Id, 0, 50, cancellationToken);
 						if (events.Any(x => x.Severity == LogEventSeverity.Error))
 						{
 							events.RemoveAll(x => x.Severity == LogEventSeverity.Warning);
@@ -1026,7 +1026,7 @@ namespace Horde.Server.Notifications.Sinks
 						List<string> eventStrings = new List<string>();
 						for (int idx = 0; idx < Math.Min(events.Count, 3); idx++)
 						{
-							ILogEventData data = await _logService.GetEventDataAsync(log, events[idx].LineIndex, events[idx].LineCount, cancellationToken);
+							ILogEventData data = await events[idx].GetDataAsync(cancellationToken);
 							eventDataItems.Add(data);
 						}
 					}
@@ -1450,10 +1450,10 @@ namespace Horde.Server.Notifications.Sinks
 			if (lastSpan != null && lastSpan.LastFailure.LogId != null)
 			{
 				LogId logId = lastSpan.LastFailure.LogId.Value;
-				ILog? log = await _logService.GetLogAsync(logId, cancellationToken);
+				ILog? log = await _logCollection.GetAsync(logId, cancellationToken);
 				if (log != null)
 				{
-					List<ILogEvent> events = await _logService.FindEventsAsync(log, lastSpan.Id, 0, 20, cancellationToken);
+					List<ILogEvent> events = await log.GetEventsAsync(lastSpan.Id, 0, 20, cancellationToken);
 					if (events.Any(x => x.Severity == LogEventSeverity.Error))
 					{
 						events.RemoveAll(x => x.Severity == LogEventSeverity.Warning);
@@ -1461,7 +1461,7 @@ namespace Horde.Server.Notifications.Sinks
 
 					for (int idx = 0; idx < Math.Min(events.Count, 3); idx++)
 					{
-						ILogEventData data = await _logService.GetEventDataAsync(log, events[idx].LineIndex, events[idx].LineCount, cancellationToken);
+						ILogEventData data = await events[idx].GetDataAsync(cancellationToken);
 						attachment.AddSection(QuoteText(data.Message));
 					}
 					if (events.Count > 3)
diff --git a/Engine/Source/Programs/Horde/Horde.Server/Startup.cs b/Engine/Source/Programs/Horde/Horde.Server/Startup.cs
--- a/Engine/Source/Programs/Horde/Horde.Server/Startup.cs
+++ b/Engine/Source/Programs/Horde/Horde.Server/Startup.cs
@@ -499,7 +499,6 @@ namespace Horde.Server
 			services.AddSingleton<IJobStepRefCollection, JobStepRefCollection>();
 			services.AddSingleton<IJobTimingCollection, JobTimingCollection>();
 			services.AddSingleton<ILeaseCollection, LeaseCollection>();
-			services.AddSingleton<ILogEventCollection, LogEventCollection>();
 			services.AddSingleton<ILogCollection, LogCollection>();
 			services.AddSingleton<INotificationTriggerCollection, NotificationTriggerCollection>();
 			services.AddSingleton<IPoolCollection, PoolCollection>();
@@ -573,7 +572,6 @@ namespace Horde.Server
 			services.AddSingleton<EnrollmentService>();
 			services.AddSingleton<ServerStatusService>();
 			services.AddHostedService(provider => provider.GetRequiredService<ServerStatusService>());
-			services.AddSingleton<ILogService, LogService>();
 			services.AddSingleton<LogTailService>();
 			services.AddHostedService(provider => provider.GetRequiredService<LogTailService>()); // NB: Runs even on workers, to receive tail notifications.
 			services.AddSingleton<INotificationService, NotificationService>();
diff --git a/Engine/Source/Programs/Horde/Horde.Server/Tasks/ConformTaskSource.cs b/Engine/Source/Programs/Horde/Horde.Server/Tasks/ConformTaskSource.cs
--- a/Engine/Source/Programs/Horde/Horde.Server/Tasks/ConformTaskSource.cs
+++ b/Engine/Source/Programs/Horde/Horde.Server/Tasks/ConformTaskSource.cs
@@ -43,7 +43,7 @@ namespace Horde.Server.Tasks
 		readonly PoolService _poolService;
 		readonly SingletonDocument<ConformList> _conformList;
 		readonly PerforceLoadBalancer _perforceLoadBalancer;
-		readonly ILogService _logService;
+		readonly ILogCollection _logCollection;
 		readonly ILogger _logger;
 		readonly IOptionsMonitor<ServerSettings> _settings;
 		readonly IOptionsMonitor<GlobalConfig> _globalConfig;
@@ -52,13 +52,13 @@ namespace Horde.Server.Tasks
 		/// <summary>
 		/// Constructor
 		/// </summary>
-		public ConformTaskSource(MongoService mongoService, IAgentCollection agentCollection, PoolService poolService, ILogService logService, PerforceLoadBalancer perforceLoadBalancer, IClock clock, IOptionsMonitor<ServerSettings> settings, IOptionsMonitor<GlobalConfig> globalConfig, ILogger<ConformTaskSource> logger)
+		public ConformTaskSource(MongoService mongoService, IAgentCollection agentCollection, PoolService poolService, ILogCollection logCollection, PerforceLoadBalancer perforceLoadBalancer, IClock clock, IOptionsMonitor<ServerSettings> settings, IOptionsMonitor<GlobalConfig> globalConfig, ILogger<ConformTaskSource> logger)
 		{
 			_agentCollection = agentCollection;
 			_poolService = poolService;
 			_conformList = new SingletonDocument<ConformList>(mongoService);
 			_perforceLoadBalancer = perforceLoadBalancer;
-			_logService = logService;
+			_logCollection = logCollection;
 			_settings = settings;
 			_globalConfig = globalConfig;
 			_logger = logger;
@@ -179,7 +179,7 @@ namespace Horde.Server.Tasks
 					{
 						try
 						{
-							ILog log = await _logService.CreateLogAsync(JobId.Empty, leaseId, agent.SessionId, LogType.Json, cancellationToken: cancellationToken);
+							ILog log = await _logCollection.AddAsync(JobId.Empty, leaseId, agent.SessionId, LogType.Json, cancellationToken: cancellationToken);
 							task.LogId = log.Id.ToString();
 							task.RemoveUntrackedFiles = agent.RequestFullConform;
 
diff --git a/Engine/Source/Programs/Horde/Horde.Server/Tasks/RestartTaskSource.cs b/Engine/Source/Programs/Horde/Horde.Server/Tasks/RestartTaskSource.cs
--- a/Engine/Source/Programs/Horde/Horde.Server/Tasks/RestartTaskSource.cs
+++ b/Engine/Source/Programs/Horde/Horde.Server/Tasks/RestartTaskSource.cs
@@ -20,11 +20,11 @@ namespace Horde.Server.Tasks
 
 		public override TaskSourceFlags Flags => TaskSourceFlags.AllowWhenDisabled | TaskSourceFlags.AllowDuringDowntime;
 
-		readonly ILogService _logService;
+		readonly ILogCollection _logCollection;
 
-		public RestartTaskSource(ILogService logService)
+		public RestartTaskSource(ILogCollection logCollection)
 		{
-			_logService = logService;
+			_logCollection = logCollection;
 
 			OnLeaseStartedProperties.Add(nameof(RestartTask.LogId), x => LogId.Parse(x.LogId));
 		}
@@ -44,7 +44,7 @@ namespace Horde.Server.Tasks
 			}
 
 			LeaseId leaseId = new LeaseId(BinaryIdUtils.CreateNew());
-			ILog log = await _logService.CreateLogAsync(JobId.Empty, leaseId, agent.SessionId, LogType.Json, cancellationToken: cancellationToken);
+			ILog log = await _logCollection.AddAsync(JobId.Empty, leaseId, agent.SessionId, LogType.Json, cancellationToken: cancellationToken);
 
 			RestartTask task = new RestartTask();
 			task.LogId = log.Id.ToString();
diff --git a/Engine/Source/Programs/Horde/Horde.Server/Tasks/ShutdownTaskSource.cs b/Engine/Source/Programs/Horde/Horde.Server/Tasks/ShutdownTaskSource.cs
--- a/Engine/Source/Programs/Horde/Horde.Server/Tasks/ShutdownTaskSource.cs
+++ b/Engine/Source/Programs/Horde/Horde.Server/Tasks/ShutdownTaskSource.cs
@@ -20,11 +20,11 @@ namespace Horde.Server.Tasks
 
 		public override TaskSourceFlags Flags => TaskSourceFlags.AllowWhenDisabled | TaskSourceFlags.AllowDuringDowntime;
 
-		readonly ILogService _logService;
+		readonly ILogCollection _logCollection;
 
-		public ShutdownTaskSource(ILogService logService)
+		public ShutdownTaskSource(ILogCollection logCollection)
 		{
-			_logService = logService;
+			_logCollection = logCollection;
 			OnLeaseStartedProperties.Add(nameof(ShutdownTask.LogId), x => LogId.Parse(x.LogId));
 		}
 
@@ -40,7 +40,7 @@ namespace Horde.Server.Tasks
 			}
 
 			LeaseId leaseId = new LeaseId(BinaryIdUtils.CreateNew());
-			ILog log = await _logService.CreateLogAsync(JobId.Empty, leaseId, agent.SessionId, LogType.Json, cancellationToken: cancellationToken);
+			ILog log = await _logCollection.AddAsync(JobId.Empty, leaseId, agent.SessionId, LogType.Json, cancellationToken: cancellationToken);
 
 			ShutdownTask task = new ShutdownTask();
 			task.LogId = log.Id.ToString();
diff --git a/Engine/Source/Programs/Horde/Horde.Server/Tasks/UpgradeTaskSource.cs b/Engine/Source/Programs/Horde/Horde.Server/Tasks/UpgradeTaskSource.cs
--- a/Engine/Source/Programs/Horde/Horde.Server/Tasks/UpgradeTaskSource.cs
+++ b/Engine/Source/Programs/Horde/Horde.Server/Tasks/UpgradeTaskSource.cs
@@ -28,15 +28,15 @@ namespace Horde.Server.Tasks
 		public override TaskSourceFlags Flags => TaskSourceFlags.AllowWhenDisabled | TaskSourceFlags.AllowDuringDowntime;
 
 		readonly IToolCollection _toolCollection;
-		readonly ILogService _logService;
+		readonly ILogCollection _logCollection;
 		readonly IOptionsMonitor<GlobalConfig> _globalConfig;
 		readonly IOptions<ServerSettings> _serverSettings;
 		readonly IClock _clock;
 
-		public UpgradeTaskSource(IToolCollection toolCollection, ILogService logService, IOptionsMonitor<GlobalConfig> globalConfig, IOptions<ServerSettings> serverSettings, IClock clock)
+		public UpgradeTaskSource(IToolCollection toolCollection, ILogCollection logCollection, IOptionsMonitor<GlobalConfig> globalConfig, IOptions<ServerSettings> serverSettings, IClock clock)
 		{
 			_toolCollection = toolCollection;
-			_logService = logService;
+			_logCollection = logCollection;
 			_globalConfig = globalConfig;
 			_serverSettings = serverSettings;
 			_clock = clock;
@@ -69,7 +69,7 @@ namespace Horde.Server.Tasks
 			}
 
 			LeaseId leaseId = new LeaseId(BinaryIdUtils.CreateNew());
-			ILog log = await _logService.CreateLogAsync(JobId.Empty, leaseId, agent.SessionId, LogType.Json, cancellationToken: cancellationToken);
+			ILog log = await _logCollection.AddAsync(JobId.Empty, leaseId, agent.SessionId, LogType.Json, cancellationToken: cancellationToken);
 
 			UpgradeTask task = new UpgradeTask();
 			task.SoftwareId = $"{tool.Id}:{deployment.Version}";
diff --git a/Engine/Source/Programs/Horde/Horde.Server/Ugs/UgsController.cs b/Engine/Source/Programs/Horde/Horde.Server/Ugs/UgsController.cs
--- a/Engine/Source/Programs/Horde/Horde.Server/Ugs/UgsController.cs
+++ b/Engine/Source/Programs/Horde/Horde.Server/Ugs/UgsController.cs
@@ -6,7 +6,6 @@ using System.Linq;
 using System.Threading;
 using System.Threading.Tasks;
 using EpicGames.Horde.Issues;
-using EpicGames.Horde.Logs;
 using EpicGames.Horde.Users;
 using Horde.Server.Issues;
 using Horde.Server.Logs;
@@ -26,45 +25,22 @@ namespace Horde.Server.Ugs
 	[Route("[controller]")]
 	public sealed class UgsController : ControllerBase
 	{
-		/// <summary>
-		/// Singleton instance of the issue service
-		/// </summary>
-		private readonly IssueService _issueService;
-
-		/// <summary>
-		/// Collection of metadata documents
-		/// </summary>
-		private readonly IUgsMetadataCollection _ugsMetadataCollection;
-
-		/// <summary>
-		/// Collection of users
-		/// </summary>
-		private readonly IUserCollection _userCollection;
-
-		/// <summary>
-		/// The log file service
-		/// </summary>
-		private readonly ILogService _logService;
-
-		/// <summary>
-		/// Server settings
-		/// </summary>
-		private readonly ServerSettings _settings;
-
-		/// <summary>
-		/// Logger 
-		/// </summary>
+		readonly IssueService _issueService;
+		readonly IUgsMetadataCollection _ugsMetadataCollection;
+		readonly IUserCollection _userCollection;
+		readonly ILogCollection _logCollection;
+		readonly ServerSettings _settings;
 		readonly ILogger _logger;
 
 		/// <summary>
 		/// Constructor
 		/// </summary>
-		public UgsController(IssueService issueService, IUgsMetadataCollection ugsMetadataCollection, IUserCollection userCollection, ILogService logService, IOptionsMonitor<ServerSettings> optionsMonitor, ILogger<UgsController> logger)
+		public UgsController(IssueService issueService, IUgsMetadataCollection ugsMetadataCollection, IUserCollection userCollection, ILogCollection logCollection, IOptionsMonitor<ServerSettings> optionsMonitor, ILogger<UgsController> logger)
 		{
 			_issueService = issueService;
 			_ugsMetadataCollection = ugsMetadataCollection;
 			_userCollection = userCollection;
-			_logService = logService;
+			_logCollection = logCollection;
 			_settings = optionsMonitor.CurrentValue;
 			_logger = logger;
 		}
@@ -270,28 +246,17 @@ namespace Horde.Server.Ugs
 		{
 			List<GetUgsIssueDiagnosticResponse> diagnostics = new List<GetUgsIssueDiagnosticResponse>();
 
-			Dictionary<LogId, ILog?> logs = new Dictionary<LogId, ILog?>();
-
 			IReadOnlyList<IIssueSpan> spans = await _issueService.Collection.FindSpansAsync(issueId, cancellationToken);
-			List<ILogEvent> events = await _logService.FindEventsForSpansAsync(spans.Select(x => x.Id), null, 0, count: 10, cancellationToken);
+			IReadOnlyList<ILogEvent> events = await _logCollection.FindEventsForSpansAsync(spans.Select(x => x.Id), null, 0, count: 10, cancellationToken);
 
 			foreach (ILogEvent logEvent in events)
 			{
-				ILog? log;
-				if (!logs.TryGetValue(logEvent.LogId, out log))
-				{
-					log = await _logService.GetLogAsync(logEvent.LogId, cancellationToken);
-					logs.Add(logEvent.LogId, log);
-				}
-				if (log != null)
-				{
-					ILogEventData eventData = await _logService.GetEventDataAsync(log, logEvent.LineIndex, logEvent.LineCount, cancellationToken);
-					long buildId = logEvent.LogId.GetHashCode();
-					Uri url = new Uri(_settings.DashboardUrl, $"log/{logEvent.LogId}?lineindex={logEvent.LineIndex}");
+				ILogEventData eventData = await logEvent.GetDataAsync(cancellationToken);
+				long buildId = logEvent.LogId.GetHashCode();
+				Uri url = new Uri(_settings.DashboardUrl, $"log/{logEvent.LogId}?lineindex={logEvent.LineIndex}");
 
-					GetUgsIssueDiagnosticResponse diagnostic = new GetUgsIssueDiagnosticResponse(buildId, eventData.Message, url);
-					diagnostics.Add(diagnostic);
-				}
+				GetUgsIssueDiagnosticResponse diagnostic = new GetUgsIssueDiagnosticResponse(buildId, eventData.Message, url);
+				diagnostics.Add(diagnostic);
 			}
 
 			return diagnostics;
