Horde: Move methods for manipulating agents onto the IAgent interface.

#rnx

diff --git a/Engine/Source/Programs/Horde/Horde.Server.Tests/Agents/AgentCollectionTests.cs b/Engine/Source/Programs/Horde/Horde.Server.Tests/Agents/AgentCollectionTests.cs
--- a/Engine/Source/Programs/Horde/Horde.Server.Tests/Agents/AgentCollectionTests.cs
+++ b/Engine/Source/Programs/Horde/Horde.Server.Tests/Agents/AgentCollectionTests.cs
@@ -33,13 +33,13 @@ public class AgentCollectionTests : TestSetup
 	[TestMethod]
 	public async Task AddLeaseAsync()
 	{
-		await AgentCollection.TryAddLeaseAsync(_agent, _lease1);
+		await _agent.TryAddLeaseAsync(_lease1);
 		await UpdateAgentAsync();
 
-		await AgentCollection.TryAddLeaseAsync(_agent, _leaseWithParent3);
+		await _agent.TryAddLeaseAsync(_leaseWithParent3);
 		await UpdateAgentAsync();
 
-		await AgentCollection.TryAddLeaseAsync(_agent, _leaseWithParent4);
+		await _agent.TryAddLeaseAsync(_leaseWithParent4);
 		await UpdateAgentAsync();
 
 		List<LeaseId> leases = await AgentCollection.FindActiveLeaseIdsAsync();
@@ -53,14 +53,14 @@ public class AgentCollectionTests : TestSetup
 	[TestMethod]
 	public async Task StartSessionAsync()
 	{
-		await AgentCollection.TryAddLeaseAsync(_agent, _lease1);
+		await _agent.TryAddLeaseAsync(_lease1);
 		await UpdateAgentAsync();
 
-		await AgentCollection.TryAddLeaseAsync(_agent, _lease2);
+		await _agent.TryAddLeaseAsync(_lease2);
 		await UpdateAgentAsync();
 
-		await AgentCollection.TryStartSessionAsync(_agent, SessionIdUtils.GenerateNewId(), DateTime.UtcNow, AgentStatus.Ok,
-			new List<string>(), new Dictionary<string, int>(), new List<PoolId>(), new List<PoolId>(), DateTime.UtcNow, null);
+		await _agent.TryCreateSessionAsync(new CreateSessionOptions(SessionIdUtils.GenerateNewId(), DateTime.UtcNow, AgentStatus.Ok,
+			new List<string>(), new Dictionary<string, int>(), new List<PoolId>(), new List<PoolId>(), DateTime.UtcNow, null));
 
 		List<LeaseId> leases = await AgentCollection.FindActiveLeaseIdsAsync();
 
@@ -70,13 +70,13 @@ public class AgentCollectionTests : TestSetup
 	[TestMethod]
 	public async Task UpdateSession_WithEmptyLeasesAsync()
 	{
-		await AgentCollection.TryAddLeaseAsync(_agent, _lease1);
+		await _agent.TryAddLeaseAsync(_lease1);
 		await UpdateAgentAsync();
 
-		await AgentCollection.TryAddLeaseAsync(_agent, _lease2);
+		await _agent.TryAddLeaseAsync(_lease2);
 		await UpdateAgentAsync();
 
-		await AgentCollection.TryUpdateSessionAsync(_agent, null, null, null, null, null, new List<AgentLease>());
+		await _agent.TryUpdateSessionAsync(new UpdateSessionOptions { Leases = new List<AgentLease>() });
 
 		List<LeaseId> leases = await AgentCollection.FindActiveLeaseIdsAsync();
 
@@ -86,10 +86,10 @@ public class AgentCollectionTests : TestSetup
 	[TestMethod]
 	public async Task UpdateSession_WithNewLeasesAsync()
 	{
-		await AgentCollection.TryAddLeaseAsync(_agent, _lease1);
+		await _agent.TryAddLeaseAsync(_lease1);
 		await UpdateAgentAsync();
 
-		await AgentCollection.TryUpdateSessionAsync(_agent, null, null, null, null, null, new List<AgentLease> { _lease1, _lease2 });
+		await _agent.TryUpdateSessionAsync(new UpdateSessionOptions { Leases = new List<AgentLease> { _lease1, _lease2 } });
 
 		List<LeaseId> leases = await AgentCollection.FindActiveLeaseIdsAsync();
 
@@ -101,13 +101,13 @@ public class AgentCollectionTests : TestSetup
 	[TestMethod]
 	public async Task UpdateSession_WithOneLeaseRemovedAsync()
 	{
-		await AgentCollection.TryAddLeaseAsync(_agent, _lease1);
+		await _agent.TryAddLeaseAsync(_lease1);
 		await UpdateAgentAsync();
 
-		await AgentCollection.TryAddLeaseAsync(_agent, _lease2);
+		await _agent.TryAddLeaseAsync(_lease2);
 		await UpdateAgentAsync();
 
-		await AgentCollection.TryUpdateSessionAsync(_agent, null, null, null, null, null, new List<AgentLease> { _lease1 });
+		await _agent.TryUpdateSessionAsync(new UpdateSessionOptions{ Leases = new List<AgentLease> { _lease1 } });
 
 		List<LeaseId> leases = await AgentCollection.FindActiveLeaseIdsAsync();
 
@@ -118,13 +118,13 @@ public class AgentCollectionTests : TestSetup
 	[TestMethod]
 	public async Task CancelLeaseAsync()
 	{
-		await AgentCollection.TryAddLeaseAsync(_agent, _lease1);
+		await _agent.TryAddLeaseAsync(_lease1);
 		await UpdateAgentAsync();
 
-		await AgentCollection.TryAddLeaseAsync(_agent, _lease2);
+		await _agent.TryAddLeaseAsync(_lease2);
 		await UpdateAgentAsync();
 
-		await AgentCollection.TryCancelLeaseAsync(_agent, 0);
+		await _agent.TryCancelLeaseAsync(0);
 
 		List<LeaseId> leases = await AgentCollection.FindActiveLeaseIdsAsync();
 
@@ -135,13 +135,13 @@ public class AgentCollectionTests : TestSetup
 	[TestMethod]
 	public async Task TerminateSessionAsync()
 	{
-		await AgentCollection.TryAddLeaseAsync(_agent, _lease1);
+		await _agent.TryAddLeaseAsync(_lease1);
 		await UpdateAgentAsync();
 
-		await AgentCollection.TryAddLeaseAsync(_agent, _lease2);
+		await _agent.TryAddLeaseAsync(_lease2);
 		await UpdateAgentAsync();
 
-		await AgentCollection.TryTerminateSessionAsync(_agent);
+		await _agent.TryTerminateSessionAsync();
 
 		List<LeaseId> leases = await AgentCollection.FindActiveLeaseIdsAsync();
 
@@ -151,13 +151,13 @@ public class AgentCollectionTests : TestSetup
 	[TestMethod]
 	public async Task GetChildLeaseIdsAsync()
 	{
-		await AgentCollection.TryAddLeaseAsync(_agent, _lease1);
+		await _agent.TryAddLeaseAsync(_lease1);
 		await UpdateAgentAsync();
 
-		await AgentCollection.TryAddLeaseAsync(_agent, _leaseWithParent3);
+		await _agent.TryAddLeaseAsync(_leaseWithParent3);
 		await UpdateAgentAsync();
 
-		await AgentCollection.TryAddLeaseAsync(_agent, _leaseWithParent4);
+		await _agent.TryAddLeaseAsync(_leaseWithParent4);
 		await UpdateAgentAsync();
 
 		List<LeaseId> leases = await AgentCollection.GetChildLeaseIdsAsync(_leaseWithParent3.ParentId!.Value);
diff --git a/Engine/Source/Programs/Horde/Horde.Server.Tests/Jobs/JobServiceTests.cs b/Engine/Source/Programs/Horde/Horde.Server.Tests/Jobs/JobServiceTests.cs
--- a/Engine/Source/Programs/Horde/Horde.Server.Tests/Jobs/JobServiceTests.cs
+++ b/Engine/Source/Programs/Horde/Horde.Server.Tests/Jobs/JobServiceTests.cs
@@ -241,7 +241,7 @@ namespace Horde.Server.Tests.Jobs
 			IAgent? agent = await AgentService.CreateAgentAsync("TestAgent", false, "");
 			Assert.IsNotNull(agent);
 
-			agent = await AgentService.Agents.TryUpdateSettingsAsync(agent, enabled: true, pools: new List<PoolId> { new PoolId("win") });
+			agent = await agent.TryUpdateAsync(new UpdateAgentOptions{ Enabled = true, Pools = new List<PoolId> { new PoolId("win") } });
 			Assert.IsNotNull(agent);
 
 			await AgentService.CreateSessionAsync(agent, AgentStatus.Ok, new List<string>(), new Dictionary<string, int>(), null);
diff --git a/Engine/Source/Programs/Horde/Horde.Server.Tests/Jobs/JobTaskSourceTests.cs b/Engine/Source/Programs/Horde/Horde.Server.Tests/Jobs/JobTaskSourceTests.cs
--- a/Engine/Source/Programs/Horde/Horde.Server.Tests/Jobs/JobTaskSourceTests.cs
+++ b/Engine/Source/Programs/Horde/Horde.Server.Tests/Jobs/JobTaskSourceTests.cs
@@ -166,7 +166,7 @@ namespace Horde.Server.Tests.Jobs
 				IAgent? agent = await AgentService.CreateAgentAsync("TestAgent", false, "");
 				Assert.IsNotNull(agent);
 
-				agent = await AgentService.Agents.TryUpdateSettingsAsync(agent, enabled: isAgentEnabled, pools: new List<PoolId> { pool.Id });
+				agent = await agent.TryUpdateAsync(new UpdateAgentOptions { Enabled = isAgentEnabled, Pools = new List<PoolId> { pool.Id } });
 				Assert.IsNotNull(agent);
 
 				await AgentService.CreateSessionAsync(agent, AgentStatus.Ok, new List<string>(), new Dictionary<string, int>(), null);
diff --git a/Engine/Source/Programs/Horde/Horde.Server.Tests/TestSetup.cs b/Engine/Source/Programs/Horde/Horde.Server.Tests/TestSetup.cs
--- a/Engine/Source/Programs/Horde/Horde.Server.Tests/TestSetup.cs
+++ b/Engine/Source/Programs/Horde/Horde.Server.Tests/TestSetup.cs
@@ -412,7 +412,7 @@ namespace Horde.Server.Tests
 			IAgent? agent = await AgentService.CreateAgentAsync("TestAgent" + s_agentIdCounter++, ephemeral, "");
 			Assert.IsNotNull(agent);
 
-			agent = await AgentService.Agents.TryUpdateSettingsAsync(agent, enabled: enabled, pools: new List<PoolId> { poolId });
+			agent = await agent.TryUpdateAsync(new UpdateAgentOptions { Enabled = enabled, Pools = new List<PoolId> { poolId } });
 			Assert.IsNotNull(agent);
 
 			agent = await AgentService.CreateSessionAsync(agent, AgentStatus.Ok, tempProps, resources, null);
@@ -420,17 +420,17 @@ namespace Horde.Server.Tests
 
 			if (workspaces is { Count: > 0 })
 			{
-				await AgentCollection.TryUpdateWorkspacesAsync(agent, workspaces, false);
+				await agent.TryUpdateWorkspacesAsync(workspaces, false);
 			}
 			
 			if (requestShutdown)
 			{
-				await AgentCollection.TryUpdateSettingsAsync(agent, requestShutdown: true);
+				await agent.TryUpdateAsync(new UpdateAgentOptions { RequestShutdown = true });
 			}
 
 			if (lease != null)
 			{
-				await AgentCollection.TryAddLeaseAsync(agent, lease);
+				await agent.TryAddLeaseAsync(lease);
 			}
 
 			Clock.UtcNow = now;
diff --git a/Engine/Source/Programs/Horde/Horde.Server/Agents/AgentCollection.cs b/Engine/Source/Programs/Horde/Horde.Server/Agents/AgentCollection.cs
--- a/Engine/Source/Programs/Horde/Horde.Server/Agents/AgentCollection.cs
+++ b/Engine/Source/Programs/Horde/Horde.Server/Agents/AgentCollection.cs
@@ -2,6 +2,8 @@
 
 using System;
 using System.Collections.Generic;
+using System.Collections.ObjectModel;
+using System.Diagnostics.CodeAnalysis;
 using System.Linq;
 using System.Linq.Expressions;
 using System.Threading;
@@ -28,6 +30,104 @@ namespace Horde.Server.Agents
 	/// </summary>
 	public class AgentCollection : IAgentCollection
 	{
+		class Agent : IAgent
+		{
+			readonly AgentCollection _collection;
+			readonly AgentDocument _document;
+
+			AgentId IAgent.Id => _document.Id;
+			SessionId? IAgent.SessionId => _document.SessionId;
+			DateTime? IAgent.SessionExpiresAt => _document.SessionExpiresAt;
+			AgentStatus IAgent.Status => _document.Status;
+			DateTime? IAgent.LastStatusChange => _document.LastStatusChange;
+			bool IAgent.Enabled => _document.Enabled;
+			bool IAgent.Ephemeral => _document.Ephemeral;
+			bool IAgent.Deleted => _document.Deleted;
+			string? IAgent.Version => _document.Version;
+			string? IAgent.Comment => _document.Comment;
+			IReadOnlyList<string> IAgent.Properties => _document.Properties ?? _document.Capabilities.Devices.FirstOrDefault()?.Properties?.ToList() ?? new List<string>();
+			IReadOnlyDictionary<string, int> IAgent.Resources => _document.Resources ?? _document.Capabilities.Devices.FirstOrDefault()?.Resources ?? new Dictionary<string, int>();
+			string? IAgent.LastUpgradeVersion => _document.LastUpgradeVersion;
+			DateTime? IAgent.LastUpgradeTime => _document.LastUpgradeTime;
+			int? IAgent.UpgradeAttemptCount => _document.UpgradeAttemptCount;
+			IReadOnlyList<PoolId> IAgent.DynamicPools => _document.DynamicPools;
+			IReadOnlyList<PoolId> IAgent.ExplicitPools => _document.Pools;
+			bool IAgent.RequestConform => _document.RequestConform;
+			bool IAgent.RequestFullConform => _document.RequestFullConform;
+			bool IAgent.RequestRestart => _document.RequestRestart;
+			bool IAgent.RequestShutdown => _document.RequestShutdown;
+			bool IAgent.RequestForceRestart => _document.RequestForceRestart;
+			string? IAgent.LastShutdownReason => _document.LastShutdownReason;
+			IReadOnlyList<AgentWorkspaceInfo> IAgent.Workspaces => _document.Workspaces;
+			DateTime IAgent.LastConformTime => _document.LastConformTime;
+			int? IAgent.ConformAttemptCount => _document.ConformAttemptCount;
+			static readonly IReadOnlyList<AgentLease> s_emptyLeases = new List<AgentLease>();
+			IReadOnlyList<AgentLease> IAgent.Leases => _document.Leases ?? s_emptyLeases;
+			string IAgent.EnrollmentKey => _document.EnrollmentKey;
+			DateTime IAgent.UpdateTime => _document.UpdateTime;
+			uint IAgent.UpdateIndex => _document.UpdateIndex;
+			
+			public Agent(AgentCollection collection, AgentDocument document)
+			{
+				_collection = collection;
+				_document = document;
+			}
+
+			public async Task<IAgent?> TryAddLeaseAsync(AgentLease newLease, CancellationToken cancellationToken = default)
+			{
+				AgentDocument? newDocument = await _collection.TryAddLeaseAsync(_document, newLease, cancellationToken);
+				return _collection.CreateAgentObject(newDocument);
+			}
+
+			public async Task<IAgent?> TryCancelLeaseAsync(int leaseIdx, CancellationToken cancellationToken = default)
+			{
+				AgentDocument? newDocument = await _collection.TryCancelLeaseAsync(_document, leaseIdx, cancellationToken);
+				return _collection.CreateAgentObject(newDocument);
+			}
+
+			public async Task<IAgent?> TryDeleteAsync(CancellationToken cancellationToken = default)
+			{
+				AgentDocument? newDocument = await _collection.TryDeleteAsync(_document, cancellationToken);
+				return _collection.CreateAgentObject(newDocument);
+			}
+
+			public async Task<IAgent?> TryResetAsync(bool ephemeral, string enrollmentKey, CancellationToken cancellationToken = default)
+			{
+				AgentDocument? newDocument = await _collection.TryResetAsync(_document, ephemeral, enrollmentKey, cancellationToken);
+				return _collection.CreateAgentObject(newDocument);
+			}
+
+			public async Task<IAgent?> TryCreateSessionAsync(CreateSessionOptions options, CancellationToken cancellationToken = default)
+			{
+				AgentDocument? newDocument = await _collection.TryCreateSessionAsync(_document, options, cancellationToken);
+				return _collection.CreateAgentObject(newDocument);
+			}
+
+			public async Task<IAgent?> TryTerminateSessionAsync(CancellationToken cancellationToken = default)
+			{
+				AgentDocument? newDocument = await _collection.TryTerminateSessionAsync(_document, cancellationToken);
+				return _collection.CreateAgentObject(newDocument);
+			}
+
+			public async Task<IAgent?> TryUpdateSessionAsync(UpdateSessionOptions options, CancellationToken cancellationToken = default)
+			{
+				AgentDocument? newDocument = await _collection.TryUpdateSessionAsync(_document, options, cancellationToken);
+				return _collection.CreateAgentObject(newDocument);
+			}
+
+			public async Task<IAgent?> TryUpdateAsync(UpdateAgentOptions options, CancellationToken cancellationToken = default)
+			{
+				AgentDocument? newDocument = await _collection.TryUpdateSettingsAsync(_document, options, cancellationToken);
+				return _collection.CreateAgentObject(newDocument);
+			}
+
+			public async Task<IAgent?> TryUpdateWorkspacesAsync(List<AgentWorkspaceInfo> workspaces, bool requestConform, CancellationToken cancellationToken = default)
+			{
+				AgentDocument? newDocument = await _collection.TryUpdateWorkspacesAsync(_document, workspaces, requestConform, cancellationToken);
+				return _collection.CreateAgentObject(newDocument);
+			}
+		}
+
 		/// <summary>
 		/// Legacy information about a device attached to an agent
 		/// </summary>
@@ -54,10 +154,8 @@ namespace Horde.Server.Agents
 		/// <summary>
 		/// Concrete implementation of an agent document
 		/// </summary>
-		class AgentDocument : IAgent
+		class AgentDocument
 		{
-			static readonly IReadOnlyList<AgentLease> s_emptyLeases = new List<AgentLease>();
-
 			[BsonRequired, BsonId]
 			public AgentId Id { get; set; }
 
@@ -124,13 +222,6 @@ namespace Horde.Server.Agents
 			public string EnrollmentKey { get; set; } = String.Empty;
 			public string? Comment { get; set; }
 
-			IReadOnlyList<PoolId> IAgent.DynamicPools => DynamicPools;
-			IReadOnlyList<PoolId> IAgent.ExplicitPools => Pools;
-			IReadOnlyList<AgentWorkspaceInfo> IAgent.Workspaces => Workspaces;
-			IReadOnlyList<AgentLease> IAgent.Leases => Leases ?? s_emptyLeases;
-			IReadOnlyList<string> IAgent.Properties => Properties ?? Capabilities.Devices.FirstOrDefault()?.Properties?.ToList() ?? new List<string>();
-			IReadOnlyDictionary<string, int> IAgent.Resources => Resources ?? Capabilities.Devices.FirstOrDefault()?.Resources ?? new Dictionary<string, int>();
-
 			[BsonConstructor]
 			private AgentDocument()
 			{
@@ -165,16 +256,29 @@ namespace Horde.Server.Agents
 			_auditLog = auditLog;
 		}
 
+		[return: NotNullIfNotNull("document")]
+		Agent? CreateAgentObject(AgentDocument? document)
+		{
+			if (document == null)
+			{
+				return null;
+			}
+			else
+			{
+				return new Agent(this, document);
+			}
+		}
+
 		/// <inheritdoc/>
 		public async Task<IAgent> AddAsync(AgentId id, bool ephemeral, string enrollmentKey, CancellationToken cancellationToken)
 		{
 			AgentDocument agent = new AgentDocument(id, ephemeral, enrollmentKey);
 			await _agents.InsertOneAsync(agent, null, cancellationToken);
-			return agent;
+			return CreateAgentObject(agent);
 		}
 
 		/// <inheritdoc/>
-		public async Task<IAgent?> TryResetAsync(IAgent agent, bool ephemeral, string enrollmentKey, CancellationToken cancellationToken = default)
+		async Task<AgentDocument?> TryResetAsync(AgentDocument agent, bool ephemeral, string enrollmentKey, CancellationToken cancellationToken = default)
 		{
 			AgentDocument agentDocument = (AgentDocument)agent;
 
@@ -184,7 +288,7 @@ namespace Horde.Server.Agents
 				.Unset(x => x.Deleted)
 				.Unset(x => x.SessionId);
 
-			IAgent? newAgent = await TryUpdateAsync(agentDocument, update, cancellationToken);
+			AgentDocument? newAgent = await TryUpdateAsync(agentDocument, update, cancellationToken);
 			if (newAgent != null)
 			{
 				await PublishUpdateEventAsync(agent.Id);
@@ -193,10 +297,8 @@ namespace Horde.Server.Agents
 		}
 
 		/// <inheritdoc/>
-		public async Task<IAgent?> TryDeleteAsync(IAgent agentInterface, CancellationToken cancellationToken)
+		async Task<AgentDocument?> TryDeleteAsync(AgentDocument agent, CancellationToken cancellationToken)
 		{
-			AgentDocument agent = (AgentDocument)agentInterface;
-
 			UpdateDefinition<AgentDocument> update = Builders<AgentDocument>.Update
 				.Set(x => x.Deleted, true)
 				.Set(x => x.EnrollmentKey, "")
@@ -214,13 +316,15 @@ namespace Horde.Server.Agents
 		/// <inheritdoc/>
 		public async Task<IAgent?> GetAsync(AgentId agentId, CancellationToken cancellationToken)
 		{
-			return await _agents.Find<AgentDocument>(x => x.Id == agentId).FirstOrDefaultAsync(cancellationToken);
+			AgentDocument? agent = await _agents.Find<AgentDocument>(x => x.Id == agentId).FirstOrDefaultAsync(cancellationToken);
+			return CreateAgentObject(agent);
 		}
 
 		/// <inheritdoc/>
 		public async Task<IReadOnlyList<IAgent>> GetManyAsync(List<AgentId> agentIds, CancellationToken cancellationToken)
 		{
-			return await _agents.Find(p => agentIds.Contains(p.Id)).ToListAsync(cancellationToken);
+			List<AgentDocument> documents = await _agents.Find(p => agentIds.Contains(p.Id)).ToListAsync(cancellationToken);
+			return documents.ConvertAll(x => CreateAgentObject(x));
 		}
 
 		/// <inheritdoc/>
@@ -269,19 +373,22 @@ namespace Horde.Server.Agents
 				search = search.Limit(count.Value);
 			}
 
-			return await search.ToListAsync(cancellationToken);
+			List<AgentDocument> documents = await search.ToListAsync(cancellationToken);
+			return documents.ConvertAll(x => CreateAgentObject(x));
 		}
 
 		/// <inheritdoc/>
 		public async Task<IReadOnlyList<IAgent>> FindExpiredAsync(DateTime utcNow, int maxAgents, CancellationToken cancellationToken)
 		{
-			return await _agents.Find(x => x.SessionId.HasValue && !(x.SessionExpiresAt > utcNow)).Limit(maxAgents).ToListAsync(cancellationToken);
+			List<AgentDocument> documents = await _agents.Find(x => x.SessionId.HasValue && !(x.SessionExpiresAt > utcNow)).Limit(maxAgents).ToListAsync(cancellationToken);
+			return documents.ConvertAll(x => CreateAgentObject(x));
 		}
 
 		/// <inheritdoc/>
 		public async Task<IReadOnlyList<IAgent>> FindDeletedAsync(CancellationToken cancellationToken)
 		{
-			return await _agents.Find(x => x.Deleted).ToListAsync(cancellationToken);
+			List<AgentDocument> documents = await _agents.Find(x => x.Deleted).ToListAsync(cancellationToken);
+			return documents.ConvertAll(x => CreateAgentObject(x));
 		}
 
 		/// <inheritdoc/>
@@ -317,35 +424,33 @@ namespace Horde.Server.Agents
 		}
 
 		/// <inheritdoc/>
-		public async Task<IAgent?> TryUpdateSettingsAsync(IAgent agentInterface, bool? enabled = null, bool? requestConform = null, bool? requestFullConform = null, bool? requestRestart = null, bool? requestShutdown = null, bool? requestForceRestart = null, string? shutdownReason = null, List<PoolId>? pools = null, string? comment = null, CancellationToken cancellationToken = default)
+		async Task<AgentDocument?> TryUpdateSettingsAsync(AgentDocument agent, UpdateAgentOptions options, CancellationToken cancellationToken)
 		{
-			AgentDocument agent = (AgentDocument)agentInterface;
-
 			// Update the database
 			UpdateDefinitionBuilder<AgentDocument> updateBuilder = new UpdateDefinitionBuilder<AgentDocument>();
 
 			List<UpdateDefinition<AgentDocument>> updates = new List<UpdateDefinition<AgentDocument>>();
-			if (pools != null)
+			if (options.Pools != null)
 			{
-				updates.Add(updateBuilder.Set(x => x.Pools, pools));
+				updates.Add(updateBuilder.Set(x => x.Pools, options.Pools));
 			}
-			if (enabled != null)
+			if (options.Enabled != null)
 			{
-				updates.Add(updateBuilder.Set(x => x.Enabled, enabled.Value));
+				updates.Add(updateBuilder.Set(x => x.Enabled, options.Enabled.Value));
 			}
-			if (requestConform != null)
+			if (options.RequestConform != null)
 			{
-				updates.Add(updateBuilder.Set(x => x.RequestConform, requestConform.Value));
+				updates.Add(updateBuilder.Set(x => x.RequestConform, options.RequestConform.Value));
 				updates.Add(updateBuilder.Unset(x => x.ConformAttemptCount));
 			}
-			if (requestFullConform != null)
+			if (options.RequestFullConform != null)
 			{
-				updates.Add(updateBuilder.Set(x => x.RequestFullConform, requestFullConform.Value));
+				updates.Add(updateBuilder.Set(x => x.RequestFullConform, options.RequestFullConform.Value));
 				updates.Add(updateBuilder.Unset(x => x.ConformAttemptCount));
 			}
-			if (requestRestart != null)
+			if (options.RequestRestart != null)
 			{
-				if (requestRestart.Value)
+				if (options.RequestRestart.Value)
 				{
 					updates.Add(updateBuilder.Set(x => x.RequestRestart, true));
 				}
@@ -354,9 +459,9 @@ namespace Horde.Server.Agents
 					updates.Add(updateBuilder.Unset(x => x.RequestRestart));
 				}
 			}
-			if (requestShutdown != null)
+			if (options.RequestShutdown != null)
 			{
-				if (requestShutdown.Value)
+				if (options.RequestShutdown.Value)
 				{
 					updates.Add(updateBuilder.Set(x => x.RequestShutdown, true));
 				}
@@ -365,9 +470,9 @@ namespace Horde.Server.Agents
 					updates.Add(updateBuilder.Unset(x => x.RequestShutdown));
 				}
 			}
-			if (requestForceRestart != null)
+			if (options.RequestForceRestart != null)
 			{
-				if (requestForceRestart.Value)
+				if (options.RequestForceRestart.Value)
 				{
 					updates.Add(updateBuilder.Set(x => x.RequestForceRestart, true));
 				}
@@ -376,17 +481,17 @@ namespace Horde.Server.Agents
 					updates.Add(updateBuilder.Unset(x => x.RequestForceRestart));
 				}
 			}
-			if (shutdownReason != null)
+			if (options.ShutdownReason != null)
 			{
-				updates.Add(updateBuilder.Set(x => x.LastShutdownReason, shutdownReason));
+				updates.Add(updateBuilder.Set(x => x.LastShutdownReason, options.ShutdownReason));
 			}
-			if (comment != null)
+			if (options.Comment != null)
 			{
-				updates.Add(updateBuilder.Set(x => x.Comment, comment));
+				updates.Add(updateBuilder.Set(x => x.Comment, options.Comment));
 			}
 
 			// Apply the update
-			IAgent? newAgent = await TryUpdateAsync(agent, updateBuilder.Combine(updates), cancellationToken);
+			AgentDocument? newAgent = await TryUpdateAsync(agent, updateBuilder.Combine(updates), cancellationToken);
 			if (newAgent != null)
 			{
 				if (newAgent.RequestRestart != agent.RequestRestart || newAgent.RequestConform != agent.RequestConform || newAgent.RequestShutdown != agent.RequestShutdown || newAgent.RequestForceRestart != agent.RequestForceRestart)
@@ -398,42 +503,40 @@ namespace Horde.Server.Agents
 		}
 
 		/// <inheritdoc/>
-		public async Task<IAgent?> TryUpdateSessionAsync(IAgent agentInterface, AgentStatus? status, DateTime? sessionExpiresAt, IReadOnlyList<string>? properties, IReadOnlyDictionary<string, int>? resources, IReadOnlyList<PoolId>? dynamicPools, List<AgentLease>? leases, CancellationToken cancellationToken)
+		async Task<AgentDocument?> TryUpdateSessionAsync(AgentDocument agent, UpdateSessionOptions options, CancellationToken cancellationToken)
 		{
-			AgentDocument agent = (AgentDocument)agentInterface;
-
 			// Create an update definition for the agent
 			UpdateDefinitionBuilder<AgentDocument> updateBuilder = Builders<AgentDocument>.Update;
 			List<UpdateDefinition<AgentDocument>> updates = new List<UpdateDefinition<AgentDocument>>();
 
-			if (status != null && agent.Status != status.Value)
+			if (options.Status != null && agent.Status != options.Status.Value)
 			{
-				updates.Add(updateBuilder.Set(x => x.Status, status.Value));
+				updates.Add(updateBuilder.Set(x => x.Status, options.Status.Value));
 				updates.Add(updateBuilder.Set(x => x.LastStatusChange, _clock.UtcNow));
 			}
-			if (sessionExpiresAt != null)
+			if (options.SessionExpiresAt != null)
 			{
-				updates.Add(updateBuilder.Set(x => x.SessionExpiresAt, sessionExpiresAt.Value));
+				updates.Add(updateBuilder.Set(x => x.SessionExpiresAt, options.SessionExpiresAt.Value));
 			}
-			if (properties != null)
+			if (options.Properties != null)
 			{
-				List<string> newProperties = properties.OrderBy(x => x, StringComparer.OrdinalIgnoreCase).ToList();
-				if (!agentInterface.Properties.SequenceEqual(newProperties, StringComparer.Ordinal))
+				List<string> newProperties = options.Properties.OrderBy(x => x, StringComparer.OrdinalIgnoreCase).ToList();
+				if (!(agent.Properties ?? Enumerable.Empty<string>()).SequenceEqual(newProperties, StringComparer.Ordinal))
 				{
 					updates.Add(updateBuilder.Set(x => x.Properties, newProperties));
 				}
 			}
-			if (resources != null && !ResourcesEqual(resources, agentInterface.Resources))
+			if (options.Resources != null && !ResourcesEqual(options.Resources, agent.Resources))
 			{
-				updates.Add(updateBuilder.Set(x => x.Resources, new Dictionary<string, int>(resources)));
+				updates.Add(updateBuilder.Set(x => x.Resources, new Dictionary<string, int>(options.Resources)));
 			}
-			if (dynamicPools != null && !dynamicPools.SequenceEqual(agent.DynamicPools))
+			if (options.DynamicPools != null && !options.DynamicPools.SequenceEqual(agent.DynamicPools))
 			{
-				updates.Add(updateBuilder.Set(x => x.DynamicPools, new List<PoolId>(dynamicPools)));
+				updates.Add(updateBuilder.Set(x => x.DynamicPools, new List<PoolId>(options.DynamicPools)));
 			}
-			if (leases != null)
+			if (options.Leases != null)
 			{
-				foreach (AgentLease lease in leases)
+				foreach (AgentLease lease in options.Leases)
 				{
 					if (lease.Payload != null && (agent.Leases == null || !agent.Leases.Any(x => x.Id == lease.Id)))
 					{
@@ -442,7 +545,7 @@ namespace Horde.Server.Agents
 				}
 
 				List<AgentLease> currentLeases = agent.Leases ?? new List<AgentLease>();
-				List<AgentLease> newLeases = leases;
+				List<AgentLease> newLeases = options.Leases;
 				List<AgentLease> leasesToAdd = newLeases.Where(nl => currentLeases.All(cl => cl.Id != nl.Id)).ToList();
 				List<AgentLease> leasesToRemove = currentLeases.Where(cl => newLeases.All(nl => nl.Id != cl.Id)).ToList();
 
@@ -456,7 +559,7 @@ namespace Horde.Server.Agents
 					await RemoveActiveLeaseAsync(lease);
 				}
 
-				updates.Add(updateBuilder.Set(x => x.Leases, leases));
+				updates.Add(updateBuilder.Set(x => x.Leases, options.Leases));
 			}
 
 			// If there are no new updates, return immediately. This is important for preventing UpdateSession calls from returning immediately.
@@ -469,8 +572,11 @@ namespace Horde.Server.Agents
 			return await TryUpdateAsync(agent, updateBuilder.Combine(updates), cancellationToken);
 		}
 
-		static bool ResourcesEqual(IReadOnlyDictionary<string, int> dictA, IReadOnlyDictionary<string, int> dictB)
+		static bool ResourcesEqual(IReadOnlyDictionary<string, int>? dictA, IReadOnlyDictionary<string, int>? dictB)
 		{
+			dictA ??= ReadOnlyDictionary<string, int>.Empty;
+			dictB ??= ReadOnlyDictionary<string, int>.Empty;
+
 			if (dictA.Count != dictB.Count)
 			{
 				return false;
@@ -489,9 +595,8 @@ namespace Horde.Server.Agents
 		}
 
 		/// <inheritdoc/>
-		public async Task<IAgent?> TryUpdateWorkspacesAsync(IAgent agentInterface, List<AgentWorkspaceInfo> workspaces, bool requestConform, CancellationToken cancellationToken)
+		async Task<AgentDocument?> TryUpdateWorkspacesAsync(AgentDocument agent, List<AgentWorkspaceInfo> workspaces, bool requestConform, CancellationToken cancellationToken)
 		{
-			AgentDocument agent = (AgentDocument)agentInterface;
 			DateTime lastConformTime = DateTime.UtcNow;
 
 			// Set the new workspaces
@@ -509,41 +614,40 @@ namespace Horde.Server.Agents
 		}
 
 		/// <inheritdoc/>
-		public async Task<IAgent?> TryStartSessionAsync(IAgent agentInterface, SessionId sessionId, DateTime sessionExpiresAt, AgentStatus status, IReadOnlyList<string> properties, IReadOnlyDictionary<string, int> resources, IReadOnlyList<PoolId> pools, IReadOnlyList<PoolId> dynamicPools, DateTime lastStatusChange, string? version, CancellationToken cancellationToken)
+		async Task<AgentDocument?> TryCreateSessionAsync(AgentDocument agent, CreateSessionOptions options, CancellationToken cancellationToken)
 		{
-			AgentDocument agent = (AgentDocument)agentInterface;
-			List<string> newProperties = properties.OrderBy(x => x, StringComparer.OrdinalIgnoreCase).ToList();
-			Dictionary<string, int> newResources = new(resources);
-			List<PoolId> newPools = new(pools);
-			List<PoolId> newDynamicPools = new(dynamicPools);
+			List<string> newProperties = options.Properties.OrderBy(x => x, StringComparer.OrdinalIgnoreCase).ToList();
+			Dictionary<string, int> newResources = new(options.Resources);
+			List<PoolId> newPools = new(options.Pools);
+			List<PoolId> newDynamicPools = new(options.DynamicPools);
 
 			// Reset the agent to use the new session
 			UpdateDefinitionBuilder<AgentDocument> updateBuilder = Builders<AgentDocument>.Update;
 
 			List<UpdateDefinition<AgentDocument>> updates = new List<UpdateDefinition<AgentDocument>>();
-			updates.Add(updateBuilder.Set(x => x.SessionId, sessionId));
-			updates.Add(updateBuilder.Set(x => x.SessionExpiresAt, sessionExpiresAt));
-			updates.Add(updateBuilder.Set(x => x.Status, status));
+			updates.Add(updateBuilder.Set(x => x.SessionId, options.SessionId));
+			updates.Add(updateBuilder.Set(x => x.SessionExpiresAt, options.SessionExpiresAt));
+			updates.Add(updateBuilder.Set(x => x.Status, options.Status));
 			updates.Add(updateBuilder.Unset(x => x.Leases));
 			updates.Add(updateBuilder.Unset(x => x.Deleted));
 			updates.Add(updateBuilder.Set(x => x.Properties, newProperties));
 			updates.Add(updateBuilder.Set(x => x.Resources, newResources));
 			updates.Add(updateBuilder.Set(x => x.Pools, newPools));
 			updates.Add(updateBuilder.Set(x => x.DynamicPools, newDynamicPools));
-			updates.Add(updateBuilder.Set(x => x.Version, version));
+			updates.Add(updateBuilder.Set(x => x.Version, options.Version));
 			updates.Add(updateBuilder.Unset(x => x.RequestRestart));
 			updates.Add(updateBuilder.Unset(x => x.RequestShutdown));
 			updates.Add(updateBuilder.Unset(x => x.RequestForceRestart));
 			updates.Add(updateBuilder.Set(x => x.LastShutdownReason, "Unexpected"));
 
-			if (String.Equals(version, agent.LastUpgradeVersion, StringComparison.Ordinal))
+			if (String.Equals(options.Version, agent.LastUpgradeVersion, StringComparison.Ordinal))
 			{
 				updates.Add(updateBuilder.Unset(x => x.UpgradeAttemptCount));
 			}
 
-			if (agent.Status != status)
+			if (agent.Status != options.Status)
 			{
-				updates.Add(updateBuilder.Set(x => x.LastStatusChange, lastStatusChange));
+				updates.Add(updateBuilder.Set(x => x.LastStatusChange, options.LastStatusChange));
 			}
 
 			foreach (AgentLease agentLease in agent.Leases ?? Enumerable.Empty<AgentLease>())
@@ -556,9 +660,8 @@ namespace Horde.Server.Agents
 		}
 
 		/// <inheritdoc/>
-		public async Task<IAgent?> TryTerminateSessionAsync(IAgent agentInterface, CancellationToken cancellationToken)
+		async Task<AgentDocument?> TryTerminateSessionAsync(AgentDocument agent, CancellationToken cancellationToken)
 		{
-			AgentDocument agent = (AgentDocument)agentInterface;
 			UpdateDefinition<AgentDocument> update = new BsonDocument();
 
 			update = update.Unset(x => x.SessionId);
@@ -584,10 +687,8 @@ namespace Horde.Server.Agents
 		private static string RedisKeyLeaseChildren(LeaseId parentId) => $"agent/lease-children/{parentId.ToString()}";
 
 		/// <inheritdoc/>
-		public async Task<IAgent?> TryAddLeaseAsync(IAgent agentInterface, AgentLease newLease, CancellationToken cancellationToken)
+		async Task<AgentDocument?> TryAddLeaseAsync(AgentDocument agent, AgentLease newLease, CancellationToken cancellationToken)
 		{
-			AgentDocument agent = (AgentDocument)agentInterface;
-
 			List<AgentLease> leases = new List<AgentLease>();
 			if (agent.Leases != null)
 			{
@@ -635,7 +736,7 @@ namespace Horde.Server.Agents
 			}
 		}
 
-		static void GetNewLeaseUpdates(IAgent agent, AgentLease lease, List<UpdateDefinition<AgentDocument>> updates)
+		static void GetNewLeaseUpdates(AgentDocument agent, AgentLease lease, List<UpdateDefinition<AgentDocument>> updates)
 		{
 			if (lease.Payload != null)
 			{
@@ -665,12 +766,10 @@ namespace Horde.Server.Agents
 		}
 
 		/// <inheritdoc/>
-		public async Task<IAgent?> TryCancelLeaseAsync(IAgent agentInterface, int leaseIdx, CancellationToken cancellationToken)
+		async Task<AgentDocument?> TryCancelLeaseAsync(AgentDocument agent, int leaseIdx, CancellationToken cancellationToken)
 		{
-			AgentDocument agent = (AgentDocument)agentInterface;
-
 			UpdateDefinition<AgentDocument> update = Builders<AgentDocument>.Update.Set(x => x.Leases![leaseIdx].State, LeaseState.Cancelled);
-			IAgent? newAgent = await TryUpdateAsync(agent, update, cancellationToken);
+			AgentDocument? newAgent = await TryUpdateAsync(agent, update, cancellationToken);
 			if (newAgent != null)
 			{
 				await PublishUpdateEventAsync(agent.Id);
diff --git a/Engine/Source/Programs/Horde/Horde.Server/Agents/AgentService.cs b/Engine/Source/Programs/Horde/Horde.Server/Agents/AgentService.cs
--- a/Engine/Source/Programs/Horde/Horde.Server/Agents/AgentService.cs
+++ b/Engine/Source/Programs/Horde/Horde.Server/Agents/AgentService.cs
@@ -220,7 +220,7 @@ namespace Horde.Server.Agents
 					return await Agents.AddAsync(agentId, ephemeral, enrollmentKey, cancellationToken);
 				}
 
-				agent = await Agents.TryResetAsync(agent, ephemeral, enrollmentKey, cancellationToken);
+				agent = await agent.TryResetAsync(ephemeral, enrollmentKey, cancellationToken);
 				if (agent != null)
 				{
 					return agent;
@@ -265,7 +265,7 @@ namespace Horde.Server.Agents
 		/// <returns>New agent state</returns>
 		public async Task<bool> TryUpdateWorkspacesAsync(IAgent agent, List<AgentWorkspaceInfo> workspaces, bool pendingConform, CancellationToken cancellationToken)
 		{
-			IAgent? newAgent = await Agents.TryUpdateWorkspacesAsync(agent, workspaces, pendingConform, cancellationToken);
+			IAgent? newAgent = await agent.TryUpdateWorkspacesAsync(workspaces, pendingConform, cancellationToken);
 			return newAgent != null;
 		}
 
@@ -291,7 +291,7 @@ namespace Horde.Server.Agents
 
 			while (agent is { Deleted: false })
 			{
-				IAgent? newAgent = await Agents.TryDeleteAsync(agent, cancellationToken);
+				IAgent? newAgent = await agent.TryDeleteAsync(cancellationToken);
 				if (newAgent != null)
 				{
 					break;
@@ -400,7 +400,7 @@ namespace Horde.Server.Agents
 					List<PoolId> pools = CombineCurrentAndRequestedPools(agent.ExplicitPools, properties);
 
 					// Reset the agent to use the new session
-					newAgent = await Agents.TryStartSessionAsync(agent, newSession.Id, sessionExpiresAt, status, properties, resources, pools, dynamicPools, lastStatusChange ?? utcNow, version, cancellationToken);
+					newAgent = await agent.TryCreateSessionAsync(new CreateSessionOptions(newSession.Id, sessionExpiresAt, status, properties, resources, pools, dynamicPools, lastStatusChange ?? utcNow, version), cancellationToken);
 					if (newAgent != null)
 					{
 						LogPropertyChanges(agentLogger, agent.Properties, newAgent.Properties);
@@ -564,7 +564,7 @@ namespace Horde.Server.Agents
 				(ITaskSource source, AgentLease lease) = result.Value;
 
 				// Add the new lease to the agent
-				IAgent? newAgent = await Agents.TryAddLeaseAsync(agent, lease, cancellationToken);
+				IAgent? newAgent = await agent.TryAddLeaseAsync(lease, cancellationToken);
 				if (newAgent != null)
 				{
 					await source.OnLeaseStartedAsync(newAgent, lease.Id, Any.Parser.ParseFrom(lease.Payload), Agents.GetLogger(agent.Id), CancellationToken.None);
@@ -657,7 +657,7 @@ namespace Horde.Server.Agents
 				return false;
 			}
 
-			await Agents.TryCancelLeaseAsync(agent, index, cancellationToken);
+			await agent.TryCancelLeaseAsync(index, cancellationToken);
 			return true;
 		}
 
@@ -783,7 +783,7 @@ namespace Horde.Server.Agents
 				List<PoolId> dynamicPools = await GetDynamicPoolsAsync(agent, cancellationToken);
 
 				// Update the agent, and try to create new lease documents if we succeed
-				IAgent? newAgent = await Agents.TryUpdateSessionAsync(agent, status, sessionExpiresAt, properties, resources, dynamicPools, updateLeases ? leases : null, cancellationToken);
+				IAgent? newAgent = await agent.TryUpdateSessionAsync(new UpdateSessionOptions(status, sessionExpiresAt, properties, resources, dynamicPools, updateLeases ? leases : null), cancellationToken);
 				if (newAgent != null)
 				{
 					agent = newAgent;
@@ -836,7 +836,7 @@ namespace Horde.Server.Agents
 			List<AgentLease> leases = new List<AgentLease>(agent.Leases);
 
 			// Clear the current session
-			IAgent? newAgent = await Agents.TryTerminateSessionAsync(agent, cancellationToken);
+			IAgent? newAgent = await agent.TryTerminateSessionAsync(cancellationToken);
 			if (newAgent != null)
 			{
 				agent = newAgent;
diff --git a/Engine/Source/Programs/Horde/Horde.Server/Agents/AgentsController.cs b/Engine/Source/Programs/Horde/Horde.Server/Agents/AgentsController.cs
--- a/Engine/Source/Programs/Horde/Horde.Server/Agents/AgentsController.cs
+++ b/Engine/Source/Programs/Horde/Horde.Server/Agents/AgentsController.cs
@@ -211,7 +211,7 @@ namespace Horde.Server.Agents
 					return NotFound(agentId);
 				}
 
-				IAgent? newAgent = await _agentService.Agents.TryUpdateSettingsAsync(agent, update.Enabled, update.RequestConform, update.RequestFullConform, update.RequestRestart, update.RequestShutdown, update.RequestForceRestart, $"Manual ({userName})", update.Pools?.ConvertAll(x => new PoolId(x)), update.Comment);
+				IAgent? newAgent = await agent.TryUpdateAsync(new UpdateAgentOptions(update.Enabled, update.RequestConform, update.RequestFullConform, update.RequestRestart, update.RequestShutdown, update.RequestForceRestart, $"Manual ({userName})", update.Pools?.ConvertAll(x => new PoolId(x)), update.Comment));
 				if (newAgent == null)
 				{
 					continue;
diff --git a/Engine/Source/Programs/Horde/Horde.Server/Agents/Fleet/AwsAutoScalingLifecycle.cs b/Engine/Source/Programs/Horde/Horde.Server/Agents/Fleet/AwsAutoScalingLifecycle.cs
--- a/Engine/Source/Programs/Horde/Horde.Server/Agents/Fleet/AwsAutoScalingLifecycle.cs
+++ b/Engine/Source/Programs/Horde/Horde.Server/Agents/Fleet/AwsAutoScalingLifecycle.cs
@@ -258,7 +258,7 @@ public sealed class AwsAutoScalingLifecycleService : IHostedService, IAsyncDispo
 			}
 
 			IAgent? agent = agentList[0];
-			await _agents.TryUpdateSettingsAsync(agent, requestShutdown: true, cancellationToken: cancellationToken);
+			await agent.TryUpdateAsync(new UpdateAgentOptions { RequestShutdown = true }, cancellationToken: cancellationToken);
 			await TrackAgentLifecycleAsync(agent.Id, e, cancellationToken);
 			return true;
 		}
diff --git a/Engine/Source/Programs/Horde/Horde.Server/Agents/Fleet/FleetService.cs b/Engine/Source/Programs/Horde/Horde.Server/Agents/Fleet/FleetService.cs
--- a/Engine/Source/Programs/Horde/Horde.Server/Agents/Fleet/FleetService.cs
+++ b/Engine/Source/Programs/Horde/Horde.Server/Agents/Fleet/FleetService.cs
@@ -462,7 +462,7 @@ namespace Horde.Server.Agents.Fleet
 			{
 				if (agent.RequestShutdown && numShutdownsCancelled < count)
 				{
-					await _agentCollection.TryUpdateSettingsAsync(agent, requestShutdown: false, cancellationToken: cancellationToken);
+					await agent.TryUpdateAsync(new UpdateAgentOptions { RequestShutdown = false }, cancellationToken: cancellationToken);
 					numShutdownsCancelled++;
 				}
 			}
diff --git a/Engine/Source/Programs/Horde/Horde.Server/Agents/Fleet/Providers/AwsFleetManager.cs b/Engine/Source/Programs/Horde/Horde.Server/Agents/Fleet/Providers/AwsFleetManager.cs
--- a/Engine/Source/Programs/Horde/Horde.Server/Agents/Fleet/Providers/AwsFleetManager.cs
+++ b/Engine/Source/Programs/Horde/Horde.Server/Agents/Fleet/Providers/AwsFleetManager.cs
@@ -271,7 +271,7 @@ namespace Horde.Server.Agents.Fleet.Providers
 			const int MaxRetries = 5;
 			for (int retryCount = 0; retryCount < MaxRetries; retryCount++)
 			{
-				if (await agentCollection.TryUpdateSettingsAsync(agent, requestShutdown: true, shutdownReason: "Autoscaler", cancellationToken: cancellationToken) != null)
+				if (await agent.TryUpdateAsync(new UpdateAgentOptions { RequestShutdown = true, ShutdownReason = "Autoscaler" }, cancellationToken: cancellationToken) != null)
 				{
 					agentLogger.LogInformation("Marked {AgentId} in pool {PoolName} for shutdown due to autoscaling (currently {NumLeases} leases outstanding, {NumRetries} retries)", agent.Id, pool.Name, agent.Leases.Count, retryCount);
 					return true;
diff --git a/Engine/Source/Programs/Horde/Horde.Server/Agents/IAgent.cs b/Engine/Source/Programs/Horde/Horde.Server/Agents/IAgent.cs
--- a/Engine/Source/Programs/Horde/Horde.Server/Agents/IAgent.cs
+++ b/Engine/Source/Programs/Horde/Horde.Server/Agents/IAgent.cs
@@ -31,406 +31,6 @@
 namespace Horde.Server.Agents
 {
 	/// <summary>
-	/// Information about a workspace synced to an agent
-	/// </summary>
-	public class AgentWorkspaceInfo
-	{
-		/// <summary>
-		/// Name of the Perforce cluster to use
-		/// </summary>
-		public string? Cluster { get; set; }
-
-		/// <summary>
-		/// User to log into Perforce with (eg. buildmachine)
-		/// </summary>
-		public string? UserName { get; set; }
-
-		/// <summary>
-		/// Identifier to distinguish this workspace from other workspaces
-		/// </summary>
-		public string Identifier { get; set; }
-
-		/// <summary>
-		/// The stream to sync
-		/// </summary>
-		public string Stream { get; set; }
-
-		/// <summary>
-		/// Custom view for the workspace
-		/// </summary>
-		public List<string>? View { get; set; }
-
-		/// <summary>
-		/// Whether to use an incremental workspace
-		/// </summary>
-		public bool Incremental { get; set; }
-
-		/// <summary>
-		/// Method to use when syncing/materializing data from Perforce
-		/// </summary>
-		public string? Method { get; set; }
-
-		/// <summary>
-		/// Minimum disk space that must be available *after* syncing this workspace (in megabytes)
-		/// </summary>
-		public long? MinScratchSpace { get; set; }
-
-		/// <summary>
-		/// Threshold for when to trigger an automatic conform of agent. Measured in megabytes free on disk.
-		/// </summary>
-		public long? ConformDiskFreeSpace { get; set; }
-
-		/// <summary>
-		/// Constructor
-		/// </summary>
-		/// <param name="cluster">Name of the Perforce cluster</param>
-		/// <param name="userName">User to log into Perforce with (eg. buildmachine)</param>
-		/// <param name="identifier">Identifier to distinguish this workspace from other workspaces</param>
-		/// <param name="stream">The stream to sync</param>
-		/// <param name="view">Custom view for the workspace</param>
-		/// <param name="incremental">Whether to use an incremental workspace</param>
-		/// <param name="method">Method to use when syncing/materializing data from Perforce</param>
-		/// <param name="minScratchSpace">Minimum disk space that must be available *after* syncing this workspace (in megabytes)</param>
-		/// <param name="conformDiskFreeSpace">Threshold for when to trigger an automatic conform of agent. Measured in megabytes free on disk</param>
-		public AgentWorkspaceInfo(string? cluster, string? userName, string identifier, string stream, List<string>? view, bool incremental, string? method, long? minScratchSpace = null, long? conformDiskFreeSpace = null)
-		{
-			if (!String.IsNullOrEmpty(cluster))
-			{
-				Cluster = cluster;
-			}
-			if (!String.IsNullOrEmpty(userName))
-			{
-				UserName = userName;
-			}
-			Identifier = identifier;
-			Stream = stream;
-			View = view;
-			Incremental = incremental;
-			Method = method;
-			MinScratchSpace = minScratchSpace;
-			ConformDiskFreeSpace = conformDiskFreeSpace;
-		}
-
-		/// <summary>
-		/// Constructor
-		/// </summary>
-		/// <param name="workspace">RPC message to construct from</param>
-		public AgentWorkspaceInfo(RpcAgentWorkspace workspace)
-			: this(workspace.ConfiguredCluster, workspace.ConfiguredUserName, workspace.Identifier, workspace.Stream, (workspace.View.Count > 0) ? workspace.View.ToList() : null, workspace.Incremental, workspace.Method, workspace.MinScratchSpace, workspace.ConformDiskFreeSpace)
-		{
-			// Treat zero value as null as Protobuf cannot store null 
-			MinScratchSpace = workspace.MinScratchSpace == 0 ? null : workspace.MinScratchSpace;
-			ConformDiskFreeSpace = workspace.ConformDiskFreeSpace == 0 ? null : workspace.ConformDiskFreeSpace;
-		}
-
-		/// <summary>
-		/// Gets a digest of the settings for this workspace
-		/// </summary>
-		/// <returns>Digest for the workspace settings</returns>
-		public string GetDigest()
-		{
-#pragma warning disable CA5351 // Do Not Use Broken Cryptographic Algorithms
-			using (MD5 hasher = MD5.Create())
-			{
-				byte[] data = BsonExtensionMethods.ToBson(this);
-				return BitConverter.ToString(hasher.ComputeHash(data)).Replace("-", "", StringComparison.Ordinal);
-			}
-#pragma warning restore CA5351
-		}
-
-		/// <inheritdoc/>
-		public override bool Equals(object? obj)
-		{
-			AgentWorkspaceInfo? other = obj as AgentWorkspaceInfo;
-			if (other == null)
-			{
-				return false;
-			}
-			if (Cluster != other.Cluster || UserName != other.UserName || Identifier != other.Identifier || Stream != other.Stream || Incremental != other.Incremental)
-			{
-				return false;
-			}
-			if (!Enumerable.SequenceEqual(View ?? new List<string>(), other.View ?? new List<string>()))
-			{
-				return false;
-			}
-			return true;
-		}
-
-		/// <inheritdoc/>
-		public override int GetHashCode()
-		{
-			return HashCode.Combine(Cluster, UserName, Identifier, Stream, Incremental); // Ignore 'View' for now
-		}
-
-		/// <summary>
-		/// Checks if two workspace sets are equivalent, ignoring order
-		/// </summary>
-		/// <param name="workspacesA">First list of workspaces</param>
-		/// <param name="workspacesB">Second list of workspaces</param>
-		/// <returns>True if the sets are equivalent</returns>
-		public static bool SetEquals(IReadOnlyList<AgentWorkspaceInfo> workspacesA, IReadOnlyList<AgentWorkspaceInfo> workspacesB)
-		{
-			HashSet<AgentWorkspaceInfo> workspacesSetA = new HashSet<AgentWorkspaceInfo>(workspacesA);
-			return workspacesSetA.SetEquals(workspacesB);
-		}
-
-		/// <summary>
-		/// Converts this workspace to an RPC message
-		/// </summary>
-		/// <param name="server">The Perforce server</param>
-		/// <param name="credentials">Credentials for the server</param>
-		/// <returns>The RPC message</returns>
-		public RpcAgentWorkspace ToRpcMessage(IPerforceServer server, PerforceCredentials? credentials)
-		{
-			// Construct the message
-			RpcAgentWorkspace result = new RpcAgentWorkspace
-			{
-				ConfiguredCluster = Cluster,
-				ConfiguredUserName = UserName,
-				ServerAndPort = server.ServerAndPort,
-				UserName = credentials?.UserName ?? UserName,
-				Password = credentials?.Password,
-				Ticket = credentials?.Ticket,
-				Identifier = Identifier,
-				Stream = Stream,
-				Incremental = Incremental,
-				Partitioned = server.SupportsPartitionedWorkspaces,
-				Method = Method ?? String.Empty,
-				MinScratchSpace = MinScratchSpace ?? 0
-			};
-
-			if (View != null)
-			{
-				result.View.AddRange(View);
-			}
-
-			return result;
-		}
-	}
-
-	/// <summary>
-	/// Configuration for an AutoSDK workspace
-	/// </summary>
-	public class AutoSdkConfig
-	{
-		/// <summary>
-		/// Disables use of AutoSDK.
-		/// </summary>
-		public static AutoSdkConfig None { get; } = new AutoSdkConfig(Array.Empty<string>());
-
-		/// <summary>
-		/// Syncs the entire AutoSDK folder.
-		/// </summary>
-		public static AutoSdkConfig Full { get; } = new AutoSdkConfig(null);
-
-		/// <summary>
-		/// Additive filter for paths to include in the workspace
-		/// </summary>
-		public List<string> View { get; set; } = new List<string>();
-
-		/// <summary>
-		/// Constructor
-		/// </summary>
-		public AutoSdkConfig()
-		{
-			View = new List<string>();
-		}
-
-		/// <summary>
-		/// Constructor
-		/// </summary>
-		/// <param name="filter">Filter for the workspace</param>
-		public AutoSdkConfig(IEnumerable<string>? filter)
-		{
-			if (filter == null)
-			{
-				View = new List<string> { "..." };
-			}
-			else
-			{
-				View = filter.OrderBy(x => x).Distinct().ToList();
-			}
-		}
-
-		/// <summary>
-		/// Merge two workspaces together
-		/// </summary>
-		[return: NotNullIfNotNull("lhs")]
-		[return: NotNullIfNotNull("rhs")]
-		public static AutoSdkConfig? Merge(AutoSdkConfig? lhs, AutoSdkConfig? rhs)
-		{
-			if (lhs?.View == null || lhs.View.Count == 0)
-			{
-				return rhs;
-			}
-			if (rhs?.View == null || rhs.View.Count == 0)
-			{
-				return lhs;
-			}
-
-			return new AutoSdkConfig(Enumerable.Concat(lhs.View, rhs.View));
-		}
-
-		/// <summary>
-		/// Test whether two views are equal
-		/// </summary>
-		public static bool Equals(AutoSdkConfig? lhs, AutoSdkConfig? rhs)
-		{
-			if (lhs?.View == null || lhs.View.Count == 0)
-			{
-				return rhs?.View == null || rhs.View.Count == 0;
-			}
-			if (rhs?.View == null || rhs.View.Count == 0)
-			{
-				return false;
-			}
-
-			return Enumerable.SequenceEqual(lhs.View, rhs.View);
-		}
-	}
-
-	/// <summary>
-	/// Document describing an active lease
-	/// </summary>
-	public class AgentLease
-	{
-		/// <summary>
-		/// Name of this lease
-		/// </summary>
-		[BsonRequired]
-		public LeaseId Id { get; set; }
-
-		/// <summary>
-		/// The parent lease id
-		/// </summary>
-		public LeaseId? ParentId { get; set; }
-
-		/// <summary>
-		/// Name of this lease
-		/// </summary>
-		public string Name { get; set; }
-
-		/// <summary>
-		/// The current state of the lease
-		/// </summary>
-		public LeaseState State { get; set; }
-
-		/// <summary>
-		/// The stream for this lease
-		/// </summary>
-		public StreamId? StreamId { get; set; }
-
-		/// <summary>
-		/// The pool for this lease
-		/// </summary>
-		public PoolId? PoolId { get; set; }
-
-		/// <summary>
-		/// Optional log for this lease
-		/// </summary>
-		public LogId? LogId { get; set; }
-
-		/// <summary>
-		/// Time at which the lease started
-		/// </summary>
-		[BsonRequired]
-		public DateTime StartTime { get; set; }
-
-		/// <summary>
-		/// Time at which the lease should be terminated
-		/// </summary>
-		public DateTime? ExpiryTime { get; set; }
-
-		/// <summary>
-		/// Flag indicating whether this lease has been accepted by the agent
-		/// </summary>
-		public bool Active { get; set; }
-
-		/// <summary>
-		/// Resources used by this lease
-		/// </summary>
-		public IReadOnlyDictionary<string, int>? Resources { get; set; }
-
-		/// <summary>
-		/// Whether the lease requires exclusive access to the agent
-		/// </summary>
-		public bool Exclusive { get; set; }
-
-		/// <summary>
-		/// For leases in the pending state, encodes an "any" protobuf containing the payload for the agent to execute the lease.
-		/// </summary>
-		public byte[]? Payload { get; set; }
-
-		/// <summary>
-		/// Private constructor
-		/// </summary>
-		[BsonConstructor]
-		private AgentLease()
-		{
-			Name = String.Empty;
-		}
-
-		/// <summary>
-		/// Constructor
-		/// </summary>
-		/// <param name="id">Identifier for the lease</param>
-		/// <param name="parentId">The parent lease id</param>
-		/// <param name="name">Name of this lease</param>
-		/// <param name="streamId"></param>
-		/// <param name="poolId"></param>
-		/// <param name="logId">Unique id for the log</param>
-		/// <param name="state">State for the lease</param>
-		/// <param name="resources">Resources required for this lease</param>
-		/// <param name="exclusive">Whether to reserve the entire device</param>
-		/// <param name="payload">Encoded "any" protobuf describing the contents of the payload</param>
-		public AgentLease(LeaseId id, LeaseId? parentId, string name, StreamId? streamId, PoolId? poolId, LogId? logId, LeaseState state, IReadOnlyDictionary<string, int>? resources, bool exclusive, byte[]? payload)
-		{
-			Id = id;
-			ParentId = parentId;
-			Name = name;
-			StreamId = streamId;
-			PoolId = poolId;
-			LogId = logId;
-			State = state;
-			Resources = resources;
-			Exclusive = exclusive;
-			Payload = payload;
-			StartTime = DateTime.UtcNow;
-		}
-
-		/// <summary>
-		/// Determines if this is a conform lease
-		/// </summary>
-		/// <returns>True if this is a conform lease</returns>
-		public bool IsConformLease()
-		{
-			if (Payload != null)
-			{
-				Any basePayload = Any.Parser.ParseFrom(Payload);
-				if (basePayload.Is(ConformTask.Descriptor))
-				{
-					return true;
-				}
-			}
-			return false;
-		}
-
-		/// <summary>
-		/// Converts this lease to an RPC message
-		/// </summary>
-		/// <returns>RPC message</returns>
-		public RpcLease ToRpcMessage()
-		{
-			RpcLease lease = new RpcLease();
-			lease.Id = Id.ToString();
-			lease.Payload = Google.Protobuf.WellKnownTypes.Any.Parser.ParseFrom(Payload);
-			lease.State = (RpcLeaseState)State;
-			return lease;
-		}
-	}
-
-	/// <summary>
 	/// Mirrors an Agent document in the database
 	/// </summary>
 	public interface IAgent
@@ -584,7 +184,117 @@
 		/// Update counter for this document. Any updates should compare-and-swap based on the value of this counter, or increment it in the case of server-side updates.
 		/// </summary>
 		public uint UpdateIndex { get; }
+
+		/// <summary>
+		/// Resets an agent to use new settings
+		/// </summary>
+		/// <param name="ephemeral">Whether the agent is ephemeral or not</param>
+		/// <param name="enrollmentKey">Key used to identify a unique enrollment for the agent with this id</param>
+		/// <param name="cancellationToken">Cancellation token for the operation</param>
+		Task<IAgent?> TryResetAsync(bool ephemeral, string enrollmentKey, CancellationToken cancellationToken = default);
+
+		/// <summary>
+		/// Deletes an agent
+		/// </summary>
+		/// <param name="cancellationToken">Cancellation token for the operation</param>
+		/// <returns>Async task</returns>
+		Task<IAgent?> TryDeleteAsync(CancellationToken cancellationToken = default);
+
+		/// <summary>
+		/// Update an agent's settings
+		/// </summary>
+		/// <param name="options">Options for updating the agent</param>
+		/// <param name="cancellationToken">Cancellation token for the operation</param>
+		/// <returns>New agent state if update was successful</returns>
+		Task<IAgent?> TryUpdateAsync(UpdateAgentOptions options, CancellationToken cancellationToken = default);
+
+		/// <summary>
+		/// Update the current workspaces for an agent.
+		/// </summary>
+		/// <param name="workspaces">Current list of workspaces</param>
+		/// <param name="requestConform">Whether the agent still needs to run another conform</param>
+		/// <param name="cancellationToken">Cancellation token for the operation</param>
+		/// <returns>New agent state</returns>
+		Task<IAgent?> TryUpdateWorkspacesAsync(List<AgentWorkspaceInfo> workspaces, bool requestConform, CancellationToken cancellationToken = default);
+
+		/// <summary>
+		/// Sets the current session
+		/// </summary>
+		/// <param name="options">Options for the new session</param>
+		/// <param name="cancellationToken">Cancellation token for the operation</param>
+		/// <returns>New agent state</returns>
+		Task<IAgent?> TryCreateSessionAsync(CreateSessionOptions options, CancellationToken cancellationToken = default);
+
+		/// <summary>
+		/// Attempt to update the current session
+		/// </summary>
+		/// <param name="options">Options for updating the session</param>
+		/// <param name="cancellationToken">Cancellation token for the operation</param>
+		/// <returns>True if the document was updated, false if another writer updated the document first</returns>
+		Task<IAgent?> TryUpdateSessionAsync(UpdateSessionOptions options, CancellationToken cancellationToken = default);
+
+		/// <summary>
+		/// Terminates the current session
+		/// </summary>
+		/// <param name="cancellationToken">Cancellation token for the operation</param>
+		/// <returns>New agent state if it succeeded, otherwise null</returns>
+		Task<IAgent?> TryTerminateSessionAsync(CancellationToken cancellationToken = default);
+
+		/// <summary>
+		/// Attempts to add a lease to an agent
+		/// </summary>
+		/// <param name="newLease">The new lease document</param>
+		/// <param name="cancellationToken">Cancellation token for the operation</param>
+		/// <returns>New agent state if it succeeded, otherwise null</returns>
+		Task<IAgent?> TryAddLeaseAsync(AgentLease newLease, CancellationToken cancellationToken = default);
+
+		/// <summary>
+		/// Attempts to cancel a lease
+		/// </summary>
+		/// <param name="leaseIdx">Index of the lease</param>
+		/// <param name="cancellationToken">Cancellation token for the operation</param>
+		/// <returns>New agent state if it succeeded, otherwise null</returns>
+		Task<IAgent?> TryCancelLeaseAsync(int leaseIdx, CancellationToken cancellationToken = default);
 	}
 
 	/// <summary>
+	/// Options for updating an agent
+	/// </summary>
+	/// <param name="Enabled">Whether the agent is enabled or not</param>
+	/// <param name="RequestConform">Whether to request a conform job be run</param>
+	/// <param name="RequestFullConform">Whether to request a full conform job be run</param>
+	/// <param name="RequestRestart">Whether to request the machine be restarted</param>
+	/// <param name="RequestShutdown">Whether to request the machine be shut down</param>
+	/// <param name="RequestForceRestart">Request an immediate restart without waiting for leases to complete</param>
+	/// <param name="ShutdownReason">The reason for shutting down agent, ex. Autoscaler/Manual/Unexpected</param>
+	/// <param name="Pools">List of pools for the agent</param>
+	/// <param name="Comment">New comment</param>
+	public record class UpdateAgentOptions(bool? Enabled = null, bool? RequestConform = null, bool? RequestFullConform = null, bool? RequestRestart = null, bool? RequestShutdown = null, bool? RequestForceRestart = null, string? ShutdownReason = null, List<PoolId>? Pools = null, string? Comment = null);
+
+	/// <summary>
+	/// Options for starting a new agent session
+	/// </summary>
+	/// <param name="SessionId">New session id</param>
+	/// <param name="SessionExpiresAt">Expiry time for the new session</param>
+	/// <param name="Status">Status of the agent</param>
+	/// <param name="Properties">Properties for the current session</param>
+	/// <param name="Resources">Resources for the agent</param>
+	/// <param name="Pools">New list of pools for the agent</param>
+	/// <param name="DynamicPools">New list of dynamic pools for the agent</param>
+	/// <param name="LastStatusChange">Time to force status change timestamp to</param>
+	/// <param name="Version">Current version of the agent software</param>
+	public record class CreateSessionOptions(SessionId SessionId, DateTime SessionExpiresAt, AgentStatus Status, IReadOnlyList<string> Properties, IReadOnlyDictionary<string, int> Resources, IReadOnlyList<PoolId> Pools, IReadOnlyList<PoolId> DynamicPools, DateTime LastStatusChange, string? Version);
+
+	/// <summary>
+	/// Options for updating a new agent session
+	/// </summary>
+	/// <param name="Status">New status of the agent</param>
+	/// <param name="SessionExpiresAt">New expiry time for the current session</param>
+	/// <param name="Properties">Properties for the current session</param>
+	/// <param name="Resources">Resources for the agent</param>
+	/// <param name="DynamicPools">New list of dynamic pools for the agent</param>
+	/// <param name="Leases">New set of leases</param>
+	public record class UpdateSessionOptions(AgentStatus? Status = null, DateTime? SessionExpiresAt = null, IReadOnlyList<string>? Properties = null, IReadOnlyDictionary<string, int>? Resources = null, IReadOnlyList<PoolId>? DynamicPools = null, List<AgentLease>? Leases = null);
+
+	/// <summary>
 	/// Extension methods for IAgent
@@ -693,7 +403,7 @@
 			List<string> values = agent.GetPropertyValues(KnownPropertyNames.SelfContained).ToList();
 			return values.Count > 0 && values[0].Equals("true", StringComparison.OrdinalIgnoreCase);
 		}
-		
+
 		/// <summary>
 		/// Get free disk space on agent (as reported through primary device capabilities)
 		/// </summary>
@@ -1082,4 +792,404 @@
 			}
 		}
 	}
+
+	/// <summary>
+	/// Information about a workspace synced to an agent
+	/// </summary>
+	public class AgentWorkspaceInfo
+	{
+		/// <summary>
+		/// Name of the Perforce cluster to use
+		/// </summary>
+		public string? Cluster { get; set; }
+
+		/// <summary>
+		/// User to log into Perforce with (eg. buildmachine)
+		/// </summary>
+		public string? UserName { get; set; }
+
+		/// <summary>
+		/// Identifier to distinguish this workspace from other workspaces
+		/// </summary>
+		public string Identifier { get; set; }
+
+		/// <summary>
+		/// The stream to sync
+		/// </summary>
+		public string Stream { get; set; }
+
+		/// <summary>
+		/// Custom view for the workspace
+		/// </summary>
+		public List<string>? View { get; set; }
+
+		/// <summary>
+		/// Whether to use an incremental workspace
+		/// </summary>
+		public bool Incremental { get; set; }
+
+		/// <summary>
+		/// Method to use when syncing/materializing data from Perforce
+		/// </summary>
+		public string? Method { get; set; }
+
+		/// <summary>
+		/// Minimum disk space that must be available *after* syncing this workspace (in megabytes)
+		/// </summary>
+		public long? MinScratchSpace { get; set; }
+
+		/// <summary>
+		/// Threshold for when to trigger an automatic conform of agent. Measured in megabytes free on disk.
+		/// </summary>
+		public long? ConformDiskFreeSpace { get; set; }
+
+		/// <summary>
+		/// Constructor
+		/// </summary>
+		/// <param name="cluster">Name of the Perforce cluster</param>
+		/// <param name="userName">User to log into Perforce with (eg. buildmachine)</param>
+		/// <param name="identifier">Identifier to distinguish this workspace from other workspaces</param>
+		/// <param name="stream">The stream to sync</param>
+		/// <param name="view">Custom view for the workspace</param>
+		/// <param name="incremental">Whether to use an incremental workspace</param>
+		/// <param name="method">Method to use when syncing/materializing data from Perforce</param>
+		/// <param name="minScratchSpace">Minimum disk space that must be available *after* syncing this workspace (in megabytes)</param>
+		/// <param name="conformDiskFreeSpace">Threshold for when to trigger an automatic conform of agent. Measured in megabytes free on disk</param>
+		public AgentWorkspaceInfo(string? cluster, string? userName, string identifier, string stream, List<string>? view, bool incremental, string? method, long? minScratchSpace = null, long? conformDiskFreeSpace = null)
+		{
+			if (!String.IsNullOrEmpty(cluster))
+			{
+				Cluster = cluster;
+			}
+			if (!String.IsNullOrEmpty(userName))
+			{
+				UserName = userName;
+			}
+			Identifier = identifier;
+			Stream = stream;
+			View = view;
+			Incremental = incremental;
+			Method = method;
+			MinScratchSpace = minScratchSpace;
+			ConformDiskFreeSpace = conformDiskFreeSpace;
+		}
+
+		/// <summary>
+		/// Constructor
+		/// </summary>
+		/// <param name="workspace">RPC message to construct from</param>
+		public AgentWorkspaceInfo(RpcAgentWorkspace workspace)
+			: this(workspace.ConfiguredCluster, workspace.ConfiguredUserName, workspace.Identifier, workspace.Stream, (workspace.View.Count > 0) ? workspace.View.ToList() : null, workspace.Incremental, workspace.Method, workspace.MinScratchSpace, workspace.ConformDiskFreeSpace)
+		{
+			// Treat zero value as null as Protobuf cannot store null 
+			MinScratchSpace = workspace.MinScratchSpace == 0 ? null : workspace.MinScratchSpace;
+			ConformDiskFreeSpace = workspace.ConformDiskFreeSpace == 0 ? null : workspace.ConformDiskFreeSpace;
+		}
+
+		/// <summary>
+		/// Gets a digest of the settings for this workspace
+		/// </summary>
+		/// <returns>Digest for the workspace settings</returns>
+		public string GetDigest()
+		{
+#pragma warning disable CA5351 // Do Not Use Broken Cryptographic Algorithms
+			using (MD5 hasher = MD5.Create())
+			{
+				byte[] data = BsonExtensionMethods.ToBson(this);
+				return BitConverter.ToString(hasher.ComputeHash(data)).Replace("-", "", StringComparison.Ordinal);
+			}
+#pragma warning restore CA5351
+		}
+
+		/// <inheritdoc/>
+		public override bool Equals(object? obj)
+		{
+			AgentWorkspaceInfo? other = obj as AgentWorkspaceInfo;
+			if (other == null)
+			{
+				return false;
+			}
+			if (Cluster != other.Cluster || UserName != other.UserName || Identifier != other.Identifier || Stream != other.Stream || Incremental != other.Incremental)
+			{
+				return false;
+			}
+			if (!Enumerable.SequenceEqual(View ?? new List<string>(), other.View ?? new List<string>()))
+			{
+				return false;
+			}
+			return true;
+		}
+
+		/// <inheritdoc/>
+		public override int GetHashCode()
+		{
+			return HashCode.Combine(Cluster, UserName, Identifier, Stream, Incremental); // Ignore 'View' for now
+		}
+
+		/// <summary>
+		/// Checks if two workspace sets are equivalent, ignoring order
+		/// </summary>
+		/// <param name="workspacesA">First list of workspaces</param>
+		/// <param name="workspacesB">Second list of workspaces</param>
+		/// <returns>True if the sets are equivalent</returns>
+		public static bool SetEquals(IReadOnlyList<AgentWorkspaceInfo> workspacesA, IReadOnlyList<AgentWorkspaceInfo> workspacesB)
+		{
+			HashSet<AgentWorkspaceInfo> workspacesSetA = new HashSet<AgentWorkspaceInfo>(workspacesA);
+			return workspacesSetA.SetEquals(workspacesB);
+		}
+
+		/// <summary>
+		/// Converts this workspace to an RPC message
+		/// </summary>
+		/// <param name="server">The Perforce server</param>
+		/// <param name="credentials">Credentials for the server</param>
+		/// <returns>The RPC message</returns>
+		public RpcAgentWorkspace ToRpcMessage(IPerforceServer server, PerforceCredentials? credentials)
+		{
+			// Construct the message
+			RpcAgentWorkspace result = new RpcAgentWorkspace
+			{
+				ConfiguredCluster = Cluster,
+				ConfiguredUserName = UserName,
+				ServerAndPort = server.ServerAndPort,
+				UserName = credentials?.UserName ?? UserName,
+				Password = credentials?.Password,
+				Ticket = credentials?.Ticket,
+				Identifier = Identifier,
+				Stream = Stream,
+				Incremental = Incremental,
+				Partitioned = server.SupportsPartitionedWorkspaces,
+				Method = Method ?? String.Empty,
+				MinScratchSpace = MinScratchSpace ?? 0
+			};
+
+			if (View != null)
+			{
+				result.View.AddRange(View);
+			}
+
+			return result;
+		}
+	}
+
+	/// <summary>
+	/// Configuration for an AutoSDK workspace
+	/// </summary>
+	public class AutoSdkConfig
+	{
+		/// <summary>
+		/// Disables use of AutoSDK.
+		/// </summary>
+		public static AutoSdkConfig None { get; } = new AutoSdkConfig(Array.Empty<string>());
+
+		/// <summary>
+		/// Syncs the entire AutoSDK folder.
+		/// </summary>
+		public static AutoSdkConfig Full { get; } = new AutoSdkConfig(null);
+
+		/// <summary>
+		/// Additive filter for paths to include in the workspace
+		/// </summary>
+		public List<string> View { get; set; } = new List<string>();
+
+		/// <summary>
+		/// Constructor
+		/// </summary>
+		public AutoSdkConfig()
+		{
+			View = new List<string>();
+		}
+
+		/// <summary>
+		/// Constructor
+		/// </summary>
+		/// <param name="filter">Filter for the workspace</param>
+		public AutoSdkConfig(IEnumerable<string>? filter)
+		{
+			if (filter == null)
+			{
+				View = new List<string> { "..." };
+			}
+			else
+			{
+				View = filter.OrderBy(x => x).Distinct().ToList();
+			}
+		}
+
+		/// <summary>
+		/// Merge two workspaces together
+		/// </summary>
+		[return: NotNullIfNotNull("lhs")]
+		[return: NotNullIfNotNull("rhs")]
+		public static AutoSdkConfig? Merge(AutoSdkConfig? lhs, AutoSdkConfig? rhs)
+		{
+			if (lhs?.View == null || lhs.View.Count == 0)
+			{
+				return rhs;
+			}
+			if (rhs?.View == null || rhs.View.Count == 0)
+			{
+				return lhs;
+			}
+
+			return new AutoSdkConfig(Enumerable.Concat(lhs.View, rhs.View));
+		}
+
+		/// <summary>
+		/// Test whether two views are equal
+		/// </summary>
+		public static bool Equals(AutoSdkConfig? lhs, AutoSdkConfig? rhs)
+		{
+			if (lhs?.View == null || lhs.View.Count == 0)
+			{
+				return rhs?.View == null || rhs.View.Count == 0;
+			}
+			if (rhs?.View == null || rhs.View.Count == 0)
+			{
+				return false;
+			}
+
+			return Enumerable.SequenceEqual(lhs.View, rhs.View);
+		}
+	}
+
+	/// <summary>
+	/// Document describing an active lease
+	/// </summary>
+	public class AgentLease
+	{
+		/// <summary>
+		/// Name of this lease
+		/// </summary>
+		[BsonRequired]
+		public LeaseId Id { get; set; }
+
+		/// <summary>
+		/// The parent lease id
+		/// </summary>
+		public LeaseId? ParentId { get; set; }
+
+		/// <summary>
+		/// Name of this lease
+		/// </summary>
+		public string Name { get; set; }
+
+		/// <summary>
+		/// The current state of the lease
+		/// </summary>
+		public LeaseState State { get; set; }
+
+		/// <summary>
+		/// The stream for this lease
+		/// </summary>
+		public StreamId? StreamId { get; set; }
+
+		/// <summary>
+		/// The pool for this lease
+		/// </summary>
+		public PoolId? PoolId { get; set; }
+
+		/// <summary>
+		/// Optional log for this lease
+		/// </summary>
+		public LogId? LogId { get; set; }
+
+		/// <summary>
+		/// Time at which the lease started
+		/// </summary>
+		[BsonRequired]
+		public DateTime StartTime { get; set; }
+
+		/// <summary>
+		/// Time at which the lease should be terminated
+		/// </summary>
+		public DateTime? ExpiryTime { get; set; }
+
+		/// <summary>
+		/// Flag indicating whether this lease has been accepted by the agent
+		/// </summary>
+		public bool Active { get; set; }
+
+		/// <summary>
+		/// Resources used by this lease
+		/// </summary>
+		public IReadOnlyDictionary<string, int>? Resources { get; set; }
+
+		/// <summary>
+		/// Whether the lease requires exclusive access to the agent
+		/// </summary>
+		public bool Exclusive { get; set; }
+
+		/// <summary>
+		/// For leases in the pending state, encodes an "any" protobuf containing the payload for the agent to execute the lease.
+		/// </summary>
+		public byte[]? Payload { get; set; }
+
+		/// <summary>
+		/// Private constructor
+		/// </summary>
+		[BsonConstructor]
+		private AgentLease()
+		{
+			Name = String.Empty;
+		}
+
+		/// <summary>
+		/// Constructor
+		/// </summary>
+		/// <param name="id">Identifier for the lease</param>
+		/// <param name="parentId">The parent lease id</param>
+		/// <param name="name">Name of this lease</param>
+		/// <param name="streamId"></param>
+		/// <param name="poolId"></param>
+		/// <param name="logId">Unique id for the log</param>
+		/// <param name="state">State for the lease</param>
+		/// <param name="resources">Resources required for this lease</param>
+		/// <param name="exclusive">Whether to reserve the entire device</param>
+		/// <param name="payload">Encoded "any" protobuf describing the contents of the payload</param>
+		public AgentLease(LeaseId id, LeaseId? parentId, string name, StreamId? streamId, PoolId? poolId, LogId? logId, LeaseState state, IReadOnlyDictionary<string, int>? resources, bool exclusive, byte[]? payload)
+		{
+			Id = id;
+			ParentId = parentId;
+			Name = name;
+			StreamId = streamId;
+			PoolId = poolId;
+			LogId = logId;
+			State = state;
+			Resources = resources;
+			Exclusive = exclusive;
+			Payload = payload;
+			StartTime = DateTime.UtcNow;
+		}
+
+		/// <summary>
+		/// Determines if this is a conform lease
+		/// </summary>
+		/// <returns>True if this is a conform lease</returns>
+		public bool IsConformLease()
+		{
+			if (Payload != null)
+			{
+				Any basePayload = Any.Parser.ParseFrom(Payload);
+				if (basePayload.Is(ConformTask.Descriptor))
+				{
+					return true;
+				}
+			}
+			return false;
+		}
+
+		/// <summary>
+		/// Converts this lease to an RPC message
+		/// </summary>
+		/// <returns>RPC message</returns>
+		public RpcLease ToRpcMessage()
+		{
+			RpcLease lease = new RpcLease();
+			lease.Id = Id.ToString();
+			lease.Payload = Google.Protobuf.WellKnownTypes.Any.Parser.ParseFrom(Payload);
+			lease.State = (RpcLeaseState)State;
+			return lease;
+		}
+	}
 }
diff --git a/Engine/Source/Programs/Horde/Horde.Server/Agents/IAgentCollection.cs b/Engine/Source/Programs/Horde/Horde.Server/Agents/IAgentCollection.cs
--- a/Engine/Source/Programs/Horde/Horde.Server/Agents/IAgentCollection.cs
+++ b/Engine/Source/Programs/Horde/Horde.Server/Agents/IAgentCollection.cs
@@ -7,7 +7,6 @@ using System.Threading.Tasks;
 using EpicGames.Horde.Agents;
 using EpicGames.Horde.Agents.Leases;
 using EpicGames.Horde.Agents.Pools;
-using EpicGames.Horde.Agents.Sessions;
 using Horde.Server.Auditing;
 
 namespace Horde.Server.Agents
@@ -26,23 +25,6 @@ namespace Horde.Server.Agents
 		/// <param name="cancellationToken">Cancellation token for the operation</param>
 		Task<IAgent> AddAsync(AgentId id, bool ephemeral, string enrollmentKey, CancellationToken cancellationToken = default);
 
-		/// <summary>
-		/// Resets an agent to use new settings
-		/// </summary>
-		/// <param name="agent">The agent to reset</param>
-		/// <param name="ephemeral">Whether the agent is ephemeral or not</param>
-		/// <param name="enrollmentKey">Key used to identify a unique enrollment for the agent with this id</param>
-		/// <param name="cancellationToken">Cancellation token for the operation</param>
-		Task<IAgent?> TryResetAsync(IAgent agent, bool ephemeral, string enrollmentKey, CancellationToken cancellationToken = default);
-
-		/// <summary>
-		/// Deletes an agent
-		/// </summary>
-		/// <param name="agent">Deletes the agent</param>
-		/// <param name="cancellationToken">Cancellation token for the operation</param>
-		/// <returns>Async task</returns>
-		Task<IAgent?> TryDeleteAsync(IAgent agent, CancellationToken cancellationToken = default);
-
 		/// <summary>
 		/// Deletes an agent
 		/// </summary>
@@ -113,90 +95,6 @@ namespace Horde.Server.Agents
 		/// <returns>List of agent lease IDs</returns>
 		Task<List<LeaseId>> GetChildLeaseIdsAsync(LeaseId id, CancellationToken cancellationToken = default);
 
-		/// <summary>
-		/// Update an agent's settings
-		/// </summary>
-		/// <param name="agent">Agent instance</param>
-		/// <param name="enabled">Whether the agent is enabled or not</param>
-		/// <param name="requestConform">Whether to request a conform job be run</param>
-		/// <param name="requestFullConform">Whether to request a full conform job be run</param>
-		/// <param name="requestRestart">Whether to request the machine be restarted</param>
-		/// <param name="requestShutdown">Whether to request the machine be shut down</param>
-		/// <param name="requestForceRestart">Request an immediate restart without waiting for leases to complete</param>
-		/// <param name="shutdownReason">The reason for shutting down agent, ex. Autoscaler/Manual/Unexpected</param>
-		/// <param name="pools">List of pools for the agent</param>
-		/// <param name="comment">New comment</param>
-		/// <param name="cancellationToken">Cancellation token for the operation</param>
-		/// <returns>New agent state if update was successful</returns>
-		Task<IAgent?> TryUpdateSettingsAsync(IAgent agent, bool? enabled = null, bool? requestConform = null, bool? requestFullConform = null, bool? requestRestart = null, bool? requestShutdown = null, bool? requestForceRestart = null, string? shutdownReason = null, List<PoolId>? pools = null, string? comment = null, CancellationToken cancellationToken = default);
-
-		/// <summary>
-		/// Update the current workspaces for an agent.
-		/// </summary>
-		/// <param name="agent">The agent to update</param>
-		/// <param name="workspaces">Current list of workspaces</param>
-		/// <param name="requestConform">Whether the agent still needs to run another conform</param>
-		/// <param name="cancellationToken">Cancellation token for the operation</param>
-		/// <returns>New agent state</returns>
-		Task<IAgent?> TryUpdateWorkspacesAsync(IAgent agent, List<AgentWorkspaceInfo> workspaces, bool requestConform, CancellationToken cancellationToken = default);
-
-		/// <summary>
-		/// Sets the current session
-		/// </summary>
-		/// <param name="agent">The agent to update</param>
-		/// <param name="sessionId">New session id</param>
-		/// <param name="sessionExpiresAt">Expiry time for the new session</param>
-		/// <param name="status">Status of the agent</param>
-		/// <param name="properties">Properties for the current session</param>
-		/// <param name="resources">Resources for the agent</param>
-		/// <param name="pools">New list of pools for the agent</param>
-		/// <param name="dynamicPools">New list of dynamic pools for the agent</param>
-		/// <param name="lastStatusChange">Time to force status change timestamp to</param>
-		/// <param name="version">Current version of the agent software</param>
-		/// <param name="cancellationToken">Cancellation token for the operation</param>
-		/// <returns>New agent state</returns>
-		Task<IAgent?> TryStartSessionAsync(IAgent agent, SessionId sessionId, DateTime sessionExpiresAt, AgentStatus status, IReadOnlyList<string> properties, IReadOnlyDictionary<string, int> resources, IReadOnlyList<PoolId> pools, IReadOnlyList<PoolId> dynamicPools, DateTime lastStatusChange, string? version, CancellationToken cancellationToken = default);
-
-		/// <summary>
-		/// Attempt to update the agent state
-		/// </summary>
-		/// <param name="agent">Agent instance</param>
-		/// <param name="status">New status of the agent</param>
-		/// <param name="sessionExpiresAt">New expiry time for the current session</param>
-		/// <param name="properties">Properties for the current session</param>
-		/// <param name="resources">Resources for the agent</param>
-		/// <param name="dynamicPools">New list of dynamic pools for the agent</param>
-		/// <param name="leases">New set of leases</param>
-		/// <param name="cancellationToken">Cancellation token for the operation</param>
-		/// <returns>True if the document was updated, false if another writer updated the document first</returns>
-		Task<IAgent?> TryUpdateSessionAsync(IAgent agent, AgentStatus? status, DateTime? sessionExpiresAt, IReadOnlyList<string>? properties, IReadOnlyDictionary<string, int>? resources, IReadOnlyList<PoolId>? dynamicPools, List<AgentLease>? leases, CancellationToken cancellationToken = default);
-
-		/// <summary>
-		/// Terminates the current session
-		/// </summary>
-		/// <param name="agent">The agent to terminate the session for</param>
-		/// <param name="cancellationToken">Cancellation token for the operation</param>
-		/// <returns>New agent state if it succeeded, otherwise null</returns>
-		Task<IAgent?> TryTerminateSessionAsync(IAgent agent, CancellationToken cancellationToken = default);
-
-		/// <summary>
-		/// Attempts to add a lease to an agent
-		/// </summary>
-		/// <param name="agent">The agent to update</param>
-		/// <param name="newLease">The new lease document</param>
-		/// <param name="cancellationToken">Cancellation token for the operation</param>
-		/// <returns>New agent state if it succeeded, otherwise null</returns>
-		Task<IAgent?> TryAddLeaseAsync(IAgent agent, AgentLease newLease, CancellationToken cancellationToken = default);
-
-		/// <summary>
-		/// Attempts to cancel a lease
-		/// </summary>
-		/// <param name="agent">The agent to update</param>
-		/// <param name="leaseIdx">Index of the lease</param>
-		/// <param name="cancellationToken">Cancellation token for the operation</param>
-		/// <returns>New agent state if it succeeded, otherwise null</returns>
-		Task<IAgent?> TryCancelLeaseAsync(IAgent agent, int leaseIdx, CancellationToken cancellationToken = default);
-
 		/// <summary>
 		/// Gets the log channel for an agent
 		/// </summary>
diff --git a/Engine/Source/Programs/Horde/Horde.Server/Agents/Pools/PoolUpdateService.cs b/Engine/Source/Programs/Horde/Horde.Server/Agents/Pools/PoolUpdateService.cs
--- a/Engine/Source/Programs/Horde/Horde.Server/Agents/Pools/PoolUpdateService.cs
+++ b/Engine/Source/Programs/Horde/Horde.Server/Agents/Pools/PoolUpdateService.cs
@@ -91,7 +91,7 @@ namespace Horde.Server.Agents.Pools
 			{
 				if (HasGracePeriodExpired(agent, pools, _globalConfig.CurrentValue.AgentShutdownIfDisabledGracePeriod))
 				{
-					await _agents.TryUpdateSettingsAsync(agent, requestShutdown: true, cancellationToken: cancellationToken);
+					await agent.TryUpdateAsync(new UpdateAgentOptions { RequestShutdown = true }, cancellationToken: cancellationToken);
 					_logger.LogInformation("Shutting down agent {AgentId} as it has been disabled for longer than grace period", agent.Id.ToString());
 					c++;
 				}
@@ -152,7 +152,7 @@ namespace Horde.Server.Agents.Pools
 
 				if (freeDiskSpace != null && maxConformDiskSpace > 0 && freeDiskSpace < maxConformDiskSpace && agent.ConformAttemptCount is null or 0)
 				{
-					await _agents.TryUpdateSettingsAsync(agent, requestConform: true, cancellationToken: cancellationToken);
+					await agent.TryUpdateAsync(new UpdateAgentOptions { RequestConform = true }, cancellationToken: cancellationToken);
 					_logger.LogInformation("Auto-conforming {AgentId} as workspace conform disk space needed ({ConformDiskSpace:F1} MB) is less than free disk space ({FreeDiskSpace:F1} MB)",
 						agent.Id.ToString(), maxConformDiskSpace / 1024.0 / 1024.0, freeDiskSpace.Value / 1024.0 / 1024.0);
 				}
diff --git a/Engine/Source/Programs/Horde/Horde.Server/Compute/ComputeService.cs b/Engine/Source/Programs/Horde/Horde.Server/Compute/ComputeService.cs
--- a/Engine/Source/Programs/Horde/Horde.Server/Compute/ComputeService.cs
+++ b/Engine/Source/Programs/Horde/Horde.Server/Compute/ComputeService.cs
@@ -420,7 +420,7 @@ namespace Horde.Server.Compute
 						{
 							using TelemetrySpan addLeaseSpan = _tracer.StartActiveSpan("Adding lease");
 
-							IAgent? newAgent = await _agentCollection.TryAddLeaseAsync(agent, lease, cancellationToken);
+							IAgent? newAgent = await agent.TryAddLeaseAsync(lease, cancellationToken);
 							if (newAgent != null)
 							{
 								await _agentCollection.PublishUpdateEventAsync(agent.Id);
diff --git a/Engine/Source/Programs/Horde/Horde.Server/Jobs/JobService.cs b/Engine/Source/Programs/Horde/Horde.Server/Jobs/JobService.cs
--- a/Engine/Source/Programs/Horde/Horde.Server/Jobs/JobService.cs
+++ b/Engine/Source/Programs/Horde/Horde.Server/Jobs/JobService.cs
@@ -499,7 +499,7 @@ namespace Horde.Server.Jobs
 					break;
 				}
 
-				IAgent? newAgent = await _agents.TryCancelLeaseAsync(agent, index, cancellationToken);
+				IAgent? newAgent = await agent.TryCancelLeaseAsync(index, cancellationToken);
 				if (newAgent != null)
 				{
 					_jobTaskSource.CancelLongPollForAgent(agent.Id);
@@ -875,7 +875,7 @@ namespace Horde.Server.Jobs
 								{
 									break;
 								}
-								if (await _agents.TryUpdateSettingsAsync(agent, requestConform: true, cancellationToken: cancellationToken) != null)
+								if (await agent.TryUpdateAsync(new UpdateAgentOptions { RequestConform = true }, cancellationToken: cancellationToken) != null)
 								{
 									_logger.LogError("Agent {AgentId} did not complete lease; marking for conform", agent.Id);
 									break;
