Horde: Updating schema documentation

#rnx

diff --git a/Engine/Source/Programs/Horde/Docs/Config/Schema/Projects.md b/Engine/Source/Programs/Horde/Docs/Config/Schema/Projects.md
--- a/Engine/Source/Programs/Horde/Docs/Config/Schema/Projects.md
+++ b/Engine/Source/Programs/Horde/Docs/Config/Schema/Projects.md
@@ -177,6 +177,7 @@ Name | Description
 `runInSeparateProcess` | `boolean`<br>Executes the job lease in a separate process
 `workspaceMaterializer` | `string`<br>What workspace materializer to use in WorkspaceExecutor. Will override any value from workspace config.
 `container` | [`JobContainerOptions`](#jobcontaineroptions)<br>Options for executing a job inside a container
+`expireAfterDays` | `integer`<br>Number of days after which to expire jobs
 
 ## JobContainerOptions
 
diff --git a/Engine/Source/Programs/Horde/Docs/Config/Schema/Streams.md b/Engine/Source/Programs/Horde/Docs/Config/Schema/Streams.md
--- a/Engine/Source/Programs/Horde/Docs/Config/Schema/Streams.md
+++ b/Engine/Source/Programs/Horde/Docs/Config/Schema/Streams.md
@@ -64,6 +64,7 @@ Name | Description
 `runInSeparateProcess` | `boolean`<br>Executes the job lease in a separate process
 `workspaceMaterializer` | `string`<br>What workspace materializer to use in WorkspaceExecutor. Will override any value from workspace config.
 `container` | [`JobContainerOptions`](#jobcontaineroptions)<br>Options for executing a job inside a container
+`expireAfterDays` | `integer`<br>Number of days after which to expire jobs
 
 ## JobContainerOptions
 
@@ -206,6 +207,8 @@ Name | Description
 `useAutoSdk` | `boolean`<br>Whether to use the AutoSDK
 `autoSdkView` | `string[]`<br>View for the AutoSDK paths to sync. If null, the whole thing will be synced.
 `method` | `string`<br>Method to use when syncing/materializing data from Perforce
+`minScratchSpace` | `integer`<br>Minimum disk space that must be available *after* syncing this workspace (in megabytes) If not available, the job will be aborted.
+`conformDiskFreeSpace` | `integer`<br>Threshold for when to trigger an automatic conform of agent. Measured in megabytes free on disk. Set to null or 0 to disable.
 
 ## TemplateRefConfig
 
