Horde: Disable timestamp validation when pulling files from temp storage. This is already checked inside BuildGraph - and supports certain files being excluded - but the same functionality is not exposed to Horde.

#rnx

diff --git a/Engine/Source/Programs/Horde/Horde.Agent/Utility/TempStorage.cs b/Engine/Source/Programs/Horde/Horde.Agent/Utility/TempStorage.cs
--- a/Engine/Source/Programs/Horde/Horde.Agent/Utility/TempStorage.cs
+++ b/Engine/Source/Programs/Horde/Horde.Agent/Utility/TempStorage.cs
@@ -788,16 +788,6 @@ namespace Horde.Storage.Utility
 				}
 			}
 
-			// Check all the local files are as expected
-			bool allMatch = true;
-			foreach (TempStorageFile file in manifest.Files)
-			{
-				allMatch &= file.Compare(rootDir, logger);
-			}
-			if (!allMatch)
-			{
-				throw new TempStorageException("Files have been modified");
-			}
 			return manifest;
 		}
 
