Horde: New Build Dialog
#rnx

diff --git a/Engine/Source/Programs/Horde/HordeDashboard/src/backend/Api.ts b/Engine/Source/Programs/Horde/HordeDashboard/src/backend/Api.ts
--- a/Engine/Source/Programs/Horde/HordeDashboard/src/backend/Api.ts
+++ b/Engine/Source/Programs/Horde/HordeDashboard/src/backend/Api.ts
@@ -1144,6 +1144,9 @@ export type CreateJobRequest = {
 
 	/** Arguments for the job */
 	arguments?: string[];
+
+	/** Parameters for the job */
+	parameters?: Record<string, string>
 }
 
 /**Response from creating a new job */
@@ -1287,6 +1290,9 @@ export type GetJobResponse = {
 	/**  Parameters for the job */
 	arguments: string[];
 
+	/** Parameters for the job */
+	parameters: Record<string, string>;
+
 	/**The last update time for this job*/
 	updateTime: Date | string;
 
@@ -2399,6 +2405,9 @@ export type GroupParameterData = ParameterData & {
 /**Free-form text entry parameter */
 export type TextParameterData = ParameterData & {
 
+	// Unique id for this parameter
+	id: string;
+
 	/**Name of the parameter associated with this parameter. */
 	label: string;
 
@@ -2425,6 +2434,9 @@ export type TextParameterData = ParameterData & {
 /**Possible option for a list parameter */
 export type ListParameterItemData = ParameterData & {
 
+	// Unique id for this parameter
+	id: string;
+
 	/**Optional group heading to display this entry under, if the picker style supports it. */
 	group?: string;
 
@@ -2473,6 +2485,10 @@ export type ListParameterData = ParameterData & {
 
 /**Allows the user to toggle an option on or off */
 export type BoolParameterData = ParameterData & {
+
+	// Unique id for this parameter
+	id: string;
+
 	/**Name of the parameter associated with this parameter. */
 	label: string;
 
diff --git a/Engine/Source/Programs/Horde/HordeDashboard/src/components/NewBuild.tsx b/Engine/Source/Programs/Horde/HordeDashboard/src/components/NewBuild.tsx
--- a/Engine/Source/Programs/Horde/HordeDashboard/src/components/NewBuild.tsx
+++ b/Engine/Source/Programs/Horde/HordeDashboard/src/components/NewBuild.tsx
@@ -13,6 +13,7 @@ import { JobDetailsV2 } from './jobDetailsV2/JobDetailsViewCommon';
 import dashboard from '../backend/Dashboard';
 import { getHordeStyling } from '../styles/Styles';
 import { Markdown } from '../base/components/Markdown';
+import { NewBuildV2 } from './build/NewBuildV2';
 
 let toolTipId = 0;
 
@@ -207,8 +208,8 @@ class BuildParameters {
 
             const data = p as BoolParameterData;
 
-            let enabledTargets:string[] = [];
-            let disabledTargets:string[] = [];
+            let enabledTargets: string[] = [];
+            let disabledTargets: string[] = [];
 
             if (data.argumentIfEnabled?.toLowerCase().startsWith("-target=")) {
                enabledTargets.push(data.argumentIfEnabled.slice(8));
@@ -222,7 +223,7 @@ class BuildParameters {
 
             if (data.argumentIfDisabled?.toLowerCase().startsWith("-target=")) {
                disabledTargets.push(data.argumentIfDisabled.slice(8));
-            }            
+            }
 
             data.argumentsIfDisabled?.forEach(a => {
                if (a.toLowerCase().startsWith("-target=")) {
@@ -250,7 +251,7 @@ class BuildParameters {
    removeTarget(target: string, updateChanged = true) {
 
       target = target.trim();
-      const unique = new Set<string>(this.targets.filter(t => t !== target));      
+      const unique = new Set<string>(this.targets.filter(t => t !== target));
       this.targets = Array.from(unique);
 
       const parameters: ParameterData[] = [];
@@ -273,8 +274,8 @@ class BuildParameters {
 
             const data = p as BoolParameterData;
 
-            let enabledTargets:string[] = [];
-            let disabledTargets:string[] = [];
+            let enabledTargets: string[] = [];
+            let disabledTargets: string[] = [];
 
             if (data.argumentIfEnabled?.toLowerCase().startsWith("-target=")) {
                enabledTargets.push(data.argumentIfEnabled.slice(8));
@@ -288,7 +289,7 @@ class BuildParameters {
 
             if (data.argumentIfDisabled?.toLowerCase().startsWith("-target=")) {
                disabledTargets.push(data.argumentIfDisabled.slice(8));
-            }    
+            }
 
             data.argumentsIfDisabled?.forEach(a => {
                if (a.toLowerCase().startsWith("-target=")) {
@@ -732,12 +733,17 @@ export const NewBuild: React.FC<{ streamId: string; show: boolean; onClose: (new
    const [mode, setMode] = useState<"Basic" | "Advanced">("Basic");
    const [showAllTemplates, setShowAllTemplates] = useState<boolean>(false);
    const [validationErrors, setValidationErrors] = useState<ValidationError[]>([]);
+   const [useNewBuildV2, setNewBuildV2] = useState<boolean>(false);
 
    const targetPicker = React.useRef(null)
    const { hordeClasses, modeColors } = getHordeStyling();
 
    const stream = projectStore.streamById(streamId);
 
+   if (useNewBuildV2) {
+      return <NewBuildV2 show={true} streamId={streamId} onClose={onClose} jobKey={jobKey} jobDetails={jobDetails} readOnly={readOnly} />
+   }
+
    if (!show) {
       return null;
    }
@@ -1620,7 +1626,6 @@ export const NewBuild: React.FC<{ streamId: string; show: boolean; onClose: (new
             <Stack verticalAlign="center" style={{ paddingBottom: 32 }}>
                <Spinner size={SpinnerSize.large} />
             </Stack>
-
          </Stack>
       </Modal>
 
@@ -1654,16 +1659,16 @@ export const NewBuild: React.FC<{ streamId: string; show: boolean; onClose: (new
       switch (param.type) {
 
          case ParameterType.List:
-            estimatedHeight += 52;
+            estimatedHeight += 54;
             break;
          case ParameterType.Bool:
-            estimatedHeight += 18;
+            estimatedHeight += 20;
             break;
          case ParameterType.Text:
-            estimatedHeight += 52;
+            estimatedHeight += 54;
             break;
          case ParameterType.Group:
-            estimatedHeight += 18;
+            estimatedHeight += 20;
             const group = (param as GroupParameterData);
             if (group.children?.length) {
                estimatedHeight += group.children.length * parameterGap;
@@ -1681,7 +1686,7 @@ export const NewBuild: React.FC<{ streamId: string; show: boolean; onClose: (new
    });
 
    if (!!buildParams.preflight) {
-      estimatedHeight += 28;
+      estimatedHeight += 32;
    }
 
    // target options for picker
@@ -1912,10 +1917,13 @@ export const NewBuild: React.FC<{ streamId: string; show: boolean; onClose: (new
                            }
                         }).join(" ")} disabled={true} label="Job Arguments" multiline resizable={false} />
                      </Stack>}
-
                   </Stack>
                </Stack>
                <Stack horizontal tokens={{ childrenGap: 16 }} styles={{ root: { paddingTop: 32, paddingLeft: 8, paddingBottom: 8 } }}>
+                  <Stack>
+                     <DefaultButton text="Use New Dialog" style={{ width: 160 }} onClick={() => setNewBuildV2(true)} />
+                  </Stack>
+
                   <Stack grow />
                   <PrimaryButton text="Start Job" disabled={submitting || !template || readOnly} onClick={() => { onSubmit(); }} />
                   <DefaultButton text="Cancel" disabled={submitting} onClick={() => { close(); }} />
diff --git a/Engine/Source/Programs/Horde/HordeDashboard/src/components/Preflight.tsx b/Engine/Source/Programs/Horde/HordeDashboard/src/components/Preflight.tsx
--- a/Engine/Source/Programs/Horde/HordeDashboard/src/components/Preflight.tsx
+++ b/Engine/Source/Programs/Horde/HordeDashboard/src/components/Preflight.tsx
@@ -5,6 +5,7 @@ import { useNavigate } from 'react-router-dom';
 import backend, { useBackend } from '../backend';
 import { useQuery } from './JobDetailCommon';
 import ErrorHandler from './ErrorHandler';
+import { GetJobsTabResponse, JobsTabData, TabType } from '../backend/Api';
 
 function setError(message: string) {
 
@@ -38,6 +39,16 @@ export const PreflightRedirector: React.FC = () => {
    // whether a template is specified
    const templateId = !query.get("templateId") ? "" : query.get("templateId")!;
 
+   const version = !query.get("version") ? "" : query.get("version")!;
+
+   const parameters:string[] = [];
+
+   query.forEach((value, key) => {
+      if (key.startsWith("id")) {
+         parameters.push(`${key}=${value}`);
+      }
+   })
+
    if (!change) {
       setError("No preflight change specified");
       return null;
@@ -79,6 +90,17 @@ export const PreflightRedirector: React.FC = () => {
       return null;
    }
 
+   let tab = "summary";
+   stream.tabs.find(t => {
+      if (t.type !== TabType.Jobs) {
+         return false;
+      }
+
+      if (!!(t as GetJobsTabResponse).templates?.find(t => t === templateId)) {
+         tab = t.title;
+      }
+   })
+
    if (!state.preflightQueried) {
 
       console.log(`Redirecting preflight: ${window.location.href}`);
@@ -97,11 +119,19 @@ export const PreflightRedirector: React.FC = () => {
                url += `&templateId=${templateId}`;
             }
 
+            if (parameters.length) {
+               url += ("&" + parameters.join("&"));
+            }
+
+            if (version === "2") {
+               url += "&newbuildversion=2"
+            }
+
             navigate(url, { replace: true });
             return;
          }
 
-         let url = `/stream/${stream!.id}?tab=summary&newbuild=true&shelvedchange=${change}&p4v=true`;
+         let url = `/stream/${stream!.id}?tab=${tab}&newbuild=true&shelvedchange=${change}&p4v=true`;
 
          if (autosubmit === "true") {
             url += "&autosubmit=true";
@@ -111,6 +141,14 @@ export const PreflightRedirector: React.FC = () => {
             url += `&templateId=${templateId}`;
          }
 
+         if (parameters.length) {
+            url += ("&" + parameters.join("&"));
+         }
+
+         if (version === "2") {
+            url += "&newbuildversion=2"
+         }
+
          navigate(url, { replace: true });
 
 
diff --git a/Engine/Source/Programs/Horde/HordeDashboard/src/components/StreamView.tsx b/Engine/Source/Programs/Horde/HordeDashboard/src/components/StreamView.tsx
--- a/Engine/Source/Programs/Horde/HordeDashboard/src/components/StreamView.tsx
+++ b/Engine/Source/Programs/Horde/HordeDashboard/src/components/StreamView.tsx
@@ -22,6 +22,8 @@ import { JobViewIncremental } from './JobViewIncremental';
 import { NewBuild } from './NewBuild';
 import { StreamSummary } from './StreamSummary';
 import { TopNav } from './TopNav';
+import { NewBuildV2 } from './build/NewBuildV2';
+import { getSiteConfig } from '../backend/Config';
 
 export const SummaryPage: React.FC = () => {
 
@@ -354,6 +356,11 @@ const StreamViewInner: React.FC = observer(() => {
    const stream = projectStore.streamById(streamId);
    const project = stream?.project;
 
+   let newBuildVersion: string = query.get("newbuildversion") ? query.get("newbuildversion") : "1";
+   if (getSiteConfig().environment === "dev") {
+      newBuildVersion = "2";
+   }
+
    if (!stream || !project) {
       console.error("Bad stream or project id in StreamView");
       return <Navigate to="/" replace={true} />
@@ -490,7 +497,17 @@ const StreamViewInner: React.FC = observer(() => {
       <Stack className={hordeClasses.horde}>
          <TopNav />
          <Breadcrumbs items={crumbItems} title={crumbTitle} />
-         <NewBuild streamId={streamId!} jobKey={newBuildTab!} show={shown} onClose={(newJobId) => {
+         {shown && newBuildVersion == "2" && <NewBuildV2 streamId={streamId!} jobKey={newBuildTab!} show={true} onClose={(newJobId) => {
+            setShown(false);
+            if (newJobId) {
+               navigate(`/job/${newJobId}`);
+            } else {
+               if (query.get("newbuild")) {
+                  navigate(`/stream/${streamId}?tab=${queryTab}`, { replace: true });
+               }
+            }
+         }} />}         
+         {shown && newBuildVersion == "1" && <NewBuild streamId={streamId!} jobKey={newBuildTab!} show={true} onClose={(newJobId) => {
             setShown(false);
             if (newJobId) {
                navigate(`/job/${newJobId}`);
@@ -499,7 +516,7 @@ const StreamViewInner: React.FC = observer(() => {
                   navigate(`/stream/${streamId}?tab=${queryTab}`, { replace: true });
                }
             }
-         }} />
+         }} />}
          {findJobsShown && <JobSearchSimpleModal onClose={() => { setFindJobsShown(false) }} streamId={stream.id} />}
          <Stack horizontal style={{ backgroundColor: hordeTheme.horde.neutralBackground }}>
             <div key={`windowsize_streamview_${windowSize.width}_${windowSize.height}`} style={{ width: (vw / 2 - (1440 / 2)) - 12, flexShrink: 0, backgroundColor: modeColors.background }} />
diff --git a/Engine/Source/Programs/Horde/HordeDashboard/src/components/build/NewBuildV2.tsx b/Engine/Source/Programs/Horde/HordeDashboard/src/components/build/NewBuildV2.tsx
new file mode 100644
--- /dev/null
+++ b/Engine/Source/Programs/Horde/HordeDashboard/src/components/build/NewBuildV2.tsx
@@ -0,0 +1,1669 @@
+import { observer } from "mobx-react-lite";
+import { JobDetailsV2 } from "../jobDetailsV2/JobDetailsViewCommon";
+import backend, { useBackend } from "../../backend";
+import { Checkbox, ComboBox, ContextualMenuItemType, DefaultButton, DirectionalHint, Dropdown, DropdownMenuItemType, IComboBoxOption, IContextualMenuItem, IContextualMenuProps, IDropdownOption, Icon, IconButton, Label, MaskedTextField, MessageBar, MessageBarType, Modal, Pivot, PivotItem, PrimaryButton, ScrollablePane, ScrollbarVisibility, Spinner, SpinnerSize, Stack, TagPicker, Text, TextField, ThemeSettingName, TooltipHost } from "@fluentui/react";
+import { getHordeStyling } from "../../styles/Styles";
+import { useEffect, useState } from "react";
+import { action, makeObservable, observable } from "mobx";
+import { useConst } from '@fluentui/react-hooks';
+import { BoolParameterData, ChangeQueryConfig, CreateJobRequest, GetJobsTabResponse, GetTemplateRefResponse, GroupParameterData, JobsTabData, ListParameterData, ListParameterItemData, ListParameterStyle, ParameterData, ParameterType, Priority, StreamData, TabType, TextParameterData } from "../../backend/Api";
+import { ProjectStore } from "../../backend/ProjectStore";
+import templateCache from "../../backend/TemplateCache";
+import ErrorHandler from "../ErrorHandler";
+import dashboard from "../../backend/Dashboard";
+import moment from "moment";
+import Markdown from "markdown-to-jsx";
+import React from "react";
+import { ITextField } from "@fluentui/react/lib-commonjs/TextField";
+import { copyToClipboard } from "../../base/utilities/clipboard";
+
+type ValidationError = {
+   paramString: string;
+   error: string;
+   value: string;
+   context?: string;
+   param?: TextParameterData;
+};
+
+const parameterWidth = 667;
+const parameterGap = 6;
+
+type NewBuildMode = "Basic" | "Advanced";
+
+class BuildOptions {
+
+   constructor(streamId: string, projectStore: ProjectStore, onClose: (newJobId: string | undefined) => void, jobDetails?: JobDetailsV2, jobKey?: string, readOnly?: boolean) {
+
+      if (BuildOptions.instance) {
+         // can happen with hot reload
+         //throw `Failed to clean up build options`;
+      }
+
+      BuildOptions.instance = this;
+      makeObservable(this);
+
+      const query = new URLSearchParams(window.location.search);
+
+      this.isPreflightSubmit = !!query.get("shelvedchange");
+      this.isFromP4V = !!query.get("p4v");
+      this.queryTemplateId = !query.get("templateId") ? "" : query.get("templateId")!;
+      const queryChange = parseInt(query.get("shelvedchange") ? query.get("shelvedchange")! : "0");
+      if (!isNaN(queryChange) && queryChange) {
+         this.queryShelvedChange = queryChange;
+      }
+
+      this.preflightChange = jobDetails?.jobData?.preflightChange;
+      this.change = jobDetails?.jobData?.change;
+      if (!this.preflightChange) {
+         this.preflightChange = this.queryShelvedChange;
+      }
+
+      if (query.get("autosubmit") === "true") {
+         this.autoSubmit = true;
+      }
+
+      this.streamId = streamId;
+      this.onClose = onClose;
+      this.readOnly = readOnly ?? false;
+      this.projectStore = projectStore;
+      this.jobKey = jobKey;
+      this.jobDetails = jobDetails;
+
+      this.showAllTemplates = jobKey === "all" || jobKey === "summary"
+
+      const stream = projectStore.streamById(streamId);
+      if (!stream) {
+         throw `Unable to get stream ${streamId}`;
+      }
+
+      this.stream = stream;
+
+      this.load();
+   }
+
+   onClose: (newJobId?: string) => void;
+
+   loaded = false;
+
+   jobKey?: string;
+   jobDetails?: JobDetailsV2;
+   readOnly: boolean;
+
+   isPreflightSubmit?: boolean;
+   isFromP4V?: boolean;
+   queryTemplateId?: string;
+
+   submitting?: boolean;
+
+   mode: NewBuildMode = "Basic";
+
+   // advanced options
+   advJobName?: string;
+   advJobPriority?: Priority;
+   advUpdateIssues?: boolean;
+
+   parameters: Record<string, string> = {};
+
+   validationErrors: ValidationError[] = [];
+
+   get defaultPreflightQuery(): ChangeQueryConfig[] | undefined {
+
+      const template = this.template;
+      let defaultPreflightQuery: ChangeQueryConfig[] | undefined;
+
+
+      if (template?.defaultChange) {
+         defaultPreflightQuery = template.defaultChange;
+      } else {
+         const defaultChange = this.stream?.defaultPreflight?.change;
+
+         if (defaultChange) {
+            if (!defaultChange.name && defaultChange.templateId) {
+               const defaultTemplate = this.allTemplates.find(t => t.id === defaultChange.templateId);
+               if (defaultTemplate) {
+                  defaultChange.name = "Latest Success - " + defaultTemplate.name;
+               }
+            }
+
+            if (defaultChange.name) {
+               defaultPreflightQuery = [defaultChange];
+            }
+         }
+      }
+
+      const defaultStreamPreflightTemplate = this.stream.templates.find(t => t.id === this.stream.defaultPreflight?.templateId);
+
+      if (!defaultPreflightQuery && this.stream.defaultPreflight?.templateId) {
+
+         if (defaultStreamPreflightTemplate) {
+            defaultPreflightQuery = defaultStreamPreflightTemplate.defaultChange;
+         }
+
+         if (!defaultPreflightQuery) {
+            console.error(`Unable to find default stream preflight template ${this.stream.defaultPreflight?.templateId} in stream templates`);
+         }
+      }
+
+      return defaultPreflightQuery;
+
+   }
+
+   get advancedModified(): boolean {
+
+      return !!((this.advJobPriority && this.advJobPriority !== Priority.Normal) || this.advUpdateIssues || !!this.advJobName)
+   }
+
+   change?: number;
+   changeOption?: string;
+   preflightChange?: number;
+   queryShelvedChange?: number;
+   autoSubmit?: boolean;
+
+   static get() { return BuildOptions.instance };
+
+   template?: GetTemplateRefResponse;
+   templates: GetTemplateRefResponse[] = [];
+
+   clear() {
+      BuildOptions.instance = undefined;
+   }
+
+   streamId: string;
+   stream: StreamData;
+   projectStore: ProjectStore;
+
+   preview: boolean = false;
+   previewSource?: string;
+   previewRestoreId?: string;
+
+   showAllTemplates = false;
+
+   setSubmitting(submitting?: boolean) {
+      if (this.submitting === submitting) {
+         return;
+      }
+
+      this.submitting = submitting;
+      this.setChanged();
+
+   }
+
+   onAutoSubmitChanged(value: boolean) 
+   {
+      if (this.autoSubmit === value) {
+         return;
+      }
+
+      this.autoSubmit = value;
+      this.setChanged();
+   }
+
+   onPreviewChanged(preview: boolean) {
+      if (this.preview == preview) {
+         return;
+      }
+
+      if (preview) {
+         this.previewRestoreId = this.template?.id;
+      } else {
+         const template = this.allTemplates.find(t => t.id === this.previewRestoreId);
+         this.template = undefined;
+         if (template) {
+            this.setTemplate(template.id);
+         } else {
+            this.selectDefaultTemplate();
+         }
+      }
+
+      this.preview = preview;
+
+      this.setChanged();
+   }
+
+   onPreviewSourceChanged(source: string) {
+      this.previewSource = source;
+      this.setPreviewTemplate();
+   }
+
+   onModeChanged(mode: NewBuildMode) {
+      if (mode === this.mode) {
+         return;
+      }
+
+      this.mode = mode;
+      this.setChanged();
+   }
+
+   onPreflightChange(preflightCL?: number) {
+
+      if (this.preflightChange === preflightCL) {
+         return;
+      }
+
+      this.preflightChange = preflightCL;
+
+      this.setChanged();
+   }
+
+   onTemplateChange(templateId: string) {
+      this.setTemplate(templateId, true);
+   }
+
+   onBooleanChanged(id: string, value: boolean) {
+      this.parameters[id] = value.toString();
+      this.setParameterChanged();
+   }
+
+   onTextChanged(id: string, value: string) {
+      this.parameters[id] = value;
+      this.setParameterChanged();
+   }
+
+   onListItemChanged(id: string, value: boolean) {
+      this.parameters[id] = value.toString();
+   }
+
+   onListItemsUpdated() {
+      this.setParameterChanged();
+   }
+
+   setValidationErrors(errors: ValidationError[]) {
+      this.validationErrors = errors;
+      this.setChanged();
+   }
+
+   subscribe() {
+      if (this.changed) { }
+   }
+
+   // so every parameter change
+   subscribeToParameterChange() {
+      if (this.parameterChanged) { }
+   }
+
+   setShowAllTemplates(value: boolean) {
+      if (this.showAllTemplates === value) {
+         return;
+      }
+      this.showAllTemplates = value;
+      this.filterTemplates();
+      this.setChanged();
+   }
+
+   @observable
+   private changed = 0;
+
+   @action
+   setChanged() {
+      this.changed++;
+   }
+
+   @observable
+   private parameterChanged = 0;
+
+   @action
+   private setParameterChanged() {
+      this.parameterChanged++;
+   }
+
+   validateParameter(p: ParameterData, errors: ValidationError[]): void {
+
+      if (p.type === ParameterType.Text) {
+
+         const param = p as TextParameterData;
+
+         if (!param.validation) {
+            return;
+         }
+
+         let value = this.parameters[param.id] ?? param.default;
+
+         let regex: RegExp | undefined;
+
+         try {
+            regex = new RegExp(param.validation, 'gm');
+         } catch (reason: any) {
+
+            errors.push({
+               param: param,
+               paramString: param.label,
+               value: value,
+               error: `Invalid validation regex: ${param.validation}`,
+               context: reason.toString()
+            });
+
+            return;
+         }
+
+         const match = value.match(regex);
+
+         if (!match || !match.length || match[0]?.trim() !== value.trim()) {
+            errors.push({
+               param: param,
+               paramString: param.label,
+               value: value,
+               error: `Did not match regex: ${param.validation}`,
+               context: param.validationError ?? param.hint
+            });
+            return;
+         }
+      }
+   }
+
+
+   validate(errors: ValidationError[]): boolean {
+
+      if (this.preflightChange && !this.template.allowPreflights) {
+         errors.push({ error: `Template "${this.template.name}" does not allow preflights`, paramString: "Shelved Change", value: this.preflightChange.toString() })
+      }
+
+      if (this.change !== undefined && isNaN(this.change)) {
+         errors.push({
+            paramString: "Change",
+            value: this.change.toString(),
+            error: "Invalid changelist"
+         });
+      }
+
+      if (this.preflightChange !== undefined && isNaN(this.preflightChange)) {
+         errors.push({
+            paramString: "Shelved Change",
+            value: this.preflightChange.toString(),
+            error: "Invalid shelved changelist"
+         });
+      }
+
+      this.template.parameters.forEach(p => {
+
+         this.validateParameter(p, errors);
+      });
+
+      if (errors.length) {
+         this.setChanged();
+      }
+
+      this.validationErrors = errors;
+
+      return !errors.length;
+
+   }
+
+   estimateHeight(previewPanel?: boolean): number {
+
+      const template = this.template;
+      let estimatedHeight = parameterGap * template.parameters.length;
+
+      const estimateHeight = (param: ParameterData) => {
+
+         switch (param.type) {
+
+            case ParameterType.List:
+               estimatedHeight += 54;
+               break;
+            case ParameterType.Bool:
+               estimatedHeight += 20;
+               break;
+            case ParameterType.Text:
+               estimatedHeight += 54;
+               break;
+         }
+
+      }
+
+      template.parameters.forEach(p => {
+         estimateHeight(p);
+      });
+
+      if (!!this.preflightChange) {
+         estimatedHeight += 32;
+      }
+
+      if (this.preview && estimatedHeight < 600) {
+         estimatedHeight = 600
+      }
+
+      if (estimatedHeight > 600) {
+         estimatedHeight = 600
+      }
+
+      if (this.preview) {
+         if (!previewPanel) {
+            if (this.mode === "Basic") {
+               estimatedHeight -= 64;
+            } else {
+               estimatedHeight -= 22;
+            }
+         }
+
+      } else if (this.mode === "Advanced") {
+         estimatedHeight += 192;
+      }
+
+
+      return estimatedHeight;
+
+   }
+
+   private setTemplateParameters() {
+
+      const template = this.template;
+      if (!template) {
+         return;
+      }
+
+      this.parameters = {};
+
+      template.parameters.forEach(param => {
+
+         switch (param.type) {
+            case ParameterType.Bool:
+               const b = param as BoolParameterData;
+               this.parameters[b.id] = this.jobDetails?.jobData?.parameters[b.id] ?? b.default.toString();
+               break;
+            case ParameterType.Text:
+               const t = param as TextParameterData;
+               this.parameters[t.id] = (this.jobDetails?.jobData?.parameters[t.id] ?? t.default) ?? "";
+               break;
+            case ParameterType.List:
+               const list = param as ListParameterData;
+               list.items.forEach(item => {
+                  this.parameters[item.id] = (this.jobDetails?.jobData?.parameters[item.id] ?? item.default.toString()) ?? "false";
+               });
+               break;
+            default:
+               console.error(`Unknown parameter type: ${param.type}`)
+         }
+      });
+
+   }
+
+   private setPreviewTemplate() {
+      const source = this.previewSource ?? ""
+      const template = JSON.parse(source) as GetTemplateRefResponse;
+      this.template = template;
+
+      this.setTemplateParameters();
+
+      this.setChanged();
+
+   }
+
+   private setTemplate(templateId: string, notifyChanged = true) {
+
+      const template = this.allTemplates.find(t => t.id === templateId);
+      if (template.id === this.template?.id) {
+         return;
+      }
+
+      if (!template) {
+         throw `Unable to find template ${templateId}`;
+      }
+
+      const source = { ...template };
+      source.hash = undefined;
+
+      this.previewSource = JSON.stringify(source, null, 4);
+
+      if (this.preview) {
+
+         this.setPreviewTemplate();
+         return;
+      }
+
+      this.template = template;
+
+      this.setTemplateParameters();
+
+      if (notifyChanged) {
+         this.setChanged();
+      }
+   }
+
+   private setDefaultParameters() {
+
+      if (!this.template) {
+         return;
+      }
+
+      if (Object.keys(this.jobDetails?.jobData?.parameters ?? {}).length) {
+         this.parameters = { ...this.parameters, ...this.jobDetails!.jobData!.parameters };
+      }
+
+      const query = new URLSearchParams(window.location.search);
+
+      this.template.parameters.forEach(p => {
+         let qid = "";
+         switch (p.type) {
+            case ParameterType.Bool:
+               const b = p as BoolParameterData;
+               qid = `id-${b.id}`;
+               if (query.get(qid)) {
+                  this.parameters[b.id] = query.get(qid) as string;
+               }
+               break;
+            case ParameterType.Text:
+               const t = p as TextParameterData;
+               qid = `id-${t.id}`;
+               if (query.get(qid)) {
+                  this.parameters[t.id] = query.get(qid) as string;
+               }
+               break;
+            case ParameterType.List:
+               const list = p as ListParameterData;
+               let any = false;
+               list.items.forEach(i => {
+                  qid = `id-${i.id}`;
+                  if (query.get(qid)) {
+                     any = true;
+                  }
+               });
+
+               if (any) {
+
+                  list.items.forEach(i => {
+                     this.parameters[i.id] = undefined;
+                     qid = `id-${i.id}`;
+                     if (query.get(qid)) {
+                        this.parameters[i.id] = query.get(qid) as string;
+                     }
+                  });
+               }
+               break;
+            default:
+               console.error(`Unknown parameter type ${p.type}`)
+         }
+
+      })
+
+   }
+
+   private selectDefaultTemplate() {
+
+
+      const jobDetails = this.jobDetails;
+      const stream = this.stream;
+      const queryTemplateId = this.queryTemplateId;
+
+      let t: GetTemplateRefResponse | undefined;
+
+      if (jobDetails?.template) {
+         t = this.allTemplates.find(t => t.name === jobDetails.template!.name);
+      }
+      // handle preflight redirect case
+      if (!t && this.queryShelvedChange && !jobDetails) {
+
+         if (queryTemplateId) {
+            let errorReason = "";
+            t = this.allTemplates?.find(t => t.id === queryTemplateId);
+            if (!t) {
+               errorReason = `Unable to find queryTemplateId ${queryTemplateId} in stream ${this.streamId}`;
+               console.error(errorReason);
+
+            } else if (!t.allowPreflights) {
+               errorReason = `Template does not allow preflights: queryTemplateId ${queryTemplateId} in stream ${this.streamId}`;
+               console.error(errorReason);
+               t = undefined;
+            }
+
+            if (errorReason) {
+               ErrorHandler.set({
+                  reason: `${errorReason}`,
+                  title: `Preflight Template Error`,
+                  message: `There was an issue with the specified preflight template.\n\nReason: ${errorReason}\n\nTime: ${moment.utc().format("MMM Do, HH:mm z")}`
+               }, true);
+            }
+         }
+
+         if (!t) {
+            const defaultStreamPreflightTemplate = stream.templates.find(t => t.id === stream.defaultPreflight?.templateId);
+            t = defaultStreamPreflightTemplate;
+            if (!t) {
+               console.error(`Stream default preflight template cannot be found for stream ${stream.fullname} : stream defaultPreflightTemplate ${stream.defaultPreflight?.templateId}, will use first template in list`);
+            }
+         }
+      }
+
+      if (!t && queryTemplateId) {
+         t = this.allTemplates.find(template => template.id === queryTemplateId);
+      }
+
+      if (!t) {
+         const pref = dashboard.getLastJobTemplateSettings(stream.id, this.templates.map(t => t.id));
+         if (pref) {
+            t = this.templates.find(t => stream.id === pref.streamId && t.id === pref.templateId)
+         }
+      }
+      // default to sane template when all are shown
+      if (!t && this.showAllTemplates && stream.tabs.length > 0) {
+         const stab = stream.tabs[0] as JobsTabData;
+         if (stab.templates && stab.templates.length > 0) {
+            t = this.allTemplates.find(t => t.id === stab.templates![0]);
+         }
+      }
+
+      if (!t) {
+         t = this.templates.length > 0 ? this.templates[0] : undefined;
+         if (!t) {
+            t = this.allTemplates[0];
+         }
+      }
+
+      if (!t) {
+         throw `Unable to get default template for jobKey: ${this.jobKey}, check that tab has a valid template list`;
+      }
+
+      // when coming from P4V, always clear the base CL
+      if (this.isFromP4V) {
+         this.change = undefined;
+         this.changeOption = undefined;
+      }
+
+      this.setTemplate(t.id, false);
+
+   }
+
+   private filterTemplates() {
+
+      let templates: GetTemplateRefResponse[] = [];
+
+      if (!this.jobDetails) {
+         if (this.showAllTemplates) {
+            templates = [...this.allTemplates].sort((a, b) => a.name.localeCompare(b.name));
+         } else {
+            const tab = this.stream.tabs.find(tab => tab.title === this.jobKey) as JobsTabData | undefined;
+            if (!tab) {
+               throw `No stream tab ${this.jobKey}`;
+            }
+
+            if (tab.type !== TabType.Jobs) {
+               throw `Tab is not of Jobs type: ${tab.title}`;
+            }
+
+            tab.templates?.forEach(templateName => {
+               const t = this.allTemplates.find(t => t.id === templateName);
+               if (!t) {
+                  console.error(`Could not find template ${templateName}`);
+                  return;
+               }
+               templates.push(t);
+            });
+         }
+      } else {
+         const query = new URLSearchParams(window.location.search);
+         const allowtemplatechange = query.get("allowtemplatechange") === "true";
+         if (allowtemplatechange) {
+            templates = [...this.allTemplates];
+         } else {
+            const template = this.allTemplates.find(t => t.name === this.jobDetails.template?.name);
+            if (template) {
+               templates = [template];
+            }
+            templates = [...this.allTemplates];
+            console.error(`Unable to find job template ${this.jobDetails.template?.name}`)
+         }
+      }
+
+      this.templates = templates;
+
+   }
+
+   private allTemplates: GetTemplateRefResponse[] = [];
+
+   private async load() {
+      this.allTemplates = await templateCache.getStreamTemplates(this.stream);
+
+      this.filterTemplates();
+
+      this.selectDefaultTemplate();
+      this.setDefaultParameters();
+
+      this.loaded = true;
+      this.setChanged();
+   }
+
+   private static instance?: BuildOptions;
+
+}
+
+const BoolParameter: React.FC<{ param: BoolParameterData }> = observer(({ param }) => {
+
+   const options = BuildOptions.get();
+
+   options.subscribeToParameterChange();
+
+   const key = param.id.replaceAll(".", "-");
+
+   return <Stack>
+      <Checkbox key={key}
+         label={param.label}
+         disabled={options.readOnly}
+         checked={options.parameters[param.id] == "true"}
+         onChange={(ev, value) => {
+            ev.preventDefault();
+            options.onBooleanChanged(param.id, value);
+         }}
+      />
+   </Stack>
+
+})
+
+const TextParameter: React.FC<{ param: TextParameterData }> = observer(({ param }) => {
+
+   const options = BuildOptions.get();
+
+   options.subscribeToParameterChange();
+
+   const key = param.id.replaceAll(".", "-");
+
+   return <Stack>
+      <TextField key={key}
+         placeholder={options.jobDetails ? "" : param.hint}
+         label={param.label}
+         spellCheck={false}
+         value={options.parameters[param.id] ?? ""}
+         disabled={options.readOnly}
+         onChange={(ev, value) => {
+            options.onTextChanged(param.id, value);
+         }}
+      />
+   </Stack>
+})
+
+
+const TagPickerParameter: React.FC<{ param: ListParameterData }> = observer(({ param }) => {
+
+   const options = BuildOptions.get();
+
+   options.subscribeToParameterChange();
+
+   const key = `parameter_key_${param.label}`
+
+   type PickerItem = {
+      itemData: ListParameterItemData;
+      key: string;
+      name: string;
+   }
+
+   const allItems: PickerItem[] = param.items.map(item => {
+      return {
+         itemData: item,
+         key: item.text,
+         name: item.text
+      };
+   });
+
+   const selectedItems = allItems.filter(item => {
+      return options.parameters[item.itemData.id] === "true";
+   });
+
+   // tag picker
+   return <Stack key={key}>
+      <Label> {param.label}</Label>
+      <TagPicker
+         disabled={options.readOnly}
+         onResolveSuggestions={(filter, selected) => {
+            return allItems.filter(i => {
+               return !selected?.find(s => i.key === s.key) && i.name.toLowerCase().indexOf(filter.toLowerCase()) !== -1;
+            });
+         }}
+
+         onEmptyResolveSuggestions={(selected) => {
+            return allItems.filter(i => {
+               return !selected?.find(s => i.key === s.key);
+            });
+         }}
+
+         selectedItems={selectedItems}
+
+         onChange={(items?) => {
+
+            allItems.forEach(item => {
+               if (!items) {
+                  return;
+               }
+
+               options.onListItemChanged(item.itemData.id, items.find(i => i.key === item.key) ? true : false);
+
+            });
+
+            options.onListItemsUpdated();
+
+         }}
+      />
+   </Stack>;
+})
+
+const BasicListParameter: React.FC<{ param: ListParameterData }> = observer(({ param }) => {
+
+   const options = BuildOptions.get();
+   options.subscribeToParameterChange();
+
+   const key = `parameter_key_${param.label}`
+
+   const doptions: IDropdownOption[] = [];
+
+   param.items.forEach((item, index) => {
+      doptions.push({
+         key: item.id,
+         text: item.text,
+         selected: options.parameters[item.id] === "true"
+      });
+   });
+
+   return <Dropdown key={key}
+      label={param.label}
+      options={doptions}
+      disabled={options.readOnly}
+      placeholder={options.jobDetails ? "" : "Select option"}
+      onChange={(ev, option, index) => {
+
+         param.items.forEach(item => {
+            options.onListItemChanged(item.id, false);
+         });
+
+         options.onListItemChanged(param.items[index!].id, true);
+         options.onListItemsUpdated();
+
+      }} />
+})
+
+const MultiListParameter: React.FC<{ param: ListParameterData }> = observer(({ param }) => {
+
+   const options = BuildOptions.get();
+   options.subscribeToParameterChange();
+
+   const { modeColors } = getHordeStyling();
+
+   const jobDetails = options.jobDetails;
+
+   const key = `parameter_key_${param.label}`
+
+   const gset: Set<string> = new Set();
+
+   param.items.forEach(item => {
+      if (!item.group) {
+         item.group = "__nogroup";
+      }
+      if (item.group) {
+         gset.add(item.group);
+      }
+   });
+
+   const groups = Array.from(gset).sort((a, b) => a.localeCompare(b));
+
+   const doptions: IDropdownOption[] = [];
+
+   const selectedKeys: string[] = [];
+
+   groups.forEach(group => {
+
+      if (group !== "__nogroup") {
+         doptions.push({
+            key: `group_${group}`,
+            text: group,
+            itemType: DropdownMenuItemType.Header
+         });
+      }
+
+      const dupes = new Map<string, number>();
+
+      param.items.forEach(item => {
+         const v = dupes.get(item.text);
+         dupes.set(item.text, !v ? 1 : v + 1);
+      });
+
+      param.items.forEach(item => {
+         if (item.group === group) {
+            const key = item.id;
+            const selected = options.parameters[item.id] === "true";
+            if (selected) {
+               selectedKeys.push(key);
+            }
+            doptions.push({
+               key: key,
+               data: item,
+               text: (item.group !== "__nogroup" && dupes.get(item.text)! > 1) ? `${item.text} - ${item.group}` : item.text,
+               selected: selected
+            });
+         }
+      });
+   });
+
+   return <Dropdown
+      key={key}
+      disabled={options.readOnly}
+      placeholder={jobDetails ? "" : "Select options"}
+      styles={{
+         callout: {
+            selectors: {
+               ".ms-Callout-main": {
+                  padding: "4px 4px 12px 12px",
+                  overflow: "hidden"
+               }
+            }
+         },
+         dropdownItemHeader: { fontSize: 12, color: modeColors.text },
+         dropdownOptionText: { fontSize: 12 },
+         dropdownItem: {
+            minHeight: 18, lineHeight: 18, selectors: {
+               '.ms-Checkbox-checkbox': {
+                  width: 14,
+                  height: 14,
+                  fontSize: 11
+               }
+            }
+         },
+         dropdownItemSelected: {
+            minHeight: 18, lineHeight: 18, backgroundColor: "inherit",
+            selectors: {
+               '.ms-Checkbox-checkbox': {
+                  width: 14,
+                  height: 14,
+                  fontSize: 11
+               }
+            }
+         }
+      }}
+      label={param.label}
+      defaultSelectedKeys={selectedKeys}
+      onChange={(event, option, index) => {
+
+         if (!option) {
+            return;
+         }
+
+         options.onListItemChanged((option.data as ListParameterItemData).id, option.selected ? true : false)
+         options.onListItemsUpdated();
+      }}
+      multiSelect
+      options={doptions}
+   />;
+})
+
+
+const ListParameter: React.FC<{ param: ListParameterData }> = ({ param }) => {
+
+
+   if (param.style === ListParameterStyle.TagPicker) {
+      return <TagPickerParameter param={param} />
+   } else if (param.style === ListParameterStyle.MultiList) {
+      return <MultiListParameter param={param} />
+   } else if (param.style === ListParameterStyle.List) {
+      return <BasicListParameter param={param} />
+   }
+
+   return null;
+
+}
+
+let toolTipId = 0;
+
+const BuildParametersPanel: React.FC = observer(() => {
+
+   const options = BuildOptions.get();
+   options.subscribe();
+
+   const template = options.template;
+
+   if (!template) {
+      console.error("Build Modal has no template");
+      return null;
+   }
+
+   const estimatedHeight = options.estimateHeight();
+
+   const renderParameter = (param: ParameterData) => {
+
+      switch (param.type) {
+         case ParameterType.Bool:
+            const b = param as BoolParameterData;
+            return <BoolParameter key={`parameter_key_${b.id}`} param={b} />
+         case ParameterType.Text:
+            const t = param as TextParameterData;
+            return <TextParameter key={`parameter_key_${t.id}`} param={t} />
+         case ParameterType.List:
+            const list = param as ListParameterData;
+            return <ListParameter key={`parameter_key_${list.label}`} param={list} />
+         default:
+            return <Text>Unknown Parameter Type</Text>;
+      }
+   };
+
+   return <Stack>
+      {!!template?.description && <Stack>
+         <Stack style={{ width: parameterWidth }}>
+            <Stack style={{ paddingTop: "4px", paddingBottom: "20px" }}>
+               <Markdown styles={{ root: { maxHeight: 240, overflow: "auto", th: { fontSize: 12 } } }}>{template.description}</Markdown>
+            </Stack>
+         </Stack>
+      </Stack>}
+
+      <Stack style={{
+         height: estimatedHeight,
+         position: 'relative',
+         width: parameterWidth
+      }}><ScrollablePane scrollbarVisibility={ScrollbarVisibility.auto}>
+            <Stack tokens={{ childrenGap: parameterGap }}>
+               {template.parameters.map((p) => {
+
+
+                  let key = "";
+
+                  switch (p.type) {
+                     case ParameterType.Bool:
+                        const b = p as BoolParameterData;
+                        key = `boolean_parameter_key_${b.id}`;
+                        break;
+                     case ParameterType.Text:
+                        const t = p as TextParameterData;
+                        key = `text_parameter_key_${t.id}`;
+                        break;
+                     case ParameterType.List:
+                        const list = p as ListParameterData;
+                        key = `list_parameter_key_${list.label}`;
+                        break;
+                     default:
+                        return <Text>Unknown Parameter Type</Text>;
+                  }
+
+                  if ((p as any).toolTip) {
+                     return <TooltipHost key={key} content={(p as any).toolTip} id={`unique_tooltip_${toolTipId++}`} directionalHint={DirectionalHint.leftCenter}>
+                        {renderParameter(p)}
+                     </TooltipHost>
+                  } else {
+                     return renderParameter(p);
+                  }
+
+               })}
+               {!!options.preflightChange && <Stack style={{ paddingTop: 8 }}><Checkbox label="Automatically submit preflight on success"
+                  checked={options.autoSubmit}
+                  onChange={(ev, checked) => {
+                     options.onAutoSubmitChanged(checked ? true : false);
+                  }}
+               /></Stack>}
+            </Stack>
+         </ScrollablePane>
+      </Stack>
+   </Stack>
+})
+
+const SubmittingModal: React.FC = () => {
+
+   const { hordeClasses } = getHordeStyling();
+
+   return <Modal isOpen={true} isBlocking={true} topOffsetFixed={true} styles={{ main: { padding: 8, width: 700, hasBeenOpened: false, top: "80px", position: "absolute" } }} className={hordeClasses.modal}>
+      <Stack tokens={{ childrenGap: 24 }} styles={{ root: { padding: 8 } }}>
+         <Stack grow verticalAlign="center">
+            <Text variant="mediumPlus" styles={{ root: { fontWeight: "unset", fontFamily: "Horde Open Sans SemiBold" } }}>Creating Job</Text>
+         </Stack>
+         <Stack horizontalAlign="center">
+            <Text variant="mediumPlus">The job is being created and will be available soon.</Text>
+         </Stack>
+         <Stack verticalAlign="center" style={{ paddingBottom: 32 }}>
+            <Spinner size={SpinnerSize.large} />
+         </Stack>
+      </Stack>
+   </Modal>
+}
+
+const TemplateSelector: React.FC = () => {
+
+   const options = BuildOptions.get();
+
+   let templateOptions: IContextualMenuItem[] = [];
+
+   if (!options.showAllTemplates) {
+      templateOptions = options.templates.map(t => {
+         return { key: t.id, text: t.name, onClick: () => options.onTemplateChange(t.id) };
+      }).sort((a, b) => a.text < b.text ? -1 : 1);
+   } else {
+
+      const sorted = new Map<string, GetTemplateRefResponse[]>();
+
+      options.templates.forEach(t => {
+
+         options.stream.tabs.forEach(tab => {
+            if (tab.type !== TabType.Jobs) {
+               return;
+            }
+
+            const jtab = tab as GetJobsTabResponse;
+            if (!jtab.templates?.find(template => template === t.id)) {
+               return;
+            }
+
+            if (!sorted.has(jtab.title)) {
+               sorted.set(jtab.title, []);
+            }
+
+            sorted.get(jtab.title)!.push(t);
+
+         })
+      })
+
+      Array.from(sorted.keys()).sort((a, b) => a < b ? -1 : 1).forEach(cat => {
+
+         const templates = sorted.get(cat);
+         if (!templates?.length) {
+            return;
+         }
+
+         const subItems = templates.sort((a, b) => a.name < b.name ? -1 : 1).map(t => {
+            return { key: t.id, text: t.name, onClick: () => options.onTemplateChange(t.id) };
+         })
+
+         templateOptions.push({ key: `${cat}_category`, text: cat, subMenuProps: { items: subItems } });
+
+      })
+   }
+
+   if (options.jobKey !== "all" && options.jobKey !== "summary") {
+
+      templateOptions.push({ key: `show_all_templates_divider`, itemType: ContextualMenuItemType.Divider });
+
+      if (!options.showAllTemplates) {
+         templateOptions.push({ key: `show_all_templates`, text: "Show All Templates", onClick: (ev) => { ev?.stopPropagation(); ev?.preventDefault(); options.setShowAllTemplates(true) } });
+      } else {
+         templateOptions.push({ key: `show_all_templates`, text: "Filter Templates", onClick: (ev) => { ev?.stopPropagation(); ev?.preventDefault(); options.setShowAllTemplates(false) } });
+      }
+   }
+
+   const templateMenuProps: IContextualMenuProps = {
+      shouldFocusOnMount: true,
+      subMenuHoverDelay: 0,
+      items: templateOptions,
+   };
+
+   return <Stack>
+      <Label>Template</Label>
+      <DefaultButton style={{ width: 280, textAlign: "left" }} disabled={options.readOnly} menuProps={templateMenuProps} text={options?.template?.name ?? "Error"} />
+   </Stack>
+}
+
+let textCounter = 0;
+
+const PreviewPanel: React.FC = observer(() => {
+
+   const options = BuildOptions.get();
+
+   const textRef = React.useRef<ITextField>(null);
+
+   options.subscribe();
+
+   const onUpdate = (source?: string) => {
+      const csource = source ?? textRef.current?.value ?? "";
+      options.onPreviewSourceChanged(csource);
+   }
+
+   const estimatedHeight = options.estimateHeight(true);
+
+   return <Stack>
+      <Stack style={{ paddingTop: 64, paddingBottom: 12 }}>
+         <Label>
+            Template Source
+         </Label>
+         <Stack>
+            <TextField key={`text_counter_${textCounter++}`} componentRef={textRef} style={{ width: parameterWidth + 4, height: estimatedHeight }} multiline spellCheck={false} resizable={false} defaultValue={options.previewSource} />
+         </Stack>
+         <Stack>
+         </Stack>
+      </Stack>
+      <Stack horizontal tokens={{ childrenGap: 16 }} styles={{ root: { paddingTop: 32, paddingLeft: 8, paddingBottom: 8 } }}>
+         <Stack grow />
+         <DefaultButton text="Copy" disabled={false} onClick={() => {
+
+            copyToClipboard(textRef.current?.value ?? "");
+
+         }} />
+         <DefaultButton text="Paste" disabled={false} onClick={async () => {
+            const text = await navigator.clipboard.readText();
+            onUpdate(text);
+         }} />
+         <DefaultButton text="Format" disabled={false} onClick={async () => {
+            let text = textRef.current?.value ?? "{}";
+            const value = JSON.parse(text);
+            onUpdate(JSON.stringify(value, null, 4));
+         }} />
+         <PrimaryButton text="Preview" onClick={() => { onUpdate(); }} />
+      </Stack>
+
+   </Stack>
+
+});
+
+const AdvancedPanel: React.FC = observer(() => {
+
+   const options = BuildOptions.get();
+   options.subscribe();
+
+   const stream = options.stream;
+
+   const template = options.template;
+
+   if (!template) {
+      console.error("Build Modal has no template");
+      return null;
+   }
+
+   const priorityOptions: IDropdownOption[] = [];
+
+   for (const p in Priority) {
+      priorityOptions.push({
+         text: p,
+         key: p,
+         isSelected: options!.advJobPriority ? p === options!.advJobPriority : Priority.Normal === p
+      });
+   }
+
+   const height = options.estimateHeight();
+
+   return <Stack style={{
+      height: height,
+      position: 'relative',
+      width: parameterWidth,
+      maxHeight: 'calc(100vh - 360px)'
+   }}>
+      <Stack style={{ paddingBottom: 12 }} tokens={{ childrenGap: 12 }}>
+         {!!stream.configRevision && <Stack>
+            <TextField label="Template Path" readOnly={true} value={stream.configPath ?? ""} />
+         </Stack>
+         }
+         <Stack>
+            <Dropdown disabled={options.readOnly} key={"key_adv_priority"} defaultValue={options.advJobPriority} label="Priority" options={priorityOptions} onChange={(ev, option) => {
+               options.advJobPriority = option?.key as Priority;
+               options.setChanged();
+            }
+            } />
+         </Stack>
+         <Stack>
+            <Checkbox disabled={options.readOnly || options.template?.updateIssues} key="key_adv_update_issues"
+               label="Update Build Health Issues"
+               defaultChecked={options.template?.updateIssues ? true : options.advUpdateIssues}
+               onChange={(ev, checked) => {
+                  options.advUpdateIssues = checked ? true : undefined;
+                  options.setChanged();
+               }} />
+         </Stack>
+         <Stack>
+            <TextField key={"key_adv_job_name"} disabled={options.readOnly} spellCheck={false} defaultValue={options.advJobName} label="Job Name" onChange={(ev, newValue) => {
+               options.advJobName = newValue;
+               options.setChanged();
+            }} />
+         </Stack>
+         <Stack>
+            <Checkbox disabled={options.readOnly}
+               label="Template Editor"
+               checked={options.preview}
+               onChange={(ev, checked) => {
+                  if (checked) {
+                     options.onModeChanged("Basic");
+                  }
+                  options.onPreviewChanged(checked);
+               }} />
+         </Stack>
+
+      </Stack>
+   </Stack>
+
+})
+
+const BuildModal: React.FC = observer(() => {
+
+   const options = BuildOptions.get();
+   options.subscribe();
+
+   const template = options.template;
+
+   if (!template) {
+      console.error("Build Modal has no template");
+      return null;
+   }
+
+   const { hordeClasses, modeColors } = getHordeStyling();
+
+   const changeOptions = ["Latest Change"];
+
+   const defaultPreflightQuery = options.defaultPreflightQuery;
+
+   if (defaultPreflightQuery) {
+      changeOptions.push(...defaultPreflightQuery.filter(n => !!n.name).map(n => n.name!));
+   }
+
+   const changeItems: IComboBoxOption[] = changeOptions.map(name => { return { key: `key_change_option_${name}`, text: name } });
+
+   let changeText: string | undefined;
+
+   if (options.change) {
+      changeText = options.change.toString();
+   }
+   if (options.changeOption) {
+      changeText = options.changeOption;
+   }
+
+   const onSubmit = () => {
+      const errors: ValidationError[] = [];
+      if (!options.validate(errors)) {
+         return;
+      }
+
+      const templateId = template.id;
+      const updateIssues = template.updateIssues ? undefined : options.advUpdateIssues;
+
+      let changeQueries: ChangeQueryConfig[] | undefined;
+
+      if (typeof options.change !== 'number') {
+         const changeOption = options.changeOption ?? "Latest Change";
+         if (changeOption !== "Latest Change" && defaultPreflightQuery) {
+            changeQueries = defaultPreflightQuery;
+         }
+         else {
+            const changeQuery = defaultPreflightQuery?.find(p => p.name === options.changeOption);
+            if (changeQuery) {
+               changeQueries = [{ ...changeQuery, condition: undefined }];
+            }
+         }
+      }
+
+      const data: CreateJobRequest = {
+         streamId: options.streamId,
+         templateId: templateId,
+         name: options.advJobName,
+         priority: options.advJobPriority,
+         updateIssues: updateIssues,
+         changeQueries: changeQueries,
+         parameters: options.parameters
+      };
+
+      if (typeof (options.change) === 'number') {
+         data.change = options.change;
+      }
+
+      if (typeof (options.preflightChange) === 'number') {
+
+         data.preflightChange = options.preflightChange;
+
+         data.updateIssues = false;
+
+         if (options.autoSubmit) {
+            data.autoSubmit = true;
+         }
+
+      }
+
+      const submit = async () => {
+
+         if (!!import.meta.env.VITE_HORDE_DEBUG_NEW_JOB) {
+
+            console.log("Debug Job Submit");
+            console.log(JSON.stringify(data));
+            return;
+         }
+
+         let errorReason: any = undefined;
+
+         options.setSubmitting(true);
+
+         console.log("Submitting job");
+         console.log(JSON.stringify(data));
+
+         let redirected = false;
+
+         await backend.createJob(data).then(async (data) => {
+            console.log(`Job created: ${JSON.stringify(data)}`);
+            console.log("Updating notifications")
+            try {
+               await backend.updateNotification({ slack: true }, "job", data.id);
+               const user = await backend.getCurrentUser();
+               if (user.jobTemplateSettings) {
+                  dashboard.jobTemplateSettings = user.jobTemplateSettings;
+               }
+
+            } catch (reason) {
+               console.log(`Error on updating notifications: ${reason}`);
+            }
+
+            redirected = true;
+            options.onClose(data.id);
+         }).catch(reason => {
+            // "Not Found" is generally a permissions error
+            errorReason = reason ? reason : "Unknown";
+            console.log(`Error on job creation: ${errorReason}`);
+
+         }).finally(() => {
+
+            if (!redirected) {
+               options.setSubmitting(false);
+               options.onClose(undefined);
+
+               if (errorReason) {
+
+                  if (errorReason?.trim().endsWith("does not exist")) {
+                     errorReason += ".  Perforce edge server replication for the change may be in progress."
+                  }
+
+                  ErrorHandler.set({
+
+                     reason: `${errorReason}`,
+                     title: `Error Creating Job`,
+                     message: `There was an issue creating the job.\n\nReason: ${errorReason}\n\nTime: ${moment.utc().format("MMM Do, HH:mm z")}`
+
+                  }, true);
+
+               }
+            }
+         });
+
+      }
+
+      submit();
+
+   }
+
+   const advancedPivotRenderer = (link: any, defaultRenderer: any): JSX.Element => {
+      return (
+         <Stack horizontal>
+            {defaultRenderer(link)}
+            {options.advancedModified && <Icon iconName="Issue" style={{ color: '#EDC74A', paddingLeft: 8 }} />}
+         </Stack>
+      );
+   }
+
+   const pivotItems = ["Basic", "Advanced"].map(tab => {
+      return <PivotItem headerText={tab} itemKey={tab} key={tab} onRenderItemLink={tab === "Advanced" ? advancedPivotRenderer : undefined} />;
+   });
+
+   let name = template?.name ? template.name : options.jobDetails?.jobData?.name;
+   if (options.preview) {
+      name += " (Preview)";
+   }
+
+   return <Modal isModeless={false} key="new_build_key" isOpen={true} isBlocking={true} topOffsetFixed={true} styles={{ main: { padding: 8, width: 700 * (options.preview ? 2 : 1), hasBeenOpened: false, top: "80px", position: "absolute" } }} className={hordeClasses.modal} onDismiss={() => { }}>
+      <Stack horizontal tokens={{ childrenGap: 18 }}>
+         <Stack>
+            <Stack horizontal styles={{ root: { paddingLeft: 8, paddingTop: 8, paddingBottom: 8 } }}>
+               <Stack grow verticalAlign="center">
+                  <Text variant="mediumPlus" styles={{ root: { fontWeight: "unset", fontFamily: "Horde Open Sans SemiBold" } }}>{name}</Text>
+               </Stack>
+
+
+               <Stack >
+                  {<Pivot className={hordeClasses.pivot}
+                     style={{ paddingLeft: 32 }}
+                     selectedKey={options.mode}
+                     linkSize="normal"
+                     linkFormat="links"
+                     onLinkClick={(item) => {
+                        options.onModeChanged(item!.props!.itemKey as NewBuildMode)
+                     }}>
+                     {pivotItems}
+                  </Pivot>
+                  }
+               </Stack>
+            </Stack>
+            <Stack styles={{ root: { paddingLeft: 8, paddingRight: 8, paddingTop: 4 } }}>
+               <Stack tokens={{ childrenGap: 8 }}>
+                  <Stack horizontal tokens={{ childrenGap: 18 }} style={{ paddingBottom: 4 }}>
+                     <Stack verticalAlign="center" verticalFill={true}>
+                        <TemplateSelector />
+                     </Stack>
+                     <Stack>
+                        <ComboBox style={{ width: 204 }} label="Change" text={changeText} options={changeItems} disabled={options.readOnly} allowFreeform autoComplete="off" placeholder="Latest Change" defaultValue={"-1"} onChange={(ev, option, index, value) => {
+                           ev.preventDefault();
+                           if (option) {
+                              options.change = undefined;
+                              options.changeOption = option.text;
+                           } else {
+                              options.changeOption = undefined;
+                              if (!value) {
+                                 options.change = undefined;
+                              } else {
+                                 const nvalue = parseInt(value);
+                                 options.change = !isNaN(nvalue) ? nvalue : undefined;
+                              }
+                           }
+                           options.setChanged();
+
+                        }}
+                        />
+                     </Stack>
+
+                     <Stack>
+                        <TextField
+                           label="Shelved Change"
+                           style={{ width: 146 }}
+                           placeholder={!template!.allowPreflights ? "Disabled by template" : "None"}
+                           title={!template!.allowPreflights ? "Preflights are disabled for this template" : undefined}
+                           autoComplete="off" value={options.preflightChange?.toString() ?? ""} disabled={!template!.allowPreflights || options.readOnly || options.isPreflightSubmit} onChange={(ev, newValue) => {
+                              ev.preventDefault();
+
+                              let change: number | undefined;
+                              if (newValue) {
+                                 change = parseInt(newValue);
+                                 if (isNaN(change)) {
+                                    change = undefined;
+                                 }
+                              }
+                              options.onPreflightChange(change);
+                           }} />
+
+                     </Stack>
+                  </Stack>
+               </Stack>
+               <Stack>
+                  {(options.mode === "Basic") && <BuildParametersPanel />}
+                  {options.mode === "Advanced" && <AdvancedPanel />}
+               </Stack>
+               <Stack horizontal tokens={{ childrenGap: 16 }} styles={{ root: { paddingTop: 32, paddingLeft: 8, paddingBottom: 8 } }}>
+                  <Stack grow />
+                  <PrimaryButton text="Start Job" disabled={!template || options.readOnly || options.preview} onClick={() => { onSubmit(); }} />
+                  <DefaultButton text="Cancel" disabled={false} onClick={() => { options.onClose(); }} />
+               </Stack>
+            </Stack>
+         </Stack>
+         {options.preview && <PreviewPanel />}
+      </Stack>
+   </Modal>
+})
+
+
+const NewBuildV2Inner: React.FC = observer(() => {
+
+   const options = BuildOptions.get();
+   options.subscribe();
+
+   return <Stack>
+      {!!options.validationErrors.length && <ValidationErrorModal errors={options.validationErrors} onClose={() => options.setValidationErrors([])} />}
+      {!options.submitting && <BuildModal />}
+      {!!options.submitting && <SubmittingModal />}
+   </Stack>
+
+})
+
+export const NewBuildV2: React.FC<{ streamId: string; show: boolean; onClose: (newJobId: string | undefined) => void, jobKey?: string; jobDetails?: JobDetailsV2, readOnly?: boolean }> = observer(({ streamId, jobKey, show, onClose, jobDetails, readOnly }) => {
+
+   const { projectStore } = useBackend();
+
+   const options = useConst(new BuildOptions(streamId, projectStore, onClose, jobDetails, jobKey, readOnly));
+
+   useEffect(() => {
+      return () => {
+         options.clear();
+      };
+   }, [options]);
+
+
+   return <Stack>
+      <NewBuildV2Inner />
+   </Stack>
+})
+
+export const ValidationErrorModal: React.FC<{ errors: ValidationError[], onClose: () => void }> = ({ errors, onClose }) => {
+
+   const { hordeClasses } = getHordeStyling();
+
+   const close = () => {
+      onClose();
+   };
+
+   const messages = errors.map((e, idx) => {
+      return <MessageBar key={`validation_error_${idx}`} messageBarType={MessageBarType.error} isMultiline={true}>
+         <Stack>
+            <Stack grow tokens={{ childrenGap: 12 }}>
+               <Stack tokens={{ childrenGap: 12 }} horizontal>
+                  <Text nowrap style={{ fontFamily: "Horde Open Sans SemiBold", width: 80 }}>Parameter:</Text>
+                  <Text nowrap >{e.paramString}</Text>
+               </Stack>
+               <Stack tokens={{ childrenGap: 12 }} horizontal>
+                  <Text nowrap style={{ fontFamily: "Horde Open Sans SemiBold", width: 80 }}>Value:</Text>
+                  <Text nowrap >{e.value ? e.value : "no value"}</Text>
+               </Stack>
+               <Stack grow tokens={{ childrenGap: 12 }} horizontal>
+                  <Text nowrap style={{ fontFamily: "Horde Open Sans SemiBold", width: 80 }}>Error:</Text>
+                  <Text>{e.error}</Text>
+               </Stack>
+            </Stack>
+            {
+               e.context && <Stack style={{ paddingTop: 18, paddingBottom: 18 }}>
+                  <Stack grow tokens={{ childrenGap: 12 }}>
+                     <Text nowrap style={{ fontFamily: "Horde Open Sans SemiBold" }}>Context:</Text>
+                     <Text>{e.context}</Text>
+                  </Stack>
+               </Stack>
+            }
+
+         </Stack>
+      </MessageBar>
+   })
+
+   return <Modal isOpen={true} className={hordeClasses.modal} styles={{ main: { padding: 8, width: 700 } }} onDismiss={() => { close() }}>
+      <Stack horizontal styles={{ root: { padding: 8 } }}>
+         <Stack style={{ paddingLeft: 8, paddingTop: 4 }} grow>
+            <Text variant="mediumPlus">Validation Errors</Text>
+         </Stack>
+         <Stack grow horizontalAlign="end">
+            <IconButton
+               iconProps={{ iconName: 'Cancel' }}
+               ariaLabel="Close popup modal"
+               onClick={() => { close(); }}
+            />
+         </Stack>
+      </Stack>
+
+      <Stack style={{ paddingLeft: 20, paddingTop: 8, paddingBottom: 8 }}>
+         <Text style={{ fontSize: 15 }}>Please correct the following errors and try again.</Text>
+      </Stack>
+
+      <Stack tokens={{ childrenGap: 8 }} styles={{ root: { paddingLeft: 20, paddingTop: 18, paddingBottom: 24, width: 660 } }}>
+         {messages}
+      </Stack>
+
+
+      <Stack horizontal styles={{ root: { padding: 8, paddingTop: 8 } }}>
+         <Stack grow />
+         <Stack horizontal tokens={{ childrenGap: 16 }} styles={{ root: { paddingTop: 12, paddingLeft: 8, paddingBottom: 8 } }}>
+            <PrimaryButton text="Ok" disabled={false} onClick={() => { close(); }} />
+         </Stack>
+      </Stack>
+   </Modal>;
+
+};
\ No newline at end of file
diff --git a/Engine/Source/Programs/Horde/HordeDashboard/src/components/jobDetailsV2/JobOperationsBar.tsx b/Engine/Source/Programs/Horde/HordeDashboard/src/components/jobDetailsV2/JobOperationsBar.tsx
--- a/Engine/Source/Programs/Horde/HordeDashboard/src/components/jobDetailsV2/JobOperationsBar.tsx
+++ b/Engine/Source/Programs/Horde/HordeDashboard/src/components/jobDetailsV2/JobOperationsBar.tsx
@@ -18,6 +18,7 @@ import { JobDetailsV2 } from './JobDetailsViewCommon';
 import { RetryStepsModal, StepRetryModal, StepRetryType } from './StepRetryModal';
 import { getSiteConfig } from '../../backend/Config';
 import { getHordeStyling } from '../../styles/Styles';
+import { NewBuildV2 } from '../build/NewBuildV2';
 
 enum ParameterState {
    Hidden,
@@ -39,6 +40,14 @@ export const JobOperations: React.FC<{ jobDetails: JobDetailsV2 }> = observer(({
 
    const stepId = query.get("step") ? query.get("step")! : undefined;
    const batchFilter = query.get("batch");
+   let newBuildVersion: string = query.get("newbuildversion") ? query.get("newbuildversion") : "1";
+
+   // show new build version 2 if we have new parameters store
+   const hasJobParameters = !!(Object.keys(jobDetails?.jobData?.parameters ?? {}).length)
+   if (hasJobParameters) {
+      newBuildVersion = "2";
+   }
+
 
    // subscribe
    if (dashboard.updated) { }
@@ -52,7 +61,7 @@ export const JobOperations: React.FC<{ jobDetails: JobDetailsV2 }> = observer(({
    }
 
    const abortDisabled = jobData.state === JobState.Complete;
-   const runAgainDisabled = false; /*jobDetails.jobdata?.state !== JobState.Complete*/
+   const runAgainDisabled = newBuildVersion === "2" && !hasJobParameters;
 
    const failedSteps = jobDetails.getSteps().filter(s => {
 
@@ -151,12 +160,15 @@ export const JobOperations: React.FC<{ jobDetails: JobDetailsV2 }> = observer(({
       return null;
    }
 
+   let showNewBuildV1 = newBuildVersion !== "2" && (parametersState === ParameterState.Parameters || parametersState === ParameterState.Clone);
+   let showNewBuildV2 = newBuildVersion === "2" && (parametersState === ParameterState.Parameters || parametersState === ParameterState.Clone);
+
    return <Stack>
       {retryStepsShown && <RetryStepsModal stepIds={failedSteps.map(s => s.id)} jobDetails={jobDetails} onClose={() => { setRetryStepsShown(false); }} />}
       <AbortJobModal jobDetails={jobDetails} show={abortShown} onClose={() => { setAbortShown(false); }} />
       <EditJobModal jobData={jobDetails.jobData} show={editShown} onClose={() => { setEditShown(false); }} />
-      <NewBuild streamId={jobDetails.stream!.id} jobDetails={jobDetails} readOnly={parametersState === ParameterState.Parameters}
-         show={parametersState === ParameterState.Parameters || parametersState === ParameterState.Clone}
+      {showNewBuildV1 && <NewBuild streamId={jobDetails.stream!.id} jobDetails={jobDetails} readOnly={parametersState === ParameterState.Parameters}
+         show={true}
          onClose={(newJobId) => {
             setParametersState(ParameterState.Hidden);
             if (newJobId) {
@@ -166,7 +178,20 @@ export const JobOperations: React.FC<{ jobDetails: JobDetailsV2 }> = observer(({
                   navigate(`/job/${jobId}`, { replace: true });
                }
             }
-         }} />
+         }} />}
+
+      {showNewBuildV2 && <NewBuildV2 streamId={jobDetails.stream!.id} jobDetails={jobDetails} readOnly={parametersState === ParameterState.Parameters}
+         show={true}
+         onClose={(newJobId) => {
+            setParametersState(ParameterState.Hidden);
+            if (newJobId) {
+               navigate(`/job/${newJobId}`);
+            } else {
+               if (query.get("newbuild")) {
+                  navigate(`/job/${jobId}`, { replace: true });
+               }
+            }
+         }} />}
 
       <Stack horizontal>
          <Stack grow />
@@ -209,7 +234,7 @@ export const JobOperations: React.FC<{ jobDetails: JobDetailsV2 }> = observer(({
 });
 
 const StepArtifactsOperations: React.FC<{ jobDetails: JobDetailsV2, stepId: string }> = observer(({ jobDetails, stepId }) => {
-   
+
    const { hordeClasses } = getHordeStyling();
 
    const navigate = useNavigate();
@@ -247,7 +272,7 @@ const StepArtifactsOperations: React.FC<{ jobDetails: JobDetailsV2, stepId: stri
    const opsList: IContextualMenuItem[] = [];
 
    const baseUrl = window.location.pathname + window.location.search;
-   
+
 
    opsList.push({
       key: 'stepops_artifacts_step',
@@ -280,7 +305,7 @@ const StepArtifactsOperations: React.FC<{ jobDetails: JobDetailsV2, stepId: stri
          text: c.description ?? c.name,
          iconProps: { iconName: "Clean" },
          onClick: () => { navigate(`${baseUrl}&artifactContext=${c.type}`, { replace: true }) }
-      });   
+      });
    })
 
    const opsItems: ICommandBarItemProps[] = [
diff --git a/Engine/Source/Programs/Horde/HordeDashboard/src/styles/Styles.tsx b/Engine/Source/Programs/Horde/HordeDashboard/src/styles/Styles.tsx
--- a/Engine/Source/Programs/Horde/HordeDashboard/src/styles/Styles.tsx
+++ b/Engine/Source/Programs/Horde/HordeDashboard/src/styles/Styles.tsx
@@ -423,17 +423,14 @@ export const getHordeStyling = () => {
                paddingTop: 0,
             },
             ".ms-TextField-fieldGroup,.ms-TextField-field,.ms-TextField-field::placeholder": {
-               height: "29px",
                fontSize: "12px",
                fontFamily: "Horde Open Sans Regular",
             },
             ".ms-Dropdown-title,.ms-Dropdown-titleIsPlaceHolder": {
                fontSize: "12px",
-               height: "29px"
             },
             ".ms-ComboBox": {
                fontSize: "12px",
-               height: "29px",
             },
             ".ms-Checkbox": {
                marginTop: 0,
