Horde: Add wrapper classes for IJobStepBatch and IJobStep.

#rnx

diff --git a/Engine/Source/Programs/Horde/Horde.Server/Jobs/IJob.cs b/Engine/Source/Programs/Horde/Horde.Server/Jobs/IJob.cs
--- a/Engine/Source/Programs/Horde/Horde.Server/Jobs/IJob.cs
+++ b/Engine/Source/Programs/Horde/Horde.Server/Jobs/IJob.cs
@@ -1397,9 +1397,21 @@ namespace Horde.Server.Jobs
 		/// <param name="groups">List of node groups</param>
 		/// <returns>Set of nodes that must have completed for this batch to start</returns>
 		public static HashSet<INode> GetStartDependencies(this IJobStepBatch batch, IReadOnlyList<INodeGroup> groups)
+		{
+			List<INode> batchNodes = batch.Steps.ConvertAll(x => groups[batch.GroupIdx].Nodes[x.NodeIdx]);
+			return GetStartDependencies(batchNodes, groups);
+		}
+
+		/// <summary>
+		/// Get the dependencies required for this batch to start, taking run-early nodes into account
+		/// </summary>
+		/// <param name="batchNodes">Nodes in the batch to search</param>
+		/// <param name="groups">List of node groups</param>
+		/// <returns>Set of nodes that must have completed for this batch to start</returns>
+		public static HashSet<INode> GetStartDependencies(IEnumerable<INode> batchNodes, IReadOnlyList<INodeGroup> groups)
 		{
 			// Find all the nodes that this group will start with.
-			List<INode> nodes = batch.Steps.ConvertAll(x => groups[batch.GroupIdx].Nodes[x.NodeIdx]);
+			List<INode> nodes = new List<INode>(batchNodes);
 			if (nodes.Any(x => x.RunEarly))
 			{
 				nodes.RemoveAll(x => !x.RunEarly);
@@ -1414,7 +1426,7 @@ namespace Horde.Server.Jobs
 			}
 
 			// Exclude all the dependencies within the same group
-			dependencies.ExceptWith(groups[batch.GroupIdx].Nodes);
+			dependencies.ExceptWith(batchNodes);
 			return dependencies;
 		}
 	}
@@ -1502,7 +1514,7 @@ namespace Horde.Server.Jobs
 		/// <summary>
 		/// Reports for this jobstep.
 		/// </summary>
-		public Dictionary<string, string>? Properties { get; }
+		public IReadOnlyDictionary<string, string>? Properties { get; }
 	}
 
 	/// <summary>
@@ -1511,12 +1523,12 @@ namespace Horde.Server.Jobs
 	public static class JobStepExtensions
 	{
 		/// <summary>
-		/// Determines if a jobstep has failed or is skipped. Can be used to determine whether dependent steps will be able to run.
+		/// Determines if a jobstep state is completed, skipped, or aborted.
 		/// </summary>
-		/// <returns>True if the step is failed or skipped</returns>
-		public static bool IsFailedOrSkipped(this IJobStep step)
+		/// <returns>True if the step is completed, skipped, or aborted</returns>
+		public static bool IsPendingState(JobStepState state)
 		{
-			return step.State == JobStepState.Skipped || step.Outcome == JobStepOutcome.Failure;
+			return state != JobStepState.Aborted && state != JobStepState.Completed && state != JobStepState.Skipped;
 		}
 
 		/// <summary>
@@ -1524,9 +1536,7 @@ namespace Horde.Server.Jobs
 		/// </summary>
 		/// <returns>True if the step is completed, skipped, or aborted</returns>
 		public static bool IsPending(this IJobStep step)
-		{
-			return step.State != JobStepState.Aborted && step.State != JobStepState.Completed && step.State != JobStepState.Skipped;
-		}
+			=> IsPendingState(step.State);
 
 		/// <summary>
 		/// Determine if a step should be timed out
diff --git a/Engine/Source/Programs/Horde/Horde.Server/Jobs/JobCollection.cs b/Engine/Source/Programs/Horde/Horde.Server/Jobs/JobCollection.cs
--- a/Engine/Source/Programs/Horde/Horde.Server/Jobs/JobCollection.cs
+++ b/Engine/Source/Programs/Horde/Horde.Server/Jobs/JobCollection.cs
@@ -43,190 +43,183 @@
 	/// </summary>
 	public class JobCollection : IJobCollection
 	{
-		/// <summary>
-		/// Embedded jobstep document
-		/// </summary>
-		[BsonIgnoreExtraElements]
-		class JobStepDocument : IJobStep
+		class Job : IJob
 		{
-			[BsonRequired]
-			public JobStepId Id { get; set; }
+			readonly JobCollection _collection;
+			readonly JobDocument _document;
+			readonly IGraph _graph;
+			readonly List<JobStepBatch> _batches;
 
-			[BsonRequired]
-			public int NodeIdx { get; set; }
+			public JobDocument Document => _document;
+			public IGraph Graph => _graph;
+			public List<JobStepBatch> Batches => _batches;
 
-			[BsonRequired]
-			public JobStepState State { get; set; } = JobStepState.Waiting;
+			JobId IJob.Id => _document.Id;
+			StreamId IJob.StreamId => _document.StreamId;
+			TemplateId IJob.TemplateId => _document.TemplateId;
+			ContentHash? IJob.TemplateHash => _document.TemplateHash;
+			ContentHash IJob.GraphHash => _document.GraphHash;
+			UserId? IJob.StartedByUserId => _document.StartedByUserId;
+			UserId? IJob.AbortedByUserId => _document.AbortedByUserId;
+			BisectTaskId? IJob.StartedByBisectTaskId => _document.StartedByBisectTaskId;
+			string IJob.Name => _document.Name;
+			int IJob.Change => _document.Change;
+			int IJob.CodeChange => _document.CodeChange;
+			int IJob.PreflightChange => _document.PreflightChange;
+			int IJob.ClonedPreflightChange => _document.ClonedPreflightChange;
+			string? IJob.PreflightDescription => _document.PreflightDescription;
+			Priority IJob.Priority => _document.Priority;
+			bool IJob.AutoSubmit => _document.AutoSubmit;
+			int? IJob.AutoSubmitChange => _document.AutoSubmitChange;
+			string? IJob.AutoSubmitMessage => _document.AutoSubmitMessage;
+			bool IJob.UpdateIssues => _document.UpdateIssues;
+			bool IJob.PromoteIssuesByDefault => _document.PromoteIssuesByDefault;
+			DateTime IJob.CreateTimeUtc => _document.GetCreateTimeOrDefault();
+			JobOptions? IJob.JobOptions => _document.JobOptions;
+			IReadOnlyList<AclClaimConfig> IJob.Claims => _document.Claims;
+			IReadOnlyList<IJobStepBatch> IJob.Batches => Batches;
+			IReadOnlyDictionary<ParameterId, string> IJob.Parameters => _document.Parameters;
+			IReadOnlyList<string> IJob.Arguments => _document.Arguments;
+			IReadOnlyDictionary<string, string> IJob.Environment => _document.Environment;
+			IReadOnlyList<int> IJob.Issues => _document.ReferencedByIssues;
+			ObjectId? IJob.NotificationTriggerId => _document.NotificationTriggerId;
+			bool IJob.ShowUgsBadges => _document.ShowUgsBadges;
+			bool IJob.ShowUgsAlerts => _document.ShowUgsAlerts;
+			string? IJob.NotificationChannel => _document.NotificationChannel;
+			string? IJob.NotificationChannelFilter => _document.NotificationChannelFilter;
+			IReadOnlyDictionary<int, ObjectId> IJob.LabelIdxToTriggerId => _document._labelNotifications.ToDictionary(x => x._labelIdx, x => x._triggerId);
+			IReadOnlyList<IReport>? IJob.Reports => _document.Reports;
+			IReadOnlyList<IChainedJob> IJob.ChainedJobs => _document.ChainedJobs;
+			DateTime IJob.UpdateTimeUtc => _document.UpdateTimeUtc ?? _document.UpdateTime?.UtcDateTime ?? DateTime.UnixEpoch;
+			int IJob.UpdateIndex => _document.UpdateIndex;
 
-			public JobStepOutcome Outcome { get; set; } = JobStepOutcome.Success;
-
-			public JobStepError Error { get; set; } = JobStepError.None;
-
-			[BsonIgnoreIfNull]
-			public LogId? LogId { get; set; }
-
-			[BsonIgnoreIfNull]
-			public ObjectId? NotificationTriggerId { get; set; }
-
-			[BsonIgnoreIfNull]
-			public DateTimeOffset? ReadyTime { get; set; }
-
-			[BsonIgnoreIfNull]
-			public DateTimeOffset? StartTime { get; set; }
-
-			[BsonIgnoreIfNull]
-			public DateTimeOffset? FinishTime { get; set; }
-
-			[BsonIgnoreIfNull]
-			public Priority? Priority { get; set; }
-
-			[BsonIgnoreIfDefault, BsonDefaultValue(false)]
-			public bool Retry { get; set; }
-
-			public UserId? RetriedByUserId { get; set; }
-
-			[BsonElement("RetryByUser")]
-			public string? RetriedByUserDeprecated { get; set; }
-
-			public bool AbortRequested { get; set; } = false;
-
-			public UserId? AbortedByUserId { get; set; }
-
-			[BsonElement("AbortByUser")]
-			public string? AbortedByUserDeprecated { get; set; }
-
-			[BsonIgnoreIfNull]
-			public List<Report>? Reports { get; set; }
-			IReadOnlyList<IReport>? IJobStep.Reports => Reports;
-
-			[BsonIgnoreIfNull]
-			public Dictionary<string, string>? Properties { get; set; }
-
-			DateTime? IJobStep.ReadyTimeUtc => ReadyTime?.UtcDateTime;
-			DateTime? IJobStep.StartTimeUtc => StartTime?.UtcDateTime;
-			DateTime? IJobStep.FinishTimeUtc => FinishTime?.UtcDateTime;
-
-			[BsonConstructor]
-			private JobStepDocument()
+			public Job(JobCollection collection, JobDocument document, IGraph graph)
 			{
+				_collection = collection;
+				_document = document;
+				_graph = graph;
+				_batches = document.Batches.ConvertAll(x => new JobStepBatch(x));
 			}
 
-			public JobStepDocument(JobStepId id, int nodeIdx)
+			public Task<IJob?> RefreshAsync(CancellationToken cancellationToken = default)
+				=> _collection.GetAsync(_document.Id, cancellationToken);
+
+			public async Task<IJob?> TryAssignLeaseAsync(int batchIdx, PoolId poolId, AgentId agentId, SessionId sessionId, LeaseId leaseId, LogId logId, CancellationToken cancellationToken = default)
 			{
-				Id = id;
-				NodeIdx = nodeIdx;
+				JobDocument? newDocument = await _collection.TryAssignLeaseAsync(Document, batchIdx, poolId, agentId, sessionId, leaseId, logId, cancellationToken);
+				return _collection.CreateJobObject(newDocument, _graph);
 			}
 
-			public override string ToString()
+			public async Task<IJob?> TryCancelLeaseAsync(int batchIdx, CancellationToken cancellationToken = default)
 			{
-				StringBuilder description = new StringBuilder($"{Id}: {State}");
-				if (Outcome != JobStepOutcome.Unspecified)
-				{
-					description.Append($" ({Outcome}");
-					if (Error != JobStepError.None)
-					{
-						description.Append($" - {Error}");
-					}
-					description.Append(')');
-				}
-				return description.ToString();
+				JobDocument? newDocument = await _collection.TryCancelLeaseAsync(Document, batchIdx, cancellationToken);
+				return _collection.CreateJobObject(newDocument, _graph);
+			}
+
+			public Task<bool> TryDeleteAsync(CancellationToken cancellationToken = default)
+				=> _collection.TryDeleteAsync(Document, cancellationToken);
+
+			public async Task<IJob?> TryFailBatchAsync(int batchIdx, JobStepBatchError reason, CancellationToken cancellationToken = default)
+			{
+				JobDocument? newDocument = await _collection.TryFailBatchAsync(Document, batchIdx, Graph, reason, cancellationToken);
+				return _collection.CreateJobObject(newDocument, _graph);
+			}
+
+			public async Task<IJob?> TryRemoveFromDispatchQueueAsync(CancellationToken cancellationToken = default)
+			{
+				JobDocument? newDocument = await _collection.TryRemoveFromDispatchQueueAsync(Document, cancellationToken);
+				return _collection.CreateJobObject(newDocument, _graph);
+			}
+
+			public async Task<IJob?> TrySkipAllBatchesAsync(JobStepBatchError reason, CancellationToken cancellationToken = default)
+			{
+				JobDocument? newDocument = await _collection.TrySkipAllBatchesAsync(Document, Graph, reason, cancellationToken);
+				return _collection.CreateJobObject(newDocument, _graph);
+			}
+
+			public async Task<IJob?> TrySkipBatchAsync(JobStepBatchId batchId, JobStepBatchError reason, CancellationToken cancellationToken = default)
+			{
+				JobDocument? newDocument = await _collection.TrySkipBatchAsync(Document, batchId, Graph, reason, cancellationToken);
+				return _collection.CreateJobObject(newDocument, _graph);
+			}
+
+			public async Task<IJob?> TryUpdateBatchAsync(JobStepBatchId batchId, LogId? newLogId, JobStepBatchState? newState, JobStepBatchError? newError, CancellationToken cancellationToken = default)
+			{
+				JobDocument? newDocument = await _collection.TryUpdateBatchAsync(Document, Graph, batchId, newLogId, newState, newError, cancellationToken);
+				return _collection.CreateJobObject(newDocument, _graph);
+			}
+
+			public async Task<IJob?> TryUpdateGraphAsync(IGraph newGraph, CancellationToken cancellationToken = default)
+			{
+				JobDocument? newDocument = await _collection.TryUpdateGraphAsync(Document, Graph, newGraph, cancellationToken);
+				return _collection.CreateJobObject(newDocument, newGraph);
+			}
+
+			public async Task<IJob?> TryUpdateJobAsync(string? name = null, Priority? priority = null, bool? autoSubmit = null, int? autoSubmitChange = null, string? autoSubmitMessage = null, UserId? abortedByUserId = null, ObjectId? notificationTriggerId = null, List<Report>? reports = null, List<string>? arguments = null, KeyValuePair<int, ObjectId>? labelIdxToTriggerId = null, KeyValuePair<TemplateId, JobId>? jobTrigger = null, CancellationToken cancellationToken = default)
+			{
+				JobDocument? newDocument = await _collection.TryUpdateJobAsync(Document, Graph, name, priority, autoSubmit, autoSubmitChange, autoSubmitMessage, abortedByUserId, notificationTriggerId, reports, arguments, labelIdxToTriggerId, jobTrigger, cancellationToken);
+				return _collection.CreateJobObject(newDocument, Graph);
+			}
+
+			public async Task<IJob?> TryUpdateStepAsync(JobStepBatchId batchId, JobStepId stepId, JobStepState newState = JobStepState.Unspecified, JobStepOutcome newOutcome = JobStepOutcome.Unspecified, JobStepError? newError = null, bool? newAbortRequested = null, UserId? newAbortByUserId = null, LogId? newLogId = null, ObjectId? newNotificationTriggerId = null, UserId? newRetryByUserId = null, Priority? newPriority = null, List<Report>? newReports = null, Dictionary<string, string?>? newProperties = null, CancellationToken cancellationToken = default)
+			{
+				JobDocument? newDocument = await _collection.TryUpdateStepAsync(_document, _graph, batchId, stepId, newState, newOutcome, newError, newAbortRequested, newAbortByUserId, newLogId, newNotificationTriggerId, newRetryByUserId, newPriority, newReports, newProperties, cancellationToken);
+				return _collection.CreateJobObject(newDocument, Graph);
 			}
 		}
 
-		class JobStepBatchDocument : IJobStepBatch
+		class JobStepBatch : IJobStepBatch
 		{
-			[BsonRequired]
-			public JobStepBatchId Id { get; set; }
+			public JobStepBatchDocument Document { get; }
+			public List<JobStep> Steps { get; }
 
-			public LogId? LogId { get; set; }
-
-			[BsonRequired]
-			public int GroupIdx { get; set; }
-
-			[BsonRequired]
-			public JobStepBatchState State { get; set; }
-
-			[BsonIgnoreIfDefault, BsonDefaultValue(JobStepBatchError.None)]
-			public JobStepBatchError Error { get; set; }
-
-			public List<JobStepDocument> Steps { get; set; } = new List<JobStepDocument>();
-
-			[BsonIgnoreIfNull]
-			public PoolId? PoolId { get; set; }
-
-			[BsonIgnoreIfNull]
-			public AgentId? AgentId { get; set; }
-
-			[BsonIgnoreIfNull]
-			public SessionId? SessionId { get; set; }
-
-			[BsonIgnoreIfNull]
-			public LeaseId? LeaseId { get; set; }
-
-			public int SchedulePriority { get; set; }
-
-			[BsonIgnoreIfNull]
-			public DateTimeOffset? ReadyTime { get; set; }
-
-			[BsonIgnoreIfNull]
-			public DateTimeOffset? StartTime { get; set; }
-
-			[BsonIgnoreIfNull]
-			public DateTimeOffset? FinishTime { get; set; }
-
+			JobStepBatchId IJobStepBatch.Id => Document.Id;
+			LogId? IJobStepBatch.LogId => Document.LogId;
+			int IJobStepBatch.GroupIdx => Document.GroupIdx;
+			JobStepBatchState IJobStepBatch.State => Document.State;
+			JobStepBatchError IJobStepBatch.Error => Document.Error;
 			IReadOnlyList<IJobStep> IJobStepBatch.Steps => Steps;
-			DateTime? IJobStepBatch.ReadyTimeUtc => ReadyTime?.UtcDateTime;
-			DateTime? IJobStepBatch.StartTimeUtc => StartTime?.UtcDateTime;
-			DateTime? IJobStepBatch.FinishTimeUtc => FinishTime?.UtcDateTime;
+			PoolId? IJobStepBatch.PoolId => Document.PoolId;
+			AgentId? IJobStepBatch.AgentId => Document.AgentId;
+			SessionId? IJobStepBatch.SessionId => Document.SessionId;
+			LeaseId? IJobStepBatch.LeaseId => Document.LeaseId;
+			int IJobStepBatch.SchedulePriority => Document.SchedulePriority;
+			DateTime? IJobStepBatch.ReadyTimeUtc => Document.ReadyTime?.UtcDateTime;
+			DateTime? IJobStepBatch.StartTimeUtc => Document.StartTime?.UtcDateTime;
+			DateTime? IJobStepBatch.FinishTimeUtc => Document.FinishTime?.UtcDateTime;
 
-			[BsonConstructor]
-			private JobStepBatchDocument()
+			public JobStepBatch(JobStepBatchDocument document)
 			{
-			}
-
-			public JobStepBatchDocument(JobStepBatchId id, int groupIdx)
-			{
-				Id = id;
-				GroupIdx = groupIdx;
-			}
-
-			public override string ToString()
-			{
-				StringBuilder description = new StringBuilder($"{Id}: {State}");
-				if (Error != JobStepBatchError.None)
-				{
-					description.Append($" - {Error}");
-				}
-				description.Append($" ({Steps.Count} steps");
-				return description.ToString();
+				Document = document;
+				Steps = document.Steps.ConvertAll(x => new JobStep(x));
 			}
 		}
 
-		class ChainedJobDocument : IChainedJob
+		class JobStep : IJobStep
 		{
-			public string Target { get; set; }
-			public TemplateId TemplateRefId { get; set; }
-			public JobId? JobId { get; set; }
-			public bool UseDefaultChangeForTemplate { get; set; }
+			readonly JobStepDocument _document;
 
-			[BsonConstructor]
-			private ChainedJobDocument()
+			JobStepId IJobStep.Id => _document.Id;
+			int IJobStep.NodeIdx => _document.NodeIdx;
+			JobStepState IJobStep.State => _document.State;
+			JobStepOutcome IJobStep.Outcome => _document.Outcome;
+			JobStepError IJobStep.Error => _document.Error;
+			LogId? IJobStep.LogId => _document.LogId;
+			ObjectId? IJobStep.NotificationTriggerId => _document.NotificationTriggerId;
+			DateTime? IJobStep.ReadyTimeUtc => _document.ReadyTime?.UtcDateTime;
+			DateTime? IJobStep.StartTimeUtc => _document.StartTime?.UtcDateTime;
+			DateTime? IJobStep.FinishTimeUtc => _document.FinishTime?.UtcDateTime;
+			Priority? IJobStep.Priority => _document.Priority;
+			UserId? IJobStep.RetriedByUserId => _document.RetriedByUserId;
+			bool IJobStep.AbortRequested => _document.AbortRequested;
+			UserId? IJobStep.AbortedByUserId => _document.AbortedByUserId;
+			IReadOnlyList<IReport>? IJobStep.Reports => _document.Reports;
+			IReadOnlyDictionary<string, string>? IJobStep.Properties => _document.Properties;
+
+			public JobStep(JobStepDocument document)
 			{
-				Target = String.Empty;
-			}
-
-			public ChainedJobDocument(ChainedJobTemplateConfig trigger)
-			{
-				Target = trigger.Trigger;
-				TemplateRefId = trigger.TemplateId;
-				UseDefaultChangeForTemplate = trigger.UseDefaultChangeForTemplate;
+				_document = document;
 			}
 		}
 
-		class LabelNotificationDocument
-		{
-			public int _labelIdx;
-			public ObjectId _triggerId;
-		}
-
 		class JobDocument
@@ -372,130 +365,179 @@
 				=> CreateTimeUtc ?? CreateTime?.UtcDateTime ?? DateTime.UnixEpoch;
 		}
 
-		class Job : IJob
+		class JobStepBatchDocument
 		{
-			readonly JobCollection _collection;
-			readonly JobDocument _document;
-			readonly IGraph _graph;
+			[BsonRequired]
+			public JobStepBatchId Id { get; set; }
 
-			public Job(JobCollection collection, JobDocument document, IGraph graph)
+			public LogId? LogId { get; set; }
+
+			[BsonRequired]
+			public int GroupIdx { get; set; }
+
+			[BsonRequired]
+			public JobStepBatchState State { get; set; }
+
+			[BsonIgnoreIfDefault, BsonDefaultValue(JobStepBatchError.None)]
+			public JobStepBatchError Error { get; set; }
+
+			public List<JobStepDocument> Steps { get; set; } = new List<JobStepDocument>();
+
+			[BsonIgnoreIfNull]
+			public PoolId? PoolId { get; set; }
+
+			[BsonIgnoreIfNull]
+			public AgentId? AgentId { get; set; }
+
+			[BsonIgnoreIfNull]
+			public SessionId? SessionId { get; set; }
+
+			[BsonIgnoreIfNull]
+			public LeaseId? LeaseId { get; set; }
+
+			public int SchedulePriority { get; set; }
+
+			[BsonIgnoreIfNull]
+			public DateTimeOffset? ReadyTime { get; set; }
+
+			[BsonIgnoreIfNull]
+			public DateTimeOffset? StartTime { get; set; }
+
+			[BsonIgnoreIfNull]
+			public DateTimeOffset? FinishTime { get; set; }
+
+			[BsonConstructor]
+			private JobStepBatchDocument()
 			{
-				_collection = collection;
-				_document = document;
-				_graph = graph;
 			}
 
-			public JobDocument Document => _document;
-
-			public JobId Id => _document.Id;
-			public StreamId StreamId => _document.StreamId;
-			public TemplateId TemplateId => _document.TemplateId;
-			public ContentHash? TemplateHash => _document.TemplateHash;
-			public ContentHash GraphHash => _document.GraphHash;
-			public IGraph Graph => _graph;
-			public UserId? StartedByUserId => _document.StartedByUserId;
-			public UserId? AbortedByUserId => _document.AbortedByUserId;
-			public BisectTaskId? StartedByBisectTaskId => _document.StartedByBisectTaskId;
-			public string Name => _document.Name;
-			public int Change => _document.Change;
-			public int CodeChange => _document.CodeChange;
-			public int PreflightChange => _document.PreflightChange;
-			public int ClonedPreflightChange => _document.ClonedPreflightChange;
-			public string? PreflightDescription => _document.PreflightDescription;
-			public Priority Priority => _document.Priority;
-			public bool AutoSubmit => _document.AutoSubmit;
-			public int? AutoSubmitChange => _document.AutoSubmitChange;
-			public string? AutoSubmitMessage => _document.AutoSubmitMessage;
-			public bool UpdateIssues => _document.UpdateIssues;
-			public bool PromoteIssuesByDefault => _document.PromoteIssuesByDefault;
-			public DateTime CreateTimeUtc => _document.GetCreateTimeOrDefault();
-			public JobOptions? JobOptions => _document.JobOptions;
-			public IReadOnlyList<AclClaimConfig> Claims => _document.Claims;
-			public int SchedulePriority => _document.SchedulePriority;
-			public IReadOnlyList<IJobStepBatch> Batches => _document.Batches;
-			public IReadOnlyDictionary<ParameterId, string> Parameters => _document.Parameters;
-			public IReadOnlyList<string> Arguments => _document.Arguments;
-			public IReadOnlyDictionary<string, string> Environment => _document.Environment;
-			public IReadOnlyList<int> Issues => _document.ReferencedByIssues;
-			public ObjectId? NotificationTriggerId => _document.NotificationTriggerId;
-			public bool ShowUgsBadges => _document.ShowUgsBadges;
-			public bool ShowUgsAlerts => _document.ShowUgsAlerts;
-			public string? NotificationChannel => _document.NotificationChannel;
-			public string? NotificationChannelFilter => _document.NotificationChannelFilter;
-			public IReadOnlyDictionary<int, ObjectId> LabelIdxToTriggerId => _document._labelNotifications.ToDictionary(x => x._labelIdx, x => x._triggerId);
-			public IReadOnlyList<IReport>? Reports => _document.Reports;
-			public IReadOnlyList<IChainedJob> ChainedJobs => _document.ChainedJobs;
-			public SubResourceId NextSubResourceId => _document.NextSubResourceId;
-			public DateTime UpdateTimeUtc => _document.UpdateTimeUtc ?? _document.UpdateTime?.UtcDateTime ?? DateTime.UnixEpoch;
-			public int UpdateIndex => _document.UpdateIndex;
-
-			public Task<IJob?> RefreshAsync(CancellationToken cancellationToken = default)
-				=> _collection.GetAsync(Id, cancellationToken);
-
-			public async Task<IJob?> TryAssignLeaseAsync(int batchIdx, PoolId poolId, AgentId agentId, SessionId sessionId, LeaseId leaseId, LogId logId, CancellationToken cancellationToken = default)
+			public JobStepBatchDocument(JobStepBatchId id, int groupIdx)
 			{
-				JobDocument? newDocument = await _collection.TryAssignLeaseAsync(Document, batchIdx, poolId, agentId, sessionId, leaseId, logId, cancellationToken);
-				return _collection.CreateJobObject(newDocument, _graph);
+				Id = id;
+				GroupIdx = groupIdx;
 			}
 
-			public async Task<IJob?> TryCancelLeaseAsync(int batchIdx, CancellationToken cancellationToken = default)
+			public override string ToString()
 			{
-				JobDocument? newDocument = await _collection.TryCancelLeaseAsync(Document, batchIdx, cancellationToken);
-				return _collection.CreateJobObject(newDocument, _graph);
-			}
-
-			public Task<bool> TryDeleteAsync(CancellationToken cancellationToken = default)
-				=> _collection.TryDeleteAsync(Document, cancellationToken);
-
-			public async Task<IJob?> TryFailBatchAsync(int batchIdx, JobStepBatchError reason, CancellationToken cancellationToken = default)
-			{
-				JobDocument? newDocument = await _collection.TryFailBatchAsync(Document, batchIdx, Graph, reason, cancellationToken);
-				return _collection.CreateJobObject(newDocument, _graph);
-			}
-
-			public async Task<IJob?> TryRemoveFromDispatchQueueAsync(CancellationToken cancellationToken = default)
-			{
-				JobDocument? newDocument = await _collection.TryRemoveFromDispatchQueueAsync(Document, cancellationToken);
-				return _collection.CreateJobObject(newDocument, _graph);
-			}
-
-			public async Task<IJob?> TrySkipAllBatchesAsync(JobStepBatchError reason, CancellationToken cancellationToken = default)
-			{
-				JobDocument? newDocument = await _collection.TrySkipAllBatchesAsync(Document, Graph, reason, cancellationToken);
-				return _collection.CreateJobObject(newDocument, _graph);
-			}
-
-			public async Task<IJob?> TrySkipBatchAsync(JobStepBatchId batchId, JobStepBatchError reason, CancellationToken cancellationToken = default)
-			{
-				JobDocument? newDocument = await _collection.TrySkipBatchAsync(Document, batchId, Graph, reason, cancellationToken);
-				return _collection.CreateJobObject(newDocument, _graph);
-			}
-
-			public async Task<IJob?> TryUpdateBatchAsync(JobStepBatchId batchId, LogId? newLogId, JobStepBatchState? newState, JobStepBatchError? newError, CancellationToken cancellationToken = default)
-			{
-				JobDocument? newDocument = await _collection.TryUpdateBatchAsync(Document, Graph, batchId, newLogId, newState, newError, cancellationToken);
-				return _collection.CreateJobObject(newDocument, _graph);
-			}
-
-			public async Task<IJob?> TryUpdateGraphAsync(IGraph newGraph, CancellationToken cancellationToken = default)
-			{
-				JobDocument? newDocument = await _collection.TryUpdateGraphAsync(Document, Graph, newGraph, cancellationToken);
-				return _collection.CreateJobObject(newDocument, newGraph);
-			}
-
-			public async Task<IJob?> TryUpdateJobAsync(string? name = null, Priority? priority = null, bool? autoSubmit = null, int? autoSubmitChange = null, string? autoSubmitMessage = null, UserId? abortedByUserId = null, ObjectId? notificationTriggerId = null, List<Report>? reports = null, List<string>? arguments = null, KeyValuePair<int, ObjectId>? labelIdxToTriggerId = null, KeyValuePair<TemplateId, JobId>? jobTrigger = null, CancellationToken cancellationToken = default)
-			{
-				JobDocument? newDocument = await _collection.TryUpdateJobAsync(Document, Graph, name, priority, autoSubmit, autoSubmitChange, autoSubmitMessage, abortedByUserId, notificationTriggerId, reports, arguments, labelIdxToTriggerId, jobTrigger, cancellationToken);
-				return _collection.CreateJobObject(newDocument, Graph);
-			}
-
-			public async Task<IJob?> TryUpdateStepAsync(JobStepBatchId batchId, JobStepId stepId, JobStepState newState = JobStepState.Unspecified, JobStepOutcome newOutcome = JobStepOutcome.Unspecified, JobStepError? newError = null, bool? newAbortRequested = null, UserId? newAbortByUserId = null, LogId? newLogId = null, ObjectId? newNotificationTriggerId = null, UserId? newRetryByUserId = null, Priority? newPriority = null, List<Report>? newReports = null, Dictionary<string, string?>? newProperties = null, CancellationToken cancellationToken = default)
-			{
-				JobDocument? newDocument = await _collection.TryUpdateStepAsync(_document, _graph, batchId, stepId, newState, newOutcome, newError, newAbortRequested, newAbortByUserId, newLogId, newNotificationTriggerId, newRetryByUserId, newPriority, newReports, newProperties, cancellationToken);
-				return _collection.CreateJobObject(newDocument, Graph);
+				StringBuilder description = new StringBuilder($"{Id}: {State}");
+				if (Error != JobStepBatchError.None)
+				{
+					description.Append($" - {Error}");
+				}
+				description.Append($" ({Steps.Count} steps");
+				return description.ToString();
 			}
 		}
 
+		[BsonIgnoreExtraElements]
+		class JobStepDocument
+		{
+			[BsonRequired]
+			public JobStepId Id { get; set; }
+
+			[BsonRequired]
+			public int NodeIdx { get; set; }
+
+			[BsonRequired]
+			public JobStepState State { get; set; } = JobStepState.Waiting;
+
+			public JobStepOutcome Outcome { get; set; } = JobStepOutcome.Success;
+
+			public JobStepError Error { get; set; } = JobStepError.None;
+
+			[BsonIgnoreIfNull]
+			public LogId? LogId { get; set; }
+
+			[BsonIgnoreIfNull]
+			public ObjectId? NotificationTriggerId { get; set; }
+
+			[BsonIgnoreIfNull]
+			public DateTimeOffset? ReadyTime { get; set; }
+
+			[BsonIgnoreIfNull]
+			public DateTimeOffset? StartTime { get; set; }
+
+			[BsonIgnoreIfNull]
+			public DateTimeOffset? FinishTime { get; set; }
+
+			[BsonIgnoreIfNull]
+			public Priority? Priority { get; set; }
+
+			[BsonIgnoreIfDefault, BsonDefaultValue(false)]
+			public bool Retry { get; set; }
+
+			public UserId? RetriedByUserId { get; set; }
+
+			[BsonElement("RetryByUser")]
+			public string? RetriedByUserDeprecated { get; set; }
+
+			public bool AbortRequested { get; set; } = false;
+
+			public UserId? AbortedByUserId { get; set; }
+
+			[BsonElement("AbortByUser")]
+			public string? AbortedByUserDeprecated { get; set; }
+
+			[BsonIgnoreIfNull]
+			public List<Report>? Reports { get; set; }
+
+			[BsonIgnoreIfNull]
+			public Dictionary<string, string>? Properties { get; set; }
+
+			[BsonConstructor]
+			private JobStepDocument()
+			{
+			}
+
+			public JobStepDocument(JobStepId id, int nodeIdx)
+			{
+				Id = id;
+				NodeIdx = nodeIdx;
+			}
+
+			public override string ToString()
+			{
+				StringBuilder description = new StringBuilder($"{Id}: {State}");
+				if (Outcome != JobStepOutcome.Unspecified)
+				{
+					description.Append($" ({Outcome}");
+					if (Error != JobStepError.None)
+					{
+						description.Append($" - {Error}");
+					}
+					description.Append(')');
+				}
+				return description.ToString();
+			}
+		}
+
+		class ChainedJobDocument : IChainedJob
+		{
+			public string Target { get; set; }
+			public TemplateId TemplateRefId { get; set; }
+			public JobId? JobId { get; set; }
+			public bool UseDefaultChangeForTemplate { get; set; }
+
+			[BsonConstructor]
+			private ChainedJobDocument()
+			{
+				Target = String.Empty;
+			}
+
+			public ChainedJobDocument(ChainedJobTemplateConfig trigger)
+			{
+				Target = trigger.Trigger;
+				TemplateRefId = trigger.TemplateId;
+				UseDefaultChangeForTemplate = trigger.UseDefaultChangeForTemplate;
+			}
+		}
+
+		class LabelNotificationDocument
+		{
+			public int _labelIdx;
+			public ObjectId _triggerId;
+		}
+
 		/// <summary>
 		/// Maximum number of times a step can be retried (after the original run)
 		/// </summary>
@@ -1027,7 +1069,7 @@
 								updates.Add(updateBuilder.Set(x => x.Batches[batchIdx].Steps[stepIdx].AbortRequested, step.AbortRequested));
 
 								// If the step is pending, and not running on an agent, set to aborted
-								if (step.IsPending() && step.State != JobStepState.Running)
+								if (JobStepExtensions.IsPendingState(step.State) && step.State != JobStepState.Running)
 								{
 									newState = JobStepState.Aborted;
 									newOutcome = JobStepOutcome.Failure;
@@ -1299,7 +1341,7 @@
 					INode oldNode = oldGroup.Nodes[oldNodeIdx];
 					INode newNode = newGroup.Nodes[newNodeIdx];
 
-					if (!step.IsPending() && !NodesMatch(oldGraph, oldNode, newGraph, newNode))
+					if (!JobStepExtensions.IsPendingState(step.State) && !NodesMatch(oldGraph, oldNode, newGraph, newNode))
 					{
 						throw new InvalidOperationException($"Definition for node '{oldNode.Name}' has changed.");
 					}
@@ -1629,7 +1671,7 @@
 				{
 					foreach (JobStepDocument step in batch.Steps)
 					{
-						if (step.IsPending())
+						if (JobStepExtensions.IsPendingState(step.State))
 						{
 							step.State = JobStepState.Skipped;
 						}
@@ -1806,7 +1848,7 @@
 			JobStepBatchDocument?[] appendToBatches = new JobStepBatchDocument?[graph.Groups.Count];
 			foreach (JobStepBatchDocument batch in job.Batches)
 			{
-				if (batch.CanBeAppendedTo())
+				if (batch.State <= JobStepBatchState.Running)
 				{
 					INodeGroup group = graph.Groups[batch.GroupIdx];
 					appendToBatches[batch.GroupIdx] = batch;
@@ -1822,10 +1864,10 @@
 					INode node = group.Nodes[nodeIdx];
 					if (newNodesToExecute.Contains(node) && !existingNodesToExecute.Contains(node))
 					{
-						IJobStepBatch? batch = appendToBatches[groupIdx];
+						JobStepBatchDocument? batch = appendToBatches[groupIdx];
 						if (batch != null && batch.Steps.Count > 0)
 						{
-							IJobStep lastStep = batch.Steps[batch.Steps.Count - 1];
+							JobStepDocument lastStep = batch.Steps[batch.Steps.Count - 1];
 							if (nodeIdx <= lastStep.NodeIdx)
 							{
 								appendToBatches[groupIdx] = null;
@@ -1909,10 +1951,10 @@
 
 			// Check we're not running a node which doesn't allow retries more than once
 			Dictionary<INode, int> nodeExecutionCount = new Dictionary<INode, int>();
-			foreach (IJobStepBatch batch in job.Batches)
+			foreach (JobStepBatchDocument batch in job.Batches)
 			{
 				INodeGroup group = graph.Groups[batch.GroupIdx];
-				foreach (IJobStep step in batch.Steps)
+				foreach (JobStepDocument step in batch.Steps)
 				{
 					INode node = group.Nodes[step.NodeIdx];
 
@@ -1957,7 +1999,7 @@
 		static int GetSchedulePriority(JobDocument job)
 		{
 			int newSchedulePriority = 0;
-			foreach (IJobStepBatch batch in job.Batches)
+			foreach (JobStepBatchDocument batch in job.Batches)
 			{
 				if (batch.State == JobStepBatchState.Ready)
 				{
@@ -1980,7 +2022,7 @@
 			UpdateDefinitionBuilder<JobDocument> updateBuilder = Builders<JobDocument>.Update;
 			if (job.Batches != null)
 			{
-				Dictionary<INode, IJobStep> stepForNode = new Dictionary<INode, IJobStep>();
+				Dictionary<INode, JobStepDocument> stepForNode = new Dictionary<INode, JobStepDocument>();
 				for (int loopBatchIdx = 0; loopBatchIdx < job.Batches.Count; loopBatchIdx++)
 				{
 					int batchIdx = loopBatchIdx; // For lambda capture
@@ -1997,13 +2039,13 @@
 						INode node = graph.Groups[batch.GroupIdx].Nodes[step.NodeIdx];
 						if (newState == JobStepState.Waiting)
 						{
-							List<IJobStep> steps = GetDependentSteps(graph, node, stepForNode);
-							if (steps.Any(x => x.AbortRequested || x.IsFailedOrSkipped()))
+							List<JobStepDocument> steps = GetDependentSteps(graph, node, stepForNode);
+							if (steps.Any(x => x.AbortRequested || x.State == JobStepState.Skipped || x.Outcome == JobStepOutcome.Failure))
 							{
 								newState = JobStepState.Skipped;
 								newOutcome = JobStepOutcome.Failure;
 							}
-							else if (!steps.Any(x => !x.AbortRequested && x.IsPending()))
+							else if (!steps.Any(x => !x.AbortRequested && JobStepExtensions.IsPendingState(x.State)))
 							{
 								logger.LogDebug("Transitioning job {JobId}, batch {BatchId}, step {StepId} to ready state ({Dependencies})", job.Id, batch.Id, step.Id, String.Join(", ", steps.Select(x => x.Id.ToString())));
 								newState = JobStepState.Ready;
@@ -2067,12 +2109,12 @@
 		/// <param name="node">The node to test</param>
 		/// <param name="stepForNode">Map of node to step</param>
 		/// <returns></returns>
-		static List<IJobStep> GetDependentSteps(IGraph graph, INode node, Dictionary<INode, IJobStep> stepForNode)
+		static List<JobStepDocument> GetDependentSteps(IGraph graph, INode node, Dictionary<INode, JobStepDocument> stepForNode)
 		{
-			List<IJobStep> steps = new List<IJobStep>();
+			List<JobStepDocument> steps = new List<JobStepDocument>();
 			foreach (NodeRef orderDependencyRef in node.OrderDependencies)
 			{
-				IJobStep? step;
+				JobStepDocument? step;
 				if (stepForNode.TryGetValue(graph.GetNode(orderDependencyRef), out step))
 				{
 					steps.Add(step);
@@ -2090,23 +2132,23 @@
 		/// <param name="stepForNode">Array mapping each index to the appropriate step for that node</param>
 		/// <param name="outReadyTimeUtc">Receives the time at which the batch was ready to execute</param>
 		/// <returns>True if the batch is ready, false otherwise</returns>
-		static JobStepBatchState GetBatchState(JobDocument job, IGraph graph, IJobStepBatch batch, Dictionary<INode, IJobStep> stepForNode, out DateTime? outReadyTimeUtc)
+		static JobStepBatchState GetBatchState(JobDocument job, IGraph graph, JobStepBatchDocument batch, Dictionary<INode, JobStepDocument> stepForNode, out DateTime? outReadyTimeUtc)
 		{
 			// Check if the batch is already complete
 			if (batch.Steps.All(x => x.State == JobStepState.Skipped || x.State == JobStepState.Completed || x.State == JobStepState.Aborted))
 			{
-				outReadyTimeUtc = batch.ReadyTimeUtc;
+				outReadyTimeUtc = batch.ReadyTime?.UtcDateTime;
 				return JobStepBatchState.Complete;
 			}
 
 			// Get the dependencies for this batch to start. Some steps may be "after" dependencies that are optional parts of the graph.
-			List<INode> nodeDependencies = batch.GetStartDependencies(graph.Groups).ToList();
+			List<INode> nodeDependencies = JobStepBatchExtensions.GetStartDependencies(batch.Steps.ConvertAll(x => graph.Groups[batch.GroupIdx].Nodes[x.NodeIdx]), graph.Groups).ToList();
 
 			// Check if we're still waiting on anything
 			DateTime readyTimeUtc = job.GetCreateTimeOrDefault();
 			foreach (INode nodeDependency in nodeDependencies)
 			{
-				IJobStep? stepDependency;
+				JobStepDocument? stepDependency;
 				if (stepForNode.TryGetValue(nodeDependency, out stepDependency))
 				{
 					if (stepDependency.State != JobStepState.Completed && stepDependency.State != JobStepState.Skipped && stepDependency.State != JobStepState.Aborted)
@@ -2115,9 +2157,9 @@
 						return JobStepBatchState.Waiting;
 					}
 
-					if (stepDependency.FinishTimeUtc != null && stepDependency.FinishTimeUtc.Value > readyTimeUtc)
+					if (stepDependency.FinishTime != null && stepDependency.FinishTime.Value.UtcDateTime > readyTimeUtc)
 					{
-						readyTimeUtc = stepDependency.FinishTimeUtc.Value;
+						readyTimeUtc = stepDependency.FinishTime.Value.UtcDateTime;
 					}
 				}
 			}
diff --git a/Engine/Source/Programs/Horde/Horde.Server/Jobs/JobsController.cs b/Engine/Source/Programs/Horde/Horde.Server/Jobs/JobsController.cs
--- a/Engine/Source/Programs/Horde/Horde.Server/Jobs/JobsController.cs
+++ b/Engine/Source/Programs/Horde/Horde.Server/Jobs/JobsController.cs
@@ -700,7 +700,7 @@ namespace Horde.Server.Jobs
 
 			if (step.Properties != null && step.Properties.Count > 0)
 			{
-				response.Properties = step.Properties;
+				response.Properties = step.Properties.ToDictionary();
 			}
 			return response;
 		}
