Horde: Fix minimum advanced tab height
#rnx

diff --git a/Engine/Source/Programs/Horde/HordeDashboard/src/components/build/NewBuildV2.tsx b/Engine/Source/Programs/Horde/HordeDashboard/src/components/build/NewBuildV2.tsx
--- a/Engine/Source/Programs/Horde/HordeDashboard/src/components/build/NewBuildV2.tsx
+++ b/Engine/Source/Programs/Horde/HordeDashboard/src/components/build/NewBuildV2.tsx
@@ -188,8 +188,7 @@ class BuildOptions {
 
    }
 
-   onAutoSubmitChanged(value: boolean) 
-   {
+   onAutoSubmitChanged(value: boolean) {
       if (this.autoSubmit === value) {
          return;
       }
@@ -436,7 +435,9 @@ class BuildOptions {
          }
 
       } else if (this.mode === "Advanced") {
-         estimatedHeight += 192;
+         if (estimatedHeight < 280) {
+            estimatedHeight = 280
+         }
       }
 
 
