Horde: Add methods to access node/group properties onto IJobStep and IJobStepBatch.

#rnx

diff --git a/Engine/Source/Programs/Horde/Horde.Server/Jobs/IJob.cs b/Engine/Source/Programs/Horde/Horde.Server/Jobs/IJob.cs
--- a/Engine/Source/Programs/Horde/Horde.Server/Jobs/IJob.cs
+++ b/Engine/Source/Programs/Horde/Horde.Server/Jobs/IJob.cs
@@ -1259,11 +1259,21 @@ namespace Horde.Server.Jobs
 		/// </summary>
 		public JobStepBatchId Id { get; }
 
+		/// <summary>
+		/// The type of agent to execute this group
+		/// </summary>
+		public string AgentType { get; }
+
 		/// <summary>
 		/// The log file id for this batch
 		/// </summary>
 		public LogId? LogId { get; }
 
+		/// <summary>
+		/// The node group for this batch
+		/// </summary>
+		public INodeGroup Group { get; }
+
 		/// <summary>
 		/// Index of the group being executed
 		/// </summary>
@@ -1431,6 +1441,20 @@ namespace Horde.Server.Jobs
 		}
 	}
 
+	/// <summary>
+	/// Reference to another step in the job
+	/// </summary>
+	/// <param name="BatchIdx">Index of the batch</param>
+	/// <param name="StepIdx">Index of the step</param>
+	public record struct StepRef(int BatchIdx, int StepIdx);
+
+	/// <summary>
+	/// Reference to the output of a step within the job
+	/// </summary>
+	/// <param name="Step">Step producing the output</param>
+	/// <param name="OutputIdx">Index of the output from this step</param>
+	public record struct StepOutputRef(StepRef Step, int OutputIdx);
+
 	/// <summary>
 	/// Embedded jobstep document
 	/// </summary>
@@ -1441,11 +1465,66 @@ namespace Horde.Server.Jobs
 		/// </summary>
 		public JobStepId Id { get; }
 
+		/// <summary>
+		/// The node for this step
+		/// </summary>
+		public INode Node { get; }
+
 		/// <summary>
 		/// Index of the node which this jobstep is to execute
 		/// </summary>
 		public int NodeIdx { get; }
 
+		/// <summary>
+		/// The name of this node 
+		/// </summary>
+		public string Name { get; }
+
+		/// <summary>
+		/// References to inputs for this node
+		/// </summary>
+		public IReadOnlyList<StepOutputRef> Inputs { get; }
+
+		/// <summary>
+		/// List of output names
+		/// </summary>
+		public IReadOnlyList<string> OutputNames { get; }
+
+		/// <summary>
+		/// Indices of nodes which must have succeeded for this node to run
+		/// </summary>
+		public IReadOnlyList<StepRef> InputDependencies { get; }
+
+		/// <summary>
+		/// Indices of nodes which must have completed for this node to run
+		/// </summary>
+		public IReadOnlyList<StepRef> OrderDependencies { get; }
+
+		/// <summary>
+		/// Whether this node can be run multiple times
+		/// </summary>
+		public bool AllowRetry { get; }
+
+		/// <summary>
+		/// This node can start running early, before dependencies of other nodes in the same group are complete
+		/// </summary>
+		public bool RunEarly { get; }
+
+		/// <summary>
+		/// Whether to include warnings in the output (defaults to true)
+		/// </summary>
+		public bool Warnings { get; }
+
+		/// <summary>
+		/// List of credentials required for this node. Each entry maps an environment variable name to a credential in the form "CredentialName.PropertyName".
+		/// </summary>
+		public IReadOnlyDictionary<string, string>? Credentials { get; }
+
+		/// <summary>
+		/// Annotations for this node
+		/// </summary>
+		public IReadOnlyNodeAnnotations Annotations { get; }
+
 		/// <summary>
 		/// Current state of the job step. This is updated automatically when runs complete.
 		/// </summary>
diff --git a/Engine/Source/Programs/Horde/Horde.Server/Jobs/JobCollection.cs b/Engine/Source/Programs/Horde/Horde.Server/Jobs/JobCollection.cs
--- a/Engine/Source/Programs/Horde/Horde.Server/Jobs/JobCollection.cs
+++ b/Engine/Source/Programs/Horde/Horde.Server/Jobs/JobCollection.cs
@@ -93,14 +93,36 @@ namespace Horde.Server.Jobs
 			DateTime IJob.UpdateTimeUtc => _document.UpdateTimeUtc ?? _document.UpdateTime?.UtcDateTime ?? DateTime.UnixEpoch;
 			int IJob.UpdateIndex => _document.UpdateIndex;
 
+			Dictionary<NodeRef, StepRef>? _cachedNodeRefToStepRef;
+			public IReadOnlyDictionary<NodeRef, StepRef> NodeRefToStepRef => _cachedNodeRefToStepRef ??= CreateNodeRefToStepRef();
+
+			Dictionary<NodeRef, StepRef> CreateNodeRefToStepRef()
+			{
+				Dictionary<NodeRef, StepRef> nodeRefToStepRef = new Dictionary<NodeRef, StepRef>();
+				for (int batchIdx = 0; batchIdx < Batches.Count; batchIdx++)
+				{
+					JobStepBatch batch = Batches[batchIdx];
+					for (int stepIdx = 0; stepIdx < batch.Steps.Count; stepIdx++)
+					{
+						JobStep step = batch.Steps[stepIdx];
+
+						NodeRef nodeRef = new NodeRef(batch.Document.GroupIdx, step.Document.NodeIdx);
+						nodeRefToStepRef[nodeRef] = new StepRef(batchIdx, stepIdx);
+					}
+				}
+				return nodeRefToStepRef;
+			}
+
 			public Job(JobCollection collection, JobDocument document, IGraph graph)
 			{
 				_collection = collection;
 				_document = document;
 				_graph = graph;
-				_batches = document.Batches.ConvertAll(x => new JobStepBatch(x));
+				_batches = document.Batches.ConvertAll(x => new JobStepBatch(this, x, graph.Groups[x.GroupIdx]));
 			}
 
+			public StepRef GetStepRef(NodeRef nodeRef) => NodeRefToStepRef[nodeRef];
+
 			public Task<IJob?> RefreshAsync(CancellationToken cancellationToken = default)
 				=> _collection.GetAsync(_document.Id, cancellationToken);
 
@@ -170,10 +192,17 @@ namespace Horde.Server.Jobs
 
 		class JobStepBatch : IJobStepBatch
 		{
+			public Job Job { get; }
 			public JobStepBatchDocument Document { get; }
+			public INodeGroup Group { get; }
 			public List<JobStep> Steps { get; }
 
 			JobStepBatchId IJobStepBatch.Id => Document.Id;
+
+			// Group properties
+			string IJobStepBatch.AgentType => Group.AgentType;
+
+			// Batch properties
 			LogId? IJobStepBatch.LogId => Document.LogId;
 			int IJobStepBatch.GroupIdx => Document.GroupIdx;
 			JobStepBatchState IJobStepBatch.State => Document.State;
@@ -188,37 +217,59 @@ namespace Horde.Server.Jobs
 			DateTime? IJobStepBatch.StartTimeUtc => Document.StartTime?.UtcDateTime;
 			DateTime? IJobStepBatch.FinishTimeUtc => Document.FinishTime?.UtcDateTime;
 
-			public JobStepBatch(JobStepBatchDocument document)
+			public JobStepBatch(Job job, JobStepBatchDocument document, INodeGroup group)
 			{
+				Job = job;
 				Document = document;
-				Steps = document.Steps.ConvertAll(x => new JobStep(x));
+				Group = group;
+				Steps = document.Steps.ConvertAll(x => new JobStep(this, x, group.Nodes[x.NodeIdx]));
 			}
 		}
 
 		class JobStep : IJobStep
 		{
-			readonly JobStepDocument _document;
-
-			JobStepId IJobStep.Id => _document.Id;
-			int IJobStep.NodeIdx => _document.NodeIdx;
-			JobStepState IJobStep.State => _document.State;
-			JobStepOutcome IJobStep.Outcome => _document.Outcome;
-			JobStepError IJobStep.Error => _document.Error;
-			LogId? IJobStep.LogId => _document.LogId;
-			ObjectId? IJobStep.NotificationTriggerId => _document.NotificationTriggerId;
-			DateTime? IJobStep.ReadyTimeUtc => _document.ReadyTime?.UtcDateTime;
-			DateTime? IJobStep.StartTimeUtc => _document.StartTime?.UtcDateTime;
-			DateTime? IJobStep.FinishTimeUtc => _document.FinishTime?.UtcDateTime;
-			Priority? IJobStep.Priority => _document.Priority;
-			UserId? IJobStep.RetriedByUserId => _document.RetriedByUserId;
-			bool IJobStep.AbortRequested => _document.AbortRequested;
-			UserId? IJobStep.AbortedByUserId => _document.AbortedByUserId;
-			IReadOnlyList<IReport>? IJobStep.Reports => _document.Reports;
-			IReadOnlyDictionary<string, string>? IJobStep.Properties => _document.Properties;
-
-			public JobStep(JobStepDocument document)
-			{
-				_document = document;
+			public Job Job => Batch.Job;
+			public JobStepBatch Batch { get; }
+			public JobStepDocument Document { get; }
+			public INode Node { get; }
+
+			JobStepId IJobStep.Id => Document.Id;
+			INode IJobStep.Node => Node;
+			int IJobStep.NodeIdx => Document.NodeIdx;
+
+			// Node properties
+			string IJobStep.Name => Node.Name;
+			IReadOnlyList<StepOutputRef> IJobStep.Inputs => Node.Inputs.ConvertAll(x => new StepOutputRef(Job.GetStepRef(x.NodeRef), x.OutputIdx));
+			IReadOnlyList<string> IJobStep.OutputNames => Node.OutputNames;
+			IReadOnlyList<StepRef> IJobStep.InputDependencies => Node.InputDependencies.ConvertAll(x => Job.GetStepRef(x));
+			IReadOnlyList<StepRef> IJobStep.OrderDependencies => Node.OrderDependencies.ConvertAll(x => Job.GetStepRef(x));
+			bool IJobStep.AllowRetry => Node.AllowRetry;
+			bool IJobStep.RunEarly => Node.RunEarly;
+			bool IJobStep.Warnings => Node.Warnings;
+			IReadOnlyDictionary<string, string>? IJobStep.Credentials => Node.Credentials;
+			IReadOnlyNodeAnnotations IJobStep.Annotations => Node.Annotations;
+
+			// Step properties
+			JobStepState IJobStep.State => Document.State;
+			JobStepOutcome IJobStep.Outcome => Document.Outcome;
+			JobStepError IJobStep.Error => Document.Error;
+			LogId? IJobStep.LogId => Document.LogId;
+			ObjectId? IJobStep.NotificationTriggerId => Document.NotificationTriggerId;
+			DateTime? IJobStep.ReadyTimeUtc => Document.ReadyTime?.UtcDateTime;
+			DateTime? IJobStep.StartTimeUtc => Document.StartTime?.UtcDateTime;
+			DateTime? IJobStep.FinishTimeUtc => Document.FinishTime?.UtcDateTime;
+			Priority? IJobStep.Priority => Document.Priority;
+			UserId? IJobStep.RetriedByUserId => Document.RetriedByUserId;
+			bool IJobStep.AbortRequested => Document.AbortRequested;
+			UserId? IJobStep.AbortedByUserId => Document.AbortedByUserId;
+			IReadOnlyList<IReport>? IJobStep.Reports => Document.Reports;
+			IReadOnlyDictionary<string, string>? IJobStep.Properties => Document.Properties;
+
+			public JobStep(JobStepBatch batch, JobStepDocument document, INode node)
+			{
+				Batch = batch;
+				Document = document;
+				Node = node;
 			}
 		}
 
