Horde: Include graph information in job responses.

diff --git a/Engine/Source/Programs/Horde/Horde.Server.Tests/Issues/IssueServiceTests.cs b/Engine/Source/Programs/Horde/Horde.Server.Tests/Issues/IssueServiceTests.cs
--- a/Engine/Source/Programs/Horde/Horde.Server.Tests/Issues/IssueServiceTests.cs
+++ b/Engine/Source/Programs/Horde/Horde.Server.Tests/Issues/IssueServiceTests.cs
@@ -228,6 +228,9 @@ namespace Horde.Server.Tests.Issues
 			return node.Object;
 		}
 
+		delegate void TryGetBatchDelegate(JobStepBatchId id, out IJobStepBatch? batch);
+		delegate void TryGetStepDelegate(JobStepId id, out IJobStep? step);
+
 		public IJob CreateJob(StreamId streamId, int change, string name, IGraph graph, TimeSpan time = default, bool promoteByDefault = true, bool updateIssues = true)
 		{
 			JobId jobId = JobIdUtils.GenerateNewId();
@@ -265,6 +268,27 @@ namespace Horde.Server.Tests.Issues
 
 			Mock<IJob> job = new Mock<IJob>(MockBehavior.Strict);
 			job.SetupGet(x => x.Id).Returns(jobId);
+
+			job.Setup(x => x.TryGetBatch(It.IsAny<JobStepBatchId>(), out It.Ref<IJobStepBatch?>.IsAny))
+				.Returns(false);
+
+			job.Setup(x => x.TryGetStep(It.IsAny<JobStepId>(), out It.Ref<IJobStep?>.IsAny))
+				.Returns(false);
+
+			foreach (IJobStepBatch batch in batches)
+			{
+				job.Setup(x => x.TryGetBatch(batch.Id, out It.Ref<IJobStepBatch?>.IsAny))
+					.Callback(new TryGetBatchDelegate((JobStepBatchId id, out IJobStepBatch? outBatch) => outBatch = batch))
+					.Returns(true);
+			}
+
+			foreach (IJobStep step in batches.SelectMany(x => x.Steps))
+			{
+				job.Setup(x => x.TryGetStep(step.Id, out It.Ref<IJobStep?>.IsAny))
+					.Callback(new TryGetStepDelegate((JobStepId id, out IJobStep? outStep) => outStep = step))
+					.Returns(true);
+			}
+
 			job.SetupGet(x => x.Name).Returns(name);
 			job.SetupGet(x => x.StreamId).Returns(streamId);
 			job.SetupGet(x => x.TemplateId).Returns(new TemplateId("test-template"));
diff --git a/Engine/Source/Programs/Horde/Horde.Server/Jobs/IJob.cs b/Engine/Source/Programs/Horde/Horde.Server/Jobs/IJob.cs
--- a/Engine/Source/Programs/Horde/Horde.Server/Jobs/IJob.cs
+++ b/Engine/Source/Programs/Horde/Horde.Server/Jobs/IJob.cs
@@ -247,6 +247,22 @@ namespace Horde.Server.Jobs
 		/// <param name="cancellationToken">Cancellation token for the operation</param>
 		Task<IJob?> RefreshAsync(CancellationToken cancellationToken = default);
 
+		/// <summary>
+		/// Attempt to get a batch with the given id
+		/// </summary>
+		/// <param name="batchId">The job batch id</param>
+		/// <param name="batch">Receives the batch interface on success</param>
+		/// <returns>True if the batch was found</returns>
+		bool TryGetBatch(JobStepBatchId batchId, [NotNullWhen(true)] out IJobStepBatch? batch);
+
+		/// <summary>
+		/// Attempt to get a step with the given id
+		/// </summary>
+		/// <param name="stepId">The job step id</param>
+		/// <param name="step">Receives the step interface on success</param>
+		/// <returns>True if the step was found</returns>
+		bool TryGetStep(JobStepId stepId, [NotNullWhen(true)] out IJobStep? step);
+
 		/// <summary>
 		/// Attempt to delete the job
 		/// </summary>
@@ -390,52 +406,6 @@ namespace Horde.Server.Jobs
 			}
 		}
 
-		/// <summary>
-		/// Attempts to get a step with a given ID
-		/// </summary>
-		/// <param name="job">The job document</param>
-		/// <param name="stepId">The step id</param>
-		/// <param name="step">On success, receives the step object</param>
-		/// <returns>True if the step was found</returns>
-		public static bool TryGetStep(this IJob job, JobStepId stepId, [NotNullWhen(true)] out IJobStep? step)
-		{
-			foreach (IJobStepBatch batch in job.Batches)
-			{
-				if (batch.TryGetStep(stepId, out step))
-				{
-					return true;
-				}
-			}
-
-			step = null;
-			return false;
-		}
-
-		/// <summary>
-		/// Attempts to get a step with a given ID
-		/// </summary>
-		/// <param name="job">The job document</param>
-		/// <param name="stepId">The step id</param>
-		/// <param name="batch">On success returns the batch containing the step</param>
-		/// <param name="step">On success, receives the step object</param>
-		/// <returns>True if the step was found</returns>
-		public static bool TryGetStep(this IJob job, JobStepId stepId, [NotNullWhen(true)] out IJobStepBatch? batch, [NotNullWhen(true)] out IJobStep? step)
-		{
-			foreach (IJobStepBatch currentBatch in job.Batches)
-			{
-				if (currentBatch.TryGetStep(stepId, out IJobStep? currentStep))
-				{
-					batch = currentBatch;
-					step = currentStep;
-					return true;
-				}
-			}
-
-			batch = null;
-			step = null;
-			return false;
-		}
-
 		/// <summary>
 		/// Gets the current job state
 		/// </summary>
@@ -841,19 +811,6 @@ namespace Horde.Server.Jobs
 			}
 		}
 
-		/// <summary>
-		/// Attempts to get a batch with the given id
-		/// </summary>
-		/// <param name="job">The job document</param>
-		/// <param name="batchId">The batch id</param>
-		/// <param name="batch">On success, receives the batch object</param>
-		/// <returns>True if the batch was found</returns>
-		public static bool TryGetBatch(this IJob job, JobStepBatchId batchId, [NotNullWhen(true)] out IJobStepBatch? batch)
-		{
-			batch = job.Batches.FirstOrDefault(x => x.Id == batchId);
-			return batch != null;
-		}
-
 		/// <summary>
 		/// Attempts to get a batch with the given id
 		/// </summary>
@@ -864,13 +821,12 @@ namespace Horde.Server.Jobs
 		/// <returns>True if the batch was found</returns>
 		public static bool TryGetStep(this IJob job, JobStepBatchId batchId, JobStepId stepId, [NotNullWhen(true)] out IJobStep? step)
 		{
-			IJobStepBatch? batch;
-			if (!TryGetBatch(job, batchId, out batch))
+			if (!job.TryGetStep(stepId, out step) || step.Batch.Id != batchId)
 			{
 				step = null;
 				return false;
 			}
-			return batch.TryGetStep(stepId, out step);
+			return true;
 		}
 
 		/// <summary>
@@ -959,7 +915,22 @@ namespace Horde.Server.Jobs
 				}
 
 				// Create the response
-				responses.Add(new GetLabelStateResponse(newState, newOutcome));
+				GetLabelStateResponse response = new GetLabelStateResponse(newState, newOutcome);
+				response.DashboardName = label.DashboardName;
+				response.DashboardCategory = label.DashboardCategory;
+				response.UgsName = label.UgsName;
+				response.UgsProject = label.UgsProject;
+
+				foreach (NodeRef includedNodeRef in label.IncludedNodes)
+				{
+					IJobStep? includedStep;
+					if (stepForNode.TryGetValue(includedNodeRef, out includedStep))
+					{
+						response.Steps.Add(includedStep.Id);
+					}
+				}
+
+				responses.Add(response);
 			}
 
 			// Remove all the nodes that don't have a step
@@ -1254,6 +1225,11 @@ namespace Horde.Server.Jobs
 	/// </summary>
 	public interface IJobStepBatch
 	{
+		/// <summary>
+		/// Job that this batch belongs to
+		/// </summary>
+		public IJob Job { get; }
+
 		/// <summary>
 		/// Unique id for this group
 		/// </summary>
@@ -1441,25 +1417,21 @@ namespace Horde.Server.Jobs
 		}
 	}
 
-	/// <summary>
-	/// Reference to another step in the job
-	/// </summary>
-	/// <param name="BatchIdx">Index of the batch</param>
-	/// <param name="StepIdx">Index of the step</param>
-	public record struct StepRef(int BatchIdx, int StepIdx);
-
-	/// <summary>
-	/// Reference to the output of a step within the job
-	/// </summary>
-	/// <param name="Step">Step producing the output</param>
-	/// <param name="OutputIdx">Index of the output from this step</param>
-	public record struct StepOutputRef(StepRef Step, int OutputIdx);
-
 	/// <summary>
 	/// Embedded jobstep document
 	/// </summary>
 	public interface IJobStep
 	{
+		/// <summary>
+		/// Job that this step belongs to
+		/// </summary>
+		public IJob Job { get; }
+
+		/// <summary>
+		/// Batch that this step belongs to
+		/// </summary>
+		public IJobStepBatch Batch { get; }
+
 		/// <summary>
 		/// Unique ID assigned to this jobstep. A new id is generated whenever a jobstep's order is changed.
 		/// </summary>
@@ -1483,7 +1455,7 @@ namespace Horde.Server.Jobs
 		/// <summary>
 		/// References to inputs for this node
 		/// </summary>
-		public IReadOnlyList<StepOutputRef> Inputs { get; }
+		public IReadOnlyList<JobStepOutputRef> Inputs { get; }
 
 		/// <summary>
 		/// List of output names
@@ -1493,12 +1465,12 @@ namespace Horde.Server.Jobs
 		/// <summary>
 		/// Indices of nodes which must have succeeded for this node to run
 		/// </summary>
-		public IReadOnlyList<StepRef> InputDependencies { get; }
+		public IReadOnlyList<JobStepId> InputDependencies { get; }
 
 		/// <summary>
 		/// Indices of nodes which must have completed for this node to run
 		/// </summary>
-		public IReadOnlyList<StepRef> OrderDependencies { get; }
+		public IReadOnlyList<JobStepId> OrderDependencies { get; }
 
 		/// <summary>
 		/// Whether this node can be run multiple times
diff --git a/Engine/Source/Programs/Horde/Horde.Server/Jobs/JobCollection.cs b/Engine/Source/Programs/Horde/Horde.Server/Jobs/JobCollection.cs
--- a/Engine/Source/Programs/Horde/Horde.Server/Jobs/JobCollection.cs
+++ b/Engine/Source/Programs/Horde/Horde.Server/Jobs/JobCollection.cs
@@ -93,12 +93,55 @@ namespace Horde.Server.Jobs
 			DateTime IJob.UpdateTimeUtc => _document.UpdateTimeUtc ?? _document.UpdateTime?.UtcDateTime ?? DateTime.UnixEpoch;
 			int IJob.UpdateIndex => _document.UpdateIndex;
 
-			Dictionary<NodeRef, StepRef>? _cachedNodeRefToStepRef;
-			public IReadOnlyDictionary<NodeRef, StepRef> NodeRefToStepRef => _cachedNodeRefToStepRef ??= CreateNodeRefToStepRef();
+			Dictionary<JobStepBatchId, JobStepBatch>? _batchIdToBatch;
+			Dictionary<JobStepId, JobStep>? _stepIdToStep;
 
-			Dictionary<NodeRef, StepRef> CreateNodeRefToStepRef()
+			Dictionary<NodeRef, JobStepId>? _cachedNodeRefToStepId;
+			public IReadOnlyDictionary<NodeRef, JobStepId> NodeRefToStepId => _cachedNodeRefToStepId ??= CreateNodeRefToStepId();
+
+			bool TryGetBatch(JobStepBatchId batchId, [NotNullWhen(true)] out JobStepBatch? batch)
+			{
+				_batchIdToBatch ??= _batches.ToDictionary(x => x.Document.Id, x => x);
+				return _batchIdToBatch.TryGetValue(batchId, out batch);
+			}
+
+			bool IJob.TryGetBatch(JobStepBatchId batchId, [NotNullWhen(true)] out IJobStepBatch? batch)
+			{
+				if (TryGetBatch(batchId, out JobStepBatch? typedBatch))
+				{
+					batch = typedBatch;
+					return true;
+				}
+				else
+				{
+					batch = null;
+					return false;
+				}
+			}
+
+			bool TryGetStep(JobStepId stepId, [NotNullWhen(true)] out JobStep? step)
+			{
+				_stepIdToStep ??= _batches.SelectMany(x => x.Steps).ToDictionary(x => x.Document.Id, x => x);
+				return _stepIdToStep.TryGetValue(stepId, out step);
+			}
+
+			bool IJob.TryGetStep(JobStepId stepId, [NotNullWhen(true)] out IJobStep? step)
+			{
+				if (TryGetStep(stepId, out JobStep? typedStep))
+				{
+					step = typedStep;
+					return true;
+				}
+				else
+				{
+					step = null;
+					return false;
+				}
+			}
+
+			Dictionary<NodeRef, JobStepId> CreateNodeRefToStepId()
 			{
-				Dictionary<NodeRef, StepRef> nodeRefToStepRef = new Dictionary<NodeRef, StepRef>();
+				Dictionary<NodeRef, JobStepId> nodeRefToStepRef = new Dictionary<NodeRef, JobStepId>();
 				for (int batchIdx = 0; batchIdx < Batches.Count; batchIdx++)
 				{
 					JobStepBatch batch = Batches[batchIdx];
@@ -107,7 +150,7 @@ namespace Horde.Server.Jobs
 						JobStep step = batch.Steps[stepIdx];
 
 						NodeRef nodeRef = new NodeRef(batch.Document.GroupIdx, step.Document.NodeIdx);
-						nodeRefToStepRef[nodeRef] = new StepRef(batchIdx, stepIdx);
+						nodeRefToStepRef[nodeRef] = step.Document.Id;
 					}
 				}
 				return nodeRefToStepRef;
@@ -121,7 +164,7 @@ namespace Horde.Server.Jobs
 				_batches = document.Batches.ConvertAll(x => new JobStepBatch(this, x, graph.Groups[x.GroupIdx]));
 			}
 
-			public StepRef GetStepRef(NodeRef nodeRef) => NodeRefToStepRef[nodeRef];
+			public JobStepId GetStepId(NodeRef nodeRef) => NodeRefToStepId[nodeRef];
 
 			public Task<IJob?> RefreshAsync(CancellationToken cancellationToken = default)
 				=> _collection.GetAsync(_document.Id, cancellationToken);
@@ -197,6 +240,7 @@ namespace Horde.Server.Jobs
 			public INodeGroup Group { get; }
 			public List<JobStep> Steps { get; }
 
+			IJob IJobStepBatch.Job => Job;
 			JobStepBatchId IJobStepBatch.Id => Document.Id;
 
 			// Group properties
@@ -233,16 +277,18 @@ namespace Horde.Server.Jobs
 			public JobStepDocument Document { get; }
 			public INode Node { get; }
 
+			IJob IJobStep.Job => Job;
+			IJobStepBatch IJobStep.Batch => Batch;
 			JobStepId IJobStep.Id => Document.Id;
 			INode IJobStep.Node => Node;
 			int IJobStep.NodeIdx => Document.NodeIdx;
 
 			// Node properties
 			string IJobStep.Name => Node.Name;
-			IReadOnlyList<StepOutputRef> IJobStep.Inputs => Node.Inputs.ConvertAll(x => new StepOutputRef(Job.GetStepRef(x.NodeRef), x.OutputIdx));
+			IReadOnlyList<JobStepOutputRef> IJobStep.Inputs => Node.Inputs.ConvertAll(x => new JobStepOutputRef(Job.GetStepId(x.NodeRef), x.OutputIdx));
 			IReadOnlyList<string> IJobStep.OutputNames => Node.OutputNames;
-			IReadOnlyList<StepRef> IJobStep.InputDependencies => Node.InputDependencies.ConvertAll(x => Job.GetStepRef(x));
-			IReadOnlyList<StepRef> IJobStep.OrderDependencies => Node.OrderDependencies.ConvertAll(x => Job.GetStepRef(x));
+			IReadOnlyList<JobStepId> IJobStep.InputDependencies => Node.InputDependencies.ConvertAll(x => Job.GetStepId(x));
+			IReadOnlyList<JobStepId> IJobStep.OrderDependencies => Node.OrderDependencies.ConvertAll(x => Job.GetStepId(x));
 			bool IJobStep.AllowRetry => Node.AllowRetry;
 			bool IJobStep.RunEarly => Node.RunEarly;
 			bool IJobStep.Warnings => Node.Warnings;
diff --git a/Engine/Source/Programs/Horde/Horde.Server/Jobs/JobRpcService.cs b/Engine/Source/Programs/Horde/Horde.Server/Jobs/JobRpcService.cs
--- a/Engine/Source/Programs/Horde/Horde.Server/Jobs/JobRpcService.cs
+++ b/Engine/Source/Programs/Horde/Horde.Server/Jobs/JobRpcService.cs
@@ -176,10 +176,12 @@ namespace Horde.Server.Jobs
 			{
 				throw new StructuredRpcException(StatusCode.NotFound, "Unable to find job {JobId}", jobId);
 			}
-			if (!job.TryGetStep(stepId, out IJobStepBatch? batch, out IJobStep? step))
+			if (!job.TryGetStep(stepId, out IJobStep? step))
 			{
 				throw new StructuredRpcException(StatusCode.NotFound, "Unable to find step {JobId}:{StepId}", job.Id, stepId);
 			}
+
+			IJobStepBatch batch = step.Batch;
 			if (batch.SessionId == null)
 			{
 				throw new StructuredRpcException(StatusCode.PermissionDenied, "Batch {JobId}:{BatchId} has no session id", job.Id, batch.Id);
diff --git a/Engine/Source/Programs/Horde/Horde.Server/Jobs/JobsController.cs b/Engine/Source/Programs/Horde/Horde.Server/Jobs/JobsController.cs
--- a/Engine/Source/Programs/Horde/Horde.Server/Jobs/JobsController.cs
+++ b/Engine/Source/Programs/Horde/Horde.Server/Jobs/JobsController.cs
@@ -636,6 +636,7 @@ namespace Horde.Server.Jobs
 			response.Id = batch.Id;
 			response.LogId = batch.LogId?.ToString();
 			response.GroupIdx = batch.GroupIdx;
+			response.AgentType = batch.AgentType;
 			response.State = batch.State;
 			response.Error = batch.Error;
 			response.Steps.AddRange(steps);
@@ -684,6 +685,20 @@ namespace Horde.Server.Jobs
 			GetStepResponse response = new GetStepResponse();
 			response.Id = step.Id;
 			response.NodeIdx = step.NodeIdx;
+
+			// Node properties
+			response.Name = step.Name;
+			response.Inputs = step.Inputs.ToList();
+			response.OutputNames = step.OutputNames.ToList();
+			response.InputDependencies = step.InputDependencies.ToList();
+			response.OrderDependencies = step.OrderDependencies.ToList();
+			response.AllowRetry = step.AllowRetry;
+			response.RunEarly = step.RunEarly;
+			response.Warnings = step.Warnings;
+			response.Credentials = step.Credentials;
+			response.Annotations = step.Annotations;
+
+			// Step properties
 			response.State = step.State;
 			response.Outcome = step.Outcome;
 			response.Error = step.Error;
@@ -1522,9 +1537,9 @@ namespace Horde.Server.Jobs
 			try
 			{
 				NodeRef? retryNodeRef = null;
-				if (request.Retry != null && job.TryGetStep(stepId, out IJobStepBatch? batch, out IJobStep? step))
+				if (request.Retry != null && job.TryGetStep(stepId, out IJobStep? step))
 				{
-					retryNodeRef = new NodeRef(batch.GroupIdx, step.NodeIdx);
+					retryNodeRef = new NodeRef(step.Batch.GroupIdx, step.NodeIdx);
 				}
 
 				IJob? newJob = await _jobService.UpdateStepAsync(job, batchId, stepId, streamConfig, request.State, request.Outcome, null, request.AbortRequested, abortByUser, (request.LogId == null) ? null : LogId.Parse(request.LogId), null, retryByUser, request.Priority, null, request.Properties);
diff --git a/Engine/Source/Programs/Shared/EpicGames.Horde/Jobs/JobMessages.cs b/Engine/Source/Programs/Shared/EpicGames.Horde/Jobs/JobMessages.cs
--- a/Engine/Source/Programs/Shared/EpicGames.Horde/Jobs/JobMessages.cs
+++ b/Engine/Source/Programs/Shared/EpicGames.Horde/Jobs/JobMessages.cs
@@ -676,6 +676,13 @@ namespace EpicGames.Horde.Jobs
 		public string? StepId { get; set; }
 	}
 
+	/// <summary>
+	/// Reference to the output of a step within the job
+	/// </summary>
+	/// <param name="StepId">Step producing the output</param>
+	/// <param name="OutputIdx">Index of the output from this step</param>
+	public record struct JobStepOutputRef(JobStepId StepId, int OutputIdx);
+
 	/// <summary>
 	/// Returns information about a jobstep
 	/// </summary>
@@ -691,6 +698,56 @@ namespace EpicGames.Horde.Jobs
 		/// </summary>
 		public int NodeIdx { get; set; }
 
+		/// <summary>
+		/// The name of this node 
+		/// </summary>
+		public string Name { get; set; } = String.Empty;
+
+		/// <summary>
+		/// References to inputs for this node
+		/// </summary>
+		public List<JobStepOutputRef>? Inputs { get; set; }
+
+		/// <summary>
+		/// List of output names
+		/// </summary>
+		public List<string>? OutputNames { get; set; }
+
+		/// <summary>
+		/// Indices of nodes which must have succeeded for this node to run
+		/// </summary>
+		public List<JobStepId>? InputDependencies { get; set; }
+
+		/// <summary>
+		/// Indices of nodes which must have completed for this node to run
+		/// </summary>
+		public List<JobStepId>? OrderDependencies { get; set; }
+
+		/// <summary>
+		/// Whether this node can be run multiple times
+		/// </summary>
+		public bool AllowRetry { get; set; }
+
+		/// <summary>
+		/// This node can start running early, before dependencies of other nodes in the same group are complete
+		/// </summary>
+		public bool RunEarly { get; set; }
+
+		/// <summary>
+		/// Whether to include warnings in the output (defaults to true)
+		/// </summary>
+		public bool Warnings { get; set; }
+
+		/// <summary>
+		/// List of credentials required for this node. Each entry maps an environment variable name to a credential in the form "CredentialName.PropertyName".
+		/// </summary>
+		public IReadOnlyDictionary<string, string>? Credentials { get; set; }
+
+		/// <summary>
+		/// Annotations for this node
+		/// </summary>
+		public IReadOnlyDictionary<string, string>? Annotations { get; set; }
+
 		/// <summary>
 		/// Current state of the job step. This is updated automatically when runs complete.
 		/// </summary>
@@ -920,6 +977,11 @@ namespace EpicGames.Horde.Jobs
 		/// </summary>
 		public int GroupIdx { get; set; }
 
+		/// <summary>
+		/// The agent type
+		/// </summary>
+		public string AgentType { get; set; } = String.Empty;
+
 		/// <summary>
 		/// The state of this batch
 		/// </summary>
@@ -1028,6 +1090,31 @@ namespace EpicGames.Horde.Jobs
 	/// </summary>
 	public class GetLabelStateResponse
 	{
+		/// <summary>
+		/// Name to show for this label on the dashboard
+		/// </summary>
+		public string? DashboardName { get; set; }
+
+		/// <summary>
+		/// Category to show this label in on the dashboard
+		/// </summary>
+		public string? DashboardCategory { get; set; }
+
+		/// <summary>
+		/// Name to show for this label in UGS
+		/// </summary>
+		public string? UgsName { get; set; }
+
+		/// <summary>
+		/// Project to display this label for in UGS
+		/// </summary>
+		public string? UgsProject { get; set; }
+
+		/// <summary>
+		/// Steps to include in the status of this label
+		/// </summary>
+		public List<JobStepId> Steps { get; set; } = new List<JobStepId>();
+
 		/// <summary>
 		/// State of the label
 		/// </summary>
