Horde: Rename step -> job step, report -> job report, batch -> job batch in job responses.

diff --git a/Engine/Source/Programs/Horde/Horde.Server.Tests/Jobs/JobsControllerDbTest.cs b/Engine/Source/Programs/Horde/Horde.Server.Tests/Jobs/JobsControllerDbTest.cs
--- a/Engine/Source/Programs/Horde/Horde.Server.Tests/Jobs/JobsControllerDbTest.cs
+++ b/Engine/Source/Programs/Horde/Horde.Server.Tests/Jobs/JobsControllerDbTest.cs
@@ -48,7 +48,7 @@ namespace Horde.Server.Tests.Jobs
 #pragma warning disable CS0618 // Type or member is obsolete
 			object obj = (await JobsController.GetStepAsync(job.Id, batchId, stepId)).Value!;
 #pragma warning restore CS0618 // Type or member is obsolete
-			GetStepResponse stepRes = (obj as GetStepResponse)!;
+			GetJobStepResponse stepRes = (obj as GetJobStepResponse)!;
 			Assert.IsFalse(stepRes.AbortRequested);
 
 			UpdateStepRequest updateReq = new UpdateStepRequest();
@@ -59,7 +59,7 @@ namespace Horde.Server.Tests.Jobs
 #pragma warning disable CS0618 // Type or member is obsolete
 			obj = (await JobsController.GetStepAsync(job.Id, batchId, stepId)).Value!;
 #pragma warning restore CS0618 // Type or member is obsolete
-			stepRes = (obj as GetStepResponse)!;
+			stepRes = (obj as GetJobStepResponse)!;
 			Assert.IsTrue(stepRes.AbortRequested);
 			//	        Assert.AreEqual("Anonymous", StepRes.AbortByUser);
 		}
diff --git a/Engine/Source/Programs/Horde/Horde.Server/Jobs/IJob.cs b/Engine/Source/Programs/Horde/Horde.Server/Jobs/IJob.cs
--- a/Engine/Source/Programs/Horde/Horde.Server/Jobs/IJob.cs
+++ b/Engine/Source/Programs/Horde/Horde.Server/Jobs/IJob.cs
@@ -224,7 +224,7 @@ namespace Horde.Server.Jobs
 		/// <summary>
 		/// List of reports for this step
 		/// </summary>
-		public IReadOnlyList<IReport>? Reports { get; }
+		public IReadOnlyList<IJobReport>? Reports { get; }
 
 		/// <summary>
 		/// List of downstream job triggers
@@ -292,7 +292,7 @@ namespace Horde.Server.Jobs
 		/// <param name="labelIdxToTriggerId">New trigger ID for a label in the job</param>
 		/// <param name="jobTrigger">New downstream job id</param>
 		/// <param name="cancellationToken">Cancellation token for the operation</param>
-		Task<IJob?> TryUpdateJobAsync(string? name = null, Priority? priority = null, bool? autoSubmit = null, int? autoSubmitChange = null, string? autoSubmitMessage = null, UserId? abortedByUserId = null, ObjectId? notificationTriggerId = null, List<Report>? reports = null, List<string>? arguments = null, KeyValuePair<int, ObjectId>? labelIdxToTriggerId = null, KeyValuePair<TemplateId, JobId>? jobTrigger = null, CancellationToken cancellationToken = default);
+		Task<IJob?> TryUpdateJobAsync(string? name = null, Priority? priority = null, bool? autoSubmit = null, int? autoSubmitChange = null, string? autoSubmitMessage = null, UserId? abortedByUserId = null, ObjectId? notificationTriggerId = null, List<JobReport>? reports = null, List<string>? arguments = null, KeyValuePair<int, ObjectId>? labelIdxToTriggerId = null, KeyValuePair<TemplateId, JobId>? jobTrigger = null, CancellationToken cancellationToken = default);
 
 		/// <summary>
 		/// Updates the state of a batch
@@ -323,7 +323,7 @@ namespace Horde.Server.Jobs
 		/// <param name="newProperties">Property changes. Any properties with a null value will be removed.</param>
 		/// <param name="cancellationToken">Cancellation token for the operation</param>
 		/// <returns>True if the job was updated, false if it was deleted in the meantime</returns>
-		Task<IJob?> TryUpdateStepAsync(JobStepBatchId batchId, JobStepId stepId, JobStepState newState = default, JobStepOutcome newOutcome = default, JobStepError? newError = null, bool? newAbortRequested = null, UserId? newAbortByUserId = null, LogId? newLogId = null, ObjectId? newNotificationTriggerId = null, UserId? newRetryByUserId = null, Priority? newPriority = null, List<Report>? newReports = null, Dictionary<string, string?>? newProperties = null, CancellationToken cancellationToken = default);
+		Task<IJob?> TryUpdateStepAsync(JobStepBatchId batchId, JobStepId stepId, JobStepState newState = default, JobStepOutcome newOutcome = default, JobStepError? newError = null, bool? newAbortRequested = null, UserId? newAbortByUserId = null, LogId? newLogId = null, ObjectId? newNotificationTriggerId = null, UserId? newRetryByUserId = null, Priority? newPriority = null, List<JobReport>? newReports = null, Dictionary<string, string?>? newProperties = null, CancellationToken cancellationToken = default);
 
 		/// <summary>
 		/// Attempts to update the node groups to be executed for a job. Fails if another write happens in the meantime.
@@ -958,7 +958,13 @@ namespace Horde.Server.Jobs
 
 			// Add a response for everything not included elsewhere.
 			GetLabelState(unlabeledNodes, stepForNode, out LabelState otherState, out LabelOutcome otherOutcome);
-			return new GetDefaultLabelStateResponse(otherState, otherOutcome, unlabeledNodes.Select(x => graph.GetNode(x).Name).ToList());
+
+			GetDefaultLabelStateResponse defaultResponse = new GetDefaultLabelStateResponse();
+			defaultResponse.State = otherState;
+			defaultResponse.Outcome = otherOutcome;
+			defaultResponse.Nodes.AddRange(unlabeledNodes.Select(x => graph.GetNode(x).Name));
+			defaultResponse.Steps.AddRange(unlabeledNodes.Select(x => stepForNode[x].Id));
+			return defaultResponse;
 		}
 
 		/// <summary>
@@ -1562,7 +1568,7 @@ namespace Horde.Server.Jobs
 		/// <summary>
 		/// List of reports for this step
 		/// </summary>
-		public IReadOnlyList<IReport>? Reports { get; }
+		public IReadOnlyList<IJobReport>? Reports { get; }
 
 		/// <summary>
 		/// Reports for this jobstep.
@@ -1775,7 +1781,7 @@ namespace Horde.Server.Jobs
 	/// <summary>
 	/// Report for a job or jobstep
 	/// </summary>
-	public interface IReport
+	public interface IJobReport
 	{
 		/// <summary>
 		/// Name of the report
@@ -1785,7 +1791,7 @@ namespace Horde.Server.Jobs
 		/// <summary>
 		/// Where to render the report
 		/// </summary>
-		ReportPlacement Placement { get; }
+		JobReportPlacement Placement { get; }
 
 		/// <summary>
 		/// The artifact id
@@ -1801,13 +1807,13 @@ namespace Horde.Server.Jobs
 	/// <summary>
 	/// Implementation of IReport
 	/// </summary>
-	public class Report : IReport
+	public class JobReport : IJobReport
 	{
 		/// <inheritdoc/>
 		public string Name { get; set; } = String.Empty;
 
 		/// <inheritdoc/>
-		public ReportPlacement Placement { get; set; }
+		public JobReportPlacement Placement { get; set; }
 
 		/// <inheritdoc/>
 		public ObjectId? ArtifactId { get; set; }
diff --git a/Engine/Source/Programs/Horde/Horde.Server/Jobs/JobCollection.cs b/Engine/Source/Programs/Horde/Horde.Server/Jobs/JobCollection.cs
--- a/Engine/Source/Programs/Horde/Horde.Server/Jobs/JobCollection.cs
+++ b/Engine/Source/Programs/Horde/Horde.Server/Jobs/JobCollection.cs
@@ -88,7 +88,7 @@ namespace Horde.Server.Jobs
 			string? IJob.NotificationChannel => _document.NotificationChannel;
 			string? IJob.NotificationChannelFilter => _document.NotificationChannelFilter;
 			IReadOnlyDictionary<int, ObjectId> IJob.LabelIdxToTriggerId => _document._labelNotifications.ToDictionary(x => x._labelIdx, x => x._triggerId);
-			IReadOnlyList<IReport>? IJob.Reports => _document.Reports;
+			IReadOnlyList<IJobReport>? IJob.Reports => _document.Reports;
 			IReadOnlyList<IChainedJob> IJob.ChainedJobs => _document.ChainedJobs;
 			DateTime IJob.UpdateTimeUtc => _document.UpdateTimeUtc ?? _document.UpdateTime?.UtcDateTime ?? DateTime.UnixEpoch;
 			int IJob.UpdateIndex => _document.UpdateIndex;
@@ -220,13 +220,13 @@ namespace Horde.Server.Jobs
 				return _collection.CreateJobObject(newDocument, newGraph);
 			}
 
-			public async Task<IJob?> TryUpdateJobAsync(string? name = null, Priority? priority = null, bool? autoSubmit = null, int? autoSubmitChange = null, string? autoSubmitMessage = null, UserId? abortedByUserId = null, ObjectId? notificationTriggerId = null, List<Report>? reports = null, List<string>? arguments = null, KeyValuePair<int, ObjectId>? labelIdxToTriggerId = null, KeyValuePair<TemplateId, JobId>? jobTrigger = null, CancellationToken cancellationToken = default)
+			public async Task<IJob?> TryUpdateJobAsync(string? name = null, Priority? priority = null, bool? autoSubmit = null, int? autoSubmitChange = null, string? autoSubmitMessage = null, UserId? abortedByUserId = null, ObjectId? notificationTriggerId = null, List<JobReport>? reports = null, List<string>? arguments = null, KeyValuePair<int, ObjectId>? labelIdxToTriggerId = null, KeyValuePair<TemplateId, JobId>? jobTrigger = null, CancellationToken cancellationToken = default)
 			{
 				JobDocument? newDocument = await _collection.TryUpdateJobAsync(Document, Graph, name, priority, autoSubmit, autoSubmitChange, autoSubmitMessage, abortedByUserId, notificationTriggerId, reports, arguments, labelIdxToTriggerId, jobTrigger, cancellationToken);
 				return _collection.CreateJobObject(newDocument, Graph);
 			}
 
-			public async Task<IJob?> TryUpdateStepAsync(JobStepBatchId batchId, JobStepId stepId, JobStepState newState = JobStepState.Unspecified, JobStepOutcome newOutcome = JobStepOutcome.Unspecified, JobStepError? newError = null, bool? newAbortRequested = null, UserId? newAbortByUserId = null, LogId? newLogId = null, ObjectId? newNotificationTriggerId = null, UserId? newRetryByUserId = null, Priority? newPriority = null, List<Report>? newReports = null, Dictionary<string, string?>? newProperties = null, CancellationToken cancellationToken = default)
+			public async Task<IJob?> TryUpdateStepAsync(JobStepBatchId batchId, JobStepId stepId, JobStepState newState = JobStepState.Unspecified, JobStepOutcome newOutcome = JobStepOutcome.Unspecified, JobStepError? newError = null, bool? newAbortRequested = null, UserId? newAbortByUserId = null, LogId? newLogId = null, ObjectId? newNotificationTriggerId = null, UserId? newRetryByUserId = null, Priority? newPriority = null, List<JobReport>? newReports = null, Dictionary<string, string?>? newProperties = null, CancellationToken cancellationToken = default)
 			{
 				JobDocument? newDocument = await _collection.TryUpdateStepAsync(_document, _graph, batchId, stepId, newState, newOutcome, newError, newAbortRequested, newAbortByUserId, newLogId, newNotificationTriggerId, newRetryByUserId, newPriority, newReports, newProperties, cancellationToken);
 				return _collection.CreateJobObject(newDocument, Graph);
@@ -308,7 +308,7 @@ namespace Horde.Server.Jobs
 			UserId? IJobStep.RetriedByUserId => Document.RetriedByUserId;
 			bool IJobStep.AbortRequested => Document.AbortRequested;
 			UserId? IJobStep.AbortedByUserId => Document.AbortedByUserId;
-			IReadOnlyList<IReport>? IJobStep.Reports => Document.Reports;
+			IReadOnlyList<IJobReport>? IJobStep.Reports => Document.Reports;
 			IReadOnlyDictionary<string, string>? IJobStep.Properties => Document.Properties;
 
 			public JobStep(JobStepBatch batch, JobStepDocument document, INode node)
@@ -380,7 +380,7 @@ namespace Horde.Server.Jobs
 
 			public int SchedulePriority { get; set; }
 			public List<JobStepBatchDocument> Batches { get; set; } = new List<JobStepBatchDocument>();
-			public List<Report>? Reports { get; set; }
+			public List<JobReport>? Reports { get; set; }
 
 			[BsonDictionaryOptions(Representation = DictionaryRepresentation.ArrayOfDocuments)]
 			public Dictionary<ParameterId, string> Parameters { get; set; } = new Dictionary<ParameterId, string>();
@@ -576,7 +576,7 @@ namespace Horde.Server.Jobs
 			public string? AbortedByUserDeprecated { get; set; }
 
 			[BsonIgnoreIfNull]
-			public List<Report>? Reports { get; set; }
+			public List<JobReport>? Reports { get; set; }
 
 			[BsonIgnoreIfNull]
 			public Dictionary<string, string>? Properties { get; set; }
@@ -898,7 +898,7 @@ namespace Horde.Server.Jobs
 		}
 
 		/// <inheritdoc/>
-		async Task<JobDocument?> TryUpdateJobAsync(JobDocument jobDocument, IGraph graph, string? name, Priority? priority, bool? autoSubmit, int? autoSubmitChange, string? autoSubmitMessage, UserId? abortedByUserId, ObjectId? notificationTriggerId, List<Report>? reports, List<string>? arguments, KeyValuePair<int, ObjectId>? labelIdxToTriggerId, KeyValuePair<TemplateId, JobId>? jobTrigger, CancellationToken cancellationToken)
+		async Task<JobDocument?> TryUpdateJobAsync(JobDocument jobDocument, IGraph graph, string? name, Priority? priority, bool? autoSubmit, int? autoSubmitChange, string? autoSubmitMessage, UserId? abortedByUserId, ObjectId? notificationTriggerId, List<JobReport>? reports, List<string>? arguments, KeyValuePair<int, ObjectId>? labelIdxToTriggerId, KeyValuePair<TemplateId, JobId>? jobTrigger, CancellationToken cancellationToken)
 		{
 			// Create the update 
 			UpdateDefinitionBuilder<JobDocument> updateBuilder = Builders<JobDocument>.Update;
@@ -969,7 +969,7 @@ namespace Horde.Server.Jobs
 			}
 			if (reports != null)
 			{
-				jobDocument.Reports ??= new List<Report>();
+				jobDocument.Reports ??= new List<JobReport>();
 				jobDocument.Reports.RemoveAll(x => reports.Any(y => y.Name == x.Name));
 				jobDocument.Reports.AddRange(reports);
 				updates.Add(updateBuilder.Set(x => x.Reports, jobDocument.Reports));
@@ -1136,7 +1136,7 @@ namespace Horde.Server.Jobs
 		}
 
 		/// <inheritdoc/>
-		async Task<JobDocument?> TryUpdateStepAsync(JobDocument jobDocument, IGraph graph, JobStepBatchId batchId, JobStepId stepId, JobStepState newState, JobStepOutcome newOutcome, JobStepError? newError, bool? newAbortRequested, UserId? newAbortByUserId, LogId? newLogId, ObjectId? newNotificationTriggerId, UserId? newRetryByUserId, Priority? newPriority, List<Report>? newReports, Dictionary<string, string?>? newProperties, CancellationToken cancellationToken)
+		async Task<JobDocument?> TryUpdateStepAsync(JobDocument jobDocument, IGraph graph, JobStepBatchId batchId, JobStepId stepId, JobStepState newState, JobStepOutcome newOutcome, JobStepError? newError, bool? newAbortRequested, UserId? newAbortByUserId, LogId? newLogId, ObjectId? newNotificationTriggerId, UserId? newRetryByUserId, Priority? newPriority, List<JobReport>? newReports, Dictionary<string, string?>? newProperties, CancellationToken cancellationToken)
 		{
 			jobDocument = Clone(jobDocument);
 
@@ -1261,7 +1261,7 @@ namespace Horde.Server.Jobs
 							// Add any new reports
 							if (newReports != null)
 							{
-								step.Reports ??= new List<Report>();
+								step.Reports ??= new List<JobReport>();
 								step.Reports.RemoveAll(x => newReports.Any(y => y.Name == x.Name));
 								step.Reports.AddRange(newReports);
 								updates.Add(updateBuilder.Set(x => x.Batches[batchIdx].Steps[stepIdx].Reports, step.Reports));
diff --git a/Engine/Source/Programs/Horde/Horde.Server/Jobs/JobRpcService.cs b/Engine/Source/Programs/Horde/Horde.Server/Jobs/JobRpcService.cs
--- a/Engine/Source/Programs/Horde/Horde.Server/Jobs/JobRpcService.cs
+++ b/Engine/Source/Programs/Horde/Horde.Server/Jobs/JobRpcService.cs
@@ -1034,16 +1034,16 @@ namespace Horde.Server.Jobs
 
 			IJobStepBatch batch = AuthorizeBatch(job, JobStepBatchId.Parse(request.BatchId), context);
 
-			Report newReport = new Report { Name = request.Name, Placement = (ReportPlacement)request.Placement, Content = request.Content };
+			JobReport newReport = new JobReport { Name = request.Name, Placement = (JobReportPlacement)request.Placement, Content = request.Content };
 			if (request.Scope == RpcReportScope.Job)
 			{
 				_logger.LogDebug("Adding report to job {JobId}: {Name} -> {Content}", job.Id, request.Name, request.Content);
-				await _jobService.UpdateJobAsync(job, reports: new List<Report> { newReport });
+				await _jobService.UpdateJobAsync(job, reports: new List<JobReport> { newReport });
 			}
 			else
 			{
 				_logger.LogDebug("Adding report to step {JobId}:{BatchId}:{StepId}: {Name} -> {Content}", job.Id, batch.Id, request.StepId, request.Name, request.Content);
-				await _jobService.UpdateStepAsync(job, batch.Id, JobStepId.Parse(request.StepId), streamConfig, newReports: new List<Report> { newReport });
+				await _jobService.UpdateStepAsync(job, batch.Id, JobStepId.Parse(request.StepId), streamConfig, newReports: new List<JobReport> { newReport });
 			}
 
 			return new RpcCreateReportResponse();
diff --git a/Engine/Source/Programs/Horde/Horde.Server/Jobs/JobService.cs b/Engine/Source/Programs/Horde/Horde.Server/Jobs/JobService.cs
--- a/Engine/Source/Programs/Horde/Horde.Server/Jobs/JobService.cs
+++ b/Engine/Source/Programs/Horde/Horde.Server/Jobs/JobService.cs
@@ -359,7 +359,7 @@ namespace Horde.Server.Jobs
 		/// <param name="arguments">New arguments for the job</param>
 		/// <param name="labelIdxToTriggerId">New trigger ID for a label in the job</param>
 		/// <param name="cancellationToken">Cancellation token for the operation</param>
-		public async Task<IJob?> UpdateJobAsync(IJob job, string? name = null, Priority? priority = null, bool? autoSubmit = null, UserId? abortedByUserId = null, ObjectId? onCompleteTriggerId = null, List<Report>? reports = null, List<string>? arguments = null, KeyValuePair<int, ObjectId>? labelIdxToTriggerId = null, CancellationToken cancellationToken = default)
+		public async Task<IJob?> UpdateJobAsync(IJob job, string? name = null, Priority? priority = null, bool? autoSubmit = null, UserId? abortedByUserId = null, ObjectId? onCompleteTriggerId = null, List<JobReport>? reports = null, List<string>? arguments = null, KeyValuePair<int, ObjectId>? labelIdxToTriggerId = null, CancellationToken cancellationToken = default)
 		{
 			using TelemetrySpan span = _tracer.StartActiveSpan($"{nameof(JobService)}.{nameof(UpdateJobAsync)}");
 			span.SetAttribute("JobId", job.Id.ToString());
@@ -975,7 +975,7 @@ namespace Horde.Server.Jobs
 		/// <param name="newProperties">Property changes. Any properties with a null value will be removed.</param>
 		/// <param name="cancellationToken">Cancellation token for the operation</param>
 		/// <returns>True if the job was updated, false if it was deleted in the meantime</returns>
-		public async Task<IJob?> UpdateStepAsync(IJob job, JobStepBatchId batchId, JobStepId stepId, StreamConfig streamConfig, JobStepState newState = JobStepState.Unspecified, JobStepOutcome newOutcome = JobStepOutcome.Unspecified, JobStepError? newError = null, bool? newAbortRequested = null, UserId? newAbortByUserId = null, LogId? newLogId = null, ObjectId? newNotificationTriggerId = null, UserId? newRetryByUserId = null, Priority? newPriority = null, List<Report>? newReports = null, Dictionary<string, string?>? newProperties = null, CancellationToken cancellationToken = default)
+		public async Task<IJob?> UpdateStepAsync(IJob job, JobStepBatchId batchId, JobStepId stepId, StreamConfig streamConfig, JobStepState newState = JobStepState.Unspecified, JobStepOutcome newOutcome = JobStepOutcome.Unspecified, JobStepError? newError = null, bool? newAbortRequested = null, UserId? newAbortByUserId = null, LogId? newLogId = null, ObjectId? newNotificationTriggerId = null, UserId? newRetryByUserId = null, Priority? newPriority = null, List<JobReport>? newReports = null, Dictionary<string, string?>? newProperties = null, CancellationToken cancellationToken = default)
 		{
 			using TelemetrySpan span = _tracer.StartActiveSpan($"{nameof(JobService)}.{nameof(UpdateStepAsync)}");
 			span.SetAttribute("Job", job.Id.ToString());
@@ -1021,7 +1021,7 @@ namespace Horde.Server.Jobs
 		/// <param name="newProperties">Property changes. Any properties with a null value will be removed.</param>
 		/// <param name="cancellationToken">Cancellation token for the operation</param>
 		/// <returns>True if the job was updated, false if it was deleted in the meantime</returns>
-		public async Task<IJob?> TryUpdateStepAsync(IJob job, JobStepBatchId batchId, JobStepId stepId, StreamConfig streamConfig, JobStepState newState = JobStepState.Unspecified, JobStepOutcome newOutcome = JobStepOutcome.Unspecified, JobStepError? newError = null, bool? newAbortRequested = null, UserId? newAbortByUserId = null, LogId? newLogId = null, ObjectId? newTriggerId = null, UserId? newRetryByUserId = null, Priority? newPriority = null, List<Report>? newReports = null, Dictionary<string, string?>? newProperties = null, CancellationToken cancellationToken = default)
+		public async Task<IJob?> TryUpdateStepAsync(IJob job, JobStepBatchId batchId, JobStepId stepId, StreamConfig streamConfig, JobStepState newState = JobStepState.Unspecified, JobStepOutcome newOutcome = JobStepOutcome.Unspecified, JobStepError? newError = null, bool? newAbortRequested = null, UserId? newAbortByUserId = null, LogId? newLogId = null, ObjectId? newTriggerId = null, UserId? newRetryByUserId = null, Priority? newPriority = null, List<JobReport>? newReports = null, Dictionary<string, string?>? newProperties = null, CancellationToken cancellationToken = default)
 		{
 			using TelemetrySpan span = _tracer.StartActiveSpan($"{nameof(JobService)}.{nameof(TryUpdateStepAsync)}");
 			span.SetAttribute("Job", job.Id.ToString());
diff --git a/Engine/Source/Programs/Horde/Horde.Server/Jobs/JobsController.cs b/Engine/Source/Programs/Horde/Horde.Server/Jobs/JobsController.cs
--- a/Engine/Source/Programs/Horde/Horde.Server/Jobs/JobsController.cs
+++ b/Engine/Source/Programs/Horde/Horde.Server/Jobs/JobsController.cs
@@ -493,7 +493,7 @@ namespace Horde.Server.Jobs
 			{
 				if (includeBatches)
 				{
-					response.Batches = new List<GetBatchResponse>();
+					response.Batches = new List<GetJobBatchResponse>();
 					foreach (IJobStepBatch batch in job.Batches)
 					{
 						response.Batches.Add(await CreateBatchResponseAsync(batch, includeCosts, cancellationToken));
@@ -591,9 +591,9 @@ namespace Horde.Server.Jobs
 			return response;
 		}
 
-		static GetReportResponse CreateGetReportResponse(IReport report)
+		static GetJobReportResponse CreateGetReportResponse(IJobReport report)
 		{
-			GetReportResponse response = new GetReportResponse(report.Name, report.Placement);
+			GetJobReportResponse response = new GetJobReportResponse(report.Name, report.Placement);
 			response.ArtifactId = report.ArtifactId?.ToString();
 			response.Content = report.Content;
 			return response;
@@ -606,9 +606,9 @@ namespace Horde.Server.Jobs
 		/// <param name="includeCosts"></param>
 		/// <param name="cancellationToken">Cancellation token for the operation</param>
 		/// <returns></returns>
-		async ValueTask<GetBatchResponse> CreateBatchResponseAsync(IJobStepBatch batch, bool includeCosts, CancellationToken cancellationToken)
+		async ValueTask<GetJobBatchResponse> CreateBatchResponseAsync(IJobStepBatch batch, bool includeCosts, CancellationToken cancellationToken)
 		{
-			List<GetStepResponse> steps = new List<GetStepResponse>();
+			List<GetJobStepResponse> steps = new List<GetJobStepResponse>();
 			foreach (IJobStep step in batch.Steps)
 			{
 				steps.Add(await CreateStepResponseAsync(step, cancellationToken));
@@ -630,9 +630,9 @@ namespace Horde.Server.Jobs
 		/// <param name="steps">Steps in this batch</param>
 		/// <param name="agentRate">Rate for this agent</param>
 		/// <returns>Response instance</returns>
-		static GetBatchResponse CreateGetBatchResponse(IJobStepBatch batch, List<GetStepResponse> steps, double? agentRate)
+		static GetJobBatchResponse CreateGetBatchResponse(IJobStepBatch batch, List<GetJobStepResponse> steps, double? agentRate)
 		{
-			GetBatchResponse response = new GetBatchResponse();
+			GetJobBatchResponse response = new GetJobBatchResponse();
 			response.Id = batch.Id;
 			response.LogId = batch.LogId?.ToString();
 			response.GroupIdx = batch.GroupIdx;
@@ -657,7 +657,7 @@ namespace Horde.Server.Jobs
 		/// <param name="step"></param>
 		/// <param name="cancellationToken">Cancellation token for the operation</param>
 		/// <returns></returns>
-		async ValueTask<GetStepResponse> CreateStepResponseAsync(IJobStep step, CancellationToken cancellationToken)
+		async ValueTask<GetJobStepResponse> CreateStepResponseAsync(IJobStep step, CancellationToken cancellationToken)
 		{
 			GetThinUserInfoResponse? abortedByUserInfo = null;
 			if (step.AbortedByUserId != null)
@@ -680,9 +680,9 @@ namespace Horde.Server.Jobs
 		/// <param name="step">The step to construct from</param>
 		/// <param name="abortedByUserInfo">User that aborted this step</param>
 		/// <param name="retriedByUserInfo">User that retried this step</param>
-		static GetStepResponse CreateGetStepResponse(IJobStep step, GetThinUserInfoResponse? abortedByUserInfo, GetThinUserInfoResponse? retriedByUserInfo)
+		static GetJobStepResponse CreateGetStepResponse(IJobStep step, GetThinUserInfoResponse? abortedByUserInfo, GetThinUserInfoResponse? retriedByUserInfo)
 		{
-			GetStepResponse response = new GetStepResponse();
+			GetJobStepResponse response = new GetJobStepResponse();
 			response.Id = step.Id;
 			response.NodeIdx = step.NodeIdx;
 
@@ -1320,7 +1320,7 @@ namespace Horde.Server.Jobs
 		[HttpGet]
 		[Obsolete("Query entire job instead")]
 		[Route("/api/v1/jobs/{jobId}/batches")]
-		[ProducesResponseType(typeof(List<GetBatchResponse>), 200)]
+		[ProducesResponseType(typeof(List<GetJobBatchResponse>), 200)]
 		public async Task<ActionResult<List<object>>> GetBatchesAsync(JobId jobId, [FromQuery] PropertyFilter? filter = null, CancellationToken cancellationToken = default)
 		{
 			IJob? job = await _jobService.GetJobAsync(jobId, cancellationToken);
@@ -1344,7 +1344,7 @@ namespace Horde.Server.Jobs
 			List<object> responses = new List<object>();
 			foreach (IJobStepBatch batch in job.Batches)
 			{
-				GetBatchResponse response = await CreateBatchResponseAsync(batch, includeCosts, cancellationToken);
+				GetJobBatchResponse response = await CreateBatchResponseAsync(batch, includeCosts, cancellationToken);
 				responses.Add(response.ApplyFilter(filter));
 			}
 			return responses;
@@ -1400,7 +1400,7 @@ namespace Horde.Server.Jobs
 		[HttpGet]
 		[Obsolete("Query entire job instead")]
 		[Route("/api/v1/jobs/{jobId}/batches/{batchId}")]
-		[ProducesResponseType(typeof(GetBatchResponse), 200)]
+		[ProducesResponseType(typeof(GetJobBatchResponse), 200)]
 		public async Task<ActionResult<object>> GetBatchAsync(JobId jobId, JobStepBatchId batchId, [FromQuery] PropertyFilter? filter = null, CancellationToken cancellationToken = default)
 		{
 			IJob? job = await _jobService.GetJobAsync(jobId, cancellationToken);
@@ -1424,7 +1424,7 @@ namespace Horde.Server.Jobs
 			{
 				if (batch.Id == batchId)
 				{
-					GetBatchResponse response = await CreateBatchResponseAsync(batch, includeCosts, cancellationToken);
+					GetJobBatchResponse response = await CreateBatchResponseAsync(batch, includeCosts, cancellationToken);
 					return response.ApplyFilter(filter);
 				}
 			}
@@ -1443,7 +1443,7 @@ namespace Horde.Server.Jobs
 		[HttpGet]
 		[Obsolete("Query entire job instead")]
 		[Route("/api/v1/jobs/{jobId}/batches/{batchId}/steps")]
-		[ProducesResponseType(typeof(List<GetStepResponse>), 200)]
+		[ProducesResponseType(typeof(List<GetJobStepResponse>), 200)]
 		public async Task<ActionResult<List<object>>> GetStepsAsync(JobId jobId, JobStepBatchId batchId, [FromQuery] PropertyFilter? filter = null, CancellationToken cancellationToken = default)
 		{
 			IJob? job = await _jobService.GetJobAsync(jobId, cancellationToken);
@@ -1469,7 +1469,7 @@ namespace Horde.Server.Jobs
 					List<object> responses = new List<object>();
 					foreach (IJobStep step in batch.Steps)
 					{
-						GetStepResponse response = await CreateStepResponseAsync(step, cancellationToken);
+						GetJobStepResponse response = await CreateStepResponseAsync(step, cancellationToken);
 						responses.Add(response.ApplyFilter(filter));
 					}
 					return responses;
@@ -1578,7 +1578,7 @@ namespace Horde.Server.Jobs
 		[HttpGet]
 		[Obsolete("Query entire job instead")]
 		[Route("/api/v1/jobs/{jobId}/batches/{batchId}/steps/{stepId}")]
-		[ProducesResponseType(typeof(GetStepResponse), 200)]
+		[ProducesResponseType(typeof(GetJobStepResponse), 200)]
 		public async Task<ActionResult<object>> GetStepAsync(JobId jobId, JobStepBatchId batchId, JobStepId stepId, [FromQuery] PropertyFilter? filter = null, CancellationToken cancellationToken = default)
 		{
 			IJob? job = await _jobService.GetJobAsync(jobId, cancellationToken);
@@ -1605,7 +1605,7 @@ namespace Horde.Server.Jobs
 					{
 						if (step.Id == stepId)
 						{
-							GetStepResponse response = await CreateStepResponseAsync(step, cancellationToken);
+							GetJobStepResponse response = await CreateStepResponseAsync(step, cancellationToken);
 							return response.ApplyFilter(filter);
 						}
 					}
diff --git a/Engine/Source/Programs/Shared/EpicGames.Horde/Jobs/JobMessages.cs b/Engine/Source/Programs/Shared/EpicGames.Horde/Jobs/JobMessages.cs
--- a/Engine/Source/Programs/Shared/EpicGames.Horde/Jobs/JobMessages.cs
+++ b/Engine/Source/Programs/Shared/EpicGames.Horde/Jobs/JobMessages.cs
@@ -307,7 +307,7 @@
 	/// <summary>
 	/// Placement for a job report
 	/// </summary>
-	public enum ReportPlacement
+	public enum JobReportPlacement
 	{
 		/// <summary>
 		/// On a panel of its own
@@ -323,7 +323,7 @@
 	/// <summary>
 	/// Information about a report associated with a job
 	/// </summary>
-	public class GetReportResponse
+	public class GetJobReportResponse
 	{
 		/// <summary>
 		/// Name of the report
@@ -333,7 +333,7 @@
 		/// <summary>
 		/// Report placement
 		/// </summary>
-		public ReportPlacement Placement { get; set; }
+		public JobReportPlacement Placement { get; set; }
 
 		/// <summary>
 		/// The artifact id
@@ -348,7 +348,7 @@
 		/// <summary>
 		/// Constructor
 		/// </summary>
-		public GetReportResponse(string name, ReportPlacement placement)
+		public GetJobReportResponse(string name, JobReportPlacement placement)
 		{
 			Name = name;
 			Placement = placement;
@@ -366,14 +366,14 @@
 		public JobId Id { get; set; }
 
 		/// <summary>
-		/// Unique id of the stream containing this job
-		/// </summary>
-		public StreamId StreamId { get; set; }
-
-		/// <summary>
 		/// Name of the job
 		/// </summary>
 		public string Name { get; set; }
 
 		/// <summary>
+		/// Unique id of the stream containing this job
+		/// </summary>
+		public StreamId StreamId { get; set; }
+
+		/// <summary>
 		/// The changelist number to build
@@ -478,7 +478,7 @@
 		/// <summary>
 		/// Array of jobstep batches
 		/// </summary>
-		public List<GetBatchResponse>? Batches { get; set; }
+		public List<GetJobBatchResponse>? Batches { get; set; }
 
 		/// <summary>
 		/// List of labels
@@ -493,7 +493,7 @@
 		/// <summary>
 		/// List of reports
 		/// </summary>
-		public List<GetReportResponse>? Reports { get; set; }
+		public List<GetJobReportResponse>? Reports { get; set; }
 
 		/// <summary>
 		/// Artifacts produced by this job
@@ -686,7 +686,7 @@
 	/// <summary>
 	/// Returns information about a jobstep
 	/// </summary>
-	public class GetStepResponse
+	public class GetJobStepResponse
 	{
 		/// <summary>
 		/// The unique id of the step
@@ -704,6 +704,21 @@
 		public string Name { get; set; } = String.Empty;
 
 		/// <summary>
+		/// Whether this node can be run multiple times
+		/// </summary>
+		public bool AllowRetry { get; set; }
+
+		/// <summary>
+		/// This node can start running early, before dependencies of other nodes in the same group are complete
+		/// </summary>
+		public bool RunEarly { get; set; }
+
+		/// <summary>
+		/// Whether to include warnings in the output (defaults to true)
+		/// </summary>
+		public bool Warnings { get; set; }
+
+		/// <summary>
 		/// References to inputs for this node
 		/// </summary>
 		public List<JobStepOutputRef>? Inputs { get; set; }
@@ -724,19 +739,4 @@
 		public List<JobStepId>? OrderDependencies { get; set; }
 
 		/// <summary>
-		/// Whether this node can be run multiple times
-		/// </summary>
-		public bool AllowRetry { get; set; }
-
-		/// <summary>
-		/// This node can start running early, before dependencies of other nodes in the same group are complete
-		/// </summary>
-		public bool RunEarly { get; set; }
-
-		/// <summary>
-		/// Whether to include warnings in the output (defaults to true)
-		/// </summary>
-		public bool Warnings { get; set; }
-
-		/// <summary>
 		/// List of credentials required for this node. Each entry maps an environment variable name to a credential in the form "CredentialName.PropertyName".
@@ -811,7 +811,7 @@
 		/// <summary>
 		/// List of reports
 		/// </summary>
-		public List<GetReportResponse>? Reports { get; set; }
+		public List<GetJobReportResponse>? Reports { get; set; }
 
 		/// <summary>
 		/// User-defined properties for this jobstep.
@@ -960,7 +960,7 @@
 	/// <summary>
 	/// Information about a jobstep batch
 	/// </summary>
-	public class GetBatchResponse
+	public class GetJobBatchResponse
 	{
 		/// <summary>
 		/// Unique id for this batch
@@ -968,10 +968,5 @@
 		public JobStepBatchId Id { get; set; }
 
 		/// <summary>
-		/// The unique log file id
-		/// </summary>
-		public string? LogId { get; set; }
-
-		/// <summary>
 		/// Index of the group being executed
 		/// </summary>
@@ -983,15 +978,5 @@
 		public string AgentType { get; set; } = String.Empty;
 
 		/// <summary>
-		/// The state of this batch
-		/// </summary>
-		public JobStepBatchState State { get; set; }
-
-		/// <summary>
-		/// Error code for this batch
-		/// </summary>
-		public JobStepBatchError Error { get; set; }
-
-		/// <summary>
 		/// The agent assigned to execute this group
 		/// </summary>
@@ -1013,5 +998,20 @@
 		public string? LeaseId { get; set; }
 
 		/// <summary>
+		/// The unique log file id
+		/// </summary>
+		public string? LogId { get; set; }
+
+		/// <summary>
+		/// The state of this batch
+		/// </summary>
+		public JobStepBatchState State { get; set; }
+
+		/// <summary>
+		/// Error code for this batch
+		/// </summary>
+		public JobStepBatchError Error { get; set; }
+
+		/// <summary>
 		/// The priority of this batch
 		/// </summary>
@@ -1035,7 +1035,7 @@
 		/// <summary>
 		/// Steps within this run
 		/// </summary>
-		public List<GetStepResponse> Steps { get; set; } = new List<GetStepResponse>();
+		public List<GetJobStepResponse> Steps { get; set; } = new List<GetJobStepResponse>();
 	}
 
 	/// <summary>
