Horde: Exclude empty dependency lists from job responses.

#rnx

diff --git a/Engine/Source/Programs/Horde/Horde.Server/Jobs/JobsController.cs b/Engine/Source/Programs/Horde/Horde.Server/Jobs/JobsController.cs
--- a/Engine/Source/Programs/Horde/Horde.Server/Jobs/JobsController.cs
+++ b/Engine/Source/Programs/Horde/Horde.Server/Jobs/JobsController.cs
@@ -688,10 +688,22 @@ namespace Horde.Server.Jobs
 
 			// Node properties
 			response.Name = step.Name;
-			response.Inputs = step.Inputs.ToList();
-			response.OutputNames = step.OutputNames.ToList();
-			response.InputDependencies = step.InputDependencies.ToList();
-			response.OrderDependencies = step.OrderDependencies.ToList();
+			if (step.Inputs.Count > 0)
+			{
+				response.Inputs = step.Inputs.ToList();
+			}
+			if (step.OutputNames.Count > 0)
+			{
+				response.OutputNames = step.OutputNames.ToList();
+			}
+			if (step.InputDependencies.Count > 0)
+			{
+				response.InputDependencies = step.InputDependencies.ToList();
+			}
+			if (step.OrderDependencies.Count > 0)
+			{
+				response.OrderDependencies = step.OrderDependencies.ToList();
+			}
 			response.AllowRetry = step.AllowRetry;
 			response.RunEarly = step.RunEarly;
 			response.Warnings = step.Warnings;
