Horde: Delete step logs when deleting a job.

diff --git a/Engine/Source/Programs/Horde/Horde.Server/Jobs/JobCollection.cs b/Engine/Source/Programs/Horde/Horde.Server/Jobs/JobCollection.cs
--- a/Engine/Source/Programs/Horde/Horde.Server/Jobs/JobCollection.cs
+++ b/Engine/Source/Programs/Horde/Horde.Server/Jobs/JobCollection.cs
@@ -22,6 +22,7 @@ using EpicGames.Horde.Streams;
 using EpicGames.Horde.Users;
 using Horde.Server.Acls;
 using Horde.Server.Jobs.Graphs;
+using Horde.Server.Logs;
 using Horde.Server.Server;
 using Horde.Server.Streams;
 using Horde.Server.Telemetry;
@@ -648,6 +649,7 @@ namespace Horde.Server.Jobs
 		readonly ITelemetrySink _telemetrySink;
 		readonly IClock _clock;
 		readonly IGraphCollection _graphCollection;
+		readonly ILogCollection _logCollection;
 		readonly Tracer _tracer;
 		readonly IOptionsMonitor<GlobalConfig> _globalConfig;
 		readonly ILogger<JobCollection> _logger;
@@ -655,10 +657,11 @@ namespace Horde.Server.Jobs
 		/// <summary>
 		/// Constructor
 		/// </summary>
-		public JobCollection(MongoService mongoService, IClock clock, IGraphCollection graphCollection, ITelemetrySink telemetrySink, IOptionsMonitor<GlobalConfig> globalConfig, Tracer tracer, ILogger<JobCollection> logger)
+		public JobCollection(MongoService mongoService, IClock clock, IGraphCollection graphCollection, ILogCollection logCollection, ITelemetrySink telemetrySink, IOptionsMonitor<GlobalConfig> globalConfig, Tracer tracer, ILogger<JobCollection> logger)
 		{
 			_clock = clock;
 			_graphCollection = graphCollection;
+			_logCollection = logCollection;
 			_telemetrySink = telemetrySink;
 			_globalConfig = globalConfig;
 			_tracer = tracer;
@@ -750,6 +753,17 @@ namespace Horde.Server.Jobs
 		async Task<bool> TryDeleteAsync(JobDocument job, CancellationToken cancellationToken)
 		{
 			DeleteResult result = await _jobs.DeleteOneAsync(x => x.Id == job.Id && x.UpdateIndex == job.UpdateIndex, null, cancellationToken);
+			foreach (JobStepDocument step in job.Batches.SelectMany(x => x.Steps))
+			{
+				if (step.LogId.HasValue)
+				{
+					ILog? log = await _logCollection.GetAsync(step.LogId.Value);
+					if (log != null)
+					{
+						await log.DeleteAsync(cancellationToken);
+					}
+				}
+			}
 			return result.DeletedCount > 0;
 		}
 
