Horde: Allow setting project-wide workspace types

Will automatically be inserted to the workspace type list of each stream belonging to the project. This also allows defining base types which can be inherited from.

diff --git a/Engine/Source/Programs/Horde/Horde.Server.Tests/Configuration/ConfigTests.cs b/Engine/Source/Programs/Horde/Horde.Server.Tests/Configuration/ConfigTests.cs
--- a/Engine/Source/Programs/Horde/Horde.Server.Tests/Configuration/ConfigTests.cs
+++ b/Engine/Source/Programs/Horde/Horde.Server.Tests/Configuration/ConfigTests.cs
@@ -187,6 +187,48 @@ namespace Horde.Server.Tests.Configuration
 			Assert.AreEqual(222, workspaces["subSubType"].ConformDiskFreeSpace);
 		}
 
+		[TestMethod]
+		public void WorkspaceInheritFromProjectConfig()
+		{
+			GlobalConfig gc = new();
+			
+			List<string> autoSdkViews = ["foo", "bar"];
+			gc.Projects.Add(new ProjectConfig
+			{
+				WorkspaceTypes =
+				{
+					{ "project1", new WorkspaceConfig { AutoSdkView = autoSdkViews } },
+					{ "project2", new WorkspaceConfig { ConformDiskFreeSpace = 111 } },
+					{ "project3", new WorkspaceConfig { MinScratchSpace = 222 } },
+				},
+				Streams = [
+					new StreamConfig { WorkspaceTypes =
+					{
+						{"stream1", new WorkspaceConfig { Base = "project1", Stream = "myStream" }},
+						{"project3", new WorkspaceConfig { Stream = "otherStream" }}
+					}}
+				]
+			});
+			
+			gc.PostLoad(new ServerSettings());
+
+			Dictionary<string, WorkspaceConfig> workspaces = gc.Projects[0].Streams[0].WorkspaceTypes;
+			Assert.AreEqual(4, workspaces.Count);
+			
+			// Project-defined streams
+			CollectionAssert.AreEquivalent(autoSdkViews, workspaces["project1"].AutoSdkView);
+			Assert.AreEqual(111, workspaces["project2"].ConformDiskFreeSpace);
+			Assert.AreEqual(111, workspaces["project2"].ConformDiskFreeSpace);
+			
+			// project3 should be overridden by stream workspace type with the same name
+			Assert.AreEqual("otherStream", workspaces["project3"].Stream);
+			Assert.IsNull(workspaces["project3"].MinScratchSpace);
+			
+			// stream1 inherits from project1 workspace type
+			CollectionAssert.AreEquivalent(autoSdkViews, workspaces["stream1"].AutoSdkView);
+			Assert.AreEqual("myStream", workspaces["stream1"].Stream);
+		}
+
 		class ObjectValue
 		{
 			public string Value { get; set; } = "";
diff --git a/Engine/Source/Programs/Horde/Horde.Server/Projects/ProjectConfig.cs b/Engine/Source/Programs/Horde/Horde.Server/Projects/ProjectConfig.cs
--- a/Engine/Source/Programs/Horde/Horde.Server/Projects/ProjectConfig.cs
+++ b/Engine/Source/Programs/Horde/Horde.Server/Projects/ProjectConfig.cs
@@ -86,6 +86,12 @@ namespace Horde.Server.Projects
 		/// </summary>
 		public JobOptions JobOptions { get; set; } = new JobOptions();
 
+		/// <summary>
+		/// Default workspace types for streams
+		/// These are added to the list of each stream's workspace types.
+		/// </summary>
+		public Dictionary<string, WorkspaceConfig> WorkspaceTypes { get; set; } = [];
+
 		/// <summary>
 		/// Telemetry store for Horde data for this project
 		/// </summary>
diff --git a/Engine/Source/Programs/Horde/Horde.Server/Streams/StreamConfig.cs b/Engine/Source/Programs/Horde/Horde.Server/Streams/StreamConfig.cs
--- a/Engine/Source/Programs/Horde/Horde.Server/Streams/StreamConfig.cs
+++ b/Engine/Source/Programs/Horde/Horde.Server/Streams/StreamConfig.cs
@@ -269,6 +269,11 @@ namespace Horde.Server.Streams
 				template.PostLoad(this);
 			}
 
+			foreach ((string wid, WorkspaceConfig wc) in projectConfig.WorkspaceTypes)
+			{
+				WorkspaceTypes.TryAdd(wid, wc);
+			}
+
 			ConfigType.MergeDefaults(AgentTypes.Select(x => (x.Key, x.Value.Base, x.Value)));
 			ConfigType.MergeDefaults(WorkspaceTypes.Select(x => (x.Key, x.Value.Base, x.Value)));
 			ConfigType.MergeDefaults(Templates.Select(x => (x.Id, x.Base, x)));
