Horde: Use full conforms when auto-conforming

#rnx

diff --git a/Engine/Source/Programs/Horde/Horde.Server/Agents/Pools/PoolUpdateService.cs b/Engine/Source/Programs/Horde/Horde.Server/Agents/Pools/PoolUpdateService.cs
--- a/Engine/Source/Programs/Horde/Horde.Server/Agents/Pools/PoolUpdateService.cs
+++ b/Engine/Source/Programs/Horde/Horde.Server/Agents/Pools/PoolUpdateService.cs
@@ -152,7 +152,7 @@ namespace Horde.Server.Agents.Pools
 
 				if (freeDiskSpace != null && maxConformDiskSpace > 0 && freeDiskSpace < maxConformDiskSpace && agent.ConformAttemptCount is null or 0)
 				{
-					await agent.TryUpdateAsync(new UpdateAgentOptions { RequestConform = true }, cancellationToken: cancellationToken);
+					await agent.TryUpdateAsync(new UpdateAgentOptions { RequestFullConform = true }, cancellationToken: cancellationToken);
 					_logger.LogInformation("Auto-conforming {AgentId} as workspace conform disk space needed ({ConformDiskSpace:F1} MB) is less than free disk space ({FreeDiskSpace:F1} MB)",
 						agent.Id.ToString(), maxConformDiskSpace / 1024.0 / 1024.0, freeDiskSpace.Value / 1024.0 / 1024.0);
 				}
