Horde: Updating schema documentation

#rnx

diff --git a/Engine/Source/Programs/Horde/Docs/Config/Schema/Projects.md b/Engine/Source/Programs/Horde/Docs/Config/Schema/Projects.md
--- a/Engine/Source/Programs/Horde/Docs/Config/Schema/Projects.md
+++ b/Engine/Source/Programs/Horde/Docs/Config/Schema/Projects.md
@@ -16,6 +16,7 @@ Name | Description
 `pools` | [`PoolConfig`](#poolconfig)`[]`<br>List of pools for this project
 `categories` | [`ProjectCategoryConfig`](#projectcategoryconfig)`[]`<br>Categories to include in this project
 `jobOptions` | [`JobOptions`](#joboptions)<br>Default settings for executing jobs
+`workspaceTypes` | `string` `->` [`WorkspaceConfig`](#workspaceconfig)<br>Default workspace types for streams These are added to the list of each stream's workspace types.
 `telemetryStoreId` | `string`<br>Telemetry store for Horde data for this project
 `streams` | [`StreamConfig`](Streams.md)`[]`<br>List of streams
 `acl` | [`AclConfig`](#aclconfig)<br>Acl entries
@@ -190,6 +191,27 @@ Name | Description
 `containerEngineExecutable` | `string`<br>Container engine executable (docker or with full path like /usr/bin/podman)
 `extraArguments` | `string`<br>Additional arguments to pass to container engine
 
+## WorkspaceConfig
+
+Information about a workspace type
+
+Name | Description
+---- | -----------
+`base` | `string`<br>Base workspace to derive from
+`cluster` | `string`<br>Name of the Perforce server cluster to use
+`serverAndPort` | `string`<br>The Perforce server and port (eg. perforce:1666)
+`userName` | `string`<br>User to log into Perforce with (defaults to buildmachine)
+`password` | `string`<br>Password to use to log into the workspace
+`identifier` | `string`<br>Identifier to distinguish this workspace from other workspaces. Defaults to the workspace type name.
+`stream` | `string`<br>Override for the stream to sync
+`view` | `string[]`<br>Custom view for the workspace
+`incremental` | `boolean`<br>Whether to use an incrementally synced workspace
+`useAutoSdk` | `boolean`<br>Whether to use the AutoSDK
+`autoSdkView` | `string[]`<br>View for the AutoSDK paths to sync. If null, the whole thing will be synced.
+`method` | `string`<br>Method to use when syncing/materializing data from Perforce
+`minScratchSpace` | `integer`<br>Minimum disk space that must be available *after* syncing this workspace (in megabytes) If not available, the job will be aborted.
+`conformDiskFreeSpace` | `integer`<br>Threshold for when to trigger an automatic conform of agent. Measured in megabytes free on disk. Set to null or 0 to disable.
+
 ## AclConfig
 
 Parameters to update an ACL
