Horde: Updating release notes.

#rnx

diff --git a/Engine/Source/Programs/Horde/Docs/ReleaseNotes.md b/Engine/Source/Programs/Horde/Docs/ReleaseNotes.md
--- a/Engine/Source/Programs/Horde/Docs/ReleaseNotes.md
+++ b/Engine/Source/Programs/Horde/Docs/ReleaseNotes.md
@@ -1,5 +1,20 @@
 # Release Notes
 
+## 2024-04-25
+
+* Allow setting project-wide workspace types. Will automatically be inserted to the workspace type list of each stream belonging to the project. This also allows defining base types which can be inherited from. (33213109)
+* EpicGames.Perforce: Fix errors parsing change views from labels. (33207649)
+* Delete step logs when deleting a job. (33206294)
+* Use the shared registry location to configure the Horde command line tool on Windows, and use a .horde.json in the user folder on Mac and Linux. (33205861)
+* Rename step -> job step, report -> job report, batch -> job batch in job responses. (33191560)
+* Include graph information in job responses. (33189353)
+* Add a service to expire jobs after a period of time. Jobs are kept forever by default, though this may be modified through the ExpireAfterDays parameter in the JobOptions object in streams or projects. (33109364)
+* Add auto conform for agents with workspaces below a certain free disk space threshold (33102671)
+* Use a standard message for server subsystems that are operating normally. (33100047)
+* Add min scratch space to workspace config (33098369)
+* Recognize tags in Perforce changelists of the form '#horde 123' as indication that a change fixes a particular issue. The specific tag can be configured in the globals.json file to support multiple deployments with the same Perforce server. (33089498)
+* Clear the default ASPNETCORE_HTTP_PORTS environment variable in Docker images to prevent warnings about port overrides on server startup. We manage HTTP port configuration through Horde settings, and don't want the NET runtime image defaults. (33072683)
+
 ## 2024-04-18
 
 * Add dashboard's time since last user activity to HTTP log (33068679)
diff --git a/Engine/Source/Programs/Horde/Tools/ReleaseNotes/Program.cs b/Engine/Source/Programs/Horde/Tools/ReleaseNotes/Program.cs
--- a/Engine/Source/Programs/Horde/Tools/ReleaseNotes/Program.cs
+++ b/Engine/Source/Programs/Horde/Tools/ReleaseNotes/Program.cs
@@ -87,28 +87,29 @@ namespace CreateReleaseNotes
 				if (changeNumbers.Add(change.Number))
 				{
 					string description = change.Description;
-					if (!Regex.IsMatch(description, @"^\s*#rnx\s*$", RegexOptions.Multiline))
-					{
-						Match jiraMatch = Regex.Match(description, @"^\s*#jira ([a-zA-Z]+-[0-9]+.*)$");
-						string? jiraTicket = jiraMatch.Success ? jiraMatch.Groups[1].Value : null;
 
-						description = Regex.Replace(description, @"^[a-zA-Z]+:\s*", "");
-						description = Regex.Replace(description, @"^\s*#.*$", "", RegexOptions.Multiline);
-						description = Regex.Replace(description, @"\n[^\n]*\n", "\n");
-						description = description.Trim();
+					Match jiraMatch = Regex.Match(description, @"^\s*#jira ([a-zA-Z]+-[0-9]+.*)$");
+					string? jiraTicket = jiraMatch.Success ? jiraMatch.Groups[1].Value : null;
 
-						parsedChanges.Add(new ChangeInfo(change.Number, description, jiraTicket));
+					description = Regex.Replace(description, @"^[a-zA-Z]+:\s*", "");
+					description = Regex.Replace(description, @"^\s*#.*$", "", RegexOptions.Multiline);
+					description = Regex.Replace(description, @"\n\s*", "\n");
+					description = description.Trim();
 
-						string logDescription = description.Replace("\r", "", StringComparison.Ordinal).Replace("\n", " ", StringComparison.Ordinal);
+					if (!Regex.IsMatch(change.Description, @"^\s*#rnx\s*$", RegexOptions.Multiline))
+					{
+						parsedChanges.Add(new ChangeInfo(change.Number, description, jiraTicket));
+					}
 
-						const int MaxDescriptionLength = 80;
-						if (logDescription.Length > MaxDescriptionLength)
-						{
-							logDescription = logDescription.Substring(0, MaxDescriptionLength);
-						}
+					string logDescription = Regex.Replace(description, @"\s+", " ");
 
-						logger.LogInformation("{Change}: {Description}", change.Number, logDescription);
+					const int MaxDescriptionLength = 80;
+					if (logDescription.Length > MaxDescriptionLength)
+					{
+						logDescription = logDescription.Substring(0, MaxDescriptionLength);
 					}
+
+					logger.LogInformation("{Change} {Author,-20} {Description}", change.Number, change.User.ToLower(), logDescription);
 				}
 			}
 
