Horde: Fix failing auto-conform test in PoolUpdateServiceTest

diff --git a/Engine/Source/Programs/Horde/Horde.Server.Tests/Agents/Pools/PoolUpdateServiceTest.cs b/Engine/Source/Programs/Horde/Horde.Server.Tests/Agents/Pools/PoolUpdateServiceTest.cs
--- a/Engine/Source/Programs/Horde/Horde.Server.Tests/Agents/Pools/PoolUpdateServiceTest.cs
+++ b/Engine/Source/Programs/Horde/Horde.Server.Tests/Agents/Pools/PoolUpdateServiceTest.cs
@@ -117,7 +117,7 @@ public class PoolUpdateServiceTest : TestSetup
 		agent = (await AgentService.GetAgentAsync(agent.Id))!;
 		
 		// Assert
-		Assert.AreEqual(conformRequested, agent.RequestConform);
+		Assert.AreEqual(conformRequested, agent.RequestFullConform);
 	}
 
 	private async Task<IAgent> CreateAutoConformAgentAsync(int freeDiskSpaceMb, params int?[] autoConformThresholdsMb)
