Horde: Fix issue with job/graph merging, and don't require version=2 on preflight queries
#rnx

diff --git a/Engine/Source/Programs/Horde/HordeDashboard/src/components/Preflight.tsx b/Engine/Source/Programs/Horde/HordeDashboard/src/components/Preflight.tsx
--- a/Engine/Source/Programs/Horde/HordeDashboard/src/components/Preflight.tsx
+++ b/Engine/Source/Programs/Horde/HordeDashboard/src/components/Preflight.tsx
@@ -39,9 +39,10 @@ export const PreflightRedirector: React.FC = () => {
    // whether a template is specified
    const templateId = !query.get("templateId") ? "" : query.get("templateId")!;
 
-   const version = !query.get("version") ? "" : query.get("version")!;
+   // explicit version otherwise keys off there being parameters
+   let version = !query.get("version") ? "" : query.get("version")!;
 
-   const parameters:string[] = [];
+   const parameters: string[] = [];
 
    query.forEach((value, key) => {
       if (key.startsWith("id")) {
@@ -49,6 +50,10 @@ export const PreflightRedirector: React.FC = () => {
       }
    })
 
+   if (parameters.length) {
+      version = "2";
+   }
+
    if (!change) {
       setError("No preflight change specified");
       return null;
@@ -68,14 +73,14 @@ export const PreflightRedirector: React.FC = () => {
 
    let stream = projectStore.streamByFullname(streamName);
 
-   if (!stream) {      
+   if (!stream) {
 
       stream = projectStore.streamById(streamName?.replace("//", "").replaceAll("/", "-").toLowerCase());
 
       if (!stream) {
-         stream = projectStore.streamByFullname(streamName + "-VS");   
-      }      
-      
+         stream = projectStore.streamByFullname(streamName + "-VS");
+      }
+
    }
 
    if (!stream) {
diff --git a/Engine/Source/Programs/Horde/HordeDashboard/src/components/UserHome.tsx b/Engine/Source/Programs/Horde/HordeDashboard/src/components/UserHome.tsx
--- a/Engine/Source/Programs/Horde/HordeDashboard/src/components/UserHome.tsx
+++ b/Engine/Source/Programs/Horde/HordeDashboard/src/components/UserHome.tsx
@@ -963,7 +963,7 @@ const JobsPanel: React.FC<{ includeOtherPreflights: boolean }> = observer(({ inc
 
       const view = labels.filter(a => {
          const idx = labels.indexOf(a)!;
-         if (!job.labels) {
+         if (idx === -1 || !job.labels?.length) {
             return false;
          }
          return job.labels[idx].state !== LabelState.Unspecified;
