Horde: Fix static analysis warning.

#rnx

diff --git a/Engine/Source/Programs/Horde/Horde.Server/Jobs/JobCollection.cs b/Engine/Source/Programs/Horde/Horde.Server/Jobs/JobCollection.cs
--- a/Engine/Source/Programs/Horde/Horde.Server/Jobs/JobCollection.cs
+++ b/Engine/Source/Programs/Horde/Horde.Server/Jobs/JobCollection.cs
@@ -757,7 +757,7 @@ namespace Horde.Server.Jobs
 			{
 				if (step.LogId.HasValue)
 				{
-					ILog? log = await _logCollection.GetAsync(step.LogId.Value);
+					ILog? log = await _logCollection.GetAsync(step.LogId.Value, cancellationToken);
 					if (log != null)
 					{
 						await log.DeleteAsync(cancellationToken);
