Horde: Support artifact link type in log render
#rnx

diff --git a/Engine/Source/Programs/Horde/HordeDashboard/src/components/LogRender.tsx b/Engine/Source/Programs/Horde/HordeDashboard/src/components/LogRender.tsx
--- a/Engine/Source/Programs/Horde/HordeDashboard/src/components/LogRender.tsx
+++ b/Engine/Source/Programs/Horde/HordeDashboard/src/components/LogRender.tsx
@@ -11,7 +11,8 @@ enum TagType {
    MSDNCode,
    AgentId,
    LeaseId,
-   Link
+   Link,
+   ArtifactId
 }
 
 export type LogItem = {
@@ -152,6 +153,10 @@ const renderTags = (navigate: NavigateFunction, line: LogLine, lineNumber: numbe
             tagType = TagType.AgentId;
          }
 
+         if (type === "ArtifactId") {
+            tagType = TagType.ArtifactId;
+         }
+
       }
 
       if (tagType === TagType.None || !record) {
@@ -219,8 +224,16 @@ const renderTags = (navigate: NavigateFunction, line: LogLine, lineNumber: numbe
 
          return <a key={key} href={`ugs://timelapse?depotPath=${(depotPath)}`} onClick={(ev) => ev.stopPropagation()}><Highlight search={search ? search : ""} className={logStyle.logLine}>{text}</Highlight></a>;
       } else if (tagType === TagType.Link) {
-         
+
          return <a key={key} rel="noreferrer" href={record.target} onClick={(ev) => ev.stopPropagation()}><Highlight search={search ? search : ""} className={logStyle.logLine}>{text}</Highlight></a>;
+      } else if (tagType === TagType.ArtifactId) {         
+         const artifactType = properties.ArtifactType;
+         if (artifactType?.length) {            
+            const search = new URLSearchParams(window.location.search);
+            search.set("artifactContext", encodeURIComponent(artifactType as string));
+            const url = `${window.location.pathname}?` + search.toString();
+            return <a key={key} href="/" onClick={(ev) => { ev.stopPropagation(); ev.preventDefault(); navigate(url, { replace: true }) }}><Highlight search={search ? search : ""} className={logStyle.logLine}>{text}</Highlight></a>;
+         }
       }
 
       return <span key={key} />;
