Horde: Fix exception accessing order dependencies for nodes that are not in the graph.

diff --git a/Engine/Source/Programs/Horde/Horde.Server/Jobs/JobCollection.cs b/Engine/Source/Programs/Horde/Horde.Server/Jobs/JobCollection.cs
--- a/Engine/Source/Programs/Horde/Horde.Server/Jobs/JobCollection.cs
+++ b/Engine/Source/Programs/Horde/Horde.Server/Jobs/JobCollection.cs
@@ -289,7 +289,10 @@ namespace Horde.Server.Jobs
 			IReadOnlyList<JobStepOutputRef> IJobStep.Inputs => Node.Inputs.ConvertAll(x => new JobStepOutputRef(Job.GetStepId(x.NodeRef), x.OutputIdx));
 			IReadOnlyList<string> IJobStep.OutputNames => Node.OutputNames;
 			IReadOnlyList<JobStepId> IJobStep.InputDependencies => Node.InputDependencies.ConvertAll(x => Job.GetStepId(x));
-			IReadOnlyList<JobStepId> IJobStep.OrderDependencies => Node.OrderDependencies.ConvertAll(x => Job.GetStepId(x));
+
+			List<JobStepId>? _orderDependencies;
+			IReadOnlyList<JobStepId> IJobStep.OrderDependencies => _orderDependencies ??= CreateOrderDependenciesList();
+
 			bool IJobStep.AllowRetry => Node.AllowRetry;
 			bool IJobStep.RunEarly => Node.RunEarly;
 			bool IJobStep.Warnings => Node.Warnings;
@@ -318,6 +321,19 @@ namespace Horde.Server.Jobs
 				Document = document;
 				Node = node;
 			}
+
+			List<JobStepId> CreateOrderDependenciesList()
+			{
+				List<JobStepId> dependencies = new List<JobStepId>();
+				foreach (NodeRef nodeRef in Node.OrderDependencies)
+				{
+					if (Job.NodeRefToStepId.TryGetValue(nodeRef, out JobStepId jobStepId))
+					{
+						dependencies.Add(jobStepId);
+					}
+				}
+				return dependencies;
+			}
 		}
 
 		class JobDocument
