Horde: Rename DatabaseConnectionString to MongoConnectionString, and add a RedisReadOnlyMode setting.

#rnx

diff --git a/Engine/Source/Programs/Horde/Horde.Server/Server/MongoService.cs b/Engine/Source/Programs/Horde/Horde.Server/Server/MongoService.cs
--- a/Engine/Source/Programs/Horde/Horde.Server/Server/MongoService.cs
+++ b/Engine/Source/Programs/Horde/Horde.Server/Server/MongoService.cs
@@ -287,7 +287,7 @@ namespace Horde.Server.Server
 
 			try
 			{
-				ReadOnlyMode = Settings.DatabaseReadOnlyMode;
+				ReadOnlyMode = Settings.MongoReadOnlyMode;
 				if (Settings.DatabasePublicCert != null)
 				{
 					X509Store localTrustStore = new X509Store(StoreName.Root);
@@ -309,7 +309,7 @@ namespace Horde.Server.Server
 					}
 				}
 
-				string? connectionString = Settings.DatabaseConnectionString;
+				string? connectionString = Settings.MongoConnectionString;
 				if (connectionString == null)
 				{
 					if (IsPortInUse(DefaultMongoPort))
@@ -343,7 +343,7 @@ namespace Horde.Server.Server
 				//TestSslConnection(MongoSettings.Server.Host, MongoSettings.Server.Port, Logger);
 
 				_client = new MongoClient(mongoSettings);
-				Database = _client.GetDatabase(Settings.DatabaseName);
+				Database = _client.GetDatabase(Settings.MongoDatabaseName);
 
 				SingletonsV1 = GetCollection<BsonDocument>("Singletons");
 				SingletonsV2 = GetCollection<BsonDocument>("SingletonsV2");
diff --git a/Engine/Source/Programs/Horde/Horde.Server/Server/RedisService.cs b/Engine/Source/Programs/Horde/Horde.Server/Server/RedisService.cs
--- a/Engine/Source/Programs/Horde/Horde.Server/Server/RedisService.cs
+++ b/Engine/Source/Programs/Horde/Horde.Server/Server/RedisService.cs
@@ -32,6 +32,11 @@ namespace Horde.Server.Server
 		/// </summary>
 		public RedisConnectionPool ConnectionPool { get; }
 
+		/// <summary>
+		/// Flag for whether the connection is read-only
+		/// </summary>
+		public bool ReadOnlyMode { get; }
+
 		RedisProcess? _redisProcess;
 		readonly ILogger<RedisService> _logger;
 
@@ -39,8 +44,9 @@ namespace Horde.Server.Server
 		/// Constructor
 		/// </summary>
 		public RedisService(IOptions<ServerSettings> options, ILogger<RedisService> logger)
-			: this(options.Value.RedisConnectionConfig, -1, logger)
+			: this(options.Value.RedisConnectionString, -1, logger)
 		{
+			ReadOnlyMode = options.Value.RedisReadOnlyMode;
 		}
 
 		/// <summary>
@@ -65,7 +71,7 @@ namespace Horde.Server.Server
 				}
 				else
 				{
-					throw new Exception($"Unable to connect to Redis. Please set {nameof(ServerSettings.RedisConnectionConfig)} in {ServerApp.ServerConfigFile}");
+					throw new Exception($"Unable to connect to Redis. Please set {nameof(ServerSettings.RedisConnectionString)} in {ServerApp.ServerConfigFile}");
 				}
 			}
 
diff --git a/Engine/Source/Programs/Horde/Horde.Server/Server/ServerStatusService.cs b/Engine/Source/Programs/Horde/Horde.Server/Server/ServerStatusService.cs
--- a/Engine/Source/Programs/Horde/Horde.Server/Server/ServerStatusService.cs
+++ b/Engine/Source/Programs/Horde/Horde.Server/Server/ServerStatusService.cs
@@ -187,7 +187,10 @@ public class ServerStatusService : IHostedService
 		}
 
 		string data = JsonSerializer.Serialize(status);
-		await redis.HashSetAsync(RedisHashKey(), id, data);
+		if(!_redis.ReadOnlyMode)
+		{
+			await redis.HashSetAsync(RedisHashKey(), id, data);
+		}
 	}
 
 	private static async Task<SubsystemStatus> GetSubsystemStatusFromRedisAsync(IDatabase redis, string id, string name)
diff --git a/Engine/Source/Programs/Horde/Horde.Server/ServerSettings.cs b/Engine/Source/Programs/Horde/Horde.Server/ServerSettings.cs
--- a/Engine/Source/Programs/Horde/Horde.Server/ServerSettings.cs
+++ b/Engine/Source/Programs/Horde/Horde.Server/ServerSettings.cs
@@ -431,15 +431,33 @@ namespace Horde.Server
 		/// </summary>
 		public string? ComputeTunnelAddress { get; set; }
 
+		/// <summary>
+		/// Connection string for the Mongo database
+		/// </summary>
+		public string? MongoConnectionString { get; set; }
+
 		/// <summary>
 		/// MongoDB connection string
 		/// </summary>
-		public string? DatabaseConnectionString { get; set; }
+		[Obsolete("Use MongoConnectionString instead")]
+		public string? DatabaseConnectionString
+		{ 
+			get => MongoConnectionString; 
+			set => MongoConnectionString = value;
+		}
 
 		/// <summary>
 		/// MongoDB database name
 		/// </summary>
-		public string DatabaseName { get; set; } = "Horde";
+		public string MongoDatabaseName { get; set; } = "Horde";
+
+		/// <inheritdoc cref="DatabaseName"/>
+		[Obsolete("Replace references to DatabaseName with MongoDatabaseName")]
+		public string DatabaseName
+		{
+			get => MongoDatabaseName;
+			set => MongoDatabaseName = value;
+		}
 
 		/// <summary>
 		/// Optional certificate to trust in order to access the database (eg. AWS public cert for TLS)
@@ -450,7 +468,15 @@ namespace Horde.Server
 		/// Access the database in read-only mode (avoids creating indices or updating content)
 		/// Useful for debugging a local instance of HordeServer against a production database.
 		/// </summary>
-		public bool DatabaseReadOnlyMode { get; set; } = false;
+		public bool MongoReadOnlyMode { get; set; } = false;
+
+		/// <inheritdoc cref="MongoReadOnlyMode"/>
+		[Obsolete("Replace DatabaseReadOnlyMode with MongoReadOnlyMode")]
+		public bool DatabaseReadOnlyMode
+		{
+			get => MongoReadOnlyMode;
+			set => MongoReadOnlyMode = value;
+		}
 
 		/// <summary>
 		/// Shutdown the current server process if memory usage reaches this threshold (specified in MB)
@@ -641,7 +667,20 @@ namespace Horde.Server
 		/// Setting it to null will disable Redis use and connection
 		/// See format at https://stackexchange.github.io/StackExchange.Redis/Configuration.html
 		/// </summary>
-		public string? RedisConnectionConfig { get; set; }
+		public string? RedisConnectionString { get; set; }
+
+		/// <inheritdoc cref="RedisConnectionString"/>
+		[Obsolete("Use RedisConnectionString instead")]
+		public string? RedisConnectionConfig
+		{ 
+			get => RedisConnectionString;
+			set => RedisConnectionString = value;
+		}
+
+		/// <summary>
+		/// Whether to disable writes to Redis.
+		/// </summary>
+		public bool RedisReadOnlyMode { get; set; }
 
 		/// <summary>
 		/// Type of write cache to use in log service
diff --git a/Engine/Source/Programs/Horde/Horde.Server/Startup.cs b/Engine/Source/Programs/Horde/Horde.Server/Startup.cs
--- a/Engine/Source/Programs/Horde/Horde.Server/Startup.cs
+++ b/Engine/Source/Programs/Horde/Horde.Server/Startup.cs
@@ -853,7 +853,7 @@ namespace Horde.Server
 			// Notifications can be triggered from any instance, so always make sure we're ticking the background task.
 			services.AddHostedService(provider => (NotificationService)provider.GetRequiredService<INotificationService>());
 
-			if (settings.IsRunModeActive(RunMode.Worker) && !settings.DatabaseReadOnlyMode)
+			if (settings.IsRunModeActive(RunMode.Worker) && !settings.MongoReadOnlyMode)
 			{
 				services.AddHostedService<AgentReportService>();
 				services.AddHostedService<BisectService>();
@@ -906,7 +906,7 @@ namespace Horde.Server
 			// Task sources. Order of registration is important here; it dictates the priority in which sources are served.
 			services.AddSingleton<JobTaskSource>();
 
-			if (!settings.DatabaseReadOnlyMode)
+			if (!settings.MongoReadOnlyMode)
 			{
 				services.AddHostedService<JobTaskSource>(provider => provider.GetRequiredService<JobTaskSource>());
 				services.AddSingleton<ConformTaskSource>();
diff --git a/Engine/Source/Programs/Horde/Horde.Server/Utilities/Clock.cs b/Engine/Source/Programs/Horde/Horde.Server/Utilities/Clock.cs
--- a/Engine/Source/Programs/Horde/Horde.Server/Utilities/Clock.cs
+++ b/Engine/Source/Programs/Horde/Horde.Server/Utilities/Clock.cs
@@ -233,6 +233,11 @@ namespace HordeCommon
 
 		async ValueTask TriggerSharedAsync(string name, RedisKey key, TimeSpan interval, Func<CancellationToken, ValueTask> tickAsync, CancellationToken cancellationToken)
 		{
+			if (_redis.ReadOnlyMode)
+			{
+				return;
+			}
+
 			using (RedisLock sharedLock = new(_redis.GetDatabase(), key))
 			{
 				if (await sharedLock.AcquireAsync(interval, false))
