Horde: Fix static analysis warnings.

diff --git a/Engine/Source/Programs/Horde/Horde/CmdApp.cs b/Engine/Source/Programs/Horde/Horde/CmdApp.cs
--- a/Engine/Source/Programs/Horde/Horde/CmdApp.cs
+++ b/Engine/Source/Programs/Horde/Horde/CmdApp.cs
@@ -7,7 +7,6 @@ using EpicGames.Horde;
 using Microsoft.Extensions.Configuration;
 using Microsoft.Extensions.DependencyInjection;
 using Microsoft.Extensions.Logging;
-using Microsoft.Extensions.Options;
 using Serilog;
 using Serilog.Formatting.Json;
 using Serilog.Sinks.SystemConsole.Themes;
diff --git a/Engine/Source/Programs/Horde/Horde/Commands/Bundles/BundleCreate.cs b/Engine/Source/Programs/Horde/Horde/Commands/Bundles/BundleCreate.cs
--- a/Engine/Source/Programs/Horde/Horde/Commands/Bundles/BundleCreate.cs
+++ b/Engine/Source/Programs/Horde/Horde/Commands/Bundles/BundleCreate.cs
@@ -8,7 +8,6 @@ using EpicGames.Horde.Storage.Bundles;
 using EpicGames.Horde.Storage.Clients;
 using EpicGames.Horde.Storage.Nodes;
 using Microsoft.Extensions.Logging;
-using Microsoft.Extensions.Options;
 
 namespace Horde.Commands.Bundles
 {
diff --git a/Engine/Source/Programs/Horde/Horde/Commands/Bundles/BundleExtract.cs b/Engine/Source/Programs/Horde/Horde/Commands/Bundles/BundleExtract.cs
--- a/Engine/Source/Programs/Horde/Horde/Commands/Bundles/BundleExtract.cs
+++ b/Engine/Source/Programs/Horde/Horde/Commands/Bundles/BundleExtract.cs
@@ -9,7 +9,6 @@ using EpicGames.Horde.Storage.Bundles;
 using EpicGames.Horde.Storage.Clients;
 using EpicGames.Horde.Storage.Nodes;
 using Microsoft.Extensions.Logging;
-using Microsoft.Extensions.Options;
 
 namespace Horde.Commands.Bundles
 {
diff --git a/Engine/Source/Programs/Horde/Horde/Commands/ConfigCommand.cs b/Engine/Source/Programs/Horde/Horde/Commands/ConfigCommand.cs
--- a/Engine/Source/Programs/Horde/Horde/Commands/ConfigCommand.cs
+++ b/Engine/Source/Programs/Horde/Horde/Commands/ConfigCommand.cs
@@ -4,7 +4,6 @@ using System.ComponentModel;
 using EpicGames.Core;
 using EpicGames.Horde;
 using Microsoft.Extensions.Logging;
-using Microsoft.Extensions.Options;
 
 namespace Horde.Commands
 {
diff --git a/Engine/Source/Programs/Horde/Horde/Commands/LoginCommand.cs b/Engine/Source/Programs/Horde/Horde/Commands/LoginCommand.cs
--- a/Engine/Source/Programs/Horde/Horde/Commands/LoginCommand.cs
+++ b/Engine/Source/Programs/Horde/Horde/Commands/LoginCommand.cs
@@ -6,7 +6,6 @@ using EpicGames.Horde;
 using EpicGames.Horde.Server;
 using Microsoft.Extensions.DependencyInjection;
 using Microsoft.Extensions.Logging;
-using Microsoft.Extensions.Options;
 
 namespace Horde.Commands
 {
diff --git a/Engine/Source/Programs/Horde/Horde/Commands/StorageCommandBase.cs b/Engine/Source/Programs/Horde/Horde/Commands/StorageCommandBase.cs
--- a/Engine/Source/Programs/Horde/Horde/Commands/StorageCommandBase.cs
+++ b/Engine/Source/Programs/Horde/Horde/Commands/StorageCommandBase.cs
@@ -5,7 +5,6 @@ using EpicGames.Core;
 using EpicGames.Horde.Storage;
 using EpicGames.Horde.Storage.Bundles;
 using EpicGames.Horde.Storage.Clients;
-using Microsoft.Extensions.Options;
 
 namespace Horde.Commands
 {
diff --git a/Engine/Source/Programs/Horde/Horde/Commands/Workspace/WorkspaceSync.cs b/Engine/Source/Programs/Horde/Horde/Commands/Workspace/WorkspaceSync.cs
--- a/Engine/Source/Programs/Horde/Horde/Commands/Workspace/WorkspaceSync.cs
+++ b/Engine/Source/Programs/Horde/Horde/Commands/Workspace/WorkspaceSync.cs
@@ -9,7 +9,6 @@ using EpicGames.Horde.Storage.Bundles;
 using EpicGames.Horde.Storage.Clients;
 using EpicGames.Horde.Storage.Nodes;
 using Microsoft.Extensions.Logging;
-using Microsoft.Extensions.Options;
 
 namespace Horde.Commands.Workspace
 {
