Horde: Add a class to handle buffering document writes to a Mongo collection.

diff --git a/Engine/Source/Programs/Horde/Horde.Server.Tests/DatabaseIntegrationTest.cs b/Engine/Source/Programs/Horde/Horde.Server.Tests/DatabaseIntegrationTest.cs
--- a/Engine/Source/Programs/Horde/Horde.Server.Tests/DatabaseIntegrationTest.cs
+++ b/Engine/Source/Programs/Horde/Horde.Server.Tests/DatabaseIntegrationTest.cs
@@ -264,8 +264,8 @@ namespace Horde.Server.Tests
 					_mongoInstance = new MongoInstance();
 
 					ServerSettings ss = new ServerSettings();
-					ss.DatabaseName = _mongoInstance.DatabaseName;
-					ss.DatabaseConnectionString = _mongoInstance.ConnectionString;
+					ss.MongoDatabaseName = _mongoInstance.DatabaseName;
+					ss.MongoConnectionString = _mongoInstance.ConnectionString;
 
 					_mongoService = new MongoService(Options.Create(ss), redisService, OpenTelemetryTracers.Horde, _loggerFactory.CreateLogger<MongoService>(), _loggerFactory);
 				}
diff --git a/Engine/Source/Programs/Horde/Horde.Server/Telemetry/Sinks/MongoTelemetrySink.cs b/Engine/Source/Programs/Horde/Horde.Server/Telemetry/Sinks/MongoTelemetrySink.cs
--- a/Engine/Source/Programs/Horde/Horde.Server/Telemetry/Sinks/MongoTelemetrySink.cs
+++ b/Engine/Source/Programs/Horde/Horde.Server/Telemetry/Sinks/MongoTelemetrySink.cs
@@ -1,8 +1,6 @@
 // Copyright Epic Games, Inc. All Rights Reserved.
 
 using System;
-using System.Collections.Concurrent;
-using System.Collections.Generic;
 using System.Linq;
 using System.Text.Json;
 using System.Threading;
@@ -10,6 +8,7 @@ using System.Threading.Tasks;
 using EpicGames.Core;
 using EpicGames.Horde.Telemetry;
 using Horde.Server.Server;
+using Horde.Server.Utilities;
 using HordeCommon;
 using Microsoft.Extensions.Hosting;
 using Microsoft.Extensions.Logging;
@@ -50,11 +49,10 @@ namespace Horde.Server.Telemetry.Sinks
 		}
 
 		readonly MongoTelemetryConfig? _config;
-		readonly IMongoCollection<EventDocument> _telemetry;
-		readonly AsyncEvent _newDataEvent = new AsyncEvent();
-		readonly AsyncEvent _flushEvent = new AsyncEvent();
-		readonly ConcurrentQueue<EventDocument> _queue = new ConcurrentQueue<EventDocument>();
-		readonly BackgroundTask _backgroundTask;
+		readonly IMongoCollection<EventDocument> _collection;
+#pragma warning disable CA2213 // False positive? _writer is disposed in DisposeAsync.
+		readonly MongoBufferedWriter<EventDocument> _writer;
+#pragma warning restore CA2213
 		readonly ITicker _cleanupTicker;
 		readonly JsonSerializerOptions _jsonOptions;
 		readonly ILogger _logger;
@@ -68,8 +66,8 @@ namespace Horde.Server.Telemetry.Sinks
 		public MongoTelemetrySink(MongoService mongoService, IClock clock, IOptions<ServerSettings> serverSettings, ILogger<MongoTelemetrySink> logger)
 		{
 			_config = serverSettings.Value.Telemetry.Select(x => x as MongoTelemetryConfig).FirstOrDefault(x => x != null);
-			_telemetry = mongoService.GetCollection<EventDocument>("Telemetry", builder => builder.Ascending(x => x.TelemetryStoreId).Descending(x => x.Id));
-			_backgroundTask = new BackgroundTask(BackgroundFlushAsync);
+			_collection = mongoService.GetCollection<EventDocument>("Telemetry", builder => builder.Ascending(x => x.TelemetryStoreId).Descending(x => x.Id));
+			_writer = new MongoBufferedWriter<EventDocument>(_collection, logger);
 			_cleanupTicker = clock.AddSharedTicker<MongoTelemetrySink>(TimeSpan.FromHours(4.0), CleanupAsync, logger);
 			_logger = logger;
 
@@ -82,7 +80,7 @@ namespace Horde.Server.Telemetry.Sinks
 		{
 			if (Enabled)
 			{
-				_backgroundTask.Start();
+				await _writer.StartAsync();
 				await _cleanupTicker.StartAsync();
 			}
 		}
@@ -91,25 +89,7 @@ namespace Horde.Server.Telemetry.Sinks
 		public async Task StopAsync(CancellationToken cancellationToken)
 		{
 			await _cleanupTicker.StopAsync();
-			await _backgroundTask.StopAsync(cancellationToken);
-		}
-
-		// Flushes the sink in the background
-		async Task BackgroundFlushAsync(CancellationToken cancellationToken)
-		{
-			Task newDataTask = _newDataEvent.Task;
-			Task flushTask = _flushEvent.Task;
-
-			while (!cancellationToken.IsCancellationRequested)
-			{
-				await newDataTask.WaitAsync(cancellationToken);
-				await Task.WhenAny(flushTask, Task.Delay(TimeSpan.FromSeconds(5.0), cancellationToken));
-
-				newDataTask = _newDataEvent.Task;
-				flushTask = _flushEvent.Task;
-
-				await FlushAsync(cancellationToken);
-			}
+			await _writer.StopAsync(cancellationToken);
 		}
 
 		async ValueTask CleanupAsync(CancellationToken cancellationToken)
@@ -123,48 +103,27 @@ namespace Horde.Server.Telemetry.Sinks
 #pragma warning disable CS0618 // Type or member is obsolete
 			ObjectId minObjectId = new ObjectId(baseTime, 0, 0, 0);
 #pragma warning restore CS0618
-			DeleteResult result = await _telemetry.DeleteManyAsync(x => x.Id < minObjectId, cancellationToken);
+			DeleteResult result = await _collection.DeleteManyAsync(x => x.Id < minObjectId, cancellationToken);
 			_logger.LogInformation("Deleted {NumItems} telemetry records", result.DeletedCount);
 		}
 
 		/// <inheritdoc/>
 		public async ValueTask DisposeAsync()
 		{
-			await _cleanupTicker.DisposeAsync();
-			await _backgroundTask.DisposeAsync();
 			await FlushAsync(default);
+			await _cleanupTicker.DisposeAsync();
+			await _writer.DisposeAsync();
 		}
 
 		/// <inheritdoc/>
-		public async ValueTask FlushAsync(CancellationToken cancellationToken)
-		{
-			// Copy all the event documents from the queue
-			List<EventDocument> eventDocuments = new List<EventDocument>(_queue.Count);
-			while (_queue.TryDequeue(out EventDocument? eventDocument))
-			{
-				eventDocuments.Add(eventDocument);
-			}
-
-			// Insert them into the database
-			if (eventDocuments.Count > 0)
-			{
-				_logger.LogInformation("Writing {NumEvents} new telemetry events to MongoDB.", eventDocuments.Count);
-				await _telemetry.InsertManyAsync(eventDocuments, cancellationToken: cancellationToken);
-			}
-		}
+		public ValueTask FlushAsync(CancellationToken cancellationToken)
+			=> _writer.FlushAsync(cancellationToken);
 
 		/// <inheritdoc/>
 		public void SendEvent(TelemetryStoreId telemetryStoreId, TelemetryEvent telemetryEvent)
 		{
 			BsonDocument bson = BsonDocument.Parse(JsonSerializer.Serialize(telemetryEvent, _jsonOptions));
-			_queue.Enqueue(new EventDocument(telemetryStoreId, bson));
-			_newDataEvent.Set();
-
-			const int FlushCount = 50;
-			if (_queue.Count > FlushCount)
-			{
-				_flushEvent.Set();
-			}
+			_writer.Write(new EventDocument(telemetryStoreId, bson));
 		}
 	}
 }
diff --git a/Engine/Source/Programs/Horde/Horde.Server/Utilities/MongoBufferedWriter.cs b/Engine/Source/Programs/Horde/Horde.Server/Utilities/MongoBufferedWriter.cs
new file mode 100644
--- /dev/null
+++ b/Engine/Source/Programs/Horde/Horde.Server/Utilities/MongoBufferedWriter.cs
@@ -0,0 +1,107 @@
+// Copyright Epic Games, Inc. All Rights Reserved.
+
+using System;
+using System.Collections.Concurrent;
+using System.Collections.Generic;
+using System.Threading;
+using System.Threading.Tasks;
+using EpicGames.Core;
+using MongoDB.Driver;
+using Microsoft.Extensions.Logging;
+
+namespace Horde.Server.Utilities
+{
+	/// <summary>
+	/// Utility class for buffering document  writes to a Mongo collection
+	/// </summary>
+	sealed class MongoBufferedWriter<TDocument> : IAsyncDisposable
+		where TDocument : class
+	{
+		readonly IMongoCollection<TDocument> _collection;
+		readonly int _flushCount;
+		readonly TimeSpan _flushTime;
+		readonly ConcurrentQueue<TDocument> _queue = new ConcurrentQueue<TDocument>();
+		readonly BackgroundTask _backgroundTask;
+		readonly AsyncEvent _newDataEvent = new AsyncEvent();
+		readonly AsyncEvent _flushEvent = new AsyncEvent();
+		readonly ILogger _logger;
+
+		public MongoBufferedWriter(IMongoCollection<TDocument> collection, ILogger logger)
+			: this(collection, 50, TimeSpan.FromSeconds(5.0), logger)
+		{
+		}
+
+		public MongoBufferedWriter(IMongoCollection<TDocument> collection, int flushCount, TimeSpan flushTime, ILogger logger)
+		{
+			_collection = collection;
+			_flushCount = flushCount;
+			_flushTime = flushTime;
+			_backgroundTask = new BackgroundTask(BackgroundFlushAsync);
+			_logger = logger;
+		}
+
+		public async ValueTask DisposeAsync()
+		{
+			await _backgroundTask.DisposeAsync();
+		}
+
+		public ValueTask StartAsync()
+		{
+			_backgroundTask.Start();
+			return new ValueTask();
+		}
+
+		public async ValueTask StopAsync(CancellationToken cancellationToken)
+		{
+			await _backgroundTask.StopAsync(cancellationToken);
+		}
+
+		// Flushes the sink in the background
+		async Task BackgroundFlushAsync(CancellationToken cancellationToken)
+		{
+			Task newDataTask = _newDataEvent.Task;
+			Task flushTask = _flushEvent.Task;
+
+			while (!cancellationToken.IsCancellationRequested)
+			{
+				await newDataTask.WaitAsync(cancellationToken);
+				await Task.WhenAny(flushTask, Task.Delay(_flushTime, cancellationToken));
+
+				newDataTask = _newDataEvent.Task;
+				flushTask = _flushEvent.Task;
+
+				await FlushAsync(cancellationToken);
+			}
+		}
+
+		/// <inheritdoc/>
+		public async ValueTask FlushAsync(CancellationToken cancellationToken)
+		{
+			// Copy all the event documents from the queue
+			List<TDocument> document = new List<TDocument>(_queue.Count);
+			while (_queue.TryDequeue(out TDocument? eventDocument))
+			{
+				document.Add(eventDocument);
+			}
+
+			// Insert them into the database
+			if (document.Count > 0)
+			{
+				_logger.LogInformation("Writing {NumEvents} new telemetry events to {CollectionName}.", document.Count, _collection.CollectionNamespace.CollectionName);
+				await _collection.InsertManyAsync(document, cancellationToken: cancellationToken);
+			}
+		}
+
+		/// <inheritdoc/>
+		public void Write(TDocument document)
+		{
+			_queue.Enqueue(document);
+			_newDataEvent.Set();
+
+			if (_queue.Count > _flushCount)
+			{
+				_flushEvent.Set();
+			}
+		}
+	}
+}
