Improve Gauntlet issue handling by adding stream id and node name to generate key as long as a Callstack property is not found.

#rnx

diff --git a/Engine/Source/Programs/Horde/Horde.Server.Tests/Issues/IssueServiceTests.cs b/Engine/Source/Programs/Horde/Horde.Server.Tests/Issues/IssueServiceTests.cs
--- a/Engine/Source/Programs/Horde/Horde.Server.Tests/Issues/IssueServiceTests.cs
+++ b/Engine/Source/Programs/Horde/Horde.Server.Tests/Issues/IssueServiceTests.cs
@@ -1749,92 +1749,54 @@ namespace Horde.Server.Tests.Issues
 		public async Task GauntletTestAsync()
 		{
 			// #1
-			// Scenario: Gauntlet test event with Name property
-			// Expected: Gauntlet fingerprint using Name prefix
+			// Scenario: Gauntlet test event with a property
+			// Expected: Gauntlet fingerprint using step identifiers
 			{
 				IJob job = CreateJob(_mainStreamId, 110, "Test Build", _graph);
 				await using (TestJsonLogger logger = await CreateLoggerAsync(job, 0, 0))
 				{
-					logger.LogError(KnownLogEvents.Gauntlet_TestEvent, "    Test {Name} failed", "Bar.Foo.Test");
+					logger.LogError(KnownLogEvents.Gauntlet_TestEvent, "Test {Name} failed", "Bar.Foo.Test");
 				}
 				await UpdateCompleteStepAsync(job, 0, 0, JobStepOutcome.Failure);
 
 				IReadOnlyList<IIssue> issues = await IssueCollection.FindIssuesAsync();
 				Assert.AreEqual(1, issues.Count);
 				Assert.AreEqual("Gauntlet", issues[0].Fingerprints[0].Type);
-				Assert.AreEqual(new IssueKey("test:Bar.Foo.Test", IssueKeyType.None), issues[0].Fingerprints[0].Keys.First());
-				Assert.AreEqual("Automation test errors with Bar.Foo.Test", issues[0].Summary);
+				Assert.AreEqual(new IssueKey($"{job.StreamId}:Update Version Files", IssueKeyType.None), issues[0].Fingerprints[0].Keys.First());
+				Assert.AreEqual("Automation test errors in Update Version Files", issues[0].Summary);
 			}
 			// #2
-			// Scenario: Gauntlet device event with Name property
-			// Expected: Gauntlet fingerprint using Device prefix
-			{
-				IJob job = CreateJob(_mainStreamId, 120, "Test Build", _graph);
-				await using (TestJsonLogger logger = await CreateLoggerAsync(job, 0, 0))
-				{
-					logger.LogWarning(KnownLogEvents.Gauntlet_DeviceEvent, "    Device {Name} reported an issue", "Foo");
-				}
-				await UpdateCompleteStepAsync(job, 0, 0, JobStepOutcome.Failure);
-
-				IReadOnlyList<IIssue> issues = await IssueCollection.FindIssuesAsync();
-				Assert.AreEqual(1, issues.Count);
-				Assert.AreEqual("Gauntlet", issues[0].Fingerprints[0].Type);
-				Assert.AreEqual(new IssueKey("device:Foo", IssueKeyType.None), issues[0].Fingerprints[0].Keys.First());
-				Assert.AreEqual("Automation device warnings with Foo", issues[0].Summary);
-			}
-			// #3
-			// Scenario: Gauntlet build drop event with File and Directory property
-			// Expected: Gauntlet fingerprint using Access prefix
-			{
-				IJob job = CreateJob(_mainStreamId, 130, "Test Build", _graph);
-				await using (TestJsonLogger logger = await CreateLoggerAsync(job, 0, 0))
-				{
-					logger.LogError(KnownLogEvents.Gauntlet_BuildDropEvent, "    File {File} reported an issue", "/Bar/Foo.txt");
-					logger.LogError(KnownLogEvents.Gauntlet_BuildDropEvent, "    Folder {Directory} reported an issue", "/Bar/Foo");
-				}
-				await UpdateCompleteStepAsync(job, 0, 0, JobStepOutcome.Failure);
-
-				IReadOnlyList<IIssue> issues = await IssueCollection.FindIssuesAsync();
-				Assert.AreEqual(1, issues.Count);
-				Assert.AreEqual("Gauntlet", issues[0].Fingerprints[0].Type);
-				Assert.AreEqual(new IssueKey("access:/Bar/Foo.txt", IssueKeyType.None), issues[0].Fingerprints[0].Keys.First());
-				Assert.AreEqual("Automation access errors with /Bar/Foo.txt and with /Bar/Foo", issues[0].Summary);
-			}
-			// #4
 			// Scenario: Gauntlet Fatal event
 			// Expected: Gauntlet fingerprint using hash prefix
 			{
-				string logMessage =
-				  "    Engine encountered a critical failure.\n"
-				+ @"Assertion failed: State.bGfxPSOSet [File:D:\build\U5M+Inc\Sync\Engine\Source\Runtime\RHI\Public\RHIValidationContext.h] [Line: 809]"
-				+ @"A Graphics PSO has to be set to set resources into a shader!"
-				+ @"	0x00007fff4b43a1f4 UnrealEditor-RHI.dll!FValidationContext::RHISetShaderParameters() [Unknown File]"
-				+ @"	0x00007fff4b3e5a01 UnrealEditor-RHI.dll!FRHICommandSetShaderParameters<FRHIGraphicsShader>::Execute() [Unknown File]"
-				+ @"	0x00007fff4b3e9a5a UnrealEditor-RHI.dll!FRHICommand<FRHICommandSetShaderParameters<FRHIGraphicsShader>,FRHICommandSetShaderParametersString1159>::ExecuteAndDestruct() [Unknown File]"
-				+ @"	0x00007fff4b3e74f7 UnrealEditor - RHI.dll!FRHICommandListBase::Execute()[Unknown File]"
-				+ @"	0x00007fff4b3f3228 UnrealEditor-RHI.dll!FRHICommandListImmediate::ExecuteAndReset() [Unknown File]"
-				+ @"	0x00007fff4b464ebf UnrealEditor-RHI.dll!FRHIComputeCommandList::SubmitCommandsHint() [Unknown File]"
-				+ @"	0x00007fff3d3066fc UnrealEditor-Renderer.dll!TBaseStaticDelegateInstance<void [Unknown File]"
-				+ @"	0x00007fff48bd7813 UnrealEditor-RenderCore.dll!FRDGBuilder::ExecutePass() [Unknown File]"
-				+ @"	0x00007fff48bd2a8a UnrealEditor-RenderCore.dll!FRDGBuilder::Execute() [Unknown File]"
-				+ @"	0x00007fff3d31fef4 UnrealEditor-Renderer.dll!FSceneRenderer::RenderThreadEnd() [Unknown File]"
-				+ @"	0x00007fff3d2ed57a UnrealEditor-Renderer.dll!`FPixelShaderUtils::AddFullscreenPass<FHZBTestPS>'::`2'::<lambda_1>::operator()() [Unknown File]"
-				+ @"	0x00007fff3d3043f9 UnrealEditor - Renderer.dll!FSceneRenderer::DoOcclusionQueries()[Unknown File]"
-				+ @"	0x00007fff3d30c888 UnrealEditor-Renderer.dll!TBaseStaticDelegateInstance<void [Unknown File]"
-				+ @"	0x00007fff501fad42 UnrealEditor-Core.dll!FNamedTaskThread::ProcessTasksNamedThread() [Unknown File]"
-				+ @"	0x00007fff501fb25e UnrealEditor-Core.dll!FNamedTaskThread::ProcessTasksUntilQuit() [Unknown File]"
-				+ @"	0x00007fff48ccbb54 UnrealEditor-RenderCore.dll!RenderingThreadMain() [Unknown File]"
-				+ @"	0x00007fff48ccff44 UnrealEditor-RenderCore.dll!FRenderingThread::Run() [Unknown File]"
-				+ @"	0x00007fff5089c862 UnrealEditor-Core.dll!FRunnableThreadWin::Run() [Unknown File]"
-				+ @"	0x00007fff5089a7bf UnrealEditor-Core.dll!FRunnableThreadWin::GuardedRun() [Unknown File]"
-				+ @"	0x00007fff8e304ed0 KERNEL32.DLL!UnknownFunction [Unknown File]"
-				+ @"	0x00007fff8f26e39b ntdll.dll!UnknownFunction [Unknown File]"
-				+ @"    Test did not run until completion. The test exited prematurely.";
-
-				IJob job = CreateJob(_mainStreamId, 140, "Test Build", _graph);
+				string logMessage = "Assertion failed: State.bGfxPSOSet [File:D:\\build\\U5M+Inc\\Sync\\Engine\\Source\\Runtime\\RHI\\Public\\RHIValidationContext.h] [Line: 809]\n"
+				+ "A Graphics PSO has to be set to set resources into a shader!";
+				string logCallstack = "	0x00007fff4b43a1f4 UnrealEditor-RHI.dll!FValidationContext::RHISetShaderParameters() [Unknown File]\n"
+				+ "	0x00007fff4b3e5a01 UnrealEditor-RHI.dll!FRHICommandSetShaderParameters<FRHIGraphicsShader>::Execute() [Unknown File]\n"
+				+ "	0x00007fff4b3e9a5a UnrealEditor-RHI.dll!FRHICommand<FRHICommandSetShaderParameters<FRHIGraphicsShader>,FRHICommandSetShaderParametersString1159>::ExecuteAndDestruct() [Unknown File]\n"
+				+ "	0x00007fff4b3e74f7 UnrealEditor - RHI.dll!FRHICommandListBase::Execute()[Unknown File]\n"
+				+ "	0x00007fff4b3f3228 UnrealEditor-RHI.dll!FRHICommandListImmediate::ExecuteAndReset() [Unknown File]\n"
+				+ "	0x00007fff4b464ebf UnrealEditor-RHI.dll!FRHIComputeCommandList::SubmitCommandsHint() [Unknown File]\n"
+				+ "	0x00007fff3d3066fc UnrealEditor-Renderer.dll!TBaseStaticDelegateInstance<void [Unknown File]\n"
+				+ "	0x00007fff48bd7813 UnrealEditor-RenderCore.dll!FRDGBuilder::ExecutePass() [Unknown File]\n"
+				+ "	0x00007fff48bd2a8a UnrealEditor-RenderCore.dll!FRDGBuilder::Execute() [Unknown File]\n"
+				+ "	0x00007fff3d31fef4 UnrealEditor-Renderer.dll!FSceneRenderer::RenderThreadEnd() [Unknown File]\n"
+				+ "	0x00007fff3d2ed57a UnrealEditor-Renderer.dll!`FPixelShaderUtils::AddFullscreenPass<FHZBTestPS>'::`2'::<lambda_1>::operator()() [Unknown File]\n"
+				+ "	0x00007fff3d3043f9 UnrealEditor - Renderer.dll!FSceneRenderer::DoOcclusionQueries()[Unknown File]\n"
+				+ "	0x00007fff3d30c888 UnrealEditor-Renderer.dll!TBaseStaticDelegateInstance<void [Unknown File]\n"
+				+ "	0x00007fff501fad42 UnrealEditor-Core.dll!FNamedTaskThread::ProcessTasksNamedThread() [Unknown File]\n"
+				+ "	0x00007fff501fb25e UnrealEditor-Core.dll!FNamedTaskThread::ProcessTasksUntilQuit() [Unknown File]\n"
+				+ "	0x00007fff48ccbb54 UnrealEditor-RenderCore.dll!RenderingThreadMain() [Unknown File]\n"
+				+ "	0x00007fff48ccff44 UnrealEditor-RenderCore.dll!FRenderingThread::Run() [Unknown File]\n"
+				+ "	0x00007fff5089c862 UnrealEditor-Core.dll!FRunnableThreadWin::Run() [Unknown File]\n"
+				+ "	0x00007fff5089a7bf UnrealEditor-Core.dll!FRunnableThreadWin::GuardedRun() [Unknown File]\n"
+				+ "	0x00007fff8e304ed0 KERNEL32.DLL!UnknownFunction [Unknown File]\n"
+				+ "	0x00007fff8f26e39b ntdll.dll!UnknownFunction [Unknown File]\n";
+
+				IJob job = CreateJob(_mainStreamId, 120, "Test Build", _graph);
 				await using (TestJsonLogger logger = await CreateLoggerAsync(job, 0, 0))
 				{
-					logger.LogError(KnownLogEvents.Gauntlet_FatalEvent, "{Message}", logMessage);
+					logger.LogError(KnownLogEvents.Gauntlet_FatalEvent, "{Message}\n{Callstack}", logMessage, logCallstack);
 				}
 				await UpdateCompleteStepAsync(job, 0, 0, JobStepOutcome.Failure);
 
@@ -1842,11 +1804,12 @@ namespace Horde.Server.Tests.Issues
 				Assert.AreEqual(1, issues.Count);
 				Assert.AreEqual("Gauntlet", issues[0].Fingerprints[0].Type);
 				Assert.AreEqual("hash:", issues[0].Fingerprints[0].Keys.First().Name.Substring(0, 5));
+				Assert.AreEqual(37, issues[0].Fingerprints[0].Keys.First().Name.Length);
 				Assert.AreEqual("Automation fatal errors in Update Version Files", issues[0].Summary);
 			}
-			// #5
+			// #3
 			// Scenario: Gauntlet Test event
-			// Expected: Gauntlet fingerprint using hash prefix
+			// Expected: Gauntlet fingerprint using hash prefix and job step salt
 			{
 				string[] logErrors =
 				{
@@ -1862,7 +1825,7 @@ namespace Horde.Server.Tests.Issues
 				"	Expected 'SimpleValueWithMetaSkipDataZenAndDirect::Get meta equality' to be true."
 				};
 
-				IJob job = CreateJob(_mainStreamId, 150, "Test Build", _graph);
+				IJob job = CreateJob(_mainStreamId, 130, "Test Build", _graph);
 				await using (TestJsonLogger logger = await CreateLoggerAsync(job, 0, 0))
 				{
 					foreach (string error in logErrors)
@@ -1875,7 +1838,9 @@ namespace Horde.Server.Tests.Issues
 				IReadOnlyList<IIssue> issues = await IssueCollection.FindIssuesAsync();
 				Assert.AreEqual(1, issues.Count);
 				Assert.AreEqual("Gauntlet", issues[0].Fingerprints[0].Type);
+				Assert.AreEqual(10, issues[0].Fingerprints[0].Keys.Count);
 				Assert.AreEqual("hash:", issues[0].Fingerprints[0].Keys.First().Name.Substring(0, 5));
+				Assert.AreEqual($":{job.StreamId}:Update Version Files", issues[0].Fingerprints[0].Keys.First().Name.Substring(37));
 				Assert.AreEqual("Automation test errors in Update Version Files", issues[0].Summary);
 			}
 		}
diff --git a/Engine/Source/Programs/Shared/EpicGames.Horde/Issues/Handlers/GauntletIssueHandler.cs b/Engine/Source/Programs/Shared/EpicGames.Horde/Issues/Handlers/GauntletIssueHandler.cs
--- a/Engine/Source/Programs/Shared/EpicGames.Horde/Issues/Handlers/GauntletIssueHandler.cs
+++ b/Engine/Source/Programs/Shared/EpicGames.Horde/Issues/Handlers/GauntletIssueHandler.cs
@@ -1,6 +1,7 @@
 // Copyright Epic Games, Inc. All Rights Reserved.
 
 using System.Collections.Generic;
+using System.Linq;
 using System.Text;
 using System.Text.Json;
 using System.Text.RegularExpressions;
@@ -18,7 +19,7 @@ namespace EpicGames.Horde.Issues.Handlers
 		/// <summary>
 		/// Prefix for framework keys
 		/// </summary>
-		const string FrameworkPrefix = "framework";
+		const string FrameworkPrefix = "test framework";
 
 		/// <summary>
 		/// Prefix for test keys
@@ -31,15 +32,20 @@ namespace EpicGames.Horde.Issues.Handlers
 		const string DevicePrefix = "device";
 
 		/// <summary>
-		/// Prefix for access keys
+		/// Prefix for build drop keys
 		/// </summary>
-		const string AccessPrefix = "access";
+		const string BuildDropPrefix = "build drop";
 
 		/// <summary>
 		/// Prefix for fatal failure keys
 		/// </summary>
 		const string FatalPrefix = "fatal";
 
+		/// <summary>
+		/// Callstack log type property
+		/// </summary>
+		static readonly Utf8String s_callstackLogType = new Utf8String("Callstack");
+
 		/// <summary>
 		/// Max Message Length to hash
 		/// </summary>
@@ -57,7 +63,7 @@ namespace EpicGames.Horde.Issues.Handlers
 			{ KnownLogEvents.Gauntlet_TestEvent, TestPrefix},
 			{ KnownLogEvents.Gauntlet_DeviceEvent, DevicePrefix},
 			{ KnownLogEvents.Gauntlet_UnrealEngineTestEvent, TestPrefix},
-			{ KnownLogEvents.Gauntlet_BuildDropEvent, AccessPrefix},
+			{ KnownLogEvents.Gauntlet_BuildDropEvent, BuildDropPrefix},
 			{ KnownLogEvents.Gauntlet_FatalEvent, FatalPrefix}
 		};
 
@@ -86,94 +92,36 @@ namespace EpicGames.Horde.Issues.Handlers
 			return s_knownGauntletEvents[eventId];
 		}
 
-		/// <summary>
-		/// Parses symbol names from a log event
-		/// </summary>
-		/// <param name="issueEvent">The log event data</param>
-		/// <param name="testNames">Receives a set of the test names</param>
-		/// <param name="metadata"></param>
-		private static void GetNames(IssueEvent issueEvent, HashSet<IssueKey> testNames, HashSet<IssueMetadata> metadata)
-		{
-			foreach (JsonLogEvent line in issueEvent.Lines)
-			{
-				JsonDocument document = JsonDocument.Parse(line.Data);
-
-				string? name = null;
-
-				string? value;
-				if (document.RootElement.TryGetNestedProperty("properties.Name", out value))
-				{
-					name = value;
-				}
-
-				if (name != null)
-				{
-					string prefix = GetEventPrefix(issueEvent.EventId!.Value);
-					testNames.Add($"{prefix}:{name}", IssueKeyType.None);
-					metadata.Add("Context", $"with {name}");
-				}
-			}
-		}
-
-		/// <summary>
-		/// Parses symbol file or directory from a log event
-		/// </summary>
-		/// <param name="issueEvent">The log event data</param>
-		/// <param name="paths">Receives a set of the paths</param>
-		/// <param name="metadata"></param>
-		private static void GetPaths(IssueEvent issueEvent, HashSet<IssueKey> paths, HashSet<IssueMetadata> metadata)
-		{
-			if (issueEvent.EventId == KnownLogEvents.Gauntlet_BuildDropEvent)
-			{
-				foreach (JsonLogEvent line in issueEvent.Lines)
-				{
-					JsonDocument document = JsonDocument.Parse(line.Data);
-
-					string? path = null;
-
-					string? value;
-					if (document.RootElement.TryGetNestedProperty("properties.File", out value))
-					{
-						path = value;
-					}
-					else if (document.RootElement.TryGetNestedProperty("properties.Directory", out value))
-					{
-						path = value;
-					}
-
-					if (path != null)
-					{
-						paths.Add($"{AccessPrefix}:{path}", IssueKeyType.None);
-						metadata.Add("Context", $"with {path}");
-					}
-				}
-			}
-		}
-
 		/// <summary>
 		/// Produce a hash from error message
 		/// </summary>
-		/// <param name="message">The log event message</param>
+		/// <param name="issueEvent">The issue event</param>
 		/// <param name="keys">Receives a set of the keys</param>
 		/// <param name="metadata">Receives a set of metadata</param>
-		private void GetHash(string message, HashSet<IssueKey> keys, HashSet<IssueMetadata> metadata)
+		private void GetHash(IssueEvent issueEvent, HashSet<IssueKey> keys, HashSet<IssueMetadata> metadata)
 		{
-			string error = message.Length > MaxMessageLength ? message.Substring(0, MaxMessageLength) : message;
-
-			if (TryGetHash(error, out Md5Hash hash))
+			if (TryGetHash(issueEvent, out Md5Hash hash))
 			{
-				keys.Add($"hash:{hash}:stream:{_context.StreamId}", IssueKeyType.None);
+				string key = $"hash:{hash}";
+				if (!EventHasCallstackProperty(issueEvent))
+				{
+					// add job step salt if no Callstack property was found
+					key += $":{_context.StreamId}:{_context.NodeName}";
+				}
+				keys.Add(key, IssueKeyType.None);
 			}
 			else
 			{
-				keys.Add($"{_context.NodeName}", IssueKeyType.None);
+				// Not enough information, make it an issue associated with only the job step
+				keys.Add($"{_context.StreamId}:{_context.NodeName}", IssueKeyType.None);
 			}
-			metadata.Add("Context", $"in {_context.NodeName}");
+			metadata.Add("Node", _context.NodeName);
 		}
 
-		private static bool TryGetHash(string message, out Md5Hash hash)
+		private static bool TryGetHash(IssueEvent issueEvent, out Md5Hash hash)
 		{
-			string sanitized = message.ToUpperInvariant();
+			string sanitized = issueEvent.Message.ToUpperInvariant();
+			sanitized = sanitized.Length > MaxMessageLength ? sanitized.Substring(0, MaxMessageLength) : sanitized;
 			sanitized = Regex.Replace(sanitized, @"(?<![a-zA-Z])(?:[A-Z]:|/)[^ :]+[/\\]SYNC[/\\]", "{root}/"); // Redact things that look like workspace roots; may be different between agents
 			sanitized = Regex.Replace(sanitized, @"0[xX][0-9a-fA-F]+", "H"); // Redact hex strings
 			sanitized = Regex.Replace(sanitized, @"\d[\d.,:]*", "n"); // Redact numbers and timestamp like things
@@ -190,22 +138,38 @@ namespace EpicGames.Horde.Issues.Handlers
 			}
 		}
 
+		private static bool EventHasCallstackProperty(IssueEvent issueEvent)
+		{
+			return issueEvent.Lines.Any(x => FindNestedPropertyOfType(x, s_callstackLogType) != null);
+		}
+
+		private static JsonProperty? FindNestedPropertyOfType(JsonLogEvent logEvent, Utf8String type)
+		{
+			JsonElement line = JsonDocument.Parse(logEvent.Data).RootElement;
+			JsonElement properties;
+			if (line.TryGetProperty("properties", out properties) && properties.ValueKind == JsonValueKind.Object)
+			{
+				foreach (JsonProperty property in properties.EnumerateObject())
+				{
+					if (property.NameEquals(type.Span))
+					{
+						return property;
+					}
+				}
+			}
+
+			return null;
+		}
+
 		/// <inheritdoc/>
 		public override bool HandleEvent(IssueEvent issueEvent)
 		{
 			if (issueEvent.EventId != null && IsMatchingEventId(issueEvent.EventId.Value))
 			{
-				IssueEventGroup issue = new IssueEventGroup("Gauntlet", "Automation {Meta:Type} {Severity} {Meta:Context}", IssueChangeFilter.Code);
+				IssueEventGroup issue = new IssueEventGroup("Gauntlet", "Automation {Meta:GauntletType} {Severity} in {Meta:Node}", IssueChangeFilter.All);
 				issue.Events.Add(issueEvent);
-
-				GetNames(issueEvent, issue.Keys, issue.Metadata);
-				GetPaths(issueEvent, issue.Keys, issue.Metadata);
-				if (issue.Keys.Count == 0)
-				{
-					GetHash(issueEvent.Message, issue.Keys, issue.Metadata);
-				}
-
-				issue.Metadata.Add(new IssueMetadata("type", GetEventPrefix(issueEvent.EventId.Value)));
+				GetHash(issueEvent, issue.Keys, issue.Metadata);
+				issue.Metadata.Add("GauntletType", GetEventPrefix(issueEvent.EventId.Value));
 				_issues.Add(issue);
 
 				return true;
