Horde: Fix log events not being tagged with the appropriate issue.

diff --git a/Engine/Source/Programs/Horde/Horde.Server/Logs/LogCollection.cs b/Engine/Source/Programs/Horde/Horde.Server/Logs/LogCollection.cs
--- a/Engine/Source/Programs/Horde/Horde.Server/Logs/LogCollection.cs
+++ b/Engine/Source/Programs/Horde/Horde.Server/Logs/LogCollection.cs
@@ -176,6 +176,8 @@ namespace Horde.Server.Logs
 			readonly LogCollection _collection;
 			readonly LogEventDocument _document;
 
+			public LogEventDocument Document => _document;
+
 			LogId ILogEvent.LogId => _document.Id.LogId;
 			LogEventSeverity ILogEvent.Severity => _document.IsWarning ? LogEventSeverity.Warning : LogEventSeverity.Error;
 			int ILogEvent.LineIndex => _document.Id.LineIndex;
@@ -722,7 +724,7 @@ namespace Horde.Server.Logs
 		/// <inheritdoc/>
 		public async Task AddSpanToEventsAsync(IEnumerable<ILogEvent> events, ObjectId spanId, CancellationToken cancellationToken)
 		{
-			FilterDefinition<LogEventDocument> eventFilter = Builders<LogEventDocument>.Filter.In(x => x.Id, events.OfType<LogEventDocument>().Select(x => x.Id));
+			FilterDefinition<LogEventDocument> eventFilter = Builders<LogEventDocument>.Filter.In(x => x.Id, events.Select(x => ((LogEvent)x).Document.Id));
 			UpdateDefinition<LogEventDocument> eventUpdate = Builders<LogEventDocument>.Update.Set(x => x.SpanId, spanId);
 			await _logEvents.UpdateManyAsync(eventFilter, eventUpdate, cancellationToken: cancellationToken);
 		}
