Horde: Filter out missing nodes from dependency information returned in job steps.

diff --git a/Engine/Source/Programs/Horde/Horde.Server/Jobs/JobCollection.cs b/Engine/Source/Programs/Horde/Horde.Server/Jobs/JobCollection.cs
--- a/Engine/Source/Programs/Horde/Horde.Server/Jobs/JobCollection.cs
+++ b/Engine/Source/Programs/Horde/Horde.Server/Jobs/JobCollection.cs
@@ -2,6 +2,7 @@
 
 using System;
 using System.Collections.Generic;
+using System.Diagnostics;
 using System.Diagnostics.CodeAnalysis;
 using System.IO;
 using System.Linq;
@@ -165,8 +166,6 @@ namespace Horde.Server.Jobs
 				_batches = document.Batches.ConvertAll(x => new JobStepBatch(this, x, graph.Groups[x.GroupIdx]));
 			}
 
-			public JobStepId GetStepId(NodeRef nodeRef) => NodeRefToStepId[nodeRef];
-
 			public Task<IJob?> RefreshAsync(CancellationToken cancellationToken = default)
 				=> _collection.GetAsync(_document.Id, cancellationToken);
 
@@ -234,6 +233,7 @@ namespace Horde.Server.Jobs
 			}
 		}
 
+		[DebuggerDisplay("{Document.Id} (Group {Document.GroupIdx})")]
 		class JobStepBatch : IJobStepBatch
 		{
 			public Job Job { get; }
@@ -271,6 +271,7 @@ namespace Horde.Server.Jobs
 			}
 		}
 
+		[DebuggerDisplay("{Document.Id} (Node {Batch.Document.GroupIdx},{Document.NodeIdx})")]
 		class JobStep : IJobStep
 		{
 			public Job Job => Batch.Job;
@@ -286,9 +287,14 @@ namespace Horde.Server.Jobs
 
 			// Node properties
 			string IJobStep.Name => Node.Name;
-			IReadOnlyList<JobStepOutputRef> IJobStep.Inputs => Node.Inputs.ConvertAll(x => new JobStepOutputRef(Job.GetStepId(x.NodeRef), x.OutputIdx));
+
+			List<JobStepOutputRef>? _inputs;
+			IReadOnlyList<JobStepOutputRef> IJobStep.Inputs => _inputs ??= CreateInputDependencyList();
+
 			IReadOnlyList<string> IJobStep.OutputNames => Node.OutputNames;
-			IReadOnlyList<JobStepId> IJobStep.InputDependencies => Node.InputDependencies.ConvertAll(x => Job.GetStepId(x));
+
+			List<JobStepId>? _inputDependencies;
+			IReadOnlyList<JobStepId> IJobStep.InputDependencies => _inputDependencies ??= CreateInputDependenciesList();
 
 			List<JobStepId>? _orderDependencies;
 			IReadOnlyList<JobStepId> IJobStep.OrderDependencies => _orderDependencies ??= CreateOrderDependenciesList();
@@ -322,6 +328,40 @@ namespace Horde.Server.Jobs
 				Node = node;
 			}
 
+			List<JobStepOutputRef> CreateInputDependencyList()
+			{
+				List<JobStepOutputRef> dependencies = new List<JobStepOutputRef>();
+				foreach (NodeOutputRef input in Node.Inputs)
+				{
+					if (Job.NodeRefToStepId.TryGetValue(input.NodeRef, out JobStepId jobStepId))
+					{
+						dependencies.Add(new JobStepOutputRef(jobStepId, input.OutputIdx));
+					}
+					else
+					{
+						MissingNodeRef(input.NodeRef);
+					}
+				}
+				return dependencies;
+			}
+
+			List<JobStepId> CreateInputDependenciesList()
+			{
+				List<JobStepId> dependencies = new List<JobStepId>();
+				foreach (NodeRef nodeRef in Node.InputDependencies)
+				{
+					if (Job.NodeRefToStepId.TryGetValue(nodeRef, out JobStepId jobStepId))
+					{
+						dependencies.Add(jobStepId);
+					}
+					else
+					{
+						MissingNodeRef(nodeRef);
+					}
+				}
+				return dependencies;
+			}
+
 			List<JobStepId> CreateOrderDependenciesList()
 			{
 				List<JobStepId> dependencies = new List<JobStepId>();
@@ -331,9 +371,18 @@ namespace Horde.Server.Jobs
 					{
 						dependencies.Add(jobStepId);
 					}
+					else
+					{
+						MissingNodeRef(nodeRef);
+					}
 				}
 				return dependencies;
 			}
+
+			void MissingNodeRef(NodeRef nodeRef)
+			{
+				_ = nodeRef;
+			}
 		}
 
 		class JobDocument
