Horde: Updating schema documentation

#rnx

diff --git a/Engine/Source/Programs/Horde/Docs/Deployment/ServerSettings.md b/Engine/Source/Programs/Horde/Docs/Deployment/ServerSettings.md
--- a/Engine/Source/Programs/Horde/Docs/Deployment/ServerSettings.md
+++ b/Engine/Source/Programs/Horde/Docs/Deployment/ServerSettings.md
@@ -14,10 +14,13 @@ Name | Description
 `http2Port` | `integer`<br>Dedicated port for serving only HTTP/2.
 `computeTunnelPort` | `integer`<br>Port to listen on for tunneling compute sockets to agents
 `computeTunnelAddress` | `string`<br>What address (host:port) clients should connect to for compute socket tunneling Port may differ from  if Horde server is behind a reverse proxy/firewall
+`mongoConnectionString` | `string`<br>Connection string for the Mongo database
 `databaseConnectionString` | `string`<br>MongoDB connection string
-`databaseName` | `string`<br>MongoDB database name
+`mongoDatabaseName` | `string`<br>MongoDB database name
+`databaseName` | `string`<br>
 `databasePublicCert` | `string`<br>Optional certificate to trust in order to access the database (eg. AWS public cert for TLS)
-`databaseReadOnlyMode` | `boolean`<br>Access the database in read-only mode (avoids creating indices or updating content) Useful for debugging a local instance of HordeServer against a production database.
+`mongoReadOnlyMode` | `boolean`<br>Access the database in read-only mode (avoids creating indices or updating content) Useful for debugging a local instance of HordeServer against a production database.
+`databaseReadOnlyMode` | `boolean`<br>
 `shutdownMemoryThreshold` | `integer`<br>Shutdown the current server process if memory usage reaches this threshold (specified in MB)<br>Usually set to 80-90% of available memory to avoid CLR heap using all of it. If a memory leak was to occur, it's usually better to restart the process rather than to let the GC work harder and harder trying to recoup memory.<br>Should only be used when multiple server processes are running behind a load balancer and one can be safely restarted automatically by the underlying process handler (Docker, Kubernetes, AWS ECS, Supervisor etc). The shutdown behaves similar to receiving a SIGTERM and will wait for outstanding requests to finish.
 `serverPrivateCert` | `string`<br>Optional PFX certificate to use for encrypting agent SSL traffic. This can be a self-signed certificate, as long as it's trusted by agents.
 `authMethod` | [`AuthMethod`](#authmethod-enum)<br>Issuer for tokens from the auth provider
@@ -46,7 +49,9 @@ Name | Description
 `noResourceBackOffTime` | `string`<br>Interval between polling for new jobs
 `initiateJobBackOffTime` | `string`<br>Interval between attempting to assign agents to take on jobs
 `unknownErrorBackOffTime` | `string`<br>Interval between scheduling jobs when an unknown error occurs
-`redisConnectionConfig` | `string`<br>Config for connecting to Redis server(s). Setting it to null will disable Redis use and connection See format at https://stackexchange.github.io/StackExchange.Redis/Configuration.html
+`redisConnectionString` | `string`<br>Config for connecting to Redis server(s). Setting it to null will disable Redis use and connection See format at https://stackexchange.github.io/StackExchange.Redis/Configuration.html
+`redisConnectionConfig` | `string`<br>
+`redisReadOnlyMode` | `boolean`<br>Whether to disable writes to Redis.
 `logServiceWriteCacheType` | `string`<br>Type of write cache to use in log service Currently Supported: "InMemory" or "Redis"
 `logStorage` | [`StorageBackendOptions`](#storagebackendoptions)<br>Settings for artifact storage
 `artifactStorage` | [`StorageBackendOptions`](#storagebackendoptions)<br>Settings for artifact storage
