Horde: Explain role of the Agent.json file in docs. Also update some references to deprecated properties.

diff --git a/Engine/Source/Programs/Horde/Docs/Deployment/Agent.md b/Engine/Source/Programs/Horde/Docs/Deployment/Agent.md
--- a/Engine/Source/Programs/Horde/Docs/Deployment/Agent.md
+++ b/Engine/Source/Programs/Horde/Docs/Deployment/Agent.md
@@ -41,7 +41,11 @@ The -H parameter and value are not required when using an unauthenticated server
 
 ### General
 
-Agent settings are configured through the [`appsettings.json`](AgentSettings.md) file in the server directory.
+Agent settings are configured through the [`Agent.json`](AgentSettings.md) file. On Windows, this file is stored
+at `C:\ProgramData\Epic\Horde\Agent\Agent.json`. On other platforms, it is stored in the `Data` folder under the
+application directory by default. Settings in this file are applied on top of the `appsettings.json` file distributed
+alongside the agent executable.
+
 All Horde-specific settings are stored under the `horde` top-level key, with middleware and standard .NET settings
 under other root keys.
 
diff --git a/Engine/Source/Programs/Horde/Docs/Deployment/Server.md b/Engine/Source/Programs/Horde/Docs/Deployment/Server.md
--- a/Engine/Source/Programs/Horde/Docs/Deployment/Server.md
+++ b/Engine/Source/Programs/Horde/Docs/Deployment/Server.md
@@ -85,7 +85,7 @@ install all the prerequisites using [Homebrew](https://brew.sh/).
    instead if you prefer.
 
     ```bash
-    export Horde__DatabaseConnectionString=mongodb://localhost:27017
+    export Horde__MongoConnectionString=mongodb://localhost:27017
     export Horde__HttpPort=37107
     export Horde__Http2Port=37109
 
@@ -136,7 +136,7 @@ As an ASP.NET application, Horde's application configuration supports the follow
 
 * Individual properties can be overridden through **environment variables** using standard ASP.NET syntax (see
   [MSDN](https://learn.microsoft.com/en-us/aspnet/core/fundamentals/configuration/?view=aspnetcore-7.0#naming-of-environment-variables)).
-  For example, the database connection string can be passed in using the `Horde__DatabaseConnectionString` environment variable.
+  For example, the database connection string can be passed in using the `Horde__MongoConnectionString` environment variable.
 * The deployment environment can be configured using the ASPNETCORE_ENVIRONMENT environment variable. Standard values
   for Horde are `Production`, `Development`, and `Local`.
 * A deployment-specific configuration file can be created called `appsettings.{Environment}.json` (e.g.
@@ -150,8 +150,8 @@ more information.
 
 ### MongoDB
 
-The MongoDB connection string can be specified via the `DatabaseConnectionString` property in the
-[Server.json](ServerSettings.md) file or the `Horde__DatabaseConnectionString` environment variable. The
+The MongoDB connection string can be specified via the `MongoConnectionString` property in the
+[Server.json](ServerSettings.md) file or the `Horde__MongoConnectionString` environment variable. The
 connection string should be in standard
 [MongoDB syntax](https://www.mongodb.com/docs/manual/reference/connection-string/), e.g.:
 
@@ -164,7 +164,7 @@ to use the primary database instance using the `readPreference=primary` argument
 secondary instance for reads can cause deadlocks because the server gets out-of-date documents in a read-modify-write
 cycle.
 
-The MongoDB connection can be configured to use a trusted set of certificates via the `DatabasePublicCert` property.
+The MongoDB connection can be configured to use a trusted set of certificates via the `MongoPublicCertificate` property.
 For example, when running on AWS using DocumentDB, this property can be set to use Amazon's
 [combined certificate bundle](https://docs.aws.amazon.com/AmazonRDS/latest/UserGuide/UsingWithRDS.SSL.html) by
 placing the `global-bundle.pem` file into the server's application directory.
@@ -172,7 +172,7 @@ placing the `global-bundle.pem` file into the server's application directory.
 ### Redis
 
 The Redis server is configured through the RedisConnectionConfig property in the [Server.json](ServerSettings.md)
-file or via the `Horde__DatabaseConnectionString` environment variable. This string is formatted as a plain server
+file or via the `Horde__MongoConnectionString` environment variable. This string is formatted as a plain server
 and port, e.g.:
 
 ```text
diff --git a/Engine/Source/Programs/Horde/Horde.Server/Commands/Generate/ConfigDocsCommand.cs b/Engine/Source/Programs/Horde/Horde.Server/Commands/Generate/ConfigDocsCommand.cs
--- a/Engine/Source/Programs/Horde/Horde.Server/Commands/Generate/ConfigDocsCommand.cs
+++ b/Engine/Source/Programs/Horde/Horde.Server/Commands/Generate/ConfigDocsCommand.cs
@@ -69,7 +69,7 @@ namespace Horde.Server.Commands.Generate
 				Type agentSettingsType = agentAssembly.GetType("Horde.Agent.AgentSettings")!;
 
 				JsonSchema agentSchema = JsonSchema.FromType(agentSettingsType, xmlDocReader);
-				await WriteDocAsync(agentSchema.RootType, "appsettings.json (Agent)", "Deployment/AgentSettings.md", "[Horde](../../README.md) > [Deployment](../Deployment.md) > [Agent](Agent.md)", AppSettingsIntro, new Dictionary<string, string>(), logger);
+				await WriteDocAsync(agentSchema.RootType, "Agent.json (Agent)", "Deployment/AgentSettings.md", "[Horde](../../README.md) > [Deployment](../Deployment.md) > [Agent](Agent.md)", AppSettingsIntro, new Dictionary<string, string>(), logger);
 			}
 
 			Dictionary<string, string> typeNameToLink = typeToPageInfo.ToDictionary(x => x.Key.Name!, x => Path.GetFileName(x.Value.FileName), StringComparer.Ordinal);
diff --git a/Engine/Source/Programs/Horde/Horde.Server/ServerSettings.cs b/Engine/Source/Programs/Horde/Horde.Server/ServerSettings.cs
--- a/Engine/Source/Programs/Horde/Horde.Server/ServerSettings.cs
+++ b/Engine/Source/Programs/Horde/Horde.Server/ServerSettings.cs
@@ -462,7 +462,15 @@ namespace Horde.Server
 		/// <summary>
 		/// Optional certificate to trust in order to access the database (eg. AWS public cert for TLS)
 		/// </summary>
-		public string? DatabasePublicCert { get; set; }
+		public string? MongoPublicCertificate { get; set; }
+
+		/// <inheritdoc cref="MongoPublicCertificate"/>
+		[Obsolete("Replace DatabasePublicCert with MongoPublicCertificate")]
+		public string? DatabasePublicCert
+		{
+			get => MongoPublicCertificate;
+			set => MongoPublicCertificate = value;
+		}
 
 		/// <summary>
 		/// Access the database in read-only mode (avoids creating indices or updating content)
