Horde: Fix deprecation warning.

#rnx

diff --git a/Engine/Source/Programs/Horde/Horde.Server/Server/MongoService.cs b/Engine/Source/Programs/Horde/Horde.Server/Server/MongoService.cs
--- a/Engine/Source/Programs/Horde/Horde.Server/Server/MongoService.cs
+++ b/Engine/Source/Programs/Horde/Horde.Server/Server/MongoService.cs
@@ -288,14 +288,14 @@ namespace Horde.Server.Server
 			try
 			{
 				ReadOnlyMode = Settings.MongoReadOnlyMode;
-				if (Settings.DatabasePublicCert != null)
+				if (Settings.MongoPublicCertificate != null)
 				{
 					X509Store localTrustStore = new X509Store(StoreName.Root);
 					try
 					{
 						localTrustStore.Open(OpenFlags.ReadWrite);
 
-						X509Certificate2Collection collection = ImportCertificateBundle(Settings.DatabasePublicCert);
+						X509Certificate2Collection collection = ImportCertificateBundle(Settings.MongoPublicCertificate);
 						foreach (X509Certificate2 certificate in collection)
 						{
 							_logger.LogInformation("Importing certificate for {Subject}", certificate.Subject);
