Horde: Fix exception running server with no streams defined.

#rnx

diff --git a/Engine/Source/Programs/Horde/Horde.Server/Issues/IssueTagService.cs b/Engine/Source/Programs/Horde/Horde.Server/Issues/IssueTagService.cs
--- a/Engine/Source/Programs/Horde/Horde.Server/Issues/IssueTagService.cs
+++ b/Engine/Source/Programs/Horde/Horde.Server/Issues/IssueTagService.cs
@@ -69,25 +69,27 @@ namespace Horde.Server.Issues
 			State initialState = await _state.GetAsync(cancellationSource.Token);
 
 			GlobalConfig globalConfig = _globalConfig.CurrentValue;
-
-			List<Task> tasks = new List<Task>();
-			try
-			{
-				foreach (StreamConfig streamConfig in globalConfig.Streams)
-				{
-					tasks.Add(Task.Run(() => TickStreamGuardedAsync(streamConfig, initialState, cancellationSource.Token), cancellationSource.Token));
-				}
-				await Task.WhenAny(tasks);
-			}
-			finally
+			if (globalConfig.Streams.Count > 0)
 			{
+				List<Task> tasks = new List<Task>();
 				try
 				{
-					await cancellationSource.CancelAsync();
-					await Task.WhenAll(tasks);
+					foreach (StreamConfig streamConfig in globalConfig.Streams)
+					{
+						tasks.Add(Task.Run(() => TickStreamGuardedAsync(streamConfig, initialState, cancellationSource.Token), cancellationSource.Token));
+					}
+					await Task.WhenAny(tasks);
 				}
-				catch
+				finally
 				{
+					try
+					{
+						await cancellationSource.CancelAsync();
+						await Task.WhenAll(tasks);
+					}
+					catch
+					{
+					}
 				}
 			}
 		}
