Horde: Store agent CPU/RAM usage telemetry in an indexed collection.

diff --git a/Engine/Source/Programs/Horde/Horde.Agent/Services/TelemetryService.cs b/Engine/Source/Programs/Horde/Horde.Agent/Services/TelemetryService.cs
--- a/Engine/Source/Programs/Horde/Horde.Agent/Services/TelemetryService.cs
+++ b/Engine/Source/Programs/Horde/Horde.Agent/Services/TelemetryService.cs
@@ -494,40 +494,57 @@ class TelemetryService : BackgroundService
 		{
 			_logger.LogDebug("Sending telemetry events to server...");
 
-			RpcSendTelemetryEventsRequest request = new();
-			Timestamp utcNow = Timestamp.FromDateTime(DateTime.UtcNow);
-			RpcExecutionMetadata em = new()
-			{
-				LeaseId = _jobHandler.CurrentLeaseId.ToString(),
-				JobId = _jobHandler.CurrentJobId,
-				JobBatchId = _jobHandler.CurrentBatchId,
-			};
+			CpuMetrics cpuMetrics = _systemMetrics.GetCpu();
+			MemoryMetrics memMetrics = _systemMetrics.GetMemory();
 
+			// Mongo method
 			{
-				RpcAgentCpuMetricsEvent cpuMetricsEvent = _systemMetrics.GetCpu().ToEvent();
-				cpuMetricsEvent.AgentId = _agentMetadataEvent.AgentId;
-				cpuMetricsEvent.Timestamp = utcNow;
-				cpuMetricsEvent.ExecutionMetadata = em;
-				request.Events.Add(new RpcWrappedTelemetryEvent { Cpu = cpuMetricsEvent });
+				RpcUploadTelemetryRequest request = new RpcUploadTelemetryRequest();
+				request.UserCpu = cpuMetrics.User;
+				request.SystemCpu = cpuMetrics.System;
+				request.IdleCpu = cpuMetrics.Idle;
+				request.TotalRam = memMetrics.Total / 1024;
+				request.FreeRam = memMetrics.Available / 1024;
+				request.UsedRam = memMetrics.Used / 1024;
 			}
 
+			// Clickhouse method
 			{
-				RpcAgentMemoryMetricsEvent memMetricsEvent = _systemMetrics.GetMemory().ToEvent();
-				memMetricsEvent.AgentId = _agentMetadataEvent.AgentId;
-				memMetricsEvent.Timestamp = utcNow;
-				memMetricsEvent.ExecutionMetadata = em;
-				request.Events.Add(new RpcWrappedTelemetryEvent { Mem = memMetricsEvent });
-			}
+				RpcSendTelemetryEventsRequest request = new();
+				Timestamp utcNow = Timestamp.FromDateTime(DateTime.UtcNow);
+				RpcExecutionMetadata em = new()
+				{
+					LeaseId = _jobHandler.CurrentLeaseId.ToString(),
+					JobId = _jobHandler.CurrentJobId,
+					JobBatchId = _jobHandler.CurrentBatchId,
+				};
 
-			if (DateTime.UtcNow > _lastTimeAgentMetadataSent + _agentMetadataReportInterval)
-			{
-				// Report agent metadata every now and then as events are not guaranteed to be delivered.
-				// Re-sending ensures the metadata will eventually make it to the server.
-				request.Events.Add(new RpcWrappedTelemetryEvent { AgentMetadata = _agentMetadataEvent });
-				_lastTimeAgentMetadataSent = DateTime.UtcNow;
-			}
+				{
+					RpcAgentCpuMetricsEvent cpuMetricsEvent = cpuMetrics.ToEvent();
+					cpuMetricsEvent.AgentId = _agentMetadataEvent.AgentId;
+					cpuMetricsEvent.Timestamp = utcNow;
+					cpuMetricsEvent.ExecutionMetadata = em;
+					request.Events.Add(new RpcWrappedTelemetryEvent { Cpu = cpuMetricsEvent });
+				}
 
-			await client.SendTelemetryEventsAsync(request, new CallOptions(cancellationToken: stoppingToken));
+				{
+					RpcAgentMemoryMetricsEvent memMetricsEvent = memMetrics.ToEvent();
+					memMetricsEvent.AgentId = _agentMetadataEvent.AgentId;
+					memMetricsEvent.Timestamp = utcNow;
+					memMetricsEvent.ExecutionMetadata = em;
+					request.Events.Add(new RpcWrappedTelemetryEvent { Mem = memMetricsEvent });
+				}
+
+				if (DateTime.UtcNow > _lastTimeAgentMetadataSent + _agentMetadataReportInterval)
+				{
+					// Report agent metadata every now and then as events are not guaranteed to be delivered.
+					// Re-sending ensures the metadata will eventually make it to the server.
+					request.Events.Add(new RpcWrappedTelemetryEvent { AgentMetadata = _agentMetadataEvent });
+					_lastTimeAgentMetadataSent = DateTime.UtcNow;
+				}
+
+				await client.SendTelemetryEventsAsync(request, new CallOptions(cancellationToken: stoppingToken));
+			}
 			await Task.Delay(_reportInterval, stoppingToken);
 		}
 
diff --git a/Engine/Source/Programs/Horde/Horde.Server/Agents/AgentsController.cs b/Engine/Source/Programs/Horde/Horde.Server/Agents/AgentsController.cs
--- a/Engine/Source/Programs/Horde/Horde.Server/Agents/AgentsController.cs
+++ b/Engine/Source/Programs/Horde/Horde.Server/Agents/AgentsController.cs
@@ -33,6 +33,7 @@ namespace Horde.Server.Agents
 	public class AgentsController : HordeControllerBase
 	{
 		readonly AgentService _agentService;
+		readonly IAgentTelemetryCollection _agentTelemetryCollection;
 		readonly IUserCollection _userCollection;
 		readonly IOptionsSnapshot<GlobalConfig> _globalConfig;
 		readonly ILogger<AgentsController> _logger;
@@ -40,9 +41,10 @@ namespace Horde.Server.Agents
 		/// <summary>
 		/// Constructor
 		/// </summary>
-		public AgentsController(AgentService agentService, IUserCollection userCollection, IOptionsSnapshot<GlobalConfig> globalConfig, ILogger<AgentsController> logger)
+		public AgentsController(AgentService agentService, IAgentTelemetryCollection agentTelemetryCollection, IUserCollection userCollection, IOptionsSnapshot<GlobalConfig> globalConfig, ILogger<AgentsController> logger)
 		{
 			_agentService = agentService;
+			_agentTelemetryCollection = agentTelemetryCollection;
 			_userCollection = userCollection;
 			_globalConfig = globalConfig;
 			_logger = logger;
@@ -361,6 +363,29 @@ namespace Horde.Server.Agents
 			return CreateGetAgentSessionResponse(session);
 		}
 
+		/// <summary>
+		/// Finds telemetry for an agent within a time range
+		/// </summary>
+		/// <param name="agentId">Agent identifier</param>
+		/// <param name="minTime">Minimum time to return</param>
+		/// <param name="maxTime">Maximum time to return</param>
+		[HttpGet]
+		[Route("/api/v1/agents/{agentId}/telemetry")]
+		public async Task<ActionResult<GetAgentTelemetryResponse>> FindTelemetryAsync(AgentId agentId, [FromQuery] DateTime minTime, [FromQuery] DateTime maxTime)
+		{
+			if (!_globalConfig.Value.Authorize(AgentAclAction.ViewAgent, User))
+			{
+				return Forbid(AgentAclAction.ViewAgent, agentId);
+			}
+
+			IReadOnlyList<IAgentTelemetry> telemetry = await _agentTelemetryCollection.FindAsync(agentId, minTime, maxTime, HttpContext.RequestAborted);
+
+			List<GetAgentTelemetrySampleResponse> samples = new List<GetAgentTelemetrySampleResponse>();
+			samples.AddRange(telemetry.Select(x => new GetAgentTelemetrySampleResponse(x.TimeUtc, x.UserCpu, x.IdleCpu, x.SystemCpu, x.FreeRam, x.UsedRam, x.TotalRam)));
+
+			return new GetAgentTelemetryResponse(samples);
+		}
+
 		/// <summary>
 		/// Find all the leases for a particular agent
 		/// </summary>
diff --git a/Engine/Source/Programs/Horde/Horde.Server/Agents/Telemetry/AgentTelemetryCollection.cs b/Engine/Source/Programs/Horde/Horde.Server/Agents/Telemetry/AgentTelemetryCollection.cs
new file mode 100644
--- /dev/null
+++ b/Engine/Source/Programs/Horde/Horde.Server/Agents/Telemetry/AgentTelemetryCollection.cs
@@ -0,0 +1,136 @@
+// Copyright Epic Games, Inc. All Rights Reserved.
+
+using System;
+using System.Collections.Generic;
+using System.Linq;
+using System.Threading;
+using System.Threading.Tasks;
+using EpicGames.Core;
+using EpicGames.Horde.Agents;
+using Horde.Server.Server;
+using Horde.Server.Utilities;
+using HordeCommon;
+using Microsoft.Extensions.Hosting;
+using Microsoft.Extensions.Logging;
+using MongoDB.Bson;
+using MongoDB.Bson.Serialization.Attributes;
+using MongoDB.Driver;
+
+namespace Horde.Server.Agents.Telemetry
+{
+	/// <summary>
+	/// Collection which writes agent data to a MongoDB collection
+	/// </summary>
+	public sealed class AgentTelemetryCollection : IAgentTelemetryCollection, IHostedService, IAsyncDisposable
+	{
+		class AgentTelemetryDocument : IAgentTelemetry
+		{
+			public ObjectId Id { get; set; }
+
+			[BsonElement("tm")]
+			public DateTime TimeUtc { get; set; }
+
+			[BsonElement("agt")]
+			public AgentId AgentId { get; set; }
+
+			[BsonElement("cpu_usr")]
+			public float UserCpu { get; set; }
+
+			[BsonElement("cpu_sys")]
+			public float SystemCpu { get; set; }
+
+			[BsonElement("cpu_idle")]
+			public float IdleCpu { get; set; }
+
+			[BsonElement("ram_tot")]
+			public int TotalRam { get; set; }
+
+			[BsonElement("ram_free")]
+			public int FreeRam { get; set; }
+
+			[BsonElement("ram_used")]
+			public int UsedRam { get; set; }
+		}
+
+		static TimeSpan RetainDays { get; } = TimeSpan.FromDays(30.0);
+
+		readonly IClock _clock;
+		readonly IMongoCollection<AgentTelemetryDocument> _collection;
+#pragma warning disable CA2213 // False positive? _writer is disposed in DisposeAsync.
+		readonly MongoBufferedWriter<AgentTelemetryDocument> _writer;
+#pragma warning restore CA2213
+		readonly ITicker _cleanupTicker;
+		readonly ILogger _logger;
+
+		/// <summary>
+		/// Constructor
+		/// </summary>
+		public AgentTelemetryCollection(MongoService mongoService, IClock clock, ILogger<AgentTelemetryCollection> logger)
+		{
+			_clock = clock;
+			_collection = mongoService.GetCollection<AgentTelemetryDocument>("Telemetry.Agents", builder => builder.Descending(x => x.Id));
+
+			List<MongoIndex<AgentTelemetryDocument>> agentIndexes = new List<MongoIndex<AgentTelemetryDocument>>();
+			agentIndexes.Add(MongoIndex.Create<AgentTelemetryDocument>(x => x.Descending(x => x.TimeUtc)));
+			agentIndexes.Add(MongoIndex.Create<AgentTelemetryDocument>(x => x.Descending(x => x.AgentId).Descending(x => x.TimeUtc)));
+			_collection = mongoService.GetCollection<AgentTelemetryDocument>("Telemetry.Agents", agentIndexes);
+
+			_writer = new MongoBufferedWriter<AgentTelemetryDocument>(_collection, logger);
+			_cleanupTicker = clock.AddSharedTicker<AgentTelemetryCollection>(TimeSpan.FromHours(4.0), CleanupAsync, logger);
+			_logger = logger;
+		}
+
+		/// <inheritdoc/>
+		public async Task StartAsync(CancellationToken cancellationToken)
+		{
+			await _writer.StartAsync();
+			await _cleanupTicker.StartAsync();
+		}
+
+		/// <inheritdoc/>
+		public async Task StopAsync(CancellationToken cancellationToken)
+		{
+			await _cleanupTicker.StopAsync();
+			await _writer.StopAsync(cancellationToken);
+		}
+
+		async ValueTask CleanupAsync(CancellationToken cancellationToken)
+		{
+			DateTime baseTime = _clock.UtcNow - RetainDays;
+			DeleteResult result = await _collection.DeleteManyAsync(x => x.TimeUtc < baseTime, cancellationToken);
+			_logger.LogInformation("Deleted {NumItems} agent telemetry records", result.DeletedCount);
+		}
+
+		/// <inheritdoc/>
+		public async ValueTask DisposeAsync()
+		{
+			await FlushAsync(default);
+			await _cleanupTicker.DisposeAsync();
+			await _writer.DisposeAsync();
+		}
+
+		/// <inheritdoc/>
+		public ValueTask FlushAsync(CancellationToken cancellationToken)
+			=> _writer.FlushAsync(cancellationToken);
+
+		/// <inheritdoc/>
+		public void Add(AgentId agentId, NewAgentTelemetry telemetry)
+		{
+			AgentTelemetryDocument document = new AgentTelemetryDocument();
+			document.AgentId = agentId;
+			document.UserCpu = telemetry.UserCpuPct;
+			document.IdleCpu = telemetry.IdleCpuPct;
+			document.SystemCpu = telemetry.SystemCpuPct;
+			document.FreeRam = telemetry.FreeRamMb;
+			document.UsedRam = telemetry.UsedRamMb;
+			document.TotalRam = telemetry.TotalRamMb;
+			_writer.Write(document);
+		}
+
+		/// <inheritdoc/>
+		public async Task<IReadOnlyList<IAgentTelemetry>> FindAsync(AgentId agentId, DateTime minTimeUtc, DateTime maxTimeUtc, CancellationToken cancellationToken = default)
+		{
+			return await _collection.Find(x => x.AgentId == agentId && (x.TimeUtc >= minTimeUtc && x.TimeUtc <= maxTimeUtc)).ToListAsync(cancellationToken);
+		}
+	}
+}
diff --git a/Engine/Source/Programs/Horde/Horde.Server/Agents/Telemetry/IAgentTelemetry.cs b/Engine/Source/Programs/Horde/Horde.Server/Agents/Telemetry/IAgentTelemetry.cs
new file mode 100644
--- /dev/null
+++ b/Engine/Source/Programs/Horde/Horde.Server/Agents/Telemetry/IAgentTelemetry.cs
@@ -0,0 +1,47 @@
+// Copyright Epic Games, Inc. All Rights Reserved.
+
+using System;
+
+namespace Horde.Server.Agents
+{
+	/// <summary>
+	/// Telemetry for an agent
+	/// </summary>
+	public interface IAgentTelemetry
+	{
+		/// <summary>
+		/// Time that the data was captured
+		/// </summary>
+		DateTime TimeUtc { get; set; }
+
+		/// <summary>
+		/// Percentage of time the CPU was busy executing code in user space
+		/// </summary>
+		float UserCpu { get; set; }
+
+		/// <summary>
+		/// Percentage of time the CPU was busy executing code in kernel space
+		/// </summary>
+		float SystemCpu { get; set; }
+
+		/// <summary>
+		/// Percentage of time the CPU was idling
+		/// </summary>
+		float IdleCpu { get; set; }
+
+		/// <summary>
+		/// Total memory installed (megabytes)
+		/// </summary>
+		int TotalRam { get; set; }
+
+		/// <summary>
+		/// Available memory (megabytes)
+		/// </summary>
+		int FreeRam { get; set; }
+
+		/// <summary>
+		/// Used memory (megabytes)
+		/// </summary>
+		int UsedRam { get; set; }
+	}
+}
diff --git a/Engine/Source/Programs/Horde/Horde.Server/Agents/Telemetry/IAgentTelemetryCollection.cs b/Engine/Source/Programs/Horde/Horde.Server/Agents/Telemetry/IAgentTelemetryCollection.cs
new file mode 100644
--- /dev/null
+++ b/Engine/Source/Programs/Horde/Horde.Server/Agents/Telemetry/IAgentTelemetryCollection.cs
@@ -0,0 +1,35 @@
+// Copyright Epic Games, Inc. All Rights Reserved.
+
+using System;
+using System.Collections.Generic;
+using System.Threading;
+using System.Threading.Tasks;
+using EpicGames.Horde.Agents;
+
+namespace Horde.Server.Agents
+{
+	/// <summary>
+	/// Interface for a collection of agent telemetry documents
+	/// </summary>
+	public interface IAgentTelemetryCollection
+	{
+		/// <summary>
+		/// Adds a new telemetry event
+		/// </summary>
+		void Add(AgentId agentId, NewAgentTelemetry telemetry);
+
+		/// <summary>
+		/// Finds telemetry samples for an agent in a particular time range
+		/// </summary>
+		/// <param name="agentId">Agent identifier</param>
+		/// <param name="minTimeUtc">Minimum time to return data for</param>
+		/// <param name="maxTimeUtc">Maximum time to return data for</param>
+		/// <param name="cancellationToken">Cancellation token for the operation</param>
+		Task<IReadOnlyList<IAgentTelemetry>> FindAsync(AgentId agentId, DateTime minTimeUtc, DateTime maxTimeUtc, CancellationToken cancellationToken = default);
+	}
+
+	/// <summary>
+	/// New telemetry data
+	/// </summary>
+	public record class NewAgentTelemetry(float UserCpuPct, float IdleCpuPct, float SystemCpuPct, int FreeRamMb, int UsedRamMb, int TotalRamMb);
+}
diff --git a/Engine/Source/Programs/Horde/Horde.Server/Jobs/JobCollection.cs b/Engine/Source/Programs/Horde/Horde.Server/Jobs/JobCollection.cs
--- a/Engine/Source/Programs/Horde/Horde.Server/Jobs/JobCollection.cs
+++ b/Engine/Source/Programs/Horde/Horde.Server/Jobs/JobCollection.cs
@@ -379,7 +379,7 @@ namespace Horde.Server.Jobs
 				return dependencies;
 			}
 
-			void MissingNodeRef(NodeRef nodeRef)
+			static void MissingNodeRef(NodeRef nodeRef)
 			{
 				_ = nodeRef;
 			}
diff --git a/Engine/Source/Programs/Horde/Horde.Server/Server/RpcService.cs b/Engine/Source/Programs/Horde/Horde.Server/Server/RpcService.cs
--- a/Engine/Source/Programs/Horde/Horde.Server/Server/RpcService.cs
+++ b/Engine/Source/Programs/Horde/Horde.Server/Server/RpcService.cs
@@ -53,6 +53,7 @@ namespace Horde.Server.Server
 		readonly ConformTaskSource _conformTaskSource;
 		readonly JobRpcCommon _jobRpcCommon;
 		readonly IToolCollection _toolCollection;
+		readonly IAgentTelemetryCollection _agentTelemetryCollection;
 		readonly AclService _aclService;
 		readonly IOptionsSnapshot<GlobalConfig> _globalConfig;
 		readonly ILogger _logger;
@@ -60,7 +61,7 @@ namespace Horde.Server.Server
 		/// <summary>
 		/// Constructor
 		/// </summary>
-		public RpcService(AgentService agentService, PoolService poolService, LifetimeService lifetimeService, ITelemetrySink telemetrySink, ConformTaskSource conformTaskSource, JobRpcCommon jobRpcCommon, IToolCollection toolCollection, AclService aclService, IOptionsSnapshot<GlobalConfig> globalConfig, ILogger<RpcService> logger)
+		public RpcService(AgentService agentService, PoolService poolService, LifetimeService lifetimeService, ITelemetrySink telemetrySink, ConformTaskSource conformTaskSource, JobRpcCommon jobRpcCommon, IToolCollection toolCollection, IAgentTelemetryCollection agentTelemetryCollection, AclService aclService, IOptionsSnapshot<GlobalConfig> globalConfig, ILogger<RpcService> logger)
 		{
 			_agentService = agentService;
 			_poolService = poolService;
@@ -69,6 +70,7 @@ namespace Horde.Server.Server
 			_conformTaskSource = conformTaskSource;
 			_jobRpcCommon = jobRpcCommon;
 			_toolCollection = toolCollection;
+			_agentTelemetryCollection = agentTelemetryCollection;
 			_aclService = aclService;
 			_globalConfig = globalConfig;
 			_logger = logger;
@@ -518,6 +520,15 @@ namespace Horde.Server.Server
 			}
 		}
 
+		/// <inheritdoc/>
+		public override Task<Empty> UploadTelemetry(RpcUploadTelemetryRequest request, ServerCallContext context)
+		{
+			NewAgentTelemetry telemetry = new NewAgentTelemetry(request.UserCpu, request.IdleCpu, request.SystemCpu, (int)request.FreeRam, (int)request.UsedRam, (int)request.TotalRam);
+			_agentTelemetryCollection.Add(new AgentId(request.AgentId), telemetry);
+
+			return Task.FromResult(new Empty());
+		}
+
 		/// <summary>
 		/// Receives telemetry events from agents
 		/// </summary>
diff --git a/Engine/Source/Programs/Horde/Horde.Server/Startup.cs b/Engine/Source/Programs/Horde/Horde.Server/Startup.cs
--- a/Engine/Source/Programs/Horde/Horde.Server/Startup.cs
+++ b/Engine/Source/Programs/Horde/Horde.Server/Startup.cs
@@ -48,6 +48,7 @@ using Horde.Server.Agents.Leases;
 using Horde.Server.Agents.Pools;
 using Horde.Server.Agents.Relay;
 using Horde.Server.Agents.Sessions;
+using Horde.Server.Agents.Telemetry;
 using Horde.Server.Agents.Utilization;
 using Horde.Server.Artifacts;
 using Horde.Server.Auditing;
@@ -844,6 +845,9 @@ namespace Horde.Server
 			services.AddHttpClient(EpicTelemetrySink.HttpClientName, client => { });
 			services.AddHttpClient(ClickHouseTelemetrySink.HttpClientName, client => { });
 
+			services.AddSingleton<AgentTelemetryCollection>();
+			services.AddSingleton<IAgentTelemetryCollection>(sp => sp.GetRequiredService<AgentTelemetryCollection>());
+
 			// Hosted service that needs to run no matter the run mode of the process (server vs worker)
 			services.AddHostedService(provider => (DowntimeService)provider.GetRequiredService<IDowntimeService>());
 
@@ -856,6 +860,7 @@ namespace Horde.Server
 			if (settings.IsRunModeActive(RunMode.Worker) && !settings.MongoReadOnlyMode)
 			{
 				services.AddHostedService<AgentReportService>();
+				services.AddHostedService(provider => provider.GetRequiredService<AgentTelemetryCollection>());
 				services.AddHostedService<BisectService>();
 				services.AddHostedService(provider => provider.GetRequiredService<FleetService>());
 				services.AddHostedService(provider => provider.GetRequiredService<ConsistencyService>());
diff --git a/Engine/Source/Programs/Shared/EpicGames.Horde/Agents/AgentMessages.cs b/Engine/Source/Programs/Shared/EpicGames.Horde/Agents/AgentMessages.cs
--- a/Engine/Source/Programs/Shared/EpicGames.Horde/Agents/AgentMessages.cs
+++ b/Engine/Source/Programs/Shared/EpicGames.Horde/Agents/AgentMessages.cs
@@ -128,4 +128,14 @@ namespace EpicGames.Horde.Agents
 	/// <param name="Workspaces">Current workspaces synced on the agent</param>
 	/// <param name="Comment"> Comment for this agent </param>
 	public record GetAgentResponse(AgentId Id, string Name, bool Enabled, AgentStatus Status, double? Rate, SessionId? SessionId, bool Ephemeral, bool Online, bool Deleted, bool PendingConform, bool PendingFullConform, bool PendingRestart, bool PendingShutdown, string LastShutdownReason, DateTime LastConformTime, int? ConformAttemptCount, DateTime? NextConformTime, string? Version, List<string> Properties, Dictionary<string, int> Resources, DateTime? UpdateTime, DateTime? LastStatusChange, List<string>? Pools, object? Capabilities, List<GetAgentLeaseResponse> Leases, List<GetAgentWorkspaceResponse> Workspaces, string? Comment);
+
+	/// <summary>
+	/// Telemetry data for an agent
+	/// </summary>
+	public record GetAgentTelemetryResponse(List<GetAgentTelemetrySampleResponse> Samples);
+
+	/// <summary>
+	/// Telemetry data sample
+	/// </summary>
+	public record GetAgentTelemetrySampleResponse(DateTime Time, float UserCpu, float IdleCpu, float SystemCpu, int FreeRam, int UsedRam, int TotalRam);
 }
