Horde: Fix exception creating indexes for telemetry collection.

#rnx

diff --git a/Engine/Source/Programs/Horde/Horde.Server/Agents/Telemetry/AgentTelemetryCollection.cs b/Engine/Source/Programs/Horde/Horde.Server/Agents/Telemetry/AgentTelemetryCollection.cs
--- a/Engine/Source/Programs/Horde/Horde.Server/Agents/Telemetry/AgentTelemetryCollection.cs
+++ b/Engine/Source/Programs/Horde/Horde.Server/Agents/Telemetry/AgentTelemetryCollection.cs
@@ -68,7 +68,6 @@ namespace Horde.Server.Agents.Telemetry
 		public AgentTelemetryCollection(MongoService mongoService, IClock clock, ILogger<AgentTelemetryCollection> logger)
 		{
 			_clock = clock;
-			_collection = mongoService.GetCollection<AgentTelemetryDocument>("Telemetry.Agents", builder => builder.Descending(x => x.Id));
 
 			List<MongoIndex<AgentTelemetryDocument>> agentIndexes = new List<MongoIndex<AgentTelemetryDocument>>();
 			agentIndexes.Add(MongoIndex.Create<AgentTelemetryDocument>(x => x.Descending(x => x.TimeUtc)));
