Horde: Fix tests.

#rnx

diff --git a/Engine/Source/Programs/Horde/Horde.Server.Tests/TestSetup.cs b/Engine/Source/Programs/Horde/Horde.Server.Tests/TestSetup.cs
--- a/Engine/Source/Programs/Horde/Horde.Server.Tests/TestSetup.cs
+++ b/Engine/Source/Programs/Horde/Horde.Server.Tests/TestSetup.cs
@@ -26,6 +26,7 @@ using Horde.Server.Agents.Leases;
 using Horde.Server.Agents.Pools;
 using Horde.Server.Agents.Relay;
 using Horde.Server.Agents.Sessions;
+using Horde.Server.Agents.Telemetry;
 using Horde.Server.Agents.Utilization;
 using Horde.Server.Artifacts;
 using Horde.Server.Auditing;
@@ -212,6 +213,7 @@ namespace Horde.Server.Tests
 
 			services.AddSingleton<IAccountCollection, AccountCollection>();
 			services.AddSingleton<IAgentCollection, AgentCollection>();
+			services.AddSingleton<IAgentTelemetryCollection, AgentTelemetryCollection>();
 			services.AddSingleton<IArtifactCollection, ArtifactCollection>();
 			services.AddSingleton<IArtifactCollectionV1, ArtifactCollectionV1>();
 			services.AddSingleton<ICommitService, CommitService>();
diff --git a/Engine/Source/Programs/Horde/Horde.Server/Agents/AgentsController.cs b/Engine/Source/Programs/Horde/Horde.Server/Agents/AgentsController.cs
--- a/Engine/Source/Programs/Horde/Horde.Server/Agents/AgentsController.cs
+++ b/Engine/Source/Programs/Horde/Horde.Server/Agents/AgentsController.cs
@@ -13,6 +13,7 @@ using EpicGames.Horde.Agents.Sessions;
 using EpicGames.Horde.Common;
 using Horde.Server.Agents.Leases;
 using Horde.Server.Agents.Sessions;
+using Horde.Server.Agents.Telemetry;
 using Horde.Server.Auditing;
 using Horde.Server.Server;
 using Horde.Server.Users;
diff --git a/Engine/Source/Programs/Horde/Horde.Server/Agents/Telemetry/IAgentTelemetry.cs b/Engine/Source/Programs/Horde/Horde.Server/Agents/Telemetry/IAgentTelemetry.cs
--- a/Engine/Source/Programs/Horde/Horde.Server/Agents/Telemetry/IAgentTelemetry.cs
+++ b/Engine/Source/Programs/Horde/Horde.Server/Agents/Telemetry/IAgentTelemetry.cs
@@ -2,7 +2,7 @@
 
 using System;
 
-namespace Horde.Server.Agents
+namespace Horde.Server.Agents.Telemetry
 {
 	/// <summary>
 	/// Telemetry for an agent
diff --git a/Engine/Source/Programs/Horde/Horde.Server/Agents/Telemetry/IAgentTelemetryCollection.cs b/Engine/Source/Programs/Horde/Horde.Server/Agents/Telemetry/IAgentTelemetryCollection.cs
--- a/Engine/Source/Programs/Horde/Horde.Server/Agents/Telemetry/IAgentTelemetryCollection.cs
+++ b/Engine/Source/Programs/Horde/Horde.Server/Agents/Telemetry/IAgentTelemetryCollection.cs
@@ -6,7 +6,7 @@ using System.Threading;
 using System.Threading.Tasks;
 using EpicGames.Horde.Agents;
 
-namespace Horde.Server.Agents
+namespace Horde.Server.Agents.Telemetry
 {
 	/// <summary>
 	/// Interface for a collection of agent telemetry documents
diff --git a/Engine/Source/Programs/Horde/Horde.Server/Server/RpcService.cs b/Engine/Source/Programs/Horde/Horde.Server/Server/RpcService.cs
--- a/Engine/Source/Programs/Horde/Horde.Server/Server/RpcService.cs
+++ b/Engine/Source/Programs/Horde/Horde.Server/Server/RpcService.cs
@@ -21,6 +21,7 @@ using Horde.Server.Acls;
 using Horde.Server.Agents;
 using Horde.Server.Agents.Pools;
 using Horde.Server.Agents.Sessions;
+using Horde.Server.Agents.Telemetry;
 using Horde.Server.Jobs;
 using Horde.Server.Tasks;
 using Horde.Server.Telemetry;
