Horde: Allow exceptions to throw when upgrading collections fails.

#rnx

diff --git a/Engine/Source/Programs/Horde/Horde.Server/Server/MongoService.cs b/Engine/Source/Programs/Horde/Horde.Server/Server/MongoService.cs
--- a/Engine/Source/Programs/Horde/Horde.Server/Server/MongoService.cs
+++ b/Engine/Source/Programs/Horde/Horde.Server/Server/MongoService.cs
@@ -752,7 +752,7 @@ namespace Horde.Server.Server
 					await UpdateIndexesInternalAsync<T>(collectionName, collection, newIndexes, cancellationToken);
 					break;
 				}
-				catch (Exception ex)
+				catch (Exception ex) when (attemptIdx < 3)
 				{
 					_logger.LogError(ex, "Error updating indexes ({Message}) - retrying (attempt {Attempt})", ex.Message, attemptIdx);
 					attemptIdx++;
