Horde: Enable agents sending telemetry to the server.

diff --git a/Engine/Source/Programs/Horde/Horde.Agent/AgentSettings.cs b/Engine/Source/Programs/Horde/Horde.Agent/AgentSettings.cs
--- a/Engine/Source/Programs/Horde/Horde.Agent/AgentSettings.cs
+++ b/Engine/Source/Programs/Horde/Horde.Agent/AgentSettings.cs
@@ -279,7 +279,7 @@ namespace Horde.Agent
 		/// <summary>
 		/// Whether to send telemetry back to Horde server
 		/// </summary>
-		public bool EnableTelemetry { get; set; } = false;
+		public bool EnableTelemetry { get; set; } = true;
 
 		/// <summary>
 		/// How often to report telemetry events to server in milliseconds
