Horde: Telemetry fixes.

#rnx

diff --git a/Engine/Source/Programs/Horde/Horde.Agent/Services/TelemetryService.cs b/Engine/Source/Programs/Horde/Horde.Agent/Services/TelemetryService.cs
--- a/Engine/Source/Programs/Horde/Horde.Agent/Services/TelemetryService.cs
+++ b/Engine/Source/Programs/Horde/Horde.Agent/Services/TelemetryService.cs
@@ -506,8 +506,11 @@ class TelemetryService : BackgroundService
 				request.TotalRam = memMetrics.Total / 1024;
 				request.FreeRam = memMetrics.Available / 1024;
 				request.UsedRam = memMetrics.Used / 1024;
+				await client.UploadTelemetryAsync(request, new CallOptions(cancellationToken: stoppingToken));
 			}
 
+			// Disabling ClickHouse telemetry for now, since we are storing in Mongo.
+/*
 			// Clickhouse method
 			{
 				RpcSendTelemetryEventsRequest request = new();
@@ -545,6 +548,7 @@ class TelemetryService : BackgroundService
 
 				await client.SendTelemetryEventsAsync(request, new CallOptions(cancellationToken: stoppingToken));
 			}
+*/
 			await Task.Delay(_reportInterval, stoppingToken);
 		}
 
