Horde: Use the connection session id to determine the agent for telemetry.

#rnx

diff --git a/Engine/Source/Programs/Horde/Horde.Server/Server/RpcService.cs b/Engine/Source/Programs/Horde/Horde.Server/Server/RpcService.cs
--- a/Engine/Source/Programs/Horde/Horde.Server/Server/RpcService.cs
+++ b/Engine/Source/Programs/Horde/Horde.Server/Server/RpcService.cs
@@ -522,12 +522,19 @@ namespace Horde.Server.Server
 		}
 
 		/// <inheritdoc/>
-		public override Task<Empty> UploadTelemetry(RpcUploadTelemetryRequest request, ServerCallContext context)
+		public override async Task<Empty> UploadTelemetry(RpcUploadTelemetryRequest request, ServerCallContext context)
 		{
-			NewAgentTelemetry telemetry = new NewAgentTelemetry(request.UserCpu, request.IdleCpu, request.SystemCpu, (int)request.FreeRam, (int)request.UsedRam, (int)request.TotalRam);
-			_agentTelemetryCollection.Add(new AgentId(request.AgentId), telemetry);
-
-			return Task.FromResult(new Empty());
+			SessionId? sessionId = context.GetHttpContext().User.GetSessionClaim();
+			if (sessionId != null)
+			{
+				ISession? session = await _agentService.GetSessionAsync(sessionId.Value);
+				if (session != null)
+				{
+					NewAgentTelemetry telemetry = new NewAgentTelemetry(request.UserCpu, request.IdleCpu, request.SystemCpu, (int)request.FreeRam, (int)request.UsedRam, (int)request.TotalRam);
+					_agentTelemetryCollection.Add(session.AgentId, telemetry);
+				}
+			}
+			return new Empty();
 		}
 
 		/// <summary>
