Horde: Send agent CPU/RAM telemetry as a background task started only when a session is active.

#rnx

diff --git a/Engine/Source/Programs/Horde/Horde.Agent.Tests/Services/TelemetryServiceTest.cs b/Engine/Source/Programs/Horde/Horde.Agent.Tests/Services/TelemetryServiceTest.cs
--- a/Engine/Source/Programs/Horde/Horde.Agent.Tests/Services/TelemetryServiceTest.cs
+++ b/Engine/Source/Programs/Horde/Horde.Agent.Tests/Services/TelemetryServiceTest.cs
@@ -25,7 +25,7 @@ public sealed class TelemetryServiceTest : IDisposable
 		AgentSettings settings = new() { Server = "Test", ServerProfiles = { ["Test"] = new ServerProfile() { Name = "Test", Url = new Uri("http://localhost:1234") } } };
 		OptionsWrapper<AgentSettings> settingsOpt = new(settings);
 
-		using WorkerService workerService = new(null!, null!, null!, new List<LeaseHandler>(), null!, null!, null!);
+		using WorkerService workerService = new(null!, null!, null!, null!, new List<LeaseHandler>(), null!, null!, null!);
 		JobHandler jobHandler = new JobHandler(new List<IJobExecutorFactory>(), settingsOpt, null!, null!);
 		GrpcService grpcService = new(settingsOpt, NullLogger<GrpcService>.Instance, _loggerFactory);
 
diff --git a/Engine/Source/Programs/Horde/Horde.Agent.Tests/WorkerServiceTest.cs b/Engine/Source/Programs/Horde/Horde.Agent.Tests/WorkerServiceTest.cs
--- a/Engine/Source/Programs/Horde/Horde.Agent.Tests/WorkerServiceTest.cs
+++ b/Engine/Source/Programs/Horde/Horde.Agent.Tests/WorkerServiceTest.cs
@@ -261,7 +261,7 @@ namespace Horde.Agent.Tests
 			await using FakeHordeRpcServer fakeServer = new();
 			await using ISession session = FakeServerSessionFactory.CreateSession(fakeServer.GetConnection(), fakeServer.GetGrpcChannel());
 
-			LeaseManager manager = new LeaseManager(session, null!, serviceProvider.GetRequiredService<StatusService>(),
+			LeaseManager manager = new LeaseManager(session, null!, serviceProvider.GetRequiredService<StatusService>(), serviceProvider.GetRequiredService<TelemetryService>(),
 				serviceProvider.GetRequiredService<IEnumerable<LeaseHandler>>(), serviceProvider.GetRequiredService<LeaseLoggerFactory>(),
 				_loggerFactory.CreateLogger<LeaseManager>());
 
diff --git a/Engine/Source/Programs/Horde/Horde.Agent/Leases/LeaseManager.cs b/Engine/Source/Programs/Horde/Horde.Agent/Leases/LeaseManager.cs
--- a/Engine/Source/Programs/Horde/Horde.Agent/Leases/LeaseManager.cs
+++ b/Engine/Source/Programs/Horde/Horde.Agent/Leases/LeaseManager.cs
@@ -120,24 +120,26 @@ namespace Horde.Agent.Leases
 		readonly ISession _session;
 		readonly CapabilitiesService _capabilitiesService;
 		readonly StatusService _statusService;
+		readonly TelemetryService _telemetryService;
 		readonly Dictionary<string, LeaseHandler> _typeUrlToLeaseHandler;
 		readonly LeaseLoggerFactory _leaseLoggerFactory;
 		readonly ILogger _logger;
 
 		RpcAgentCapabilities? _capabilities;
 
-		public LeaseManager(ISession session, CapabilitiesService capabilitiesService, StatusService statusService, IEnumerable<LeaseHandler> leaseHandlers, LeaseLoggerFactory leaseLoggerFactory, ILogger logger)
+		public LeaseManager(ISession session, CapabilitiesService capabilitiesService, StatusService statusService, TelemetryService telemetryService, IEnumerable<LeaseHandler> leaseHandlers, LeaseLoggerFactory leaseLoggerFactory, ILogger logger)
 		{
 			_session = session;
 			_capabilitiesService = capabilitiesService;
 			_statusService = statusService;
+			_telemetryService = telemetryService;
 			_typeUrlToLeaseHandler = leaseHandlers.ToDictionary(x => x.LeaseType, x => x);
 			_leaseLoggerFactory = leaseLoggerFactory;
 			_logger = logger;
 		}
 
 		public LeaseManager(ISession session, IServiceProvider serviceProvider)
-			: this(session, serviceProvider.GetRequiredService<CapabilitiesService>(), serviceProvider.GetRequiredService<StatusService>(), serviceProvider.GetRequiredService<IEnumerable<LeaseHandler>>(), serviceProvider.GetRequiredService<LeaseLoggerFactory>(), serviceProvider.GetRequiredService<ILogger<LeaseManager>>())
+			: this(session, serviceProvider.GetRequiredService<CapabilitiesService>(), serviceProvider.GetRequiredService<StatusService>(), serviceProvider.GetRequiredService<TelemetryService>(), serviceProvider.GetRequiredService<IEnumerable<LeaseHandler>>(), serviceProvider.GetRequiredService<LeaseLoggerFactory>(), serviceProvider.GetRequiredService<ILogger<LeaseManager>>())
 		{
 		}
 
@@ -240,6 +242,9 @@ namespace Horde.Agent.Leases
 			// Run a background task to update the capabilities of this agent
 			await using BackgroundTask updateCapsTask = BackgroundTask.StartNew(ctx => UpdateCapabilitiesBackgroundAsync(_session.WorkingDir, ctx));
 
+			// Send telemetry to the server in the background
+			await using BackgroundTask telemetryTask = _telemetryService.CreateBackgroundTask(_session.AgentId);
+
 			// Loop until we're ready to exit
 			Stopwatch updateCapabilitiesTimer = Stopwatch.StartNew();
 			for (; ; )
diff --git a/Engine/Source/Programs/Horde/Horde.Agent/Services/TelemetryService.cs b/Engine/Source/Programs/Horde/Horde.Agent/Services/TelemetryService.cs
--- a/Engine/Source/Programs/Horde/Horde.Agent/Services/TelemetryService.cs
+++ b/Engine/Source/Programs/Horde/Horde.Agent/Services/TelemetryService.cs
@@ -4,6 +4,7 @@ using System.Diagnostics;
 using System.Runtime.InteropServices;
 using System.Text;
 using EpicGames.Core;
+using EpicGames.Horde.Agents;
 using Google.Protobuf.WellKnownTypes;
 using Grpc.Core;
 using Grpc.Net.Client;
@@ -454,13 +455,30 @@ class TelemetryService : BackgroundService
 	}
 
 	/// <inheritdoc />
-	protected override async Task ExecuteAsync(CancellationToken stoppingToken)
+	protected override Task ExecuteAsync(CancellationToken stoppingToken)
 	{
+		return Task.CompletedTask;
+	}
+
+	bool SendLegacyMetrics { get; set; } = false;
+
+	/// <summary>
+	/// Create a background task that sends telemetry for a session
+	/// </summary>
+	public BackgroundTask CreateBackgroundTask(AgentId agentId)
+	{
+		return BackgroundTask.StartNew(ctx => ExecuteBackgroundAsync(agentId, ctx));
+	}
+
+	/// <inheritdoc />
+	async Task ExecuteBackgroundAsync(AgentId agentId, CancellationToken stoppingToken)
+	{
+		_logger.LogDebug("Starting telemetry background task");
 		while (!stoppingToken.IsCancellationRequested)
 		{
 			try
 			{
-				if (!await ExecuteInternalAsync(stoppingToken))
+				if (!await ExecuteBackgroundInternalAsync(agentId, stoppingToken))
 				{
 					break;
 				}
@@ -475,11 +493,19 @@ class TelemetryService : BackgroundService
 			}
 
 			// Wait a moment before attempting to restart the background work
-			await Task.Delay(TimeSpan.FromSeconds(20), stoppingToken);
+			try
+			{
+				await Task.Delay(TimeSpan.FromSeconds(20), stoppingToken);
+			}
+			catch (OperationCanceledException) when (stoppingToken.IsCancellationRequested)
+			{
+				break;
+			}
 		}
+		_logger.LogDebug("Stopping telemetry background task");
 	}
 
-	private async Task<bool> ExecuteInternalAsync(CancellationToken stoppingToken)
+	private async Task<bool> ExecuteBackgroundInternalAsync(AgentId agentId, CancellationToken stoppingToken)
 	{
 		if (_systemMetrics == null || !_agentSettings.EnableTelemetry)
 		{
@@ -509,9 +535,8 @@ class TelemetryService : BackgroundService
 				await client.UploadTelemetryAsync(request, new CallOptions(cancellationToken: stoppingToken));
 			}
 
-			// Disabling ClickHouse telemetry for now, since we are storing in Mongo.
-/*
 			// Clickhouse method
+			if (SendLegacyMetrics)
 			{
 				RpcSendTelemetryEventsRequest request = new();
 				Timestamp utcNow = Timestamp.FromDateTime(DateTime.UtcNow);
@@ -548,7 +573,7 @@ class TelemetryService : BackgroundService
 
 				await client.SendTelemetryEventsAsync(request, new CallOptions(cancellationToken: stoppingToken));
 			}
-*/
+
 			await Task.Delay(_reportInterval, stoppingToken);
 		}
 
diff --git a/Engine/Source/Programs/Horde/Horde.Agent/Services/WorkerService.cs b/Engine/Source/Programs/Horde/Horde.Agent/Services/WorkerService.cs
--- a/Engine/Source/Programs/Horde/Horde.Agent/Services/WorkerService.cs
+++ b/Engine/Source/Programs/Horde/Horde.Agent/Services/WorkerService.cs
@@ -18,6 +18,7 @@ namespace Horde.Agent.Services
 		readonly ISessionFactory _sessionFactory;
 		readonly CapabilitiesService _capabilitiesService;
 		readonly StatusService _statusService;
+		readonly TelemetryService _telemetryService;
 		readonly LeaseHandler[] _leaseHandlers;
 		readonly LeaseLoggerFactory _leaseLoggerFactory;
 		readonly IServiceProvider _serviceProvider;
@@ -36,11 +37,12 @@ namespace Horde.Agent.Services
 		/// <summary>
 		/// Constructor. Registers with the server and starts accepting connections.
 		/// </summary>
-		public WorkerService(ISessionFactory sessionFactory, CapabilitiesService capabilitiesService, StatusService statusService, IEnumerable<LeaseHandler> leaseHandlers, LeaseLoggerFactory leaseLoggerFactory, IServiceProvider serviceProvider, ILogger<WorkerService> logger)
+		public WorkerService(ISessionFactory sessionFactory, CapabilitiesService capabilitiesService, StatusService statusService, TelemetryService telemetryService, IEnumerable<LeaseHandler> leaseHandlers, LeaseLoggerFactory leaseLoggerFactory, IServiceProvider serviceProvider, ILogger<WorkerService> logger)
 		{
 			_sessionFactory = sessionFactory;
 			_capabilitiesService = capabilitiesService;
 			_statusService = statusService;
+			_telemetryService = telemetryService;
 			_logger = logger;
 			_leaseHandlers = leaseHandlers.ToArray();
 			_leaseLoggerFactory = leaseLoggerFactory;
@@ -103,7 +105,7 @@ namespace Horde.Agent.Services
 
 						await using (ISession session = await _sessionFactory.CreateAsync(stoppingToken))
 						{
-							_currentLeaseManager = new LeaseManager(session, _capabilitiesService, _statusService, _leaseHandlers, _leaseLoggerFactory, _logger);
+							_currentLeaseManager = new LeaseManager(session, _capabilitiesService, _statusService, _telemetryService, _leaseHandlers, _leaseLoggerFactory, _logger);
 							result = await _currentLeaseManager.RunAsync(false, stoppingToken);
 						}
 
diff --git a/Engine/Source/Programs/Horde/Horde.Server/Server/RpcService.cs b/Engine/Source/Programs/Horde/Horde.Server/Server/RpcService.cs
--- a/Engine/Source/Programs/Horde/Horde.Server/Server/RpcService.cs
+++ b/Engine/Source/Programs/Horde/Horde.Server/Server/RpcService.cs
@@ -522,19 +522,12 @@ namespace Horde.Server.Server
 		}
 
 		/// <inheritdoc/>
-		public override async Task<Empty> UploadTelemetry(RpcUploadTelemetryRequest request, ServerCallContext context)
+		public override Task<Empty> UploadTelemetry(RpcUploadTelemetryRequest request, ServerCallContext context)
 		{
-			SessionId? sessionId = context.GetHttpContext().User.GetSessionClaim();
-			if (sessionId != null)
-			{
-				ISession? session = await _agentService.GetSessionAsync(sessionId.Value);
-				if (session != null)
-				{
-					NewAgentTelemetry telemetry = new NewAgentTelemetry(request.UserCpu, request.IdleCpu, request.SystemCpu, (int)request.FreeRam, (int)request.UsedRam, (int)request.TotalRam);
-					_agentTelemetryCollection.Add(session.AgentId, telemetry);
-				}
-			}
-			return new Empty();
+			NewAgentTelemetry telemetry = new NewAgentTelemetry(request.UserCpu, request.IdleCpu, request.SystemCpu, (int)request.FreeRam, (int)request.UsedRam, (int)request.TotalRam);
+			_agentTelemetryCollection.Add(new AgentId(request.AgentId), telemetry);
+
+			return Task.FromResult(new Empty());
 		}
 
 		/// <summary>
