Horde: Back out change to introduce dependency on telmetry service, causing agents to fail to start.

#rnx

diff --git a/Engine/Source/Programs/Horde/Horde.Agent.Tests/Services/TelemetryServiceTest.cs b/Engine/Source/Programs/Horde/Horde.Agent.Tests/Services/TelemetryServiceTest.cs
--- a/Engine/Source/Programs/Horde/Horde.Agent.Tests/Services/TelemetryServiceTest.cs
+++ b/Engine/Source/Programs/Horde/Horde.Agent.Tests/Services/TelemetryServiceTest.cs
@@ -25,7 +25,7 @@ public sealed class TelemetryServiceTest : IDisposable
 		AgentSettings settings = new() { Server = "Test", ServerProfiles = { ["Test"] = new ServerProfile() { Name = "Test", Url = new Uri("http://localhost:1234") } } };
 		OptionsWrapper<AgentSettings> settingsOpt = new(settings);
 
-		using WorkerService workerService = new(null!, null!, null!, null!, new List<LeaseHandler>(), null!, null!, null!);
+		using WorkerService workerService = new(null!, null!, null!, new List<LeaseHandler>(), null!, null!, null!);
 		JobHandler jobHandler = new JobHandler(new List<IJobExecutorFactory>(), settingsOpt, null!, null!);
 		GrpcService grpcService = new(settingsOpt, NullLogger<GrpcService>.Instance, _loggerFactory);
 
diff --git a/Engine/Source/Programs/Horde/Horde.Agent.Tests/WorkerServiceTest.cs b/Engine/Source/Programs/Horde/Horde.Agent.Tests/WorkerServiceTest.cs
--- a/Engine/Source/Programs/Horde/Horde.Agent.Tests/WorkerServiceTest.cs
+++ b/Engine/Source/Programs/Horde/Horde.Agent.Tests/WorkerServiceTest.cs
@@ -261,7 +261,7 @@ namespace Horde.Agent.Tests
 			await using FakeHordeRpcServer fakeServer = new();
 			await using ISession session = FakeServerSessionFactory.CreateSession(fakeServer.GetConnection(), fakeServer.GetGrpcChannel());
 
-			LeaseManager manager = new LeaseManager(session, null!, serviceProvider.GetRequiredService<StatusService>(), serviceProvider.GetRequiredService<TelemetryService>(),
+			LeaseManager manager = new LeaseManager(session, null!, serviceProvider.GetRequiredService<StatusService>(),
 				serviceProvider.GetRequiredService<IEnumerable<LeaseHandler>>(), serviceProvider.GetRequiredService<LeaseLoggerFactory>(),
 				_loggerFactory.CreateLogger<LeaseManager>());
 
diff --git a/Engine/Source/Programs/Horde/Horde.Agent/Leases/LeaseManager.cs b/Engine/Source/Programs/Horde/Horde.Agent/Leases/LeaseManager.cs
--- a/Engine/Source/Programs/Horde/Horde.Agent/Leases/LeaseManager.cs
+++ b/Engine/Source/Programs/Horde/Horde.Agent/Leases/LeaseManager.cs
@@ -120,26 +120,24 @@ namespace Horde.Agent.Leases
 		readonly ISession _session;
 		readonly CapabilitiesService _capabilitiesService;
 		readonly StatusService _statusService;
-		readonly TelemetryService _telemetryService;
 		readonly Dictionary<string, LeaseHandler> _typeUrlToLeaseHandler;
 		readonly LeaseLoggerFactory _leaseLoggerFactory;
 		readonly ILogger _logger;
 
 		RpcAgentCapabilities? _capabilities;
 
-		public LeaseManager(ISession session, CapabilitiesService capabilitiesService, StatusService statusService, TelemetryService telemetryService, IEnumerable<LeaseHandler> leaseHandlers, LeaseLoggerFactory leaseLoggerFactory, ILogger logger)
+		public LeaseManager(ISession session, CapabilitiesService capabilitiesService, StatusService statusService, IEnumerable<LeaseHandler> leaseHandlers, LeaseLoggerFactory leaseLoggerFactory, ILogger logger)
 		{
 			_session = session;
 			_capabilitiesService = capabilitiesService;
 			_statusService = statusService;
-			_telemetryService = telemetryService;
 			_typeUrlToLeaseHandler = leaseHandlers.ToDictionary(x => x.LeaseType, x => x);
 			_leaseLoggerFactory = leaseLoggerFactory;
 			_logger = logger;
 		}
 
 		public LeaseManager(ISession session, IServiceProvider serviceProvider)
-			: this(session, serviceProvider.GetRequiredService<CapabilitiesService>(), serviceProvider.GetRequiredService<StatusService>(), serviceProvider.GetRequiredService<TelemetryService>(), serviceProvider.GetRequiredService<IEnumerable<LeaseHandler>>(), serviceProvider.GetRequiredService<LeaseLoggerFactory>(), serviceProvider.GetRequiredService<ILogger<LeaseManager>>())
+			: this(session, serviceProvider.GetRequiredService<CapabilitiesService>(), serviceProvider.GetRequiredService<StatusService>(), serviceProvider.GetRequiredService<IEnumerable<LeaseHandler>>(), serviceProvider.GetRequiredService<LeaseLoggerFactory>(), serviceProvider.GetRequiredService<ILogger<LeaseManager>>())
 		{
 		}
 
@@ -242,9 +240,6 @@ namespace Horde.Agent.Leases
 			// Run a background task to update the capabilities of this agent
 			await using BackgroundTask updateCapsTask = BackgroundTask.StartNew(ctx => UpdateCapabilitiesBackgroundAsync(_session.WorkingDir, ctx));
 
-			// Send telemetry to the server in the background
-			await using BackgroundTask telemetryTask = _telemetryService.CreateBackgroundTask(_session.AgentId);
-
 			// Loop until we're ready to exit
 			Stopwatch updateCapabilitiesTimer = Stopwatch.StartNew();
 			for (; ; )
diff --git a/Engine/Source/Programs/Horde/Horde.Agent/Services/WorkerService.cs b/Engine/Source/Programs/Horde/Horde.Agent/Services/WorkerService.cs
--- a/Engine/Source/Programs/Horde/Horde.Agent/Services/WorkerService.cs
+++ b/Engine/Source/Programs/Horde/Horde.Agent/Services/WorkerService.cs
@@ -18,7 +18,6 @@ namespace Horde.Agent.Services
 		readonly ISessionFactory _sessionFactory;
 		readonly CapabilitiesService _capabilitiesService;
 		readonly StatusService _statusService;
-		readonly TelemetryService _telemetryService;
 		readonly LeaseHandler[] _leaseHandlers;
 		readonly LeaseLoggerFactory _leaseLoggerFactory;
 		readonly IServiceProvider _serviceProvider;
@@ -37,12 +36,11 @@ namespace Horde.Agent.Services
 		/// <summary>
 		/// Constructor. Registers with the server and starts accepting connections.
 		/// </summary>
-		public WorkerService(ISessionFactory sessionFactory, CapabilitiesService capabilitiesService, StatusService statusService, TelemetryService telemetryService, IEnumerable<LeaseHandler> leaseHandlers, LeaseLoggerFactory leaseLoggerFactory, IServiceProvider serviceProvider, ILogger<WorkerService> logger)
+		public WorkerService(ISessionFactory sessionFactory, CapabilitiesService capabilitiesService, StatusService statusService, IEnumerable<LeaseHandler> leaseHandlers, LeaseLoggerFactory leaseLoggerFactory, IServiceProvider serviceProvider, ILogger<WorkerService> logger)
 		{
 			_sessionFactory = sessionFactory;
 			_capabilitiesService = capabilitiesService;
 			_statusService = statusService;
-			_telemetryService = telemetryService;
 			_logger = logger;
 			_leaseHandlers = leaseHandlers.ToArray();
 			_leaseLoggerFactory = leaseLoggerFactory;
@@ -105,7 +103,7 @@ namespace Horde.Agent.Services
 
 						await using (ISession session = await _sessionFactory.CreateAsync(stoppingToken))
 						{
-							_currentLeaseManager = new LeaseManager(session, _capabilitiesService, _statusService, _telemetryService, _leaseHandlers, _leaseLoggerFactory, _logger);
+							_currentLeaseManager = new LeaseManager(session, _capabilitiesService, _statusService, _leaseHandlers, _leaseLoggerFactory, _logger);
 							result = await _currentLeaseManager.RunAsync(false, stoppingToken);
 						}
 
