Horde: Register ISystemMetrics as its own service, with different implementations for each OS. Use the regular gRPC connection for sending updates.

diff --git a/Engine/Source/Programs/Horde/Horde.Agent.Tests/Services/TelemetryServiceTest.cs b/Engine/Source/Programs/Horde/Horde.Agent.Tests/Services/TelemetryServiceTest.cs
--- a/Engine/Source/Programs/Horde/Horde.Agent.Tests/Services/TelemetryServiceTest.cs
+++ b/Engine/Source/Programs/Horde/Horde.Agent.Tests/Services/TelemetryServiceTest.cs
@@ -29,7 +29,7 @@ public sealed class TelemetryServiceTest : IDisposable
 		JobHandler jobHandler = new JobHandler(new List<IJobExecutorFactory>(), settingsOpt, null!, null!);
 		GrpcService grpcService = new(settingsOpt, NullLogger<GrpcService>.Instance, _loggerFactory);
 
-		_telemetryService = new TelemetryService(workerService, jobHandler, grpcService, settingsOpt, NullLogger<TelemetryService>.Instance);
+		_telemetryService = new TelemetryService(workerService, jobHandler, grpcService, null!, settingsOpt, NullLogger<TelemetryService>.Instance);
 	}
 
 	[TestMethod]
diff --git a/Engine/Source/Programs/Horde/Horde.Agent.Tests/WorkerServiceTest.cs b/Engine/Source/Programs/Horde/Horde.Agent.Tests/WorkerServiceTest.cs
--- a/Engine/Source/Programs/Horde/Horde.Agent.Tests/WorkerServiceTest.cs
+++ b/Engine/Source/Programs/Horde/Horde.Agent.Tests/WorkerServiceTest.cs
@@ -261,9 +261,7 @@ namespace Horde.Agent.Tests
 			await using FakeHordeRpcServer fakeServer = new();
 			await using ISession session = FakeServerSessionFactory.CreateSession(fakeServer.GetConnection(), fakeServer.GetGrpcChannel());
 
-			LeaseManager manager = new LeaseManager(session, null!, serviceProvider.GetRequiredService<StatusService>(),
-				serviceProvider.GetRequiredService<IEnumerable<LeaseHandler>>(), serviceProvider.GetRequiredService<LeaseLoggerFactory>(),
-				_loggerFactory.CreateLogger<LeaseManager>());
+			LeaseManager manager = new LeaseManager(session, serviceProvider);
 
 			Task handleSessionTask = Task.Run(() => manager.RunAsync(false, cts.Token), cts.Token);
 			await fakeServer.UpdateSessionReceived.Task.WaitAsync(cts.Token);
diff --git a/Engine/Source/Programs/Horde/Horde.Agent/AgentApp.cs b/Engine/Source/Programs/Horde/Horde.Agent/AgentApp.cs
--- a/Engine/Source/Programs/Horde/Horde.Agent/AgentApp.cs
+++ b/Engine/Source/Programs/Horde/Horde.Agent/AgentApp.cs
@@ -247,6 +247,8 @@ namespace Horde.Agent
 
 			services.AddMemoryCache();
 
+			services.AddSingleton<ISystemMetrics>(sp => CreateSystemMetrics(sp.GetRequiredService<ILogger<ISystemMetrics>>()));
+
 			// Allow commands to augment the service collection for their own DI service providers
 			services.AddSingleton<DefaultServices>(x => new DefaultServices(configuration, services));
 
@@ -255,6 +257,22 @@ namespace Horde.Agent
 			return await CommandHost.RunAsync(arguments, serviceProvider, typeof(Commands.Service.RunCommand));
 		}
 
+		static ISystemMetrics CreateSystemMetrics(ILogger logger)
+		{
+			if (RuntimeInformation.IsOSPlatform(OSPlatform.Windows))
+			{
+				try
+				{
+					return new WindowsSystemMetrics();
+				}
+				catch (Exception e)
+				{
+					logger.LogError(e, "Unable to initialize system metric collector for telemetry. Disabling. Reason: {Message}", e.Message);
+				}
+			}
+			return new DefaultSystemMetrics();
+		}
+
 		static IConfiguration CreateConfig(bool readInstalledConfig, FileReference? agentConfigFile, Dictionary<string, string?> configOverrides)
 		{
 			string? environment = Environment.GetEnvironmentVariable("DOTNET_ENVIRONMENT");
diff --git a/Engine/Source/Programs/Horde/Horde.Agent/AgentSettings.cs b/Engine/Source/Programs/Horde/Horde.Agent/AgentSettings.cs
--- a/Engine/Source/Programs/Horde/Horde.Agent/AgentSettings.cs
+++ b/Engine/Source/Programs/Horde/Horde.Agent/AgentSettings.cs
@@ -279,7 +279,7 @@ namespace Horde.Agent
 		/// <summary>
 		/// Whether to send telemetry back to Horde server
 		/// </summary>
-		public bool EnableTelemetry { get; set; } = true;
+		public bool EnableTelemetry { get; set; } = false;
 
 		/// <summary>
 		/// How often to report telemetry events to server in milliseconds
diff --git a/Engine/Source/Programs/Horde/Horde.Agent/Leases/LeaseManager.cs b/Engine/Source/Programs/Horde/Horde.Agent/Leases/LeaseManager.cs
--- a/Engine/Source/Programs/Horde/Horde.Agent/Leases/LeaseManager.cs
+++ b/Engine/Source/Programs/Horde/Horde.Agent/Leases/LeaseManager.cs
@@ -120,24 +120,26 @@ namespace Horde.Agent.Leases
 		readonly ISession _session;
 		readonly CapabilitiesService _capabilitiesService;
 		readonly StatusService _statusService;
+		readonly ISystemMetrics _systemMetrics;
 		readonly Dictionary<string, LeaseHandler> _typeUrlToLeaseHandler;
 		readonly LeaseLoggerFactory _leaseLoggerFactory;
 		readonly ILogger _logger;
 
 		RpcAgentCapabilities? _capabilities;
 
-		public LeaseManager(ISession session, CapabilitiesService capabilitiesService, StatusService statusService, IEnumerable<LeaseHandler> leaseHandlers, LeaseLoggerFactory leaseLoggerFactory, ILogger logger)
+		public LeaseManager(ISession session, CapabilitiesService capabilitiesService, StatusService statusService, ISystemMetrics systemMetrics, IEnumerable<LeaseHandler> leaseHandlers, LeaseLoggerFactory leaseLoggerFactory, ILogger logger)
 		{
 			_session = session;
 			_capabilitiesService = capabilitiesService;
 			_statusService = statusService;
+			_systemMetrics = systemMetrics;
 			_typeUrlToLeaseHandler = leaseHandlers.ToDictionary(x => x.LeaseType, x => x);
 			_leaseLoggerFactory = leaseLoggerFactory;
 			_logger = logger;
 		}
 
 		public LeaseManager(ISession session, IServiceProvider serviceProvider)
-			: this(session, serviceProvider.GetRequiredService<CapabilitiesService>(), serviceProvider.GetRequiredService<StatusService>(), serviceProvider.GetRequiredService<IEnumerable<LeaseHandler>>(), serviceProvider.GetRequiredService<LeaseLoggerFactory>(), serviceProvider.GetRequiredService<ILogger<LeaseManager>>())
+			: this(session, serviceProvider.GetRequiredService<CapabilitiesService>(), serviceProvider.GetRequiredService<StatusService>(), serviceProvider.GetRequiredService<ISystemMetrics>(), serviceProvider.GetRequiredService<IEnumerable<LeaseHandler>>(), serviceProvider.GetRequiredService<LeaseLoggerFactory>(), serviceProvider.GetRequiredService<ILogger<LeaseManager>>())
 		{
 		}
 
@@ -240,6 +242,9 @@ namespace Horde.Agent.Leases
 			// Run a background task to update the capabilities of this agent
 			await using BackgroundTask updateCapsTask = BackgroundTask.StartNew(ctx => UpdateCapabilitiesBackgroundAsync(_session.WorkingDir, ctx));
 
+			// Run another background task to send telemetry data
+			await using BackgroundTask telemetryTask = BackgroundTask.StartNew(ctx => SendTelemetryAsync(rpcClient, ctx));
+
 			// Loop until we're ready to exit
 			Stopwatch updateCapabilitiesTimer = Stopwatch.StartNew();
 			for (; ; )
@@ -583,5 +588,48 @@ namespace Horde.Agent.Leases
 				}
 			}
 		}
+
+		/// <summary>
+		/// Periodically send agent telemetry to the server
+		/// </summary>
+		async Task SendTelemetryAsync(HordeRpc.HordeRpcClient rpcClient, CancellationToken cancellationToken)
+		{
+			while (!cancellationToken.IsCancellationRequested)
+			{
+				try
+				{
+					CpuMetrics? cpuMetrics = _systemMetrics.GetCpu();
+					MemoryMetrics? memMetrics = _systemMetrics.GetMemory();
+
+					if (cpuMetrics != null || memMetrics != null)
+					{
+						RpcUploadTelemetryRequest request = new RpcUploadTelemetryRequest();
+						if (cpuMetrics != null)
+						{
+							request.UserCpu = cpuMetrics.User;
+							request.SystemCpu = cpuMetrics.System;
+							request.IdleCpu = cpuMetrics.Idle;
+						}
+						if (memMetrics != null)
+						{
+							request.TotalRam = memMetrics.Total / 1024;
+							request.FreeRam = memMetrics.Available / 1024;
+							request.UsedRam = memMetrics.Used / 1024;
+						}
+						await rpcClient.UploadTelemetryAsync(request, new CallOptions(cancellationToken: cancellationToken));
+					}
+				}
+				catch (OperationCanceledException) when (cancellationToken.IsCancellationRequested)
+				{
+					break;
+				}
+				catch (Exception ex)
+				{
+					_logger.LogWarning(ex, "Exception sending telemetry data: {Message}", ex.Message);
+				}
+
+				await Task.Delay(TimeSpan.FromSeconds(30.0), cancellationToken);
+			}
+		}
 	}
 }
diff --git a/Engine/Source/Programs/Horde/Horde.Agent/Services/TelemetryService.cs b/Engine/Source/Programs/Horde/Horde.Agent/Services/TelemetryService.cs
--- a/Engine/Source/Programs/Horde/Horde.Agent/Services/TelemetryService.cs
+++ b/Engine/Source/Programs/Horde/Horde.Agent/Services/TelemetryService.cs
@@ -2,6 +2,7 @@
 
 using System.Diagnostics;
 using System.Runtime.InteropServices;
+using System.Runtime.Versioning;
 using System.Text;
 using EpicGames.Core;
 using EpicGames.Horde.Agents;
@@ -98,28 +99,40 @@ public class MemoryMetrics
 /// <summary>
 /// OS agnostic interface for retrieving system metrics (CPU, memory etc)
 /// </summary>
-public interface ISystemMetrics : IDisposable
+public interface ISystemMetrics
 {
 	/// <summary>
 	/// Get CPU usage metrics
 	/// </summary>
 	/// <returns>An object with CPU usage metrics</returns>
-	CpuMetrics GetCpu();
+	CpuMetrics? GetCpu();
 
 	/// <summary>
 	/// Get memory usage metrics
 	/// </summary>
 	/// <returns>An object with memory usage metrics</returns>
-	MemoryMetrics GetMemory();
+	MemoryMetrics? GetMemory();
 }
 
-// Suppress call sites not available on all platforms
-#pragma warning disable CA1416
+/// <summary>
+/// Default implementation of <see cref="ISystemMetrics"/>
+/// </summary>
+public sealed class DefaultSystemMetrics : ISystemMetrics
+{
+	/// <inheritdoc/>
+	public CpuMetrics? GetCpu()
+		=> null;
+
+	/// <inheritdoc/>
+	public MemoryMetrics? GetMemory()
+		=> null;
+}
 
 /// <summary>
 /// Windows specific implementation for gathering system metrics
 /// </summary>
-public sealed class WindowsSystemMetrics : ISystemMetrics
+[SupportedOSPlatform("windows")]
+public sealed class WindowsSystemMetrics : ISystemMetrics, IDisposable
 {
 	private const string ProcessorInfo = "Processor Information"; // Prefer this over "Processor" as it's more modern
 	private const string Memory = "Memory";
@@ -224,7 +237,7 @@ class TelemetryService : BackgroundService
 	/// <summary>
 	/// Constructor
 	/// </summary>
-	public TelemetryService(WorkerService workerService, JobHandler jobHandler, GrpcService grpcService, IOptions<AgentSettings> settings, ILogger<TelemetryService> logger)
+	public TelemetryService(WorkerService workerService, JobHandler jobHandler, GrpcService grpcService, ISystemMetrics systemMetrics, IOptions<AgentSettings> settings, ILogger<TelemetryService> logger)
 	{
 		_workerService = workerService;
 		_jobHandler = jobHandler;
@@ -232,6 +245,7 @@ class TelemetryService : BackgroundService
 		_agentSettings = settings.Value;
 		_logger = logger;
 		_reportInterval = TimeSpan.FromMilliseconds(_agentSettings.TelemetryReportInterval);
+		_systemMetrics = systemMetrics;
 
 		// Calculate this once at startup as it should not change during lifetime of process
 		_agentMetadataEvent = GetAgentMetadataEvent();
@@ -241,29 +255,12 @@ class TelemetryService : BackgroundService
 	public override void Dispose()
 	{
 		base.Dispose();
-		_systemMetrics?.Dispose();
 		_eventLoopHeartbeatCts?.Dispose();
 	}
 
 	/// <inheritdoc />
 	public override Task StartAsync(CancellationToken cancellationToken)
 	{
-		if (RuntimeInformation.IsOSPlatform(OSPlatform.Windows))
-		{
-			try
-			{
-				_systemMetrics ??= new WindowsSystemMetrics();
-			}
-			catch (Exception e)
-			{
-				_logger.LogError(e, "Unable to initialize system metric collector for telemetry. Disabling. Reason: {Message}", e.Message);
-			}
-		}
-		else
-		{
-			_logger.LogInformation("System metric collection only implemented on Windows");
-		}
-
 		_eventLoopHeartbeatCts = CancellationTokenSource.CreateLinkedTokenSource(cancellationToken);
 		_eventLoopTask = EventLoopHeartbeatAsync(_eventLoopHeartbeatCts.Token);
 
@@ -455,30 +452,14 @@ class TelemetryService : BackgroundService
 	}
 
 	/// <inheritdoc />
-	protected override Task ExecuteAsync(CancellationToken stoppingToken)
-	{
-		return Task.CompletedTask;
-	}
-
-	bool SendLegacyMetrics { get; set; } = false;
-
-	/// <summary>
-	/// Create a background task that sends telemetry for a session
-	/// </summary>
-	public BackgroundTask CreateBackgroundTask(AgentId agentId)
-	{
-		return BackgroundTask.StartNew(ctx => ExecuteBackgroundAsync(agentId, ctx));
-	}
-
-	/// <inheritdoc />
-	async Task ExecuteBackgroundAsync(AgentId agentId, CancellationToken stoppingToken)
+	protected override async Task ExecuteAsync(CancellationToken stoppingToken)
 	{
 		_logger.LogDebug("Starting telemetry background task");
 		while (!stoppingToken.IsCancellationRequested)
 		{
 			try
 			{
-				if (!await ExecuteBackgroundInternalAsync(agentId, stoppingToken))
+				if (!await ExecuteBackgroundInternalAsync(stoppingToken))
 				{
 					break;
 				}
@@ -505,7 +486,7 @@ class TelemetryService : BackgroundService
 		_logger.LogDebug("Stopping telemetry background task");
 	}
 
-	private async Task<bool> ExecuteBackgroundInternalAsync(AgentId agentId, CancellationToken stoppingToken)
+	private async Task<bool> ExecuteBackgroundInternalAsync(CancellationToken stoppingToken)
 	{
 		if (_systemMetrics == null || !_agentSettings.EnableTelemetry)
 		{
@@ -520,60 +501,46 @@ class TelemetryService : BackgroundService
 		{
 			_logger.LogDebug("Sending telemetry events to server...");
 
-			CpuMetrics cpuMetrics = _systemMetrics.GetCpu();
-			MemoryMetrics memMetrics = _systemMetrics.GetMemory();
+			CpuMetrics? cpuMetrics = _systemMetrics.GetCpu();
+			MemoryMetrics? memMetrics = _systemMetrics.GetMemory();
 
-			// Mongo method
+			RpcSendTelemetryEventsRequest request = new();
+			Timestamp utcNow = Timestamp.FromDateTime(DateTime.UtcNow);
+			RpcExecutionMetadata em = new()
 			{
-				RpcUploadTelemetryRequest request = new RpcUploadTelemetryRequest();
-				request.UserCpu = cpuMetrics.User;
-				request.SystemCpu = cpuMetrics.System;
-				request.IdleCpu = cpuMetrics.Idle;
-				request.TotalRam = memMetrics.Total / 1024;
-				request.FreeRam = memMetrics.Available / 1024;
-				request.UsedRam = memMetrics.Used / 1024;
-				await client.UploadTelemetryAsync(request, new CallOptions(cancellationToken: stoppingToken));
-			}
+				LeaseId = _jobHandler.CurrentLeaseId.ToString(),
+				JobId = _jobHandler.CurrentJobId,
+				JobBatchId = _jobHandler.CurrentBatchId,
+			};
 
-			// Clickhouse method
-			if (SendLegacyMetrics)
+			if (cpuMetrics != null)
 			{
-				RpcSendTelemetryEventsRequest request = new();
-				Timestamp utcNow = Timestamp.FromDateTime(DateTime.UtcNow);
-				RpcExecutionMetadata em = new()
-				{
-					LeaseId = _jobHandler.CurrentLeaseId.ToString(),
-					JobId = _jobHandler.CurrentJobId,
-					JobBatchId = _jobHandler.CurrentBatchId,
-				};
-
-				{
-					RpcAgentCpuMetricsEvent cpuMetricsEvent = cpuMetrics.ToEvent();
-					cpuMetricsEvent.AgentId = _agentMetadataEvent.AgentId;
-					cpuMetricsEvent.Timestamp = utcNow;
-					cpuMetricsEvent.ExecutionMetadata = em;
-					request.Events.Add(new RpcWrappedTelemetryEvent { Cpu = cpuMetricsEvent });
-				}
-
-				{
-					RpcAgentMemoryMetricsEvent memMetricsEvent = memMetrics.ToEvent();
-					memMetricsEvent.AgentId = _agentMetadataEvent.AgentId;
-					memMetricsEvent.Timestamp = utcNow;
-					memMetricsEvent.ExecutionMetadata = em;
-					request.Events.Add(new RpcWrappedTelemetryEvent { Mem = memMetricsEvent });
-				}
+				RpcAgentCpuMetricsEvent cpuMetricsEvent = cpuMetrics.ToEvent();
+				cpuMetricsEvent.AgentId = _agentMetadataEvent.AgentId;
+				cpuMetricsEvent.Timestamp = utcNow;
+				cpuMetricsEvent.ExecutionMetadata = em;
+				request.Events.Add(new RpcWrappedTelemetryEvent { Cpu = cpuMetricsEvent });
+			}
 
-				if (DateTime.UtcNow > _lastTimeAgentMetadataSent + _agentMetadataReportInterval)
-				{
-					// Report agent metadata every now and then as events are not guaranteed to be delivered.
-					// Re-sending ensures the metadata will eventually make it to the server.
-					request.Events.Add(new RpcWrappedTelemetryEvent { AgentMetadata = _agentMetadataEvent });
-					_lastTimeAgentMetadataSent = DateTime.UtcNow;
-				}
+			if (memMetrics != null)
+			{
+				RpcAgentMemoryMetricsEvent memMetricsEvent = memMetrics.ToEvent();
+				memMetricsEvent.AgentId = _agentMetadataEvent.AgentId;
+				memMetricsEvent.Timestamp = utcNow;
+				memMetricsEvent.ExecutionMetadata = em;
+				request.Events.Add(new RpcWrappedTelemetryEvent { Mem = memMetricsEvent });
+			}
 
-				await client.SendTelemetryEventsAsync(request, new CallOptions(cancellationToken: stoppingToken));
+			if (DateTime.UtcNow > _lastTimeAgentMetadataSent + _agentMetadataReportInterval)
+			{
+				// Report agent metadata every now and then as events are not guaranteed to be delivered.
+				// Re-sending ensures the metadata will eventually make it to the server.
+				request.Events.Add(new RpcWrappedTelemetryEvent { AgentMetadata = _agentMetadataEvent });
+				_lastTimeAgentMetadataSent = DateTime.UtcNow;
 			}
 
+			await client.SendTelemetryEventsAsync(request, new CallOptions(cancellationToken: stoppingToken));
+
 			await Task.Delay(_reportInterval, stoppingToken);
 		}
 
diff --git a/Engine/Source/Programs/Horde/Horde.Agent/Services/WorkerService.cs b/Engine/Source/Programs/Horde/Horde.Agent/Services/WorkerService.cs
--- a/Engine/Source/Programs/Horde/Horde.Agent/Services/WorkerService.cs
+++ b/Engine/Source/Programs/Horde/Horde.Agent/Services/WorkerService.cs
@@ -103,7 +103,7 @@ namespace Horde.Agent.Services
 
 						await using (ISession session = await _sessionFactory.CreateAsync(stoppingToken))
 						{
-							_currentLeaseManager = new LeaseManager(session, _capabilitiesService, _statusService, _leaseHandlers, _leaseLoggerFactory, _logger);
+							_currentLeaseManager = new LeaseManager(session, _serviceProvider);
 							result = await _currentLeaseManager.RunAsync(false, stoppingToken);
 						}
 
