Horde: Add missing agent id to telemetry events.

#rnx

diff --git a/Engine/Source/Programs/Horde/Horde.Agent/Leases/LeaseManager.cs b/Engine/Source/Programs/Horde/Horde.Agent/Leases/LeaseManager.cs
--- a/Engine/Source/Programs/Horde/Horde.Agent/Leases/LeaseManager.cs
+++ b/Engine/Source/Programs/Horde/Horde.Agent/Leases/LeaseManager.cs
@@ -604,6 +604,7 @@ namespace Horde.Agent.Leases
 					if (cpuMetrics != null || memMetrics != null)
 					{
 						RpcUploadTelemetryRequest request = new RpcUploadTelemetryRequest();
+						request.AgentId = _session.AgentId.ToString();
 						if (cpuMetrics != null)
 						{
 							request.UserCpu = cpuMetrics.User;
