Horde: Additional logging for telemetry writing.

#rnx

diff --git a/Engine/Source/Programs/Horde/Horde.Server/Agents/Telemetry/AgentTelemetryCollection.cs b/Engine/Source/Programs/Horde/Horde.Server/Agents/Telemetry/AgentTelemetryCollection.cs
--- a/Engine/Source/Programs/Horde/Horde.Server/Agents/Telemetry/AgentTelemetryCollection.cs
+++ b/Engine/Source/Programs/Horde/Horde.Server/Agents/Telemetry/AgentTelemetryCollection.cs
@@ -117,6 +117,7 @@ namespace Horde.Server.Agents.Telemetry
 		{
 			AgentTelemetryDocument document = new AgentTelemetryDocument();
 			document.AgentId = agentId;
+			document.TimeUtc = _clock.UtcNow;
 			document.UserCpu = telemetry.UserCpuPct;
 			document.IdleCpu = telemetry.IdleCpuPct;
 			document.SystemCpu = telemetry.SystemCpuPct;
diff --git a/Engine/Source/Programs/Horde/Horde.Server/Server/RpcService.cs b/Engine/Source/Programs/Horde/Horde.Server/Server/RpcService.cs
--- a/Engine/Source/Programs/Horde/Horde.Server/Server/RpcService.cs
+++ b/Engine/Source/Programs/Horde/Horde.Server/Server/RpcService.cs
@@ -524,6 +524,8 @@ namespace Horde.Server.Server
 		/// <inheritdoc/>
 		public override Task<Empty> UploadTelemetry(RpcUploadTelemetryRequest request, ServerCallContext context)
 		{
+			_logger.LogDebug("Posting telemetry data for {AgentId}", new AgentId(request.AgentId));
+
 			NewAgentTelemetry telemetry = new NewAgentTelemetry(request.UserCpu, request.IdleCpu, request.SystemCpu, (int)request.FreeRam, (int)request.UsedRam, (int)request.TotalRam);
 			_agentTelemetryCollection.Add(new AgentId(request.AgentId), telemetry);
 
diff --git a/Engine/Source/Programs/Horde/Horde.Server/Utilities/MongoBufferedWriter.cs b/Engine/Source/Programs/Horde/Horde.Server/Utilities/MongoBufferedWriter.cs
--- a/Engine/Source/Programs/Horde/Horde.Server/Utilities/MongoBufferedWriter.cs
+++ b/Engine/Source/Programs/Horde/Horde.Server/Utilities/MongoBufferedWriter.cs
@@ -40,17 +40,24 @@ namespace Horde.Server.Utilities
 			_logger = logger;
 		}
 
+		/// <inheritdoc/>
 		public async ValueTask DisposeAsync()
 		{
 			await _backgroundTask.DisposeAsync();
 		}
 
+		/// <summary>
+		/// Start the background task to periodically flush data to the DB
+		/// </summary>
 		public ValueTask StartAsync()
 		{
 			_backgroundTask.Start();
 			return new ValueTask();
 		}
 
+		/// <summary>
+		/// Stops the background task
+		/// </summary>
 		public async ValueTask StopAsync(CancellationToken cancellationToken)
 		{
 			await _backgroundTask.StopAsync(cancellationToken);
@@ -64,13 +71,21 @@ namespace Horde.Server.Utilities
 
 			while (!cancellationToken.IsCancellationRequested)
 			{
-				await newDataTask.WaitAsync(cancellationToken);
-				await Task.WhenAny(flushTask, Task.Delay(_flushTime, cancellationToken));
+				try
+				{
+					await newDataTask.WaitAsync(cancellationToken);
+					await Task.WhenAny(flushTask, Task.Delay(_flushTime, cancellationToken));
 
-				newDataTask = _newDataEvent.Task;
-				flushTask = _flushEvent.Task;
+					newDataTask = _newDataEvent.Task;
+					flushTask = _flushEvent.Task;
 
-				await FlushAsync(cancellationToken);
+					await FlushAsync(cancellationToken);
+				}
+				catch (Exception ex)
+				{
+					_logger.LogError(ex, "Exception in buffered writer: {Message}", ex.Message);
+					await Task.Delay(TimeSpan.FromSeconds(30.0), cancellationToken);
+				}
 			}
 		}
 
@@ -78,17 +93,17 @@ namespace Horde.Server.Utilities
 		public async ValueTask FlushAsync(CancellationToken cancellationToken)
 		{
 			// Copy all the event documents from the queue
-			List<TDocument> document = new List<TDocument>(_queue.Count);
-			while (_queue.TryDequeue(out TDocument? eventDocument))
+			List<TDocument> documents = new List<TDocument>(_queue.Count);
+			while (_queue.TryDequeue(out TDocument? document))
 			{
-				document.Add(eventDocument);
+				documents.Add(document);
 			}
 
 			// Insert them into the database
-			if (document.Count > 0)
+			if (documents.Count > 0)
 			{
-				_logger.LogInformation("Writing {NumEvents} new telemetry events to {CollectionName}.", document.Count, _collection.CollectionNamespace.CollectionName);
-				await _collection.InsertManyAsync(document, cancellationToken: cancellationToken);
+				_logger.LogInformation("Writing {NumEvents} new telemetry events to {CollectionName}.", documents.Count, _collection.CollectionNamespace.CollectionName);
+				await _collection.InsertManyAsync(documents, cancellationToken: cancellationToken);
 			}
 		}
 
