Horde: Fix telemetry not being flushed from non-worker runmodes.

#rnx

diff --git a/Engine/Source/Programs/Horde/Horde.Server/Startup.cs b/Engine/Source/Programs/Horde/Horde.Server/Startup.cs
--- a/Engine/Source/Programs/Horde/Horde.Server/Startup.cs
+++ b/Engine/Source/Programs/Horde/Horde.Server/Startup.cs
@@ -845,8 +845,10 @@ namespace Horde.Server
 			services.AddHttpClient(EpicTelemetrySink.HttpClientName, client => { });
 			services.AddHttpClient(ClickHouseTelemetrySink.HttpClientName, client => { });
 
+			// Create the agent telemetry collection, and register the hosted service so we can flush from any server.
 			services.AddSingleton<AgentTelemetryCollection>();
 			services.AddSingleton<IAgentTelemetryCollection>(sp => sp.GetRequiredService<AgentTelemetryCollection>());
+			services.AddHostedService(provider => provider.GetRequiredService<AgentTelemetryCollection>());
 
 			// Hosted service that needs to run no matter the run mode of the process (server vs worker)
 			services.AddHostedService(provider => (DowntimeService)provider.GetRequiredService<IDowntimeService>());
@@ -860,7 +862,6 @@ namespace Horde.Server
 			if (settings.IsRunModeActive(RunMode.Worker) && !settings.MongoReadOnlyMode)
 			{
 				services.AddHostedService<AgentReportService>();
-				services.AddHostedService(provider => provider.GetRequiredService<AgentTelemetryCollection>());
 				services.AddHostedService<BisectService>();
 				services.AddHostedService(provider => provider.GetRequiredService<FleetService>());
 				services.AddHostedService(provider => provider.GetRequiredService<ConsistencyService>());
diff --git a/Engine/Source/Programs/Horde/Horde.Server/Utilities/MongoBufferedWriter.cs b/Engine/Source/Programs/Horde/Horde.Server/Utilities/MongoBufferedWriter.cs
--- a/Engine/Source/Programs/Horde/Horde.Server/Utilities/MongoBufferedWriter.cs
+++ b/Engine/Source/Programs/Horde/Horde.Server/Utilities/MongoBufferedWriter.cs
@@ -102,7 +102,7 @@ namespace Horde.Server.Utilities
 			// Insert them into the database
 			if (documents.Count > 0)
 			{
-				_logger.LogInformation("Writing {NumEvents} new telemetry events to {CollectionName}.", documents.Count, _collection.CollectionNamespace.CollectionName);
+				_logger.LogDebug("Writing {NumEvents} new telemetry events to {CollectionName}.", documents.Count, _collection.CollectionNamespace.CollectionName);
 				await _collection.InsertManyAsync(documents, cancellationToken: cancellationToken);
 			}
 		}
