Horde: Updating schema documentation

#rnx

diff --git a/Engine/Source/Programs/Horde/Docs/Deployment/AgentSettings.md b/Engine/Source/Programs/Horde/Docs/Deployment/AgentSettings.md
--- a/Engine/Source/Programs/Horde/Docs/Deployment/AgentSettings.md
+++ b/Engine/Source/Programs/Horde/Docs/Deployment/AgentSettings.md
@@ -1,6 +1,6 @@
-[Horde](../../README.md) > [Deployment](../Deployment.md) > [Agent](Agent.md) > appsettings.json (Agent)
+[Horde](../../README.md) > [Deployment](../Deployment.md) > [Agent](Agent.md) > Agent.json (Agent)
 
-# appsettings.json (Agent)
+# Agent.json (Agent)
 
 All Horde-specific settings are stored in a root object called `Horde`. Other .NET functionality may be configured using properties in the root of this file.
 
diff --git a/Engine/Source/Programs/Horde/Docs/Deployment/ServerSettings.md b/Engine/Source/Programs/Horde/Docs/Deployment/ServerSettings.md
--- a/Engine/Source/Programs/Horde/Docs/Deployment/ServerSettings.md
+++ b/Engine/Source/Programs/Horde/Docs/Deployment/ServerSettings.md
@@ -18,7 +18,8 @@ Name | Description
 `databaseConnectionString` | `string`<br>MongoDB connection string
 `mongoDatabaseName` | `string`<br>MongoDB database name
 `databaseName` | `string`<br>
-`databasePublicCert` | `string`<br>Optional certificate to trust in order to access the database (eg. AWS public cert for TLS)
+`mongoPublicCertificate` | `string`<br>Optional certificate to trust in order to access the database (eg. AWS public cert for TLS)
+`databasePublicCert` | `string`<br>
 `mongoReadOnlyMode` | `boolean`<br>Access the database in read-only mode (avoids creating indices or updating content) Useful for debugging a local instance of HordeServer against a production database.
 `databaseReadOnlyMode` | `boolean`<br>
 `shutdownMemoryThreshold` | `integer`<br>Shutdown the current server process if memory usage reaches this threshold (specified in MB)<br>Usually set to 80-90% of available memory to avoid CLR heap using all of it. If a memory leak was to occur, it's usually better to restart the process rather than to let the GC work harder and harder trying to recoup memory.<br>Should only be used when multiple server processes are running behind a load balancer and one can be safely restarted automatically by the underlying process handler (Docker, Kubernetes, AWS ECS, Supervisor etc). The shutdown behaves similar to receiving a SIGTERM and will wait for outstanding requests to finish.
