Horde: Fix references to appsettings.json in docs.

#rnx

diff --git a/Engine/Source/Programs/Horde/Docs/Deployment/Agent.md b/Engine/Source/Programs/Horde/Docs/Deployment/Agent.md
--- a/Engine/Source/Programs/Horde/Docs/Deployment/Agent.md
+++ b/Engine/Source/Programs/Horde/Docs/Deployment/Agent.md
@@ -51,7 +51,7 @@ under other root keys.
 
 ### Server Profiles
 
-The agent's [`appsettings.json`](AgentSettings.md) file can contain settings for connecting to multiple Horde servers
+The agent's [`Agent.json`](AgentSettings.md) file can contain settings for connecting to multiple Horde servers
 through the `ServerProfiles` property. Setting up multiple profiles can be useful when running Horde in multiple
 environments (eg. dev vs production), and each [server profile](AgentSettings.md#serverprofile) contains a name,
 server URL and authentication token.
@@ -200,7 +200,7 @@ Launch the daemon:
 The default location for data used by the agent (Perforce workspaces, caches, scratch space) is
 `C:\ProgramData\HordeAgent` on Windows and the application directory on Mac/Linux.
 
-This path can be overridden using the `WorkingDir` property in the agent's [`appsettings.json`](AgentSettings.md) file.
+This path can be overridden using the `WorkingDir` property in the agent's [`Agent.json`](AgentSettings.md) file.
 
 ITo prevent agents from being affected by runaway jobs filling up the disk with data, 
 having the agent store data on a drive other than the system disk is a good idea.  
@@ -209,7 +209,7 @@ Setting the `%TEMP%` and `%TMP%` environment variables to this drive on Windows
 ### Mounting Network Shares
 
 The agent can be configured to mount certain network shares at startup before taking on any work using the `Shares`
-property in the agent's [`appsettings.json`](AgentSettings.md) file.
+property in the agent's [`Agent.json`](AgentSettings.md) file.
 
 ### GPU Access
 
