Horde: Remove support for legacy artifact uploads.

diff --git a/Engine/Source/Programs/Horde/Horde.Agent.Tests/Services/TelemetryServiceTest.cs b/Engine/Source/Programs/Horde/Horde.Agent.Tests/Services/TelemetryServiceTest.cs
--- a/Engine/Source/Programs/Horde/Horde.Agent.Tests/Services/TelemetryServiceTest.cs
+++ b/Engine/Source/Programs/Horde/Horde.Agent.Tests/Services/TelemetryServiceTest.cs
@@ -25,7 +25,7 @@ public sealed class TelemetryServiceTest : IDisposable
 		AgentSettings settings = new() { Server = "Test", ServerProfiles = { ["Test"] = new ServerProfile() { Name = "Test", Url = new Uri("http://localhost:1234") } } };
 		OptionsWrapper<AgentSettings> settingsOpt = new(settings);
 
-		using WorkerService workerService = new(null!, null!, null!, new List<LeaseHandler>(), null!, null!, null!);
+		using WorkerService workerService = new WorkerService(null!, null!, null!, null!);
 		JobHandler jobHandler = new JobHandler(new List<IJobExecutorFactory>(), settingsOpt, null!, null!);
 		GrpcService grpcService = new(settingsOpt, NullLogger<GrpcService>.Instance, _loggerFactory);
 
diff --git a/Engine/Source/Programs/Horde/Horde.Agent.Tests/WorkerServiceTest.cs b/Engine/Source/Programs/Horde/Horde.Agent.Tests/WorkerServiceTest.cs
--- a/Engine/Source/Programs/Horde/Horde.Agent.Tests/WorkerServiceTest.cs
+++ b/Engine/Source/Programs/Horde/Horde.Agent.Tests/WorkerServiceTest.cs
@@ -42,7 +42,6 @@ namespace Horde.Agent.Tests
 	{
 		private readonly ServiceCollection _serviceCollection;
 
-		private readonly ILoggerFactory _loggerFactory;
 		private readonly JobId _jobId = JobId.Parse("65bd0655591b5d5d7d047b58");
 		private readonly JobStepBatchId _batchId = new JobStepBatchId(0x1234);
 		private readonly JobStepId _stepId1 = new JobStepId(1);
@@ -78,12 +77,6 @@ namespace Horde.Agent.Tests
 
 		public WorkerServiceTest()
 		{
-			_loggerFactory = LoggerFactory.Create(builder =>
-			{
-				builder.SetMinimumLevel(LogLevel.Debug);
-				builder.AddSimpleConsole(options => { options.SingleLine = true; });
-			});
-
 			_serviceCollection = new ServiceCollection();
 			_serviceCollection.AddLogging();
 			_serviceCollection.AddHordeHttpClient();
diff --git a/Engine/Source/Programs/Horde/Horde.Agent/Execution/JobExecutor.cs b/Engine/Source/Programs/Horde/Horde.Agent/Execution/JobExecutor.cs
--- a/Engine/Source/Programs/Horde/Horde.Agent/Execution/JobExecutor.cs
+++ b/Engine/Source/Programs/Horde/Horde.Agent/Execution/JobExecutor.cs
@@ -435,38 +435,6 @@ namespace Horde.Agent.Execution
 			}
 		}
 
-		private async Task StorePreprocessedFileAsync(FileReference? localFile, JobStepId stepId, DirectoryReference? sharedStorageDir, ILogger logger, CancellationToken cancellationToken)
-		{
-			if (localFile != null)
-			{
-				string fileName = localFile.GetFileName();
-				await ArtifactUploader.UploadAsync(RpcConnection, JobId, BatchId, stepId, fileName, localFile, logger, cancellationToken);
-
-				if (sharedStorageDir != null)
-				{
-					FileReference remoteFile = FileReference.Combine(sharedStorageDir, fileName);
-					DirectoryReference.CreateDirectory(remoteFile.Directory);
-					FileReference.Copy(localFile, remoteFile);
-				}
-			}
-		}
-
-		private static void FetchPreprocessedFile(FileReference localFile, DirectoryReference? sharedStorageDir, ILogger logger)
-		{
-			if (!FileReference.Exists(localFile))
-			{
-				if (sharedStorageDir == null)
-				{
-					throw new FileNotFoundException($"Missing preprocessed script {localFile}");
-				}
-
-				FileReference remoteFile = FileReference.Combine(sharedStorageDir, localFile.GetFileName());
-				logger.LogInformation("Copying {RemoteFile} to {LocalFile}", remoteFile, localFile);
-				DirectoryReference.CreateDirectory(localFile.Directory);
-				FileReference.Copy(remoteFile, localFile, false);
-			}
-		}
-
 		protected abstract Task<bool> SetupAsync(JobStepInfo step, ILogger logger, CancellationToken cancellationToken);
 
 		protected abstract Task<bool> ExecuteAsync(JobStepInfo step, ILogger logger, CancellationToken cancellationToken);
@@ -479,7 +447,7 @@ namespace Horde.Agent.Execution
 			return StorageFactory.CreateClient(namespaceId, token);
 		}
 
-		protected virtual async Task<bool> SetupAsync(JobStepInfo step, DirectoryReference workspaceDir, DirectoryReference? sharedStorageDir, bool? useP4, ILogger logger, CancellationToken cancellationToken)
+		protected virtual async Task<bool> SetupAsync(JobStepInfo step, DirectoryReference workspaceDir, bool? useP4, ILogger logger, CancellationToken cancellationToken)
 		{
 			FileReference definitionFile = FileReference.Combine(workspaceDir, "Engine", "Saved", "Horde", "Exported.json");
 
@@ -515,76 +483,62 @@ namespace Horde.Agent.Execution
 				preprocessedSchemaFile = FileReference.Combine(workspaceDir, PreprocessedSchema);
 				arguments.AppendArgument("-Schema=", preprocessedSchemaFile.FullName);
 			}
-			if (sharedStorageDir != null && !_preprocessScript) // Do not precompile when preprocessing the script; other agents may have a different view of UAT
-			{
-				DirectoryReference buildDir = GetAutomationToolDir(sharedStorageDir);
-				arguments.Append($" CopyUAT -WithLauncher -TargetDir=\"{buildDir}\"");
-			}
 
-			int result = await ExecuteAutomationToolAsync(step, workspaceDir, sharedStorageDir, arguments.ToString(), useP4, logger, cancellationToken);
+			int result = await ExecuteAutomationToolAsync(step, workspaceDir, arguments.ToString(), useP4, logger, cancellationToken);
 			if (result != 0)
 			{
 				return false;
 			}
 
-			if (JobOptions.UseNewTempStorage ?? true)
+			List<FileReference> buildGraphFiles = new List<FileReference>();
+			buildGraphFiles.Add(definitionFile);
+			if (preprocessedScriptFile != null)
 			{
-				List<FileReference> buildGraphFiles = new List<FileReference>();
-				buildGraphFiles.Add(definitionFile);
-				if (preprocessedScriptFile != null)
-				{
-					buildGraphFiles.Add(preprocessedScriptFile);
-				}
-				if (preprocessedSchemaFile != null)
-				{
-					buildGraphFiles.Add(preprocessedSchemaFile);
-				}
-
-				using (GlobalTracer.Instance.BuildSpan("TempStorage").WithTag("resource", "Write").StartActive())
-				{
-					// Create the artifact
-					using IRpcClientRef<JobRpc.JobRpcClient> jobRpc = await RpcConnection.GetClientRefAsync<JobRpc.JobRpcClient>(cancellationToken);
+				buildGraphFiles.Add(preprocessedScriptFile);
+			}
+			if (preprocessedSchemaFile != null)
+			{
+				buildGraphFiles.Add(preprocessedSchemaFile);
+			}
 
-					ArtifactName artifactName = TempStorage.GetArtifactNameForNode(SetupStepName);
-					ArtifactType artifactType = ArtifactType.StepOutput;
+			using (GlobalTracer.Instance.BuildSpan("TempStorage").WithTag("resource", "Write").StartActive())
+			{
+				// Create the artifact
+				using IRpcClientRef<JobRpc.JobRpcClient> jobRpc = await RpcConnection.GetClientRefAsync<JobRpc.JobRpcClient>(cancellationToken);
 
-					RpcCreateJobArtifactRequestV2 artifactRequest = new RpcCreateJobArtifactRequestV2();
-					artifactRequest.JobId = JobId.ToString();
-					artifactRequest.StepId = step.StepId.ToString();
-					artifactRequest.Name = artifactName.ToString();
-					artifactRequest.Type = artifactType.ToString();
+				ArtifactName artifactName = TempStorage.GetArtifactNameForNode(SetupStepName);
+				ArtifactType artifactType = ArtifactType.StepOutput;
 
-					RpcCreateJobArtifactResponseV2 artifact = await jobRpc.Client.CreateArtifactV2Async(artifactRequest, cancellationToken: cancellationToken);
-					ArtifactId artifactId = ArtifactId.Parse(artifact.Id);
-					logger.LogInformation("Creating output artifact {ArtifactId} '{ArtifactName}' ({ArtifactType}) with ref {RefName} in namespace {NamespaceId}", artifactId, artifactName, artifactType, artifact.RefName, artifact.NamespaceId);
+				RpcCreateJobArtifactRequestV2 artifactRequest = new RpcCreateJobArtifactRequestV2();
+				artifactRequest.JobId = JobId.ToString();
+				artifactRequest.StepId = step.StepId.ToString();
+				artifactRequest.Name = artifactName.ToString();
+				artifactRequest.Type = artifactType.ToString();
 
-					// Write the data
-					using IStorageClient storage = CreateStorageClient(new NamespaceId(artifact.NamespaceId), artifact.Token);
+				RpcCreateJobArtifactResponseV2 artifact = await jobRpc.Client.CreateArtifactV2Async(artifactRequest, cancellationToken: cancellationToken);
+				ArtifactId artifactId = ArtifactId.Parse(artifact.Id);
+				logger.LogInformation("Creating output artifact {ArtifactId} '{ArtifactName}' ({ArtifactType}) with ref {RefName} in namespace {NamespaceId}", artifactId, artifactName, artifactType, artifact.RefName, artifact.NamespaceId);
 
-					Stopwatch timer = Stopwatch.StartNew();
+				// Write the data
+				using IStorageClient storage = CreateStorageClient(new NamespaceId(artifact.NamespaceId), artifact.Token);
 
-					IBlobRef<DirectoryNode> rootNodeRef;
-					await using (IBlobWriter blobWriter = storage.CreateBlobWriter(artifact.RefName))
-					{
-						DirectoryNode buildGraphNode = new DirectoryNode();
-						await buildGraphNode.AddFilesAsync(workspaceDir, buildGraphFiles, blobWriter, cancellationToken: cancellationToken);
-						IBlobRef<DirectoryNode> outputNodeRef = await blobWriter.WriteBlobAsync(buildGraphNode, cancellationToken);
+				Stopwatch timer = Stopwatch.StartNew();
 
-						DirectoryNode rootNode = new DirectoryNode();
-						rootNode.AddDirectory(new DirectoryEntry(BuildGraphTempStorageDir, buildGraphNode.Length, outputNodeRef));
+				IBlobRef<DirectoryNode> rootNodeRef;
+				await using (IBlobWriter blobWriter = storage.CreateBlobWriter(artifact.RefName))
+				{
+					DirectoryNode buildGraphNode = new DirectoryNode();
+					await buildGraphNode.AddFilesAsync(workspaceDir, buildGraphFiles, blobWriter, cancellationToken: cancellationToken);
+					IBlobRef<DirectoryNode> outputNodeRef = await blobWriter.WriteBlobAsync(buildGraphNode, cancellationToken);
 
-						rootNodeRef = await blobWriter.WriteBlobAsync(rootNode, cancellationToken);
-					}
-					await storage.WriteRefAsync(artifact.RefName, rootNodeRef, new RefOptions(), cancellationToken);
+					DirectoryNode rootNode = new DirectoryNode();
+					rootNode.AddDirectory(new DirectoryEntry(BuildGraphTempStorageDir, buildGraphNode.Length, outputNodeRef));
 
-					logger.LogInformation("Upload took {Time:n1}s", timer.Elapsed.TotalSeconds);
+					rootNodeRef = await blobWriter.WriteBlobAsync(rootNode, cancellationToken);
 				}
-			}
-			else
-			{
-				await ArtifactUploader.UploadAsync(RpcConnection, JobId, BatchId, step.StepId, definitionFile.GetFileName(), definitionFile, logger, cancellationToken);
-				await StorePreprocessedFileAsync(preprocessedScriptFile, step.StepId, sharedStorageDir, logger, cancellationToken);
-				await StorePreprocessedFileAsync(preprocessedSchemaFile, step.StepId, sharedStorageDir, logger, cancellationToken);
+				await storage.WriteRefAsync(artifact.RefName, rootNodeRef, new RefOptions(), cancellationToken);
+
+				logger.LogInformation("Upload took {Time:n1}s", timer.Elapsed.TotalSeconds);
 			}
 
 			RpcUpdateGraphRequest updateGraph = await ParseGraphUpdateAsync(definitionFile, logger, cancellationToken);
@@ -796,7 +750,7 @@ namespace Horde.Agent.Execution
 			}
 		}
 
-		protected async Task<bool> ExecuteAsync(JobStepInfo step, DirectoryReference workspaceDir, DirectoryReference? sharedStorageDir, bool? useP4, ILogger logger, CancellationToken cancellationToken)
+		protected async Task<bool> ExecuteAsync(JobStepInfo step, DirectoryReference workspaceDir, bool? useP4, ILogger logger, CancellationToken cancellationToken)
 		{
 			StringBuilder arguments = new StringBuilder("BuildGraph");
 			if (_preprocessScript)
@@ -807,44 +761,35 @@ namespace Horde.Agent.Execution
 				FileReference localPreprocessedSchema = FileReference.Combine(workspaceDir, PreprocessedSchema);
 				arguments.AppendArgument("-ImportSchema=", localPreprocessedSchema.FullName);
 
-				if (JobOptions.UseNewTempStorage ?? true)
-				{
-					ArtifactName artifactName = TempStorage.GetArtifactNameForNode(SetupStepName);
+				ArtifactName artifactName = TempStorage.GetArtifactNameForNode(SetupStepName);
 
-					using IRpcClientRef<JobRpc.JobRpcClient> jobRpc = await RpcConnection.GetClientRefAsync<JobRpc.JobRpcClient>(cancellationToken);
+				using IRpcClientRef<JobRpc.JobRpcClient> jobRpc = await RpcConnection.GetClientRefAsync<JobRpc.JobRpcClient>(cancellationToken);
 
-					RpcGetJobArtifactRequest artifactRequest = new RpcGetJobArtifactRequest();
-					artifactRequest.JobId = JobId.ToString();
-					artifactRequest.StepId = step.StepId.ToString();
-					artifactRequest.Name = artifactName.ToString();
-					artifactRequest.Type = ArtifactType.StepOutput.ToString();
+				RpcGetJobArtifactRequest artifactRequest = new RpcGetJobArtifactRequest();
+				artifactRequest.JobId = JobId.ToString();
+				artifactRequest.StepId = step.StepId.ToString();
+				artifactRequest.Name = artifactName.ToString();
+				artifactRequest.Type = ArtifactType.StepOutput.ToString();
 
-					RpcGetJobArtifactResponse artifact = await jobRpc.Client.GetArtifactAsync(artifactRequest, cancellationToken: cancellationToken);
+				RpcGetJobArtifactResponse artifact = await jobRpc.Client.GetArtifactAsync(artifactRequest, cancellationToken: cancellationToken);
 
-					NamespaceId namespaceId = new NamespaceId(artifact.NamespaceId);
-					RefName refName = new RefName(artifact.RefName);
+				NamespaceId namespaceId = new NamespaceId(artifact.NamespaceId);
+				RefName refName = new RefName(artifact.RefName);
 
-					logger.LogInformation("Reading preprocessed script from {NamespaceId}:{RefName}", namespaceId, refName);
+				logger.LogInformation("Reading preprocessed script from {NamespaceId}:{RefName}", namespaceId, refName);
 
-					using IStorageClient storage = CreateStorageClient(namespaceId, artifact.Token);
+				using IStorageClient storage = CreateStorageClient(namespaceId, artifact.Token);
 
-					DirectoryNode node = await storage.ReadRefTargetAsync<DirectoryNode>(refName, cancellationToken: cancellationToken);
-					DirectoryNode? buildGraphDir = await node.TryOpenDirectoryAsync(BuildGraphTempStorageDir, cancellationToken: cancellationToken);
-					if (buildGraphDir != null)
-					{
-						await buildGraphDir.CopyToDirectoryAsync(new DirectoryInfo(workspaceDir.FullName), null, logger, cancellationToken);
-						logger.LogInformation("Copying preprocessed script from {BuildGraphFolderName} into {OutputDir}", BuildGraphTempStorageDir, workspaceDir);
-					}
-					else
-					{
-						logger.LogInformation("Bundle has no {BuildGraphFolderName} folder; not copying any files", BuildGraphTempStorageDir);
-					}
+				DirectoryNode node = await storage.ReadRefTargetAsync<DirectoryNode>(refName, cancellationToken: cancellationToken);
+				DirectoryNode? buildGraphDir = await node.TryOpenDirectoryAsync(BuildGraphTempStorageDir, cancellationToken: cancellationToken);
+				if (buildGraphDir != null)
+				{
+					await buildGraphDir.CopyToDirectoryAsync(new DirectoryInfo(workspaceDir.FullName), null, logger, cancellationToken);
+					logger.LogInformation("Copying preprocessed script from {BuildGraphFolderName} into {OutputDir}", BuildGraphTempStorageDir, workspaceDir);
 				}
 				else
 				{
-					logger.LogInformation("Fetching preprocessed script from {SharedDir}", sharedStorageDir);
-					FetchPreprocessedFile(localPreprocessedScript, sharedStorageDir, logger);
-					FetchPreprocessedFile(localPreprocessedSchema, sharedStorageDir, logger);
+					logger.LogInformation("Bundle has no {BuildGraphFolderName} folder; not copying any files", BuildGraphTempStorageDir);
 				}
 			}
 			else if (_scriptFileName != null)
@@ -862,33 +807,13 @@ namespace Horde.Agent.Execution
 				}
 			}
 
-			if (JobOptions.UseNewTempStorage ?? true)
-			{
-				bool result = await ExecuteWithTempStorageAsync(step, workspaceDir, arguments.ToString(), useP4, logger, cancellationToken);
-				return result;
-			}
-			else
-			{
-				if (sharedStorageDir != null)
-				{
-					arguments.AppendArgument("-SharedStorageDir=", sharedStorageDir.FullName);
-				}
-
-				bool result = await ExecuteAutomationToolAsync(step, workspaceDir, sharedStorageDir, arguments.ToString(), useP4, logger, cancellationToken) == 0;
-				return result;
-			}
+			bool result = await ExecuteWithTempStorageAsync(step, workspaceDir, arguments.ToString(), useP4, logger, cancellationToken);
+			return result;
 		}
 
 		protected async Task CreateArtifactsAsync(JobStepId stepId, ArtifactName name, ArtifactType type, DirectoryReference baseDir, IEnumerable<(string, FileReference)> files, ILogger logger, CancellationToken cancellationToken)
 		{
-			if (JobOptions.UseNewTempStorage ?? true)
-			{
-				await CreateArtifactAsync(stepId, name, type, baseDir, files.Select(x => x.Item2), logger, cancellationToken);
-			}
-			else
-			{
-				await ArtifactUploader.UploadAsync(RpcConnection, JobId, BatchId, stepId, files, logger, CancellationToken.None);
-			}
+			await CreateArtifactAsync(stepId, name, type, baseDir, files.Select(x => x.Item2), logger, cancellationToken);
 		}
 
 		protected async Task CreateArtifactAsync(JobStepId stepId, ArtifactName name, ArtifactType type, DirectoryReference baseDir, IEnumerable<FileReference> files, ILogger logger, CancellationToken cancellationToken)
@@ -994,7 +919,7 @@ namespace Horde.Agent.Execution
 			}
 
 			// Run UAT
-			if (await ExecuteAutomationToolAsync(step, workspaceDir, null, arguments, useP4, logger, cancellationToken) != 0)
+			if (await ExecuteAutomationToolAsync(step, workspaceDir, arguments, useP4, logger, cancellationToken) != 0)
 			{
 				return false;
 			}
@@ -1219,7 +1144,7 @@ namespace Horde.Agent.Execution
 			return true;
 		}
 
-		protected async Task<int> ExecuteAutomationToolAsync(JobStepInfo step, DirectoryReference workspaceDir, DirectoryReference? sharedStorageDir, string? arguments, bool? useP4, ILogger logger, CancellationToken cancellationToken)
+		protected async Task<int> ExecuteAutomationToolAsync(JobStepInfo step, DirectoryReference workspaceDir, string? arguments, bool? useP4, ILogger logger, CancellationToken cancellationToken)
 		{
 			int result;
 			using IScope scope = GlobalTracer.Instance.BuildSpan("BuildGraph").StartActive();
@@ -1236,7 +1161,7 @@ namespace Horde.Agent.Execution
 
 			if (RuntimeInformation.IsOSPlatform(OSPlatform.Windows))
 			{
-				result = await ExecuteCommandAsync(step, workspaceDir, sharedStorageDir, Environment.GetEnvironmentVariable("COMSPEC") ?? "cmd.exe", $"/C \"\"{workspaceDir}\\Engine\\Build\\BatchFiles\\RunUAT.bat\" {arguments}\"", logger, cancellationToken);
+				result = await ExecuteCommandAsync(step, workspaceDir, Environment.GetEnvironmentVariable("COMSPEC") ?? "cmd.exe", $"/C \"\"{workspaceDir}\\Engine\\Build\\BatchFiles\\RunUAT.bat\" {arguments}\"", logger, cancellationToken);
 			}
 			else if (RuntimeInformation.IsOSPlatform(OSPlatform.Linux))
 			{
@@ -1247,11 +1172,11 @@ namespace Horde.Agent.Execution
 					args = $"\"{workspaceDir}/Engine/Build/BatchFiles/RunWineUAT.sh\" {arguments}";
 				}
 
-				result = await ExecuteCommandAsync(step, workspaceDir, sharedStorageDir, "/bin/bash", args, logger, cancellationToken);
+				result = await ExecuteCommandAsync(step, workspaceDir, "/bin/bash", args, logger, cancellationToken);
 			}
 			else if (RuntimeInformation.IsOSPlatform(OSPlatform.OSX))
 			{
-				result = await ExecuteCommandAsync(step, workspaceDir, sharedStorageDir, "/bin/sh", $"\"{workspaceDir}/Engine/Build/BatchFiles/RunUAT.sh\" {arguments}", logger, cancellationToken);
+				result = await ExecuteCommandAsync(step, workspaceDir, "/bin/sh", $"\"{workspaceDir}/Engine/Build/BatchFiles/RunUAT.sh\" {arguments}", logger, cancellationToken);
 			}
 			else
 			{
@@ -1495,14 +1420,10 @@ namespace Horde.Agent.Execution
 			return await ExecuteProcessAsync(executable, arguments, new Dictionary<string, string>(), filter, logger, cancellationToken);
 		}
 
-		private static List<DirectoryReference> GetContainerMountDirs(DirectoryReference workspaceDir, DirectoryReference? sharedStorageDir, IReadOnlyDictionary<string, string> envVars)
+		private static List<DirectoryReference> GetContainerMountDirs(DirectoryReference workspaceDir, IReadOnlyDictionary<string, string> envVars)
 		{
 			List<DirectoryReference> dirs = new();
 			dirs.Add(workspaceDir);
-			if (sharedStorageDir != null)
-			{
-				dirs.Add(sharedStorageDir);
-			}
 			if (envVars.TryGetValue("UE_SDKS_ROOT", out string? autoSdkDirPath))
 			{
 				dirs.Add(new DirectoryReference(autoSdkDirPath));
@@ -1520,7 +1441,7 @@ namespace Horde.Agent.Execution
 			return name;
 		}
 
-		async Task<int> ExecuteCommandAsync(JobStepInfo step, DirectoryReference workspaceDir, DirectoryReference? sharedStorageDir, string fileName, string arguments, ILogger jobLogger, CancellationToken cancellationToken)
+		async Task<int> ExecuteCommandAsync(JobStepInfo step, DirectoryReference workspaceDir, string fileName, string arguments, ILogger jobLogger, CancellationToken cancellationToken)
 		{
 			// Method for expanding environment variable properties related to this step
 			Dictionary<string, string> properties = new Dictionary<string, string>(StringComparer.OrdinalIgnoreCase);
@@ -1630,7 +1551,7 @@ namespace Horde.Agent.Execution
 				{
 					if (RuntimeInformation.IsOSPlatform(OSPlatform.Linux) && JobOptions?.Container?.Enabled is true)
 					{
-						List<DirectoryReference> mountDirs = GetContainerMountDirs(workspaceDir, sharedStorageDir, newEnvVars);
+						List<DirectoryReference> mountDirs = GetContainerMountDirs(workspaceDir, newEnvVars);
 						exitCode = await ExecuteProcessInContainerAsync(arguments, mountDirs, newEnvVars, filter, jobLogger, cancellationToken);
 					}
 					else
@@ -1705,24 +1626,7 @@ namespace Horde.Agent.Execution
 			if (DirectoryReference.Exists(logDir))
 			{
 				List<FileReference> artifactFiles = DirectoryReference.EnumerateFiles(logDir, "*", SearchOption.AllDirectories).ToList();
-				if (JobOptions.UseNewTempStorage ?? true)
-				{
-					await CreateArtifactAsync(step.StepId, TempStorage.GetArtifactNameForNode(step.Name), ArtifactType.StepSaved, workspaceDir, artifactFiles, jobLogger, cancellationToken);
-				}
-				else
-				{
-					Dictionary<FileReference, string> artifactFileToId = new Dictionary<FileReference, string>();
-					foreach (FileReference artifactFile in artifactFiles)
-					{
-						string artifactName = artifactFile.MakeRelativeTo(logDir);
-
-						string? artifactId = await ArtifactUploader.UploadAsync(RpcConnection, JobId, BatchId, step.StepId, artifactName, artifactFile, jobLogger, cancellationToken);
-						if (artifactId != null)
-						{
-							artifactFileToId[artifactFile] = artifactId;
-						}
-					}
-				}
+				await CreateArtifactAsync(step.StepId, TempStorage.GetArtifactNameForNode(step.Name), ArtifactType.StepSaved, workspaceDir, artifactFiles, jobLogger, cancellationToken);
 
 				foreach (FileReference reportFile in artifactFiles.Where(x => x.HasExtension(".report.json")))
 				{
diff --git a/Engine/Source/Programs/Horde/Horde.Agent/Execution/LocalExecutor.cs b/Engine/Source/Programs/Horde/Horde.Agent/Execution/LocalExecutor.cs
--- a/Engine/Source/Programs/Horde/Horde.Agent/Execution/LocalExecutor.cs
+++ b/Engine/Source/Programs/Horde/Horde.Agent/Execution/LocalExecutor.cs
@@ -45,14 +45,14 @@ namespace Horde.Agent.Execution
 
 		protected override Task<bool> SetupAsync(JobStepInfo step, ILogger logger, CancellationToken cancellationToken)
 		{
-			return SetupAsync(step, _localWorkspaceDir, null, null, logger, cancellationToken);
+			return SetupAsync(step, _localWorkspaceDir, null, logger, cancellationToken);
 		}
 
 		protected override Task<bool> ExecuteAsync(JobStepInfo step, ILogger logger, CancellationToken cancellationToken)
 		{
 			if (_settings.RunSteps)
 			{
-				return ExecuteAsync(step, _localWorkspaceDir, null, null, logger, cancellationToken);
+				return ExecuteAsync(step, _localWorkspaceDir, null, logger, cancellationToken);
 			}
 			else
 			{
diff --git a/Engine/Source/Programs/Horde/Horde.Agent/Execution/PerforceExecutor.cs b/Engine/Source/Programs/Horde/Horde.Agent/Execution/PerforceExecutor.cs
--- a/Engine/Source/Programs/Horde/Horde.Agent/Execution/PerforceExecutor.cs
+++ b/Engine/Source/Programs/Horde/Horde.Agent/Execution/PerforceExecutor.cs
@@ -2,7 +2,6 @@
 
 using System.Runtime.InteropServices;
 using System.Text;
-using System.Text.RegularExpressions;
 using EpicGames.Core;
 using EpicGames.Perforce;
 using EpicGames.Perforce.Managed;
@@ -23,7 +22,6 @@ namespace Horde.Agent.Execution
 		protected RpcAgentWorkspace _workspaceInfo;
 		protected RpcAgentWorkspace? _autoSdkWorkspaceInfo;
 		protected DirectoryReference _rootDir;
-		protected DirectoryReference? _sharedStorageDir;
 
 		protected WorkspaceInfo? _autoSdkWorkspace;
 		protected WorkspaceInfo _workspace;
@@ -137,14 +135,6 @@ namespace Horde.Agent.Execution
 			// Remove all the local settings directories
 			DeleteEngineUserSettings(logger);
 
-			// Get the temp storage directory
-			if (!String.IsNullOrEmpty(Batch.TempStorageDir))
-			{
-				string escapedStreamName = Regex.Replace(Batch.StreamName, "[^a-zA-Z0-9_-]", "+");
-				_sharedStorageDir = DirectoryReference.Combine(new DirectoryReference(Batch.TempStorageDir), escapedStreamName, $"CL {Batch.Change} - Job {JobId}");
-				CopyAutomationTool(_sharedStorageDir, _workspace.WorkspaceDir, logger);
-			}
-
 			// Get all the environment variables for jobs
 			_envVars["IsBuildMachine"] = "1";
 			_envVars["uebp_LOCAL_ROOT"] = _workspace.WorkspaceDir.FullName;
@@ -222,14 +212,14 @@ namespace Horde.Agent.Execution
 		{
 			PerforceMetadataLogger perforceLogger = CreatePerforceLogger(logger, Batch.Change, _workspace, _autoSdkWorkspace);
 			bool useP4 = WorkspaceInfo.ShouldUseHaveTable(_workspaceInfo.Method);
-			return await SetupAsync(step, _workspace.WorkspaceDir, _sharedStorageDir, useP4, perforceLogger, cancellationToken);
+			return await SetupAsync(step, _workspace.WorkspaceDir, useP4, perforceLogger, cancellationToken);
 		}
 
 		protected override async Task<bool> ExecuteAsync(JobStepInfo step, ILogger logger, CancellationToken cancellationToken)
 		{
 			PerforceMetadataLogger perforceLogger = CreatePerforceLogger(logger, Batch.Change, _workspace, _autoSdkWorkspace);
 			bool useP4 = WorkspaceInfo.ShouldUseHaveTable(_workspaceInfo.Method);
-			return await ExecuteAsync(step, _workspace.WorkspaceDir, _sharedStorageDir, useP4, perforceLogger, cancellationToken);
+			return await ExecuteAsync(step, _workspace.WorkspaceDir, useP4, perforceLogger, cancellationToken);
 		}
 
 		public override async Task FinalizeAsync(ILogger logger, CancellationToken cancellationToken)
diff --git a/Engine/Source/Programs/Horde/Horde.Agent/Execution/TestExecutor.cs b/Engine/Source/Programs/Horde/Horde.Agent/Execution/TestExecutor.cs
--- a/Engine/Source/Programs/Horde/Horde.Agent/Execution/TestExecutor.cs
+++ b/Engine/Source/Programs/Horde/Horde.Agent/Execution/TestExecutor.cs
@@ -211,10 +211,6 @@ namespace Horde.Agent.Execution
 				}
 			}
 
-			FileReference tempFileName = new(Path.GetTempFileName());
-			await File.WriteAllTextAsync(tempFileName.FullName, "Some example data", cancellationToken);
-			await ArtifactUploader.UploadAsync(RpcConnection, JobId, BatchId, step.StepId, tempFileName.GetFileName(), tempFileName, logger, cancellationToken);
-
 			logger.LogInformation("**** FINISH NODE {StepName} ****", step.Name);
 
 			return true;
diff --git a/Engine/Source/Programs/Horde/Horde.Agent/Execution/WorkspaceExecutor.cs b/Engine/Source/Programs/Horde/Horde.Agent/Execution/WorkspaceExecutor.cs
--- a/Engine/Source/Programs/Horde/Horde.Agent/Execution/WorkspaceExecutor.cs
+++ b/Engine/Source/Programs/Horde/Horde.Agent/Execution/WorkspaceExecutor.cs
@@ -110,7 +110,7 @@ namespace Horde.Agent.Execution
 			// Loop back to JobExecutor's SetupAsync again, but with workspace and shared storage dir set
 			WorkspaceMaterializerSettings settings = await _workspace.GetSettingsAsync(cancellationToken);
 			DirectoryReference workspaceDir = settings.DirectoryPath;
-			return await SetupAsync(step, workspaceDir, _sharedStorageDir, settings.IsPerforceWorkspace, GetLogger(settings, logger), cancellationToken);
+			return await SetupAsync(step, workspaceDir, settings.IsPerforceWorkspace, GetLogger(settings, logger), cancellationToken);
 		}
 
 		/// <inheritdoc/>
@@ -119,7 +119,7 @@ namespace Horde.Agent.Execution
 			// Loop back to JobExecutor's ExecuteAsync again, but with workspace and shared storage dir set
 			WorkspaceMaterializerSettings settings = await _workspace.GetSettingsAsync(cancellationToken);
 			DirectoryReference workspaceDir = settings.DirectoryPath;
-			return await ExecuteAsync(step, workspaceDir, _sharedStorageDir, settings.IsPerforceWorkspace, GetLogger(settings, logger), cancellationToken);
+			return await ExecuteAsync(step, workspaceDir, settings.IsPerforceWorkspace, GetLogger(settings, logger), cancellationToken);
 		}
 
 		/// <inheritdoc/>
diff --git a/Engine/Source/Programs/Horde/Horde.Agent/Leases/Handlers/JobHandler.cs b/Engine/Source/Programs/Horde/Horde.Agent/Leases/Handlers/JobHandler.cs
--- a/Engine/Source/Programs/Horde/Horde.Agent/Leases/Handlers/JobHandler.cs
+++ b/Engine/Source/Programs/Horde/Horde.Agent/Leases/Handlers/JobHandler.cs
@@ -164,7 +164,7 @@ namespace Horde.Agent.Leases.Handlers
 			GlobalTracer.Instance.ActiveSpan?.SetTag("jobName", executeTask.JobName.ToString());
 			GlobalTracer.Instance.ActiveSpan?.SetTag("batchId", executeTask.BatchId.ToString());
 
-			logger.LogInformation("Executor: {Name}, UseNewTempStorage: {UseNewTempStorage}", jobOptions.Executor, jobOptions.UseNewTempStorage ?? false);
+			logger.LogInformation("Executor: {Name}", jobOptions.Executor);
 
 			// Start executing the current batch
 			RpcBeginBatchResponse batch = await session.RpcConnection.InvokeAsync<JobRpc.JobRpcClient, RpcBeginBatchResponse>(x => x.BeginBatchAsync(new RpcBeginBatchRequest(executeTask.JobId, executeTask.BatchId, leaseId), null, null, cancellationToken), cancellationToken);
diff --git a/Engine/Source/Programs/Horde/Horde.Agent/Services/TelemetryService.cs b/Engine/Source/Programs/Horde/Horde.Agent/Services/TelemetryService.cs
--- a/Engine/Source/Programs/Horde/Horde.Agent/Services/TelemetryService.cs
+++ b/Engine/Source/Programs/Horde/Horde.Agent/Services/TelemetryService.cs
@@ -5,7 +5,6 @@ using System.Runtime.InteropServices;
 using System.Runtime.Versioning;
 using System.Text;
 using EpicGames.Core;
-using EpicGames.Horde.Agents;
 using Google.Protobuf.WellKnownTypes;
 using Grpc.Core;
 using Grpc.Net.Client;
@@ -227,7 +226,7 @@ class TelemetryService : BackgroundService
 	private readonly TimeSpan _reportInterval;
 	private readonly RpcAgentMetadataEvent _agentMetadataEvent;
 	private readonly TimeSpan _agentMetadataReportInterval = TimeSpan.FromMinutes(2);
-	private ISystemMetrics? _systemMetrics;
+	private readonly ISystemMetrics? _systemMetrics;
 	private DateTime _lastTimeAgentMetadataSent = DateTime.UnixEpoch;
 
 	private CancellationTokenSource? _eventLoopHeartbeatCts;
diff --git a/Engine/Source/Programs/Horde/Horde.Agent/Services/WorkerService.cs b/Engine/Source/Programs/Horde/Horde.Agent/Services/WorkerService.cs
--- a/Engine/Source/Programs/Horde/Horde.Agent/Services/WorkerService.cs
+++ b/Engine/Source/Programs/Horde/Horde.Agent/Services/WorkerService.cs
@@ -16,10 +16,7 @@ namespace Horde.Agent.Services
 	{
 		readonly ILogger _logger;
 		readonly ISessionFactory _sessionFactory;
-		readonly CapabilitiesService _capabilitiesService;
 		readonly StatusService _statusService;
-		readonly LeaseHandler[] _leaseHandlers;
-		readonly LeaseLoggerFactory _leaseLoggerFactory;
 		readonly IServiceProvider _serviceProvider;
 		LeaseManager? _currentLeaseManager;
 
@@ -36,14 +33,11 @@ namespace Horde.Agent.Services
 		/// <summary>
 		/// Constructor. Registers with the server and starts accepting connections.
 		/// </summary>
-		public WorkerService(ISessionFactory sessionFactory, CapabilitiesService capabilitiesService, StatusService statusService, IEnumerable<LeaseHandler> leaseHandlers, LeaseLoggerFactory leaseLoggerFactory, IServiceProvider serviceProvider, ILogger<WorkerService> logger)
+		public WorkerService(ISessionFactory sessionFactory, StatusService statusService, IServiceProvider serviceProvider, ILogger<WorkerService> logger)
 		{
 			_sessionFactory = sessionFactory;
-			_capabilitiesService = capabilitiesService;
 			_statusService = statusService;
 			_logger = logger;
-			_leaseHandlers = leaseHandlers.ToArray();
-			_leaseLoggerFactory = leaseLoggerFactory;
 			_serviceProvider = serviceProvider;
 		}
 
diff --git a/Engine/Source/Programs/Horde/Horde.Agent/Utility/ArtifactUploader.cs b/Engine/Source/Programs/Horde/Horde.Agent/Utility/ArtifactUploader.cs
deleted file mode 100644
diff --git a/Engine/Source/Programs/Horde/Horde.Server.Tests/Server/RpcServiceTest.cs b/Engine/Source/Programs/Horde/Horde.Server.Tests/Server/RpcServiceTest.cs
--- a/Engine/Source/Programs/Horde/Horde.Server.Tests/Server/RpcServiceTest.cs
+++ b/Engine/Source/Programs/Horde/Horde.Server.Tests/Server/RpcServiceTest.cs
@@ -10,15 +10,9 @@ using System.Threading;
 using System.Threading.Channels;
 using System.Threading.Tasks;
 using EpicGames.Horde.Agents;
-using EpicGames.Horde.Agents.Leases;
 using EpicGames.Horde.Agents.Pools;
-using EpicGames.Horde.Agents.Sessions;
-using Google.Protobuf;
 using Grpc.Core;
 using Horde.Server.Agents;
-using Horde.Server.Agents.Sessions;
-using Horde.Server.Jobs.Artifacts;
-using Horde.Server.Logs;
 using Horde.Server.Server;
 using Horde.Server.Utilities;
 using HordeCommon.Rpc;
@@ -28,7 +22,6 @@ using Microsoft.AspNetCore.Http.Features;
 using Microsoft.Extensions.FileProviders;
 using Microsoft.Extensions.Hosting;
 using Microsoft.VisualStudio.TestTools.UnitTesting;
-using MongoDB.Bson;
 
 namespace Horde.Server.Tests.Server
 {
@@ -454,63 +447,6 @@ namespace Horde.Server.Tests.Server
 			Assert.IsTrue(re.Message.Contains("Invalid agent name", StringComparison.OrdinalIgnoreCase));
 		}
 
-		[TestMethod]
-		public async Task UploadArtifactTestAsync()
-		{
-			Fixture fixture = await CreateFixtureAsync();
-
-			SessionId sessionId = SessionIdUtils.GenerateNewId();
-			ServerCallContext context = new ServerCallContextStub(new ClaimsPrincipal(new ClaimsIdentity(new List<Claim>
-			{
-				HordeClaims.AdminClaim.ToClaim(),
-				new Claim(HordeClaimTypes.AgentSessionId, sessionId.ToString()),
-			}, "TestAuthType")));
-
-			string[] data = { "foo", "bar", "baz", "qux" };
-			string dataStr = String.Join("", data);
-
-			RpcUploadArtifactMetadata metadata = new RpcUploadArtifactMetadata
-			{
-				JobId = fixture.Job1.Id.ToString(),
-				BatchId = fixture.Job1.Batches[0].Id.ToString(),
-				StepId = fixture.Job1.Batches[0].Steps[0].Id.ToString(),
-				Name = "testfile.txt",
-				MimeType = "text/plain",
-				Length = dataStr.Length
-			};
-
-			// Set the session ID on the job batch to pass auth later
-			Deref(await fixture.Job1.TryAssignLeaseAsync(0, new PoolId("foo"),
-				new AgentId("test"), sessionId,
-				new LeaseId(BinaryIdUtils.CreateNew()), LogIdUtils.GenerateNewId()));
-			/*
-						TestAsyncStreamReader<UploadArtifactRequest> RequestStream = new TestAsyncStreamReader<UploadArtifactRequest>(Context);
-						Task<UploadArtifactResponse> Call = TestSetup.RpcService.UploadArtifact(RequestStream,  Context);
-						RequestStream.AddMessage(new UploadArtifactRequest { Metadata = Metadata });
-						RequestStream.AddMessage(new UploadArtifactRequest { Data = ByteString.CopyFromUtf8(Data[0]) });
-						RequestStream.AddMessage(new UploadArtifactRequest { Data = ByteString.CopyFromUtf8(Data[1]) });
-						RequestStream.AddMessage(new UploadArtifactRequest { Data = ByteString.CopyFromUtf8(Data[2]) });
-						// Only send three messages and not the last one.
-						// Aborting the upload here and retry in next code section below.
-						RequestStream.Complete();
-						await Task.Delay(500);
-			*/
-			TestAsyncStreamReader<RpcUploadArtifactRequest> requestStream = new TestAsyncStreamReader<RpcUploadArtifactRequest>(context);
-			Task<RpcUploadArtifactResponse> call = RpcService.UploadArtifact(requestStream, context);
-			requestStream.AddMessage(new RpcUploadArtifactRequest { Metadata = metadata });
-			requestStream.AddMessage(new RpcUploadArtifactRequest { Data = ByteString.CopyFromUtf8(data[0]) });
-			requestStream.AddMessage(new RpcUploadArtifactRequest { Data = ByteString.CopyFromUtf8(data[1]) });
-			requestStream.AddMessage(new RpcUploadArtifactRequest { Data = ByteString.CopyFromUtf8(data[2]) });
-			requestStream.AddMessage(new RpcUploadArtifactRequest { Data = ByteString.CopyFromUtf8(data[3]) });
-			RpcUploadArtifactResponse res = await call;
-
-			IArtifactV1? artifact = await ArtifactCollection.GetArtifactAsync(ObjectId.Parse(res.Id));
-			Assert.IsNotNull(artifact);
-			Stream stream = await ArtifactCollection.OpenArtifactReadStreamAsync(artifact!);
-			using StreamReader reader = new StreamReader(stream);
-			string text = await reader.ReadToEndAsync();
-			Assert.AreEqual(dataStr, text);
-		}
 		/*
 		[TestMethod]
 		public async Task UploadSoftwareAsync()
diff --git a/Engine/Source/Programs/Horde/Horde.Server/Jobs/JobRpcService.cs b/Engine/Source/Programs/Horde/Horde.Server/Jobs/JobRpcService.cs
--- a/Engine/Source/Programs/Horde/Horde.Server/Jobs/JobRpcService.cs
+++ b/Engine/Source/Programs/Horde/Horde.Server/Jobs/JobRpcService.cs
@@ -230,9 +230,6 @@ namespace Horde.Server.Jobs
 		/// <inheritdoc/>
 		public override Task<Empty> CreateEvents(RpcCreateEventsRequest request, ServerCallContext context) => _jobRpcCommon.CreateEventsAsync(request, context);
 
-		/// <inheritdoc/>
-		public override Task<RpcUploadArtifactResponse> UploadArtifact(IAsyncStreamReader<RpcUploadArtifactRequest> reader, ServerCallContext context) => _jobRpcCommon.UploadArtifactAsync(reader, context);
-
 		/// <inheritdoc/>
 		public override Task<RpcUploadTestDataResponse> UploadTestData(IAsyncStreamReader<RpcUploadTestDataRequest> reader, ServerCallContext context) => _jobRpcCommon.UploadTestDataAsync(reader, context);
 
diff --git a/Engine/Source/Programs/Horde/Horde.Server/Jobs/JobTaskSource.cs b/Engine/Source/Programs/Horde/Horde.Server/Jobs/JobTaskSource.cs
--- a/Engine/Source/Programs/Horde/Horde.Server/Jobs/JobTaskSource.cs
+++ b/Engine/Source/Programs/Horde/Horde.Server/Jobs/JobTaskSource.cs
@@ -945,7 +945,6 @@ namespace Horde.Server.Jobs
 		{
 			RpcJobOptions rpcOptions = new RpcJobOptions();
 			rpcOptions.Executor = options.Executor;
-			rpcOptions.UseNewTempStorage = options.UseNewTempStorage;
 			rpcOptions.UseWine = options.UseWine;
 			rpcOptions.RunInSeparateProcess = options.RunInSeparateProcess;
 			rpcOptions.WorkspaceMaterializer = options.WorkspaceMaterializer;
diff --git a/Engine/Source/Programs/Horde/Horde.Server/Jobs/JobsController.cs b/Engine/Source/Programs/Horde/Horde.Server/Jobs/JobsController.cs
--- a/Engine/Source/Programs/Horde/Horde.Server/Jobs/JobsController.cs
+++ b/Engine/Source/Programs/Horde/Horde.Server/Jobs/JobsController.cs
@@ -586,7 +586,7 @@ namespace Horde.Server.Jobs
 			response.Parameters = job.Parameters.ToDictionary();
 			response.Arguments = job.Arguments.ToList();
 			response.UpdateTime = new DateTimeOffset(job.UpdateTimeUtc);
-			response.UseArtifactsV2 = job.JobOptions?.UseNewTempStorage ?? true;
+			response.UseArtifactsV2 = true;
 			response.UpdateIssues = job.UpdateIssues;
 			return response;
 		}
diff --git a/Engine/Source/Programs/Shared/EpicGames.Horde/Jobs/JobMessages.cs b/Engine/Source/Programs/Shared/EpicGames.Horde/Jobs/JobMessages.cs
--- a/Engine/Source/Programs/Shared/EpicGames.Horde/Jobs/JobMessages.cs
+++ b/Engine/Source/Programs/Shared/EpicGames.Horde/Jobs/JobMessages.cs
@@ -48,11 +48,6 @@ namespace EpicGames.Horde.Jobs
 		/// </summary>
 		public string? Executor { get; set; }
 
-		/// <summary>
-		/// Whether to use the new temp storage backend
-		/// </summary>
-		public bool? UseNewTempStorage { get; set; }
-
 		/// <summary>
 		/// Whether to execute using Wine emulation on Linux
 		/// </summary>
@@ -84,7 +79,6 @@ namespace EpicGames.Horde.Jobs
 		public void MergeDefaults(JobOptions other)
 		{
 			Executor ??= other.Executor;
-			UseNewTempStorage ??= other.UseNewTempStorage;
 			UseWine ??= other.UseWine;
 			RunInSeparateProcess ??= other.RunInSeparateProcess;
 			WorkspaceMaterializer ??= other.WorkspaceMaterializer;
