Horde: Remove legacy job methods from horde_rpc interface.

diff --git a/Engine/Source/Programs/Horde/Horde.Agent.Tests/Services/TelemetryServiceTest.cs b/Engine/Source/Programs/Horde/Horde.Agent.Tests/Services/TelemetryServiceTest.cs
--- a/Engine/Source/Programs/Horde/Horde.Agent.Tests/Services/TelemetryServiceTest.cs
+++ b/Engine/Source/Programs/Horde/Horde.Agent.Tests/Services/TelemetryServiceTest.cs
@@ -5,7 +5,6 @@ using System.Collections.Generic;
 using System.Threading;
 using System.Threading.Tasks;
 using Horde.Agent.Execution;
-using Horde.Agent.Leases;
 using Horde.Agent.Leases.Handlers;
 using Horde.Agent.Services;
 using Microsoft.Extensions.Logging.Abstractions;
diff --git a/Engine/Source/Programs/Horde/Horde.Server.Tests/TestSetup.cs b/Engine/Source/Programs/Horde/Horde.Server.Tests/TestSetup.cs
--- a/Engine/Source/Programs/Horde/Horde.Server.Tests/TestSetup.cs
+++ b/Engine/Source/Programs/Horde/Horde.Server.Tests/TestSetup.cs
@@ -283,8 +283,6 @@ namespace Horde.Server.Tests
 			services.AddSingleton(typeof(IHealthMonitor<>), typeof(HealthMonitor<>));
 			services.AddSingleton<ServerStatusService>();
 
-			services.AddScoped<JobRpcCommon>();
-
 			services.AddSingleton<ConformTaskSource>();
 			services.AddSingleton<ICommitService, CommitService>();
 
diff --git a/Engine/Source/Programs/Horde/Horde.Server/Jobs/JobRpcService.cs b/Engine/Source/Programs/Horde/Horde.Server/Jobs/JobRpcService.cs
--- a/Engine/Source/Programs/Horde/Horde.Server/Jobs/JobRpcService.cs
+++ b/Engine/Source/Programs/Horde/Horde.Server/Jobs/JobRpcService.cs
@@ -23,7 +23,6 @@ using Grpc.Core;
 using Horde.Common.Rpc;
 using Horde.Server.Acls;
 using Horde.Server.Artifacts;
-using Horde.Server.Jobs.Artifacts;
 using Horde.Server.Jobs.Graphs;
 using Horde.Server.Jobs.Templates;
 using Horde.Server.Jobs.TestData;
@@ -47,203 +46,9 @@ namespace Horde.Server.Jobs
 	public class JobRpcService : JobRpc.JobRpcBase
 	{
 		readonly AclService _aclService;
-		readonly IJobCollection _jobCollection;
-		readonly IArtifactCollection _artifactCollection;
-		readonly JobRpcCommon _jobRpcCommon;
-		readonly GlobalConfig _globalConfig;
-
-		/// <summary>
-		/// Constructor
-		/// </summary>
-		public JobRpcService(AclService aclService, IJobCollection jobCollection, IArtifactCollection artifactCollection, JobRpcCommon jobRpcCommon, IOptionsSnapshot<GlobalConfig> globalConfig)
-		{
-			_aclService = aclService;
-			_jobCollection = jobCollection;
-			_artifactCollection = artifactCollection;
-			_jobRpcCommon = jobRpcCommon;
-			_globalConfig = globalConfig.Value;
-		}
-
-		/// <inheritdoc/>
-		public override async Task<RpcCreateJobArtifactResponse> CreateArtifact(RpcCreateJobArtifactRequest request, ServerCallContext context)
-		{
-			(IJob job, _, IJobStep step) = await AuthorizeAsync(request.JobId, request.StepId, context);
-
-			ArtifactType type = request.Type switch
-			{
-				JobArtifactType.TempStorage => ArtifactType.StepOutput,
-				JobArtifactType.Saved => ArtifactType.StepSaved,
-				JobArtifactType.Trace => ArtifactType.StepTrace,
-				JobArtifactType.TestData => ArtifactType.StepTestData,
-				_ => throw new StructuredRpcException(StatusCode.InvalidArgument, "Invalid artifact type")
-			};
-
-			List<string> keys = new List<string>();
-			keys.Add($"job:{job.Id}");
-			keys.Add($"job:{job.Id}/step:{step.Id}");
-
-			List<string> metadata = new List<string>();
-
-			if (!_globalConfig.TryGetTemplate(job.StreamId, job.TemplateId, out TemplateRefConfig? templateConfig))
-			{
-				throw new StructuredRpcException(StatusCode.NotFound, "Couldn't find template {TemplateId} in stream {StreamId}", job.TemplateId, job.StreamId);
-			}
-
-			IArtifact artifact = await _artifactCollection.AddAsync(new ArtifactName("default"), type, null, job.StreamId, job.Change, keys, metadata, templateConfig.Acl.ScopeName, context.CancellationToken);
-
-			List<AclClaimConfig> claims = new List<AclClaimConfig>();
-			claims.Add(new AclClaimConfig(HordeClaimTypes.WriteNamespace, artifact.NamespaceId.ToString()));
-
-			string token = await _aclService.IssueBearerTokenAsync(claims, TimeSpan.FromHours(8.0), context.CancellationToken);
-			return new RpcCreateJobArtifactResponse { Id = artifact.Id.ToString(), NamespaceId = artifact.NamespaceId.ToString(), RefName = artifact.RefName.ToString(), Token = token };
-		}
-
-		/// <inheritdoc/>
-		public override async Task<RpcCreateJobArtifactResponseV2> CreateArtifactV2(RpcCreateJobArtifactRequestV2 request, ServerCallContext context)
-		{
-			(IJob job, _, IJobStep step) = await AuthorizeAsync(request.JobId, request.StepId, context);
-
-			ArtifactName name = new ArtifactName(request.Name);
-			ArtifactType type = new ArtifactType(request.Type);
-
-			List<string> keys = new List<string>();
-			keys.Add(job.GetArtifactKey());
-			keys.Add(job.GetArtifactKey(step));
-			keys.AddRange(request.Keys);
-
-			if (!_globalConfig.TryGetTemplate(job.StreamId, job.TemplateId, out TemplateRefConfig? templateConfig))
-			{
-				throw new StructuredRpcException(StatusCode.NotFound, "Couldn't find template {TemplateId} in stream {StreamId}", job.TemplateId, job.StreamId);
-			}
-
-			string? description = request.Description;
-			if (String.IsNullOrEmpty(description))
-			{
-				description = request.Name;
-			}
-
-			IArtifact artifact = await _artifactCollection.AddAsync(name, type, description, job.StreamId, job.Change, keys, request.Metadata, templateConfig.Acl.ScopeName, context.CancellationToken);
-
-			List<AclClaimConfig> claims = new List<AclClaimConfig>();
-			claims.Add(new AclClaimConfig(HordeClaimTypes.WriteNamespace, $"{artifact.NamespaceId}:{artifact.RefName}"));
-
-			string token = await _aclService.IssueBearerTokenAsync(claims, TimeSpan.FromHours(8.0), context.CancellationToken);
-
-			RpcCreateJobArtifactResponseV2 response = new RpcCreateJobArtifactResponseV2();
-			response.Id = artifact.Id.ToString();
-			response.NamespaceId = artifact.NamespaceId.ToString();
-			response.RefName = artifact.RefName.ToString();
-			response.Token = token;
-			return response;
-		}
-
-		/// <inheritdoc/>
-		public override async Task<RpcGetJobArtifactResponse> GetArtifact(RpcGetJobArtifactRequest request, ServerCallContext context)
-		{
-			(IJob job, _, _) = await AuthorizeAsync(request.JobId, request.StepId, context);
-
-			ArtifactName name = new ArtifactName(request.Name);
-			ArtifactType type = new ArtifactType(request.Type);
-
-			IArtifact? artifact = await _artifactCollection.FindAsync(name: name, type: type, keys: new[] { job.GetArtifactKey() }, cancellationToken: context.CancellationToken).FirstOrDefaultAsync(context.CancellationToken);
-			if (artifact == null)
-			{
-				throw new StructuredRpcException(StatusCode.NotFound, "No artifact {ArtifactName} of type {ArtifactType} was found for job {JobId}", name, type, job.Id);
-			}
-
-			List<AclClaimConfig> claims = new List<AclClaimConfig>();
-			claims.Add(new AclClaimConfig(HordeClaimTypes.ReadNamespace, $"{artifact.NamespaceId}:{artifact.RefName}"));
-
-			string token = await _aclService.IssueBearerTokenAsync(claims, TimeSpan.FromHours(8.0), context.CancellationToken);
-
-			RpcGetJobArtifactResponse response = new RpcGetJobArtifactResponse();
-			response.Id = artifact.Id.ToString();
-			response.NamespaceId = artifact.NamespaceId.ToString();
-			response.RefName = artifact.RefName.ToString();
-			response.Token = token;
-			return response;
-		}
-
-		Task<(IJob, IJobStepBatch, IJobStep)> AuthorizeAsync(string jobId, string stepId, ServerCallContext context)
-		{
-			return AuthorizeAsync(JobId.Parse(jobId), JobStepId.Parse(stepId), context);
-		}
-
-		async Task<(IJob, IJobStepBatch, IJobStep)> AuthorizeAsync(JobId jobId, JobStepId stepId, ServerCallContext context)
-		{
-			IJob? job = await _jobCollection.GetAsync(jobId);
-			if (job == null)
-			{
-				throw new StructuredRpcException(StatusCode.NotFound, "Unable to find job {JobId}", jobId);
-			}
-			if (!job.TryGetStep(stepId, out IJobStep? step))
-			{
-				throw new StructuredRpcException(StatusCode.NotFound, "Unable to find step {JobId}:{StepId}", job.Id, stepId);
-			}
-
-			IJobStepBatch batch = step.Batch;
-			if (batch.SessionId == null)
-			{
-				throw new StructuredRpcException(StatusCode.PermissionDenied, "Batch {JobId}:{BatchId} has no session id", job.Id, batch.Id);
-			}
-			if (batch.LeaseId == null)
-			{
-				throw new StructuredRpcException(StatusCode.PermissionDenied, "Batch {JobId}:{BatchId} has no lease id", job.Id, batch.Id);
-			}
-
-			ClaimsPrincipal principal = context.GetHttpContext().User;
-			if (!principal.HasSessionClaim(batch.SessionId.Value) && !principal.HasLeaseClaim(batch.LeaseId.Value))
-			{
-				throw new StructuredRpcException(StatusCode.PermissionDenied, "Session id {SessionId} not valid for step {JobId}:{BatchId}:{StepId}. Expected {ExpectedSessionId}.", principal.GetSessionClaim() ?? SessionId.Empty, job.Id, batch.Id, step.Id, batch.SessionId.Value);
-			}
-
-			return (job, batch, step);
-		}
-
-		/// <inheritdoc/>
-		public override Task<RpcGetStreamResponse> GetStream(RpcGetStreamRequest request, ServerCallContext context) => _jobRpcCommon.GetStreamAsync(request, context);
-
-		/// <inheritdoc/>
-		public override Task<RpcGetJobResponse> GetJob(RpcGetJobRequest request, ServerCallContext context) => _jobRpcCommon.GetJobAsync(request, context);
-
-		/// <inheritdoc/>
-		public override Task<Empty> UpdateJob(RpcUpdateJobRequest request, ServerCallContext context) => _jobRpcCommon.UpdateJobAsync(request, context);
-
-		/// <inheritdoc/>
-		public override Task<RpcBeginBatchResponse> BeginBatch(RpcBeginBatchRequest request, ServerCallContext context) => _jobRpcCommon.BeginBatchAsync(request, context);
-
-		/// <inheritdoc/>
-		public override Task<Empty> FinishBatch(RpcFinishBatchRequest request, ServerCallContext context) => _jobRpcCommon.FinishBatchAsync(request, context);
-
-		/// <inheritdoc/>
-		public override Task<RpcBeginStepResponse> BeginStep(RpcBeginStepRequest request, ServerCallContext context) => _jobRpcCommon.BeginStepAsync(request, context);
-
-		/// <inheritdoc/>
-		public override Task<Empty> UpdateStep(RpcUpdateStepRequest request, ServerCallContext context) => _jobRpcCommon.UpdateStepAsync(request, context);
-
-		/// <inheritdoc/>
-		public override Task<RpcGetStepResponse> GetStep(RpcGetStepRequest request, ServerCallContext context) => _jobRpcCommon.GetStepAsync(request, context);
-
-		/// <inheritdoc/>
-		public override Task<RpcUpdateGraphResponse> UpdateGraph(RpcUpdateGraphRequest request, ServerCallContext context) => _jobRpcCommon.UpdateGraphAsync(request, context);
-
-		/// <inheritdoc/>
-		public override Task<Empty> CreateEvents(RpcCreateEventsRequest request, ServerCallContext context) => _jobRpcCommon.CreateEventsAsync(request, context);
-
-		/// <inheritdoc/>
-		public override Task<RpcUploadTestDataResponse> UploadTestData(IAsyncStreamReader<RpcUploadTestDataRequest> reader, ServerCallContext context) => _jobRpcCommon.UploadTestDataAsync(reader, context);
-
-		/// <inheritdoc/>
-		public override Task<RpcCreateReportResponse> CreateReport(RpcCreateReportRequest request, ServerCallContext context) => _jobRpcCommon.CreateReportAsync(request, context);
-	}
-
-	/// <summary>
-	/// Common methods between HordeRpc and JobRpc.
-	/// </summary>
-	public class JobRpcCommon
-	{
 		readonly JobService _jobService;
-		readonly IArtifactCollectionV1 _artifactCollection;
+		readonly IArtifactCollection _artifactCollection;
+		readonly IJobCollection _jobCollection;
 		readonly ILogCollection _logCollection;
 		readonly IGraphCollection _graphs;
 		readonly ITestDataCollection _testData;
@@ -257,9 +62,11 @@ namespace Horde.Server.Jobs
 		/// <summary>
 		/// Constructor
 		/// </summary>
-		public JobRpcCommon(JobService jobService, IArtifactCollectionV1 artifactCollection, ILogCollection logCollection, IGraphCollection graphs, ITestDataCollection testData, IJobStepRefCollection jobStepRefCollection, ITemplateCollection templateCollection, HttpClient httpClient, IClock clock, IOptionsSnapshot<GlobalConfig> globalConfig, ILogger<JobRpcCommon> logger)
+		public JobRpcService(AclService aclService, JobService jobService, IArtifactCollection artifactCollection, IJobCollection jobCollection, ILogCollection logCollection, IGraphCollection graphs, ITestDataCollection testData, IJobStepRefCollection jobStepRefCollection, ITemplateCollection templateCollection, HttpClient httpClient, IClock clock, IOptionsSnapshot<GlobalConfig> globalConfig, ILogger<JobRpcService> logger)
 		{
+			_aclService = aclService;
 			_jobService = jobService;
+			_jobCollection = jobCollection;
 			_artifactCollection = artifactCollection;
 			_logCollection = logCollection;
 			_graphs = graphs;
@@ -911,48 +718,6 @@ namespace Horde.Server.Jobs
 			return new Empty();
 		}
 
-		/// <summary>
-		/// Uploads a new artifact
-		/// </summary>
-		/// <param name="reader">Request arguments</param>
-		/// <param name="context">Context for the RPC call</param>
-		/// <returns>Information about the new agent</returns>
-		public async Task<RpcUploadArtifactResponse> UploadArtifactAsync(IAsyncStreamReader<RpcUploadArtifactRequest> reader, ServerCallContext context)
-		{
-			// Advance to the metadata object
-			if (!await reader.MoveNext())
-			{
-				throw new StructuredRpcException(StatusCode.DataLoss, "Missing request for artifact upload");
-			}
-
-			// Read the request object
-			RpcUploadArtifactMetadata? metadata = reader.Current.Metadata;
-			if (metadata == null)
-			{
-				throw new StructuredRpcException(StatusCode.DataLoss, "Expected metadata in first artifact request");
-			}
-
-			// Get the job and step
-			IJob job = await GetJobAsync(JobId.Parse(metadata.JobId), context.CancellationToken);
-			AuthorizeBatch(job, JobStepBatchId.Parse(metadata.BatchId), context);
-
-			IJobStep? step;
-			if (!job.TryGetStep(JobStepBatchId.Parse(metadata.BatchId), JobStepId.Parse(metadata.StepId), out step))
-			{
-				throw new StructuredRpcException(StatusCode.NotFound, "Unable to find step {JobId}:{BatchId}:{StepId}", job.Id, metadata.BatchId, metadata.StepId);
-			}
-
-			// Upload the stream
-			using (ArtifactChunkStream inputStream = new ArtifactChunkStream(reader, metadata.Length))
-			{
-				IArtifactV1 artifact = await _artifactCollection.CreateArtifactAsync(job.Id, step.Id, metadata.Name, metadata.MimeType, inputStream, context.CancellationToken);
-
-				RpcUploadArtifactResponse response = new RpcUploadArtifactResponse();
-				response.Id = artifact.Id.ToString();
-				return response;
-			}
-		}
-
 		/// <summary>
 		/// Uploads new test data
 		/// </summary>
@@ -1015,6 +780,108 @@ namespace Horde.Server.Jobs
 			return new RpcUploadTestDataResponse();
 		}
 
+		/// <inheritdoc/>
+		public override async Task<RpcCreateJobArtifactResponseV2> CreateArtifactV2(RpcCreateJobArtifactRequestV2 request, ServerCallContext context)
+		{
+			(IJob job, _, IJobStep step) = await AuthorizeAsync(request.JobId, request.StepId, context);
+
+			ArtifactName name = new ArtifactName(request.Name);
+			ArtifactType type = new ArtifactType(request.Type);
+
+			List<string> keys = new List<string>();
+			keys.Add(job.GetArtifactKey());
+			keys.Add(job.GetArtifactKey(step));
+			keys.AddRange(request.Keys);
+
+			if (!_globalConfig.Value.TryGetTemplate(job.StreamId, job.TemplateId, out TemplateRefConfig? templateConfig))
+			{
+				throw new StructuredRpcException(StatusCode.NotFound, "Couldn't find template {TemplateId} in stream {StreamId}", job.TemplateId, job.StreamId);
+			}
+
+			string? description = request.Description;
+			if (String.IsNullOrEmpty(description))
+			{
+				description = request.Name;
+			}
+
+			IArtifact artifact = await _artifactCollection.AddAsync(name, type, description, job.StreamId, job.Change, keys, request.Metadata, templateConfig.Acl.ScopeName, context.CancellationToken);
+
+			List<AclClaimConfig> claims = new List<AclClaimConfig>();
+			claims.Add(new AclClaimConfig(HordeClaimTypes.WriteNamespace, $"{artifact.NamespaceId}:{artifact.RefName}"));
+
+			string token = await _aclService.IssueBearerTokenAsync(claims, TimeSpan.FromHours(8.0), context.CancellationToken);
+
+			RpcCreateJobArtifactResponseV2 response = new RpcCreateJobArtifactResponseV2();
+			response.Id = artifact.Id.ToString();
+			response.NamespaceId = artifact.NamespaceId.ToString();
+			response.RefName = artifact.RefName.ToString();
+			response.Token = token;
+			return response;
+		}
+
+		/// <inheritdoc/>
+		public override async Task<RpcGetJobArtifactResponse> GetArtifact(RpcGetJobArtifactRequest request, ServerCallContext context)
+		{
+			(IJob job, _, _) = await AuthorizeAsync(request.JobId, request.StepId, context);
+
+			ArtifactName name = new ArtifactName(request.Name);
+			ArtifactType type = new ArtifactType(request.Type);
+
+			IArtifact? artifact = await _artifactCollection.FindAsync(name: name, type: type, keys: new[] { job.GetArtifactKey() }, cancellationToken: context.CancellationToken).FirstOrDefaultAsync(context.CancellationToken);
+			if (artifact == null)
+			{
+				throw new StructuredRpcException(StatusCode.NotFound, "No artifact {ArtifactName} of type {ArtifactType} was found for job {JobId}", name, type, job.Id);
+			}
+
+			List<AclClaimConfig> claims = new List<AclClaimConfig>();
+			claims.Add(new AclClaimConfig(HordeClaimTypes.ReadNamespace, $"{artifact.NamespaceId}:{artifact.RefName}"));
+
+			string token = await _aclService.IssueBearerTokenAsync(claims, TimeSpan.FromHours(8.0), context.CancellationToken);
+
+			RpcGetJobArtifactResponse response = new RpcGetJobArtifactResponse();
+			response.Id = artifact.Id.ToString();
+			response.NamespaceId = artifact.NamespaceId.ToString();
+			response.RefName = artifact.RefName.ToString();
+			response.Token = token;
+			return response;
+		}
+
+		Task<(IJob, IJobStepBatch, IJobStep)> AuthorizeAsync(string jobId, string stepId, ServerCallContext context)
+		{
+			return AuthorizeAsync(JobId.Parse(jobId), JobStepId.Parse(stepId), context);
+		}
+
+		async Task<(IJob, IJobStepBatch, IJobStep)> AuthorizeAsync(JobId jobId, JobStepId stepId, ServerCallContext context)
+		{
+			IJob? job = await _jobCollection.GetAsync(jobId);
+			if (job == null)
+			{
+				throw new StructuredRpcException(StatusCode.NotFound, "Unable to find job {JobId}", jobId);
+			}
+			if (!job.TryGetStep(stepId, out IJobStep? step))
+			{
+				throw new StructuredRpcException(StatusCode.NotFound, "Unable to find step {JobId}:{StepId}", job.Id, stepId);
+			}
+
+			IJobStepBatch batch = step.Batch;
+			if (batch.SessionId == null)
+			{
+				throw new StructuredRpcException(StatusCode.PermissionDenied, "Batch {JobId}:{BatchId} has no session id", job.Id, batch.Id);
+			}
+			if (batch.LeaseId == null)
+			{
+				throw new StructuredRpcException(StatusCode.PermissionDenied, "Batch {JobId}:{BatchId} has no lease id", job.Id, batch.Id);
+			}
+
+			ClaimsPrincipal principal = context.GetHttpContext().User;
+			if (!principal.HasSessionClaim(batch.SessionId.Value) && !principal.HasLeaseClaim(batch.LeaseId.Value))
+			{
+				throw new StructuredRpcException(StatusCode.PermissionDenied, "Session id {SessionId} not valid for step {JobId}:{BatchId}:{StepId}. Expected {ExpectedSessionId}.", principal.GetSessionClaim() ?? SessionId.Empty, job.Id, batch.Id, step.Id, batch.SessionId.Value);
+			}
+
+			return (job, batch, step);
+		}
+
 		/// <summary>
 		/// Create a new report on a job or job step
 		/// </summary>
diff --git a/Engine/Source/Programs/Horde/Horde.Server/Server/RpcService.cs b/Engine/Source/Programs/Horde/Horde.Server/Server/RpcService.cs
--- a/Engine/Source/Programs/Horde/Horde.Server/Server/RpcService.cs
+++ b/Engine/Source/Programs/Horde/Horde.Server/Server/RpcService.cs
@@ -22,7 +22,6 @@ using Horde.Server.Agents;
 using Horde.Server.Agents.Pools;
 using Horde.Server.Agents.Sessions;
 using Horde.Server.Agents.Telemetry;
-using Horde.Server.Jobs;
 using Horde.Server.Tasks;
 using Horde.Server.Telemetry;
 using Horde.Server.Tools;
@@ -52,7 +51,6 @@ namespace Horde.Server.Server
 		readonly LifetimeService _lifetimeService;
 		readonly ITelemetrySink _telemetrySink;
 		readonly ConformTaskSource _conformTaskSource;
-		readonly JobRpcCommon _jobRpcCommon;
 		readonly IToolCollection _toolCollection;
 		readonly IAgentTelemetryCollection _agentTelemetryCollection;
 		readonly AclService _aclService;
@@ -62,14 +60,13 @@ namespace Horde.Server.Server
 		/// <summary>
 		/// Constructor
 		/// </summary>
-		public RpcService(AgentService agentService, PoolService poolService, LifetimeService lifetimeService, ITelemetrySink telemetrySink, ConformTaskSource conformTaskSource, JobRpcCommon jobRpcCommon, IToolCollection toolCollection, IAgentTelemetryCollection agentTelemetryCollection, AclService aclService, IOptionsSnapshot<GlobalConfig> globalConfig, ILogger<RpcService> logger)
+		public RpcService(AgentService agentService, PoolService poolService, LifetimeService lifetimeService, ITelemetrySink telemetrySink, ConformTaskSource conformTaskSource, IToolCollection toolCollection, IAgentTelemetryCollection agentTelemetryCollection, AclService aclService, IOptionsSnapshot<GlobalConfig> globalConfig, ILogger<RpcService> logger)
 		{
 			_agentService = agentService;
 			_poolService = poolService;
 			_lifetimeService = lifetimeService;
 			_telemetrySink = telemetrySink;
 			_conformTaskSource = conformTaskSource;
-			_jobRpcCommon = jobRpcCommon;
 			_toolCollection = toolCollection;
 			_agentTelemetryCollection = agentTelemetryCollection;
 			_aclService = aclService;
@@ -439,36 +436,6 @@ namespace Horde.Server.Server
 			}
 		}
 
-		/// <inheritdoc/>
-		public override Task<RpcGetStreamResponse> GetStream(RpcGetStreamRequest request, ServerCallContext context) => _jobRpcCommon.GetStreamAsync(request, context);
-
-		/// <inheritdoc/>
-		public override Task<RpcGetJobResponse> GetJob(RpcGetJobRequest request, ServerCallContext context) => _jobRpcCommon.GetJobAsync(request, context);
-
-		/// <inheritdoc/>
-		public override Task<Empty> UpdateJob(RpcUpdateJobRequest request, ServerCallContext context) => _jobRpcCommon.UpdateJobAsync(request, context);
-
-		/// <inheritdoc/>
-		public override Task<RpcBeginBatchResponse> BeginBatch(RpcBeginBatchRequest request, ServerCallContext context) => _jobRpcCommon.BeginBatchAsync(request, context);
-
-		/// <inheritdoc/>
-		public override Task<Empty> FinishBatch(RpcFinishBatchRequest request, ServerCallContext context) => _jobRpcCommon.FinishBatchAsync(request, context);
-
-		/// <inheritdoc/>
-		public override Task<RpcBeginStepResponse> BeginStep(RpcBeginStepRequest request, ServerCallContext context) => _jobRpcCommon.BeginStepAsync(request, context);
-
-		/// <inheritdoc/>
-		public override Task<Empty> UpdateStep(RpcUpdateStepRequest request, ServerCallContext context) => _jobRpcCommon.UpdateStepAsync(request, context);
-
-		/// <inheritdoc/>
-		public override Task<RpcGetStepResponse> GetStep(RpcGetStepRequest request, ServerCallContext context) => _jobRpcCommon.GetStepAsync(request, context);
-
-		/// <inheritdoc/>
-		public override Task<RpcUpdateGraphResponse> UpdateGraph(RpcUpdateGraphRequest request, ServerCallContext context) => _jobRpcCommon.UpdateGraphAsync(request, context);
-
-		/// <inheritdoc/>
-		public override Task<Empty> CreateEvents(RpcCreateEventsRequest request, ServerCallContext context) => _jobRpcCommon.CreateEventsAsync(request, context);
-
 		/// <summary>
 		/// Downloads a new agent archive
 		/// </summary>
@@ -560,14 +527,5 @@ namespace Horde.Server.Server
 
 			return new Empty();
 		}
-
-		/// <inheritdoc/>
-		public override Task<RpcUploadArtifactResponse> UploadArtifact(IAsyncStreamReader<RpcUploadArtifactRequest> reader, ServerCallContext context) => _jobRpcCommon.UploadArtifactAsync(reader, context);
-
-		/// <inheritdoc/>
-		public override Task<RpcUploadTestDataResponse> UploadTestData(IAsyncStreamReader<RpcUploadTestDataRequest> reader, ServerCallContext context) => _jobRpcCommon.UploadTestDataAsync(reader, context);
-
-		/// <inheritdoc/>
-		public override Task<RpcCreateReportResponse> CreateReport(RpcCreateReportRequest request, ServerCallContext context) => _jobRpcCommon.CreateReportAsync(request, context);
 	}
 }
diff --git a/Engine/Source/Programs/Horde/Horde.Server/Startup.cs b/Engine/Source/Programs/Horde/Horde.Server/Startup.cs
--- a/Engine/Source/Programs/Horde/Horde.Server/Startup.cs
+++ b/Engine/Source/Programs/Horde/Horde.Server/Startup.cs
@@ -6,7 +6,6 @@ using System.Diagnostics;
 using System.Globalization;
 using System.IO;
 using System.Linq;
-using System.Net.Http;
 using System.Net.Mime;
 using System.Reflection;
 using System.Runtime.InteropServices;
@@ -121,8 +120,6 @@ using MongoDB.Bson;
 using MongoDB.Bson.Serialization;
 using MongoDB.Bson.Serialization.Conventions;
 using MongoDB.Bson.Serialization.Serializers;
-using Polly;
-using Polly.Extensions.Http;
 using Serilog;
 using Serilog.Events;
 using StackExchange.Redis;
@@ -487,9 +484,6 @@ namespace Horde.Server
 			});
 			services.AddGrpcReflection();
 
-			services.AddHttpClient<JobRpcCommon>().AddPolicyHandler(GetDefaultHttpRetryPolicy());
-			services.AddScoped<JobRpcCommon>();
-
 			services.AddSingleton<IAccountCollection, AccountCollection>();
 			services.AddSingleton<IAgentCollection, AgentCollection>();
 			services.AddSingleton<IArtifactCollection, ArtifactCollection>();
@@ -1047,13 +1041,6 @@ namespace Horde.Server
 			});
 		}
 
-		private static IAsyncPolicy<HttpResponseMessage> GetDefaultHttpRetryPolicy()
-		{
-			return HttpPolicyExtensions
-				.HandleTransientHttpError()
-				.WaitAndRetryAsync(new[] { TimeSpan.FromSeconds(1), TimeSpan.FromSeconds(5), TimeSpan.FromSeconds(10), TimeSpan.FromSeconds(10) });
-		}
-
 		public sealed class BlobLocatorBsonSerializer : SerializerBase<BlobLocator>
 		{
 			/// <inheritdoc/>
@@ -1296,7 +1283,7 @@ namespace Horde.Server
 				options.EnrichDiagnosticContext = (diagnosticContext, httpContext) =>
 				{
 					diagnosticContext.Set("RemoteIP", httpContext?.Connection?.RemoteIpAddress);
-					
+
 					// Header sent by the dashboard to indicate how long a user has been inactive for a particular browser page (in seconds)
 					if (httpContext?.Request.Headers.TryGetValue("X-Horde-LastUserActivity", out StringValues values) is true)
 					{
