Horde: Disable file modified errors in job executor.

diff --git a/Engine/Source/Programs/Horde/Horde.Agent/Execution/JobExecutor.cs b/Engine/Source/Programs/Horde/Horde.Agent/Execution/JobExecutor.cs
--- a/Engine/Source/Programs/Horde/Horde.Agent/Execution/JobExecutor.cs
+++ b/Engine/Source/Programs/Horde/Horde.Agent/Execution/JobExecutor.cs
@@ -946,18 +946,20 @@ namespace Horde.Agent.Execution
 			}
 			if (modifiedFiles.Count > 0)
 			{
+				const int NumFiles = 50;
+
 				string modifiedFileList = "";
-				if (modifiedFiles.Count < 100)
+				if (modifiedFiles.Count < NumFiles)
 				{
 					modifiedFileList = String.Join("\n", modifiedFiles.Select(x => x.Value));
 				}
 				else
 				{
-					modifiedFileList = String.Join("\n", modifiedFiles.Take(100).Select(x => x.Value));
-					modifiedFileList += $"{Environment.NewLine}And {modifiedFiles.Count - 100} more.";
+					modifiedFileList = String.Join("\n", modifiedFiles.Take(NumFiles).Select(x => x.Value));
+					modifiedFileList += $"{Environment.NewLine}...and {modifiedFiles.Count - NumFiles} more.";
 				}
 
-				logger.LogError("Build product(s) from a previous step have been modified:\n{FileList}", modifiedFileList);
+				logger.LogInformation("Build product(s) from a previous step have been modified:\n{FileList}", modifiedFileList);
 				return false;
 			}
 
