Horde: Remove need for separate job graph data
#rnx

diff --git a/Engine/Source/Programs/Horde/HordeDashboard/package.json b/Engine/Source/Programs/Horde/HordeDashboard/package.json
--- a/Engine/Source/Programs/Horde/HordeDashboard/package.json
+++ b/Engine/Source/Programs/Horde/HordeDashboard/package.json
@@ -17,7 +17,6 @@
     "@datadog/browser-logs": "^4.30.0",
     "@fluentui/react": "^8.106.2",
     "@fluentui/react-charting": "5.2.0",
-    "crypto-js": "^4.1.1",
     "d3": "^7.6.1",
     "idb": "^7.1.1",
     "markdown-to-jsx": "^7.1.8",
@@ -43,7 +42,6 @@
     "eslint-plugin-react-refresh": "^0.4.6",
     "typescript": "^5.2.2",
     "vite": "^5.2.0",
-    "@types/crypto-js": "^4.1.1",
     "@types/d3": "^7.4.0",
     "@types/markdown-to-jsx": "7.0.1",
     "@types/react-custom-scrollbars": "^4.0.10",
diff --git a/Engine/Source/Programs/Horde/HordeDashboard/src/Plugins/example/components/JobDetailPluginExample.tsx b/Engine/Source/Programs/Horde/HordeDashboard/src/Plugins/example/components/JobDetailPluginExample.tsx
--- a/Engine/Source/Programs/Horde/HordeDashboard/src/Plugins/example/components/JobDetailPluginExample.tsx
+++ b/Engine/Source/Programs/Horde/HordeDashboard/src/Plugins/example/components/JobDetailPluginExample.tsx
@@ -2,16 +2,16 @@
 
 import { Stack, Text } from '@fluentui/react';
 import React from 'react';
-import { GetBatchResponse, GetLabelResponse, GetStepResponse } from '../../../backend/Api';
+import { GetBatchResponse,  GetStepResponse } from '../../../backend/Api';
 import { JobDetails } from '../../../backend/JobDetails';
 import { getHordeStyling } from '../../../styles/Styles';
 
-export const JobDetailPluginExample: React.FC<{ jobDetails: JobDetails; step?: GetStepResponse; batch?: GetBatchResponse, label?: GetLabelResponse}> = ({ jobDetails, step, batch, label }) => {
+export const JobDetailPluginExample: React.FC<{ jobDetails: JobDetails; step?: GetStepResponse; batch?: GetBatchResponse}> = ({ jobDetails, step, batch }) => {
 
    const { hordeClasses } = getHordeStyling();
 
-    // do not render on batch or label views
-    if (batch || label || !step) {
+    // do not render on batch
+    if (batch || !step) {
         return null;
     }
 
diff --git a/Engine/Source/Programs/Horde/HordeDashboard/src/backend/Api.ts b/Engine/Source/Programs/Horde/HordeDashboard/src/backend/Api.ts
--- a/Engine/Source/Programs/Horde/HordeDashboard/src/backend/Api.ts
+++ b/Engine/Source/Programs/Horde/HordeDashboard/src/backend/Api.ts
@@ -210,8 +210,6 @@ export type AgentData = GetAgentResponse
 export type LeaseData = GetAgentLeaseResponse
 export type SessionData = GetAgentSessionResponse
 export type PoolData = GetPoolResponse
-export type GroupData = GetGroupResponse
-export type NodeData = GetNodeResponse
 export type BatchData = GetBatchResponse
 export type EventData = GetLogEventResponse
 export type LogData = GetLogFileResponse
@@ -231,10 +229,6 @@ export type IssueData = GetIssueResponse & {
 	events?: GetLogEventResponse[];
 }
 
-export type LabelData = GetLabelResponse & {
-	defaultLabel?: GetDefaultLabelStateResponse;
-}
-
 export type ProjectData = GetProjectResponse & {
 	streams?: StreamData[];
 }
@@ -246,9 +240,7 @@ export type StreamData = GetStreamResponse & {
 
 }
 
-export type JobData = GetJobResponse & {
-	graphRef?: GetGraphResponse;
-}
+export type JobData = GetJobResponse;
 
 export type StepData = GetStepResponse & {
 
@@ -1243,7 +1235,7 @@ export type GetJobResponse = {
 	templateHash?: string;
 
 	/** Hash of the graph for this job */
-	graphHash?: string;
+	// graphHash?: string;
 
 	/** The user that started this job */
 	startedByUserInfo?: GetThinUserInfoResponse;
@@ -1282,7 +1274,7 @@ export type GetJobResponse = {
 	labels?: GetLabelStateResponse[];
 
 	/** The default label, containing the state of all steps that are otherwise not matched. */
-	defaultLabel?: GetDefaultLabelStateResponse;
+	defaultLabel?: GetLabelStateResponse;
 
 	/** List of reports */
 	reports?: GetReportResponse[];
@@ -1334,6 +1326,17 @@ export type UpdateStepRequest = {
 	properties?: { [key: string]: string | null };
 }
 
+
+/** Reference to the output of a step within the job */
+export type JobStepOutputRef = {
+	// Step producing the output
+	stepId: string;
+	// <param name="OutputIdx">Index of the output from this step</param>
+	outputIdx: number;
+}
+
+export type JobStepId = string;
+
 /**Returns information about a jobstep */
 export type GetStepResponse = {
 
@@ -1379,6 +1382,35 @@ export type GetStepResponse = {
 	/**User-defined properties for this jobstep. */
 	properties: { [key: string]: string };
 
+	/// The name of this node 
+	name: string;
+
+	/// References to inputs for this node
+	inputs?: JobStepOutputRef[];
+
+	/// List of output names
+	outputNames?: string[];
+
+	/// Indices of nodes which must have succeeded for this node to run
+	inputDependencies?: JobStepId[];
+
+	/// Indices of nodes which must have completed for this node to run
+	orderDependencies?: JobStepId[];
+
+	/// Whether this node can be run multiple times
+	allowRetry: boolean;
+
+	/// This node can start running early, before dependencies of other nodes in the same group are complete
+	runEarly: boolean;
+
+	/// Whether to include warnings in the output (defaults to true)
+	warnings: boolean;
+
+	/// List of credentials required for this node. Each entry maps an environment variable name to a credential in the form "CredentialName.PropertyName".
+	credentials?: Record<string, string>;
+
+	/// Annotations for this node
+	annotations?: Record<string, string>;
 }
 
 //**Returns information about test data */
@@ -1448,6 +1480,9 @@ export type GetBatchResponse = {
 	/**The agent assigned to execute this group */
 	agentId?: string;
 
+	/// The agent type
+	agentType: string;
+
 	/** The USD rate of an agent hour */
 	agentRate?: number;
 
@@ -1468,7 +1503,6 @@ export type GetBatchResponse = {
 
 	/**Time at which the group finished (UTC) */
 	finishTime?: Date | string;
-
 }
 
 /**Describes the history of a step */
@@ -1720,54 +1754,6 @@ export type CreateGroupRequest = {
 
 }
 
-/**Information required to create a node */
-export type GetNodeResponse = {
-
-	/**The name of this node  */
-	name: string;
-
-	/**Indices of nodes which must have succeeded for this node to run */
-	inputDependencies: string[];
-
-	/**Indices of nodes which must have completed for this node to run */
-	orderDependencies: string[];
-
-	/**The priority of this node */
-	priority: Priority;
-
-	/**Whether this node can be retried */
-	allowRetry: boolean;
-
-	/**This node can start running early, before dependencies of other nodes in the same group are complete */
-	runEarly: boolean;
-
-	/**Sets this node as a target to be built */
-	target: boolean;
-
-	/**Expected time to execute this node based on historical trends */
-	averageDuration: number;
-
-	/**Aggregates that this node belongs do */
-	aggregates?: string[];
-
-	/**Properties for this node */
-	properties: { [key: string]: string };
-
-}
-
-/**Information about a group of nodes */
-export type GetGroupResponse = {
-
-	/**The executor to use for this group */
-	executor: string;
-
-	/**The type of agent to execute this group */
-	agentType: string;
-
-	/**Nodes in the group */
-	nodes: GetNodeResponse[];
-}
-
 /**Request to update a node */
 export type UpdateNodeRequest = {
 
@@ -2770,60 +2756,6 @@ export enum JobState {
 }
 
 
-/**Information about a label */
-export type GetLabelResponse = {
-
-	/**Category of the aggregate */
-	category: string;
-
-	/**Label for this aggregate */
-	name: string;
-
-	/**Label for this aggregate, currently mapped to name property on server */
-	dashboardName?: string;
-
-	/**Name to show for this label in UGS */
-	ugsName?: string;
-
-	/** Project to display this label for in UGS */
-	ugsProject?: string;
-
-	/**Nodes which must be part of the job for the aggregate to be shown */
-	requiredNodes: string[];
-
-	/**Nodes to include in the status of this aggregate, if present in the job */
-	includedNodes: string[];
-}
-
-
-/**Information about an aggregate */
-export type GetAggregateResponse = {
-
-	/**Name of the aggregate */
-	name: string;
-
-	/**Nodes which must be part of the job for the aggregate to be shown */
-	nodes: string[];
-
-}
-
-/**Information about a graph */
-export type GetGraphResponse = {
-
-	/**The hash of the graph */
-	hash: string;
-
-	/**Array of nodes for this job */
-	groups?: GetGroupResponse[];
-
-	/**List of aggregates */
-	namedAggregates?: GetAggregateResponse[];
-
-	/**List of labels for the graph */
-	labels?: GetLabelResponse[];
-
-}
-
 /**The timing info for a job*/
 export type GetJobTimingResponse = {
 
@@ -2853,6 +2785,22 @@ export type GetLabelTimingInfoResponse = GetTimingInfoResponse &
 
 /**State of an label within a job */
 export type GetLabelStateResponse = {
+
+	// Name to show for this label on the dashboard
+	dashboardName?: string;
+
+	// Category to show this label in on the dashboard
+	dashboardCategory?: string;
+
+	// Name to show for this label in UGS
+	ugsName?: string;
+
+	// Project to display this label for in UGS
+	ugsProject?: string;
+
+	// Steps to include in the status of this label
+	steps: JobStepId[];
+
 	/**State of the label */
 	state?: LabelState;
 
@@ -2860,14 +2808,6 @@ export type GetLabelStateResponse = {
 	outcome?: LabelOutcome;
 }
 
-/**Information about the default label (ie. with inlined list of nodes) */
-export type GetDefaultLabelStateResponse = GetLabelStateResponse &
-{
-	/**List of nodes covered by default label */
-	nodes: string[];
-}
-
-
 /**Information about the timing info for a particular target */
 export type GetTimingInfoResponse = {
 
diff --git a/Engine/Source/Programs/Horde/HordeDashboard/src/backend/Backend.ts b/Engine/Source/Programs/Horde/HordeDashboard/src/backend/Backend.ts
--- a/Engine/Source/Programs/Horde/HordeDashboard/src/backend/Backend.ts
+++ b/Engine/Source/Programs/Horde/HordeDashboard/src/backend/Backend.ts
@@ -1,10 +1,9 @@
 // Copyright Epic Games, Inc. All Rights Reserved.
 
 import templateCache from '../backend/TemplateCache';
-import { AccountClaimMessage, AgentData, AgentQuery, ApproveAgentsRequest, ArtifactData, AuditLogEntry, AuditLogQuery, BatchUpdatePoolRequest, ChangeSummaryData, CreateAccountRequest, CreateBisectTaskRequest, CreateBisectTaskResponse, CreateDeviceRequest, CreateDeviceResponse, CreateExternalIssueRequest, CreateExternalIssueResponse, CreateJobRequest, CreateJobResponse, CreateNoticeRequest, CreatePoolRequest, CreateServiceAccountRequest, CreateServiceAccountResponse, CreateSoftwareResponse, CreateSubscriptionRequest, CreateSubscriptionResponse, CreateZipRequest, DashboardPreference, DevicePoolTelemetryQuery, DeviceTelemetryQuery, EventData, FindArtifactsResponse, FindIssueResponse, FindJobTimingsResponse, GetAccountResponse, GetArtifactDirectoryResponse, GetArtifactZipRequest, GetBisectTaskResponse, GetDashboardConfigResponse, GetDevicePlatformResponse, GetDevicePoolResponse, GetDevicePoolTelemetryResponse, GetDeviceReservationResponse, GetDeviceResponse, GetDeviceTelemetryResponse, GetExternalIssueProjectResponse, GetExternalIssueResponse, GetGraphResponse, GetIssueStreamResponse, GetJobsTabResponse, GetJobStepRefResponse, GetJobStepTraceResponse, GetJobTimingResponse, GetLogEventResponse, GetLogFileResponse, GetNoticeResponse, GetNotificationResponse, GetPendingAgentsResponse, GetPerforceServerStatusResponse, GetPoolResponse, GetPoolSummaryResponse, GetServerInfoResponse, GetServerSettingsResponse, GetServiceAccountResponse, GetSoftwareResponse, GetSubscriptionResponse, GetTelemetryMetricsResponse, GetTestDataDetailsResponse, GetTestDataRefResponse, GetTestMetaResponse, GetTestResponse, GetTestsRequest, GetTestStreamResponse, GetToolSummaryResponse, GetUserResponse, GetUtilizationTelemetryResponse, GlobalConfig, IssueData, IssueQuery, IssueQueryV2, JobData, JobQuery, JobsTabColumnType, JobStepOutcome, JobStreamQuery, JobTimingsQuery, LeaseData, LogData, LogEventQuery, LogLineData, MetricsQuery, PoolData, PoolQuery, PreflightConfigResponse, ProjectData, ScheduleData, ScheduleQuery, SearchLogFileResponse, ServerStatusResponse, ServerUpdateResponse, SessionData, StreamData, TabType, TestData, UpdateAccountRequest, UpdateAgentRequest, UpdateBisectTaskRequest, UpdateCurrentAccountRequest, UpdateDeviceRequest, UpdateGlobalConfigRequest, UpdateIssueRequest, UpdateJobRequest, UpdateLeaseRequest, UpdateNoticeRequest, UpdateNotificationsRequest, UpdatePoolRequest, UpdateServerSettingsRequest, UpdateServiceAccountRequest, UpdateServiceAccountResponse, UpdateStepRequest, UpdateStepResponse, UpdateTemplateRefRequest, UpdateUserRequest, UsersQuery } from './Api';
+import { AccountClaimMessage, AgentData, AgentQuery, ApproveAgentsRequest, ArtifactData, AuditLogEntry, AuditLogQuery, BatchUpdatePoolRequest, ChangeSummaryData, CreateAccountRequest, CreateBisectTaskRequest, CreateBisectTaskResponse, CreateDeviceRequest, CreateDeviceResponse, CreateExternalIssueRequest, CreateExternalIssueResponse, CreateJobRequest, CreateJobResponse, CreateNoticeRequest, CreatePoolRequest, CreateServiceAccountRequest, CreateServiceAccountResponse, CreateSoftwareResponse, CreateSubscriptionRequest, CreateSubscriptionResponse, CreateZipRequest, DashboardPreference, DevicePoolTelemetryQuery, DeviceTelemetryQuery, EventData, FindArtifactsResponse, FindIssueResponse, FindJobTimingsResponse, GetAccountResponse, GetArtifactDirectoryResponse, GetArtifactZipRequest, GetBisectTaskResponse, GetDashboardConfigResponse, GetDevicePlatformResponse, GetDevicePoolResponse, GetDevicePoolTelemetryResponse, GetDeviceReservationResponse, GetDeviceResponse, GetDeviceTelemetryResponse, GetExternalIssueProjectResponse, GetExternalIssueResponse, GetIssueStreamResponse, GetJobsTabResponse, GetJobStepRefResponse, GetJobStepTraceResponse, GetJobTimingResponse, GetLogEventResponse, GetLogFileResponse, GetNoticeResponse, GetNotificationResponse, GetPendingAgentsResponse, GetPerforceServerStatusResponse, GetPoolResponse, GetPoolSummaryResponse, GetServerInfoResponse, GetServerSettingsResponse, GetServiceAccountResponse, GetSoftwareResponse, GetSubscriptionResponse, GetTelemetryMetricsResponse, GetTestDataDetailsResponse, GetTestDataRefResponse, GetTestMetaResponse, GetTestResponse, GetTestsRequest, GetTestStreamResponse, GetToolSummaryResponse, GetUserResponse, GetUtilizationTelemetryResponse, GlobalConfig, IssueData, IssueQuery, IssueQueryV2, JobData, JobQuery, JobsTabColumnType, JobStepOutcome, JobStreamQuery, JobTimingsQuery, LeaseData, LogData, LogEventQuery, LogLineData, MetricsQuery, PoolData, PoolQuery, PreflightConfigResponse, ProjectData, ScheduleData, ScheduleQuery, SearchLogFileResponse, ServerStatusResponse, ServerUpdateResponse, SessionData, StreamData, TabType, TestData, UpdateAccountRequest, UpdateAgentRequest, UpdateBisectTaskRequest, UpdateCurrentAccountRequest, UpdateDeviceRequest, UpdateGlobalConfigRequest, UpdateIssueRequest, UpdateJobRequest, UpdateLeaseRequest, UpdateNoticeRequest, UpdateNotificationsRequest, UpdatePoolRequest, UpdateServerSettingsRequest, UpdateServiceAccountRequest, UpdateServiceAccountResponse, UpdateStepRequest, UpdateStepResponse, UpdateTemplateRefRequest, UpdateUserRequest, UsersQuery } from './Api';
 import dashboard, { Dashboard } from './Dashboard';
 import { ChallengeStatus, Fetch } from './Fetch';
-import graphCache, { GraphQuery } from './GraphCache';
 import { projectStore } from './ProjectStore';
 
 
@@ -146,7 +145,7 @@ export class Backend {
             });
         });
     }
-    
+
 
     getAgent(id: string): Promise<AgentData> {
         return new Promise<AgentData>((resolve, reject) => {
@@ -278,24 +277,7 @@ export class Backend {
 
     }
 
-    /** Get's a job's graph data, used by caching and shouldn't need to call directly */
-    getGraph(jobId: string): Promise<GetGraphResponse> {
-
-        return new Promise<GetGraphResponse>((resolve, reject) => {
-
-            this.backend.get(`/api/v1/jobs/${jobId}/graph`).then((value) => {
-                const response = value.data as GetGraphResponse;
-                // filter out ugs labels
-                response.labels = response.labels?.filter(label => !!label.dashboardName);
-                resolve(response);
-            }).catch(reason => {
-                reject(reason);
-            });
-        });
-
-    }
-
-    getJob(id: string, query?: JobQuery, includeGraph = true, show404Error = false): Promise<JobData> {
+    getJob(id: string, query?: JobQuery, show404Error = false): Promise<JobData> {
 
         return new Promise<JobData>((resolve, reject) => {
 
@@ -305,22 +287,7 @@ export class Backend {
             }).then((value) => {
 
                 const response = value.data as JobData;
-                if (includeGraph && !response.graphHash) {
-                    return reject(`Job ${id} has undefined graph hash`);
-                }
-
-                if (!includeGraph) {
-                    resolve(response);
-                    return;
-                }
-
-                graphCache.get({ graphHash: response.graphHash!, jobId: id }).then(graph => {
-                    response.graphRef = graph;
-                    resolve(response);
-                }).catch(reason => {
-                    reject(reason);
-                });
-
+                resolve(response);
 
             }).catch(reason => {
                 reject(reason);
@@ -329,7 +296,7 @@ export class Backend {
 
     }
 
-    getJobsByIds(id: string[], query?: JobQuery, includeGraphs?: boolean): Promise<JobData[]> {
+    getJobsByIds(id: string[], query?: JobQuery): Promise<JobData[]> {
 
         if (!query) {
             (query as any) = {
@@ -347,24 +314,7 @@ export class Backend {
 
                 const response = value.data as JobData[];
 
-                if (includeGraphs === false || includeGraphs === undefined) {
-                    resolve(response);
-                    return;
-                }
-
-                graphCache.getGraphs(response.map(j => {
-                    return {
-                        jobId: j.id!,
-                        graphHash: j.graphHash!
-                    };
-                })).then((graphs => {
-                    response.forEach(j => {
-                        j.graphRef = graphs.find(g => g.hash === j.graphHash)
-                    })
-                    resolve(response)
-                })).catch(reason => {
-                    reject(reason);
-                });
+                resolve(response);
 
             }).catch(reason => {
                 reject(reason);
@@ -373,7 +323,7 @@ export class Backend {
 
     }
 
-    getStreamJobs(streamId: string, query: JobStreamQuery, queryGraph: boolean = false): Promise<JobData[]> {
+    getStreamJobs(streamId: string, query: JobStreamQuery): Promise<JobData[]> {
 
         if (typeof query.index === 'number') {
             query.index = 0;
@@ -389,31 +339,8 @@ export class Backend {
                 params: query
             }).then((value) => {
                 const jobs = value.data as JobData[];
-
-                if (!queryGraph) {
-                    resolve(jobs);
-                    return;
-                }
-
-                const query: GraphQuery[] = [];
-                jobs.forEach(j => {
-                    if (!j.graphHash) {
-                        console.error(`Job ${j.id} has no graph hash`);
-                        return;
-                    }
-
-                    if (!query.find(q => q.graphHash === j.graphHash)) {
-                        query.push({ graphHash: j.graphHash!, jobId: j.id });
-                    }
-                });
-
-                graphCache.getGraphs(query).then(graphs => {
-                    jobs.forEach(j => {
-                        j.graphRef = graphs.find(g => g.hash === j.graphHash);
-                    });
-                    resolve(jobs);
-                }).catch(reason => reject(reason));
-
+                resolve(jobs);
+                return;
             }).catch((reason) => {
                 reject(reason);
             });
@@ -422,7 +349,7 @@ export class Backend {
     }
 
 
-    getJobs(query: JobQuery, queryGraph: boolean = false): Promise<JobData[]> {
+    getJobs(query: JobQuery): Promise<JobData[]> {
 
         if (typeof query.index === 'number') {
             query.index = 0;
@@ -439,29 +366,8 @@ export class Backend {
             }).then((value) => {
                 const jobs = value.data as JobData[];
 
-                if (!queryGraph) {
-                    resolve(jobs);
-                    return;
-                }
-
-                const query: GraphQuery[] = [];
-                jobs.forEach(j => {
-                    if (!j.graphHash) {
-                        console.error(`Job ${j.id} has no graph hash`);
-                        return;
-                    }
-
-                    if (!query.find(q => q.graphHash === j.graphHash)) {
-                        query.push({ graphHash: j.graphHash!, jobId: j.id });
-                    }
-                });
-
-                graphCache.getGraphs(query).then(graphs => {
-                    jobs.forEach(j => {
-                        j.graphRef = graphs.find(g => g.hash === j.graphHash);
-                    });
-                    resolve(jobs);
-                }).catch(reason => reject(reason));
+                resolve(jobs);
+                return;
 
             }).catch((reason) => {
                 reject(reason);
@@ -489,20 +395,7 @@ export class Backend {
                 params: query
             }).then((response) => {
                 let timingsWrapper = response.data as FindJobTimingsResponse;
-                let graphCalls: { jobResponse: JobData, call: any }[] = [];
-                Object.values(timingsWrapper.timings).forEach(timing => {
-                    if (!timing.jobResponse.graphHash) {
-                        return reject(`Job ${timing.jobResponse.id} has undefined graph hash`);
-                    }
-                    graphCalls.push({ jobResponse: timing.jobResponse, call: graphCache.get({ graphHash: timing.jobResponse.graphHash, jobId: timing.jobResponse.id }) });
-                });
-                let allPromises = graphCalls.map(item => item.call);
-                Promise.all(allPromises).then(responses => {
-                    for (let idx = 0; idx < responses.length; idx++) {
-                        graphCalls[idx].jobResponse.graphRef = responses[idx] as GetGraphResponse;
-                    }
-                    resolve(timingsWrapper);
-                });
+                resolve(timingsWrapper);
             }).catch(reason => {
                 reject(reason);
             });
@@ -1336,7 +1229,7 @@ export class Backend {
         }
 
     }
-    
+
     updateJobStep(jobId: string, batchId: string, stepId: string, request: UpdateStepRequest): Promise<UpdateStepResponse> {
         return new Promise<UpdateStepResponse>((resolve, reject) => {
             this.backend.put(`api/v1/jobs/${jobId}/batches/${batchId}/steps/${stepId}`, request).then((value) => {
@@ -1786,7 +1679,7 @@ export class Backend {
             });
         });
     }
-    
+
     // update current account 
     updateCurrentAccount(request: UpdateCurrentAccountRequest): Promise<boolean> {
         return new Promise<boolean>((resolve, reject) => {
@@ -1968,7 +1861,6 @@ export class Backend {
             await dashboard.update();
 
             if (dashboard.localCache) {
-                graphCache.initialize();
                 templateCache.initialize();
             }
 
diff --git a/Engine/Source/Programs/Horde/HordeDashboard/src/backend/GraphCache.ts b/Engine/Source/Programs/Horde/HordeDashboard/src/backend/GraphCache.ts
deleted file mode 100644
diff --git a/Engine/Source/Programs/Horde/HordeDashboard/src/backend/JobDetails.ts b/Engine/Source/Programs/Horde/HordeDashboard/src/backend/JobDetails.ts
--- a/Engine/Source/Programs/Horde/HordeDashboard/src/backend/JobDetails.ts
+++ b/Engine/Source/Programs/Horde/HordeDashboard/src/backend/JobDetails.ts
@@ -6,11 +6,11 @@ import moment from 'moment';
 import backend from '.';
 import { getBatchInitElapsed, getNiceTime, getStepETA, getStepElapsed, getStepFinishTime, getStepTimingDelta } from '../base/utilities/timeUtils';
 import { getBatchText } from '../components/JobDetailCommon';
-import { AgentData, ArtifactData, BatchData, EventData, GetGroupResponse, GetJobTimingResponse, GetLabelResponse, GetLabelStateResponse, GetLabelTimingInfoResponse, GroupData, IssueData, JobData, JobState, JobStepBatchState, JobStepError, JobStepOutcome, JobStepState, LabelState, NodeData, ReportPlacement, StepData, StreamData, TestData } from './Api';
+import { AgentData, ArtifactData, BatchData, EventData, GetJobTimingResponse, GetLabelStateResponse, GetLabelTimingInfoResponse, IssueData, JobData, JobState, JobStepBatchState, JobStepError, JobStepOutcome, JobStepState, LabelState, ReportPlacement, StepData, StreamData, TestData } from './Api';
 import { projectStore } from './ProjectStore';
 
 
-export type JobLabel = GetLabelResponse & {
+export type JobLabel = {
     stateResponse: GetLabelStateResponse;
     timing?: GetLabelTimingInfoResponse;
     default: boolean;
@@ -146,25 +146,11 @@ export class JobDetails {
         this.process();
     }
 
-    stepsByLabel(label: GetLabelResponse): StepData[] {
+    stepsByLabel(label: GetLabelStateResponse): StepData[] {
 
-        const steps: StepData[] = [];
+        const stepIds = new Set(label.steps);
 
-        this.getSteps().forEach(step => {
-
-            const node = this.nodeByStepId(step.id);
-
-            if (!node) {
-                return;
-            }
-
-            if (label.includedNodes.indexOf(node.name) !== -1 && !steps.find(s => s.id === step.id)) {
-                steps.push(step);
-            }
-
-        })
-
-        return steps;
+        return this.getSteps().filter(s => stepIds.has(s.id));
     }
 
     stepById(id?: string): StepData | undefined {
@@ -229,7 +215,7 @@ export class JobDetails {
                 return;
             }
 
-            step = b.steps.find(s => this.nodeByStepId(s.id)?.name === name);
+            step = b.steps.find(s => s.name === name);
 
         });
 
@@ -262,16 +248,15 @@ export class JobDetails {
 
     getStepRetries(stepId: string): StepData[] {
 
-        const node = this.nodeByStepId(stepId);
-
-        if (!node) {
+        const step = this.stepById(stepId);
+        if (!step) {
             return [];
         }
 
         let steps: StepData[] = [];
 
         this.batches.forEach(b => {
-            b.steps.filter(s => this.nodeByStepId(s.id) === node).forEach(s => {
+            b.steps.filter(s => s.name === step.name).forEach(s => {
                 steps.push(s);
             });
         })
@@ -301,21 +286,6 @@ export class JobDetails {
     }
 
 
-    getNodeGroupIdx(node: NodeData): number {
-        const group = this.groups.find(g => g.nodes.indexOf(node) !== -1);
-        if (!group) {
-            return -1;
-        }
-        return this.groups.indexOf(group);
-    }
-
-
-    getTargets(): NodeData[] {
-
-        return this.nodes.filter(n => n.target);
-
-    }
-
     batchByStepId(stepId: string | undefined): BatchData | undefined {
         if (!stepId) {
             return undefined;
@@ -338,38 +308,20 @@ export class JobDetails {
     }
 
 
-    nodeByStepId(stepId: string | undefined): NodeData | undefined {
-
-        if (!stepId) {
-            return undefined;
-        }
-
-        const batch = this.batchByStepId(stepId);
-        const step = this.stepById(stepId);
-
-        if (!batch || !step) return undefined;
-
-        return this.groups[batch.groupIdx].nodes[step.nodeIdx];
-
-    }
-
     getStepName(stepId: string | undefined, includeRetry: boolean = true): string {
 
-        if (!stepId) {
-            return "";
-        }
+        const step = this.stepById(stepId);
 
-        const node = this.nodeByStepId(stepId);
-        if (!node) {
+        if (!step?.name) {
             return "";
         }
 
         const idx = this.getStepRetryNumber(stepId);
         if (!idx || !includeRetry) {
-            return node.name;
+            return step.name;
         }
 
-        return `${node.name} (${idx + 1})`;
+        return `${step.name} (${idx + 1})`;
 
     }
 
@@ -381,14 +333,6 @@ export class JobDetails {
         return this.testdata.filter(t => t.stepId === stepId);
     }
 
-    getLabelIndex(labelIn: JobLabel): number {
-        const label = this.findLabel(labelIn.name, labelIn.category);
-        if (!label) {
-            return -1;
-        }
-        return this.labels.indexOf(label);
-    }
-
     stepPrice(stepId: string): number | undefined {
 
         const step = this.stepById(stepId);
@@ -460,9 +404,9 @@ export class JobDetails {
 
         const label = labels.find(label => {
             if (category) {
-                return category === label.category && name === label.name;
+                return category === label.stateResponse.dashboardCategory && name === label.stateResponse.dashboardName;
             }
-            return name === label.name;
+            return name === label.stateResponse.dashboardName;
         });
 
         return label;
@@ -489,21 +433,6 @@ export class JobDetails {
     }
 
 
-    groupByNodeName(name: string | undefined): GetGroupResponse | undefined {
-
-        if (!name) {
-            return undefined;
-        }
-
-        const groups = this.jobdata?.graphRef?.groups;
-
-        return groups?.find(g => {
-            return g?.nodes.find(n => n.name.toLowerCase() === name.toLowerCase());
-        });
-
-    }
-
-
     resumePolling() {
 
         if (!this.id) {
@@ -540,7 +469,6 @@ export class JobDetails {
         this.logId = undefined;
         this.stepId = undefined;
         this.labelIdx = undefined;
-        this.groups = [];
         this.batches = [];
         this.events = [];
         this.issues = [];
@@ -578,24 +506,10 @@ export class JobDetails {
         this.id = jobdata.id;
         this.stream = projectStore.streamById(jobdata.streamId);
         this.batches = jobdata.batches ?? [];
-        this.groups = jobdata.graphRef?.groups ?? [];
-        this.nodes = [];
-        if (this.groups) {
-            this.nodes = this.groups.map(g => g.nodes).flat();
-        }
-
-        let labels: GetLabelResponse[] = [];
-        if (jobdata.graphRef?.labels && jobdata.labels) {
-            labels = jobdata.graphRef.labels;
-        }
 
-        this.labels = labels.map((label, index) => {
+        this.labels = jobdata.labels.map((label, index) => {
 
             return {
-                category: label.category,
-                name: label.name,
-                includedNodes: label.includedNodes,
-                requiredNodes: label.requiredNodes,
                 default: false,
                 stateResponse: jobdata.labels![index],
                 internal: label,
@@ -606,10 +520,6 @@ export class JobDetails {
         const defaultLabel = jobdata.defaultLabel;
         if (defaultLabel) {
             this.labels.push({
-                category: "Other",
-                name: "Other",
-                includedNodes: defaultLabel.nodes,
-                requiredNodes: [],
                 stateResponse: defaultLabel,
                 default: true
             })
@@ -935,8 +845,6 @@ export class JobDetails {
     jobdata?: JobData
     stream?: StreamData
     labels: JobLabel[] = []
-    groups: GroupData[] = []
-    nodes: NodeData[] = []
     batches: BatchData[] = []
     events: EventData[] = []
     artifacts: ArtifactData[] = []
@@ -1049,8 +957,7 @@ export const getBatchSummaryMarkdown = (jobDetails: JobDetails, batchId: string)
     }
 
 
-    const group = jobDetails.groups[batch.groupIdx];
-    const agentType = group?.agentType ?? "";
+    const agentType = batch?.agentType ?? "";
     const agentPool = jobDetails.stream?.agentTypes[agentType!]?.pool ?? "";
 
     let summaryText = "";
@@ -1109,8 +1016,7 @@ export const getStepSummaryMarkdown = (jobDetails: JobDetails, stepId: string):
             return undefined;
         }
 
-        const group = jobDetails.groups[batch!.groupIdx];
-        const agentType = group?.agentType;
+        const agentType = batch?.agentType;
         const agentPool = jobDetails.stream?.agentTypes[agentType!]?.pool;
         return getBatchText({ batch: batch, agentType: agentType, agentPool: agentPool });
 
diff --git a/Engine/Source/Programs/Horde/HordeDashboard/src/backend/JobHandler.ts b/Engine/Source/Programs/Horde/HordeDashboard/src/backend/JobHandler.ts
--- a/Engine/Source/Programs/Horde/HordeDashboard/src/backend/JobHandler.ts
+++ b/Engine/Source/Programs/Horde/HordeDashboard/src/backend/JobHandler.ts
@@ -4,8 +4,6 @@ import { action, makeObservable, observable } from 'mobx';
 import moment from 'moment';
 import backend from '../backend';
 import { JobData, JobState, JobStepOutcome, JobStreamQuery, StreamData } from '../backend/Api';
-import graphCache, { GraphQuery } from '../backend/GraphCache';
-
 
 export type FilterStatus = "Running" | "Complete" | "Succeeded" | "Failed" | "Waiting";
 
@@ -125,7 +123,7 @@ export class JobHandler {
 
             const cancelId = this.cancelId++;
 
-            const mjobs = await backend.getStreamJobs(this.stream!.id, query, false);
+            const mjobs = await backend.getStreamJobs(this.stream!.id, query);
 
             // check for canceled after modified test
             if (this.canceled.has(cancelId)) {
@@ -151,10 +149,6 @@ export class JobHandler {
             const jobs: JobData[] = [];
 
             mjobs.forEach(j1 => {
-                const existing = this.jobs.find(j2 => j1.id === j2.id);
-                if (existing) {
-                    j1.graphRef = existing.graphRef;
-                }
                 jobs.push(j1);
             })
 
@@ -167,56 +161,13 @@ export class JobHandler {
                 return timeA < timeB ? 1 : -1;
             });
 
-            const graphHashes = new Set<string>();
-
-            this.jobs.forEach(j => {
-
-                if (j.graphRef?.hash !== j.graphHash) {
-                    j.graphRef = undefined;
-                }
-
-                if (graphHashes.size > 5) {
-                    return;
-                }
-
-                if (j.graphHash && !j.graphRef) {
-                    j.graphRef = graphCache.cache.get(j.graphHash);
-                    if (!j.graphRef) {
-                        graphHashes.add(j.graphHash);
-                    }
-                }
-            })
-
-            if (graphHashes.size) {
-
-                const graphQuery: GraphQuery[] = [];
-                Array.from(graphHashes.values()).forEach(h => {
-                    const jobId = this.jobs.find(j => j.graphHash === h)!.id;
-                    graphQuery.push({
-                        jobId: jobId,
-                        graphHash: h
-                    })
-                })
-
-                const graphs = await graphCache.getGraphs(graphQuery);
-
-                graphs.forEach(graph => {
-
-                    this.jobs.forEach(j => {
-                        if (graph.hash === j.graphHash) {
-                            j.graphRef = graph;
-                        }
-                    })
-                })
-
-            }
 
             // check for canceled during graph request
             if (this.canceled.has(cancelId)) {
                 return;
             }
 
-            if (graphHashes.size || mjobs.length || this.initial) {
+            if (mjobs.length || this.initial) {
                 this.initial = false;
                 wasUpdated = true;
             }
diff --git a/Engine/Source/Programs/Horde/HordeDashboard/src/backend/LocalCache.ts b/Engine/Source/Programs/Horde/HordeDashboard/src/backend/LocalCache.ts
deleted file mode 100644
diff --git a/Engine/Source/Programs/Horde/HordeDashboard/src/backend/TestDataHandler.ts b/Engine/Source/Programs/Horde/HordeDashboard/src/backend/TestDataHandler.ts
--- a/Engine/Source/Programs/Horde/HordeDashboard/src/backend/TestDataHandler.ts
+++ b/Engine/Source/Programs/Horde/HordeDashboard/src/backend/TestDataHandler.ts
@@ -196,11 +196,7 @@ export class TestDataWrapper implements TestData {
         let stepName = "";
         const batch = this.jobdata.batches?.find(b => b.steps.find(s => s.id === this.stepId));
         const stepNode = batch?.steps.find(s => s.id === this.stepId);
-        const groups = this.jobdata?.graphRef?.groups;
-        if (groups && stepNode && batch) {
-            stepName = groups[batch.groupIdx]?.nodes[stepNode.nodeIdx]?.name;
-        }
-        this.stepName = stepName;
+        stepName = stepNode.name;
 
         return this.stepName;
     }
diff --git a/Engine/Source/Programs/Horde/HordeDashboard/src/base/utilities/filter.ts b/Engine/Source/Programs/Horde/HordeDashboard/src/base/utilities/filter.ts
--- a/Engine/Source/Programs/Horde/HordeDashboard/src/base/utilities/filter.ts
+++ b/Engine/Source/Programs/Horde/HordeDashboard/src/base/utilities/filter.ts
@@ -1,7 +1,5 @@
 
 import { GetJobResponse } from "../../backend/Api";
-import graphCache from "../../backend/GraphCache";
-
 
 export type JobFilterSimple = {
     filterKeyword?: string;
@@ -30,34 +28,11 @@ const getJobKeywords = (job: GetJobResponse): string[] => {
         keywords.push("Scheduler");
     }
 
-    if (job.graphHash) {
-        const graph = graphCache.cache.get(job.graphHash);
-
-        if (graph) {
-
-            const nodes: string[] = [];
-            graph.groups?.forEach(g => g.nodes.forEach(n => nodes.push(n.name)));
-
-            job.batches?.forEach(batch => {
-
-                const group = graph.groups?.[batch.groupIdx];
-
-                if (group) {
-
-                    batch.steps.forEach(step => {
-                        const name = group.nodes[step.nodeIdx]?.name;
-                        if (name && keywords.indexOf(name) === -1) {                            
-                            keywords.push(name);
-                        }
-
-                    })
-                }
-
-            });
-
-        } 
-    }
-
+    job.batches?.forEach(batch => {
+        batch.steps.forEach(step => {
+            keywords.push(step.name);
+        })
+    })
 
     keywords.push(...job.arguments);
 
diff --git a/Engine/Source/Programs/Horde/HordeDashboard/src/components/AutomationReport.tsx b/Engine/Source/Programs/Horde/HordeDashboard/src/components/AutomationReport.tsx
--- a/Engine/Source/Programs/Horde/HordeDashboard/src/components/AutomationReport.tsx
+++ b/Engine/Source/Programs/Horde/HordeDashboard/src/components/AutomationReport.tsx
@@ -223,7 +223,7 @@ class TestReportGenerator {
       });
 
       if (jobIds.size) {
-         const jobs = await backend.getJobsByIds(Array.from(jobIds), { filter: "id,batches" }, false);
+         const jobs = await backend.getJobsByIds(Array.from(jobIds), { filter: "id,batches" });
          jobs.forEach(j => {
             const metas = allMetaEvents.filter(m => m.jobId === j.id);
             metas.forEach(m => {
diff --git a/Engine/Source/Programs/Horde/HordeDashboard/src/components/ErrorPane.tsx b/Engine/Source/Programs/Horde/HordeDashboard/src/components/ErrorPane.tsx
--- a/Engine/Source/Programs/Horde/HordeDashboard/src/components/ErrorPane.tsx
+++ b/Engine/Source/Programs/Horde/HordeDashboard/src/components/ErrorPane.tsx
@@ -5,7 +5,7 @@ import { IColumn, List, Stack, Text } from '@fluentui/react';
 import { getFocusStyle, mergeStyleSets } from '@fluentui/react/lib/Styling';
 import React from 'react';
 import { Link, useNavigate } from 'react-router-dom';
-import { EventData, EventSeverity, GetLabelResponse } from '../backend/Api';
+import { EventData, EventSeverity, GetLabelStateResponse } from '../backend/Api';
 import { JobDetails } from '../backend/JobDetails';
 import { JobEventHandler } from '../backend/JobEventHandler';
 import { renderLine } from './LogRender';
@@ -237,7 +237,7 @@ export const JobEventListPanel: React.FC<{ jobDetails: JobDetails, stepIds: stri
       const step = jobDetails.stepById(stepId)!;
 
       return <Stack styles={{ root: { background: 'rgb(233, 232, 231)', marginRight: 8, height: 42, fontSize: 12, selectors: { "a, a:hover, a:visited": { color: "#FFFFFF" }, ":hover": { background: 'rgb(223, 222, 221)' } } } }} verticalFill={true} verticalAlign="center" >
-         <Link to={url}><Stack style={{ paddingLeft: 12 }} horizontal><StepStatusIcon step={step} /><Text>{jobDetails.nodeByStepId(stepId)?.name}</Text>
+         <Link to={url}><Stack style={{ paddingLeft: 12 }} horizontal><StepStatusIcon step={step} /><Text>{step?.name}</Text>
          </Stack>
          </Link>
       </Stack>
@@ -254,7 +254,7 @@ export const JobEventListPanel: React.FC<{ jobDetails: JobDetails, stepIds: stri
 };
 
 
-export const JobEventPanel: React.FC<{ jobDetails: JobDetails, label?: GetLabelResponse, eventHandler: JobEventHandler }> = observer(({ jobDetails, label, eventHandler }) => {
+export const JobEventPanel: React.FC<{ jobDetails: JobDetails, label?: GetLabelStateResponse, eventHandler: JobEventHandler }> = observer(({ jobDetails, label, eventHandler }) => {
 
    const { hordeClasses } = getHordeStyling();
 
@@ -269,8 +269,7 @@ export const JobEventPanel: React.FC<{ jobDetails: JobDetails, label?: GetLabelR
    let stepIds = jobDetails.getSteps().map(s => s.id);
 
    if (label) {
-      const nodes = jobDetails.nodes?.filter(n => label.includedNodes?.find(on => on === n.name));
-      stepIds = stepIds.filter(stepId => nodes.indexOf(jobDetails.nodeByStepId(stepId)!) !== -1);
+      stepIds = stepIds.filter(stepId => label.steps.indexOf(stepId) !== -1);
    }
 
    stepIds = stepIds.filter(s => {
diff --git a/Engine/Source/Programs/Horde/HordeDashboard/src/components/JobDetailCommon.tsx b/Engine/Source/Programs/Horde/HordeDashboard/src/components/JobDetailCommon.tsx
--- a/Engine/Source/Programs/Horde/HordeDashboard/src/components/JobDetailCommon.tsx
+++ b/Engine/Source/Programs/Horde/HordeDashboard/src/components/JobDetailCommon.tsx
@@ -1,12 +1,11 @@
 // Copyright Epic Games, Inc. All Rights Reserved.
 
 import { useLocation } from 'react-router-dom';
-import { BatchData, GetJobStepRefResponse, JobStepBatchError, JobStepBatchState, JobStepOutcome, JobStepState, NodeData, StepData } from '../backend/Api';
+import { BatchData, GetJobStepRefResponse, JobStepBatchError, JobStepBatchState, JobStepOutcome, JobStepState, StepData } from '../backend/Api';
 
 type StepItem = {
    step?: StepData;
    batch?: BatchData;
-   node?: NodeData;
    agentId?: string;
    agentRow?: boolean;
    agentType?: string;
diff --git a/Engine/Source/Programs/Horde/HordeDashboard/src/components/JobLabelCallout.tsx b/Engine/Source/Programs/Horde/HordeDashboard/src/components/JobLabelCallout.tsx
--- a/Engine/Source/Programs/Horde/HordeDashboard/src/components/JobLabelCallout.tsx
+++ b/Engine/Source/Programs/Horde/HordeDashboard/src/components/JobLabelCallout.tsx
@@ -5,7 +5,7 @@ import { observer } from 'mobx-react-lite';
 import { Callout, DirectionalHint, List, Stack, Text } from "@fluentui/react";
 import React from "react";
 import { Link } from "react-router-dom";
-import { LabelData, StepData } from "../backend/Api";
+import { GetLabelStateResponse, StepData } from "../backend/Api";
 import { JobDetails } from "../backend/JobDetails";
 import { StepStatusIcon } from './StatusIcon';
 
@@ -15,8 +15,8 @@ type StepItem = {
 
 export type CalloutState = {
    jobId?: string;
-   label?: LabelData;
    target?: string;
+   label?: GetLabelStateResponse;
 }
 
 export class CalloutController {
@@ -117,6 +117,7 @@ export class CalloutController {
 export const JobLabelCallout: React.FC<{ controller: CalloutController }> = observer(({ controller }) => {
 
    const state = controller.state;
+
    const label = state.label;
    const jobId = state.jobId;
    const target = state.target;
@@ -126,12 +127,10 @@ export const JobLabelCallout: React.FC<{ controller: CalloutController }> = obse
       return <div />;
    }
 
-   const nodes = details.nodes?.filter(n => label.includedNodes?.find(on => on === n.name));
    const allSteps = details.getSteps();
-   const steps = allSteps.filter(step => nodes.indexOf(details.nodeByStepId(step.id)!) !== -1);
+   const steps = allSteps.filter(step => label.steps.indexOf(step.id) !== -1);
    const items = steps.map(step => { return { step: step }; });
 
-
    const onRenderCell = (stepItem?: StepItem): JSX.Element => {
 
       const step = stepItem?.step;
diff --git a/Engine/Source/Programs/Horde/HordeDashboard/src/components/JobSearchSimple.tsx b/Engine/Source/Programs/Horde/HordeDashboard/src/components/JobSearchSimple.tsx
--- a/Engine/Source/Programs/Horde/HordeDashboard/src/components/JobSearchSimple.tsx
+++ b/Engine/Source/Programs/Horde/HordeDashboard/src/components/JobSearchSimple.tsx
@@ -327,7 +327,7 @@ export const JobSearchSimpleModal: React.FC<{ streamId: string, onClose: () => v
                count: maxCount
             }
 
-            jobs = await backend.getJobs(query, false);
+            jobs = await backend.getJobs(query);
 
             const p = searchState.containsParameter?.trim().toLowerCase();
             if (p) {
diff --git a/Engine/Source/Programs/Horde/HordeDashboard/src/components/JobView.tsx b/Engine/Source/Programs/Horde/HordeDashboard/src/components/JobView.tsx
--- a/Engine/Source/Programs/Horde/HordeDashboard/src/components/JobView.tsx
+++ b/Engine/Source/Programs/Horde/HordeDashboard/src/components/JobView.tsx
@@ -6,7 +6,7 @@ import moment from 'moment-timezone';
 import React, { useEffect, useState } from 'react';
 import { Link, useParams } from 'react-router-dom';
 import { useBackend } from '../backend';
-import { GetJobsTabLabelColumnResponse, GetJobsTabParameterColumnResponse, JobData, JobsTabColumnType, JobsTabData, JobState, LabelData, LabelOutcome, LabelState, StreamData } from '../backend/Api';
+import { GetJobsTabLabelColumnResponse, GetJobsTabParameterColumnResponse, GetLabelStateResponse, JobData, JobsTabColumnType, JobsTabData, JobState, LabelOutcome, LabelState, StreamData } from '../backend/Api';
 import dashboard from '../backend/Dashboard';
 import { JobHandler } from '../backend/JobHandler';
 import { filterJob, JobFilterSimple } from '../base/utilities/filter';
@@ -309,8 +309,8 @@ const JobList: React.FC<{ tab: string; filter: JobFilterSimple, controller: Call
    };
 
    const filterLabels = (job: JobData, category: string | undefined) => {
-
-      const labels = job.graphRef?.labels;
+      
+      const labels = job.labels;
 
       if (!labels || !labels.length) {
          return [];
@@ -335,31 +335,30 @@ const JobList: React.FC<{ tab: string; filter: JobFilterSimple, controller: Call
          return job.labels[idx].state !== LabelState.Unspecified;
       });
 
-      const unassigned: LabelData[] = [];
+      const unassigned: GetLabelStateResponse[] = [];
 
       view.forEach(label => {
-         if (!label.category || !jobTab?.columns?.find(c => { return ((c as GetJobsTabLabelColumnResponse).category === label.category) })) {
+         if (!label.dashboardCategory || !jobTab?.columns?.find(c => { return ((c as GetJobsTabLabelColumnResponse).category === label.dashboardCategory) })) {
             unassigned.push(label);
          }
       });
 
-      return view.filter(label => (label.name && ((label.category === category) || ((category === "Other" || !category) && unassigned.indexOf(label) !== -1)))).sort((a, b) => a.name! < b.name! ? -1 : 1);
+      return view.filter(label => (label.dashboardName && ((label.dashboardCategory === category) || ((category === "Other" || !category) && unassigned.indexOf(label) !== -1)))).sort((a, b) => a.dashboardName! < b.dashboardName! ? -1 : 1);      
 
    };
 
+   const JobLabel: React.FC<{ item: JobItem; column: GetJobsTabLabelColumnResponse; label: GetLabelStateResponse }> = ({ item, column, label }) => {
 
-   const JobLabel: React.FC<{ item: JobItem; column: GetJobsTabLabelColumnResponse; label: LabelData }> = ({ item, column, label }) => {
+      const defaultLabel = item.job.defaultLabel;
 
-      if (label.defaultLabel && column.category !== "Other" && column.heading !== "Other") {
+      if (label === defaultLabel && column.category !== "Other" && column.heading !== "Other") {
          return <div />;
       }
 
-
-
-      const aggregates = item.job.graphRef?.labels;
+      const aggregates = item.job.labels;
 
       // note details may not be loaded here, as only initialized on callout for optimization (details.getLabelIndex(label.Name, label.Category);)
-      const jlabel = aggregates?.find((l, idx) => l.category === label.category && l.name === label.name && item.job.labels![idx]?.state !== LabelState.Unspecified);
+      const jlabel = aggregates?.find((l, idx) => l.dashboardCategory === label.dashboardCategory && l.dashboardName === label.dashboardName && item.job.labels![idx]?.state !== LabelState.Unspecified);
       let labelIdx = -1;
       if (jlabel) {
          labelIdx = aggregates?.indexOf(jlabel)!;
@@ -367,9 +366,9 @@ const JobList: React.FC<{ tab: string; filter: JobFilterSimple, controller: Call
 
       let state: LabelState | undefined;
       let outcome: LabelOutcome | undefined;
-      if (label.defaultLabel) {
-         state = label.defaultLabel.state;
-         outcome = label.defaultLabel.outcome;
+      if (label === defaultLabel) {
+         state = defaultLabel.state;
+         outcome = defaultLabel.outcome;
       } else {
          state = item.job.labels![labelIdx]?.state;
          outcome = item.job.labels![labelIdx]?.outcome;
@@ -383,12 +382,12 @@ const JobList: React.FC<{ tab: string; filter: JobFilterSimple, controller: Call
          url = `/job/${item.job.id}?label=${labelIdx}`;
       }
 
-      const target = `label_${item.job.id}_${label.name}_${label.category}`.replace(/[^A-Za-z0-9]/g, "");
+      const target = `label_${item.job.id}_${label.dashboardName}_${label.dashboardCategory}`.replace(/[^A-Za-z0-9]/g, "");
 
       return <Stack>
          <div id={target}>
             <Link to={url}><Stack className={hordeClasses.badgeNoIcon}>
-               <DefaultButton key={label.name} style={{ backgroundColor: color.primaryColor }} text={label.name}
+               <DefaultButton key={label.dashboardName} style={{ backgroundColor: color.primaryColor }} text={label.dashboardName}
                   onMouseOver={(ev) => {
                      ev.stopPropagation();
                      controller.setState({ target: `#${target}`, label: label, jobId: item.job.id })
@@ -414,7 +413,7 @@ const JobList: React.FC<{ tab: string; filter: JobFilterSimple, controller: Call
 
    };
 
-   const renderLabels = (item: JobItem, labels: LabelData[], jobColumn: GetJobsTabLabelColumnResponse, column?: IColumn): JSX.Element => {
+   const renderLabels = (item: JobItem, labels: GetLabelStateResponse[], jobColumn: GetJobsTabLabelColumnResponse, column?: IColumn): JSX.Element => {
 
       const defaultLabel = item.job.defaultLabel;
 
@@ -422,23 +421,15 @@ const JobList: React.FC<{ tab: string; filter: JobFilterSimple, controller: Call
          return <div />;
       }
 
-      if (defaultLabel?.nodes?.length && (!jobColumn.category || jobColumn.category === "Other")) {
-
-         const otherLabel = {
-            category: "",
-            name: "Other",
-            requiredNodes: [],
-            includedNodes: defaultLabel!.nodes,
-            defaultLabel: item.job.defaultLabel
-         };
-
-         labels.push(otherLabel);
+      if (defaultLabel && (!jobColumn.category || jobColumn.category === "Other")) {
+         defaultLabel.dashboardName = "Other";
+         labels.push(defaultLabel);
       }
 
       let key = 0;
 
       const buttons = labels.map(label => {
-         return <JobLabel key={`label_${item.job.id}_${jobColumn.category}_${label.name}_${key++}`} item={item} column={jobColumn} label={label} />;
+         return <JobLabel key={`label_${item.job.id}_${jobColumn.category}_${label.dashboardName}_${key++}`} item={item} column={jobColumn} label={label} />;
       });
 
       return (
@@ -465,6 +456,7 @@ const JobList: React.FC<{ tab: string; filter: JobFilterSimple, controller: Call
       if (jobColumn?.type === JobsTabColumnType.Labels) {
          const labelColumn = jobColumn as GetJobsTabLabelColumnResponse;
          const labels = filterLabels(item.job!, labelColumn?.category);
+
          return <Stack className="horde-no-darktheme" verticalAlign="center" verticalFill={true}>{renderLabels(item, labels, labelColumn, column)}</Stack>;
       }
 
diff --git a/Engine/Source/Programs/Horde/HordeDashboard/src/components/JobViewAll.tsx b/Engine/Source/Programs/Horde/HordeDashboard/src/components/JobViewAll.tsx
--- a/Engine/Source/Programs/Horde/HordeDashboard/src/components/JobViewAll.tsx
+++ b/Engine/Source/Programs/Horde/HordeDashboard/src/components/JobViewAll.tsx
@@ -6,7 +6,7 @@ import moment from "moment-timezone";
 import React, { useEffect } from "react";
 import { Link, useParams } from "react-router-dom";
 import { useBackend } from "../backend";
-import { GetStepResponse, JobData, JobState, JobStepOutcome, LabelData, LabelOutcome, LabelState, StepData, StreamData } from "../backend/Api";
+import { GetLabelStateResponse, GetStepResponse, JobData, JobState, JobStepOutcome, LabelOutcome, LabelState, StepData, StreamData } from "../backend/Api";
 import dashboard from "../backend/Dashboard";
 import { JobHandler } from "../backend/JobHandler";
 import { filterJob, JobFilterSimple } from "../base/utilities/filter";
@@ -20,37 +20,37 @@ import { StepStatusIcon } from "./StatusIcon";
 
 
 type JobItem = {
-    key: string;
-    job: JobData;
-    stream: StreamData;
-    startedby: string;
+   key: string;
+   job: JobData;
+   stream: StreamData;
+   startedby: string;
 };
 
 const customStyles = mergeStyleSets({
-    detailsRow: {
-        selectors: {
-            '.ms-DetailsRow': {
-                borderBottom: '0px',                
-                width: "100%"
-            },
-            '.ms-DetailsRow-cell': {
-                position: 'relative',
-                textAlign: "center",
-                padding: 0,
-                overflow: "visible",
-                whiteSpace: "nowrap"
-            }
-
-        }
-    },
-    header: {
-        selectors: {
-            ".ms-DetailsHeader-cellTitle": {
-                padding: 0,
-                paddingLeft: 4
-            }
-        }
-    }
+   detailsRow: {
+      selectors: {
+         '.ms-DetailsRow': {
+            borderBottom: '0px',
+            width: "100%"
+         },
+         '.ms-DetailsRow-cell': {
+            position: 'relative',
+            textAlign: "center",
+            padding: 0,
+            overflow: "visible",
+            whiteSpace: "nowrap"
+         }
+
+      }
+   },
+   header: {
+      selectors: {
+         ".ms-DetailsHeader-cellTitle": {
+            padding: 0,
+            paddingLeft: 4
+         }
+      }
+   }
 
 });
 
@@ -59,511 +59,497 @@ const jobHandler = new JobHandler(true, 300);
 
 const JobFilterPanel: React.FC = () => {
 
-    const { streamId } = useParams<{ streamId: string }>();
+   const { streamId } = useParams<{ streamId: string }>();
 
-    return <Stack>
-        <JobFilterBar streamId={streamId!} />
-    </Stack>
+   return <Stack>
+      <JobFilterBar streamId={streamId!} />
+   </Stack>
 }
 
 const JobViewAllInner: React.FC<{ filter: JobFilterSimple }> = observer(({ filter }) => {
 
-    const { streamId } = useParams<{ streamId: string }>();
-    const { projectStore } = useBackend();
+   const { streamId } = useParams<{ streamId: string }>();
+   const { projectStore } = useBackend();
 
-    useEffect(() => {
+   useEffect(() => {
 
-        return () => {
-            jobHandler.clear();
-        };
+      return () => {
+         jobHandler.clear();
+      };
 
-    }, []);
-   
-    const { hordeClasses, detailClasses, modeColors } = getHordeStyling();
+   }, []);
 
+   const { hordeClasses, detailClasses, modeColors } = getHordeStyling();
 
-    // subscribe
-    if (jobHandler.updated) { }
-    if (jobFilter.updated) { }
 
-    const stream = projectStore.streamById(streamId);
+   // subscribe
+   if (jobHandler.updated) { }
+   if (jobFilter.updated) { }
 
-    if (!streamId || !stream || !projectStore.streamById(streamId)) {
-        console.error("bad stream id setting up JobList");
-        return <div />;
-    }
+   const stream = projectStore.streamById(streamId);
 
-    let jobs = jobHandler.getFilteredJobs(jobFilter.status);
+   if (!streamId || !stream || !projectStore.streamById(streamId)) {
+      console.error("bad stream id setting up JobList");
+      return <div />;
+   }
 
-    jobs = jobs.filter(j => filterJob(j, filter.filterKeyword));
+   let jobs = jobHandler.getFilteredJobs(jobFilter.status);
 
-    let jobItems: JobItem[] = jobs.map(j => {
+   jobs = jobs.filter(j => filterJob(j, filter.filterKeyword));
 
-        let startedBy = j.startedByUserInfo?.name ?? "Scheduler";
-        if (startedBy.toLowerCase().indexOf("testaccount") !== -1) {
-            startedBy = "TestAccount";
-        }
+   let jobItems: JobItem[] = jobs.map(j => {
 
-        return {
-            key: j.id,
-            job: j,
-            startedby: startedBy,
-            stream: stream
-        } as JobItem;
+      let startedBy = j.startedByUserInfo?.name ?? "Scheduler";
+      if (startedBy.toLowerCase().indexOf("testaccount") !== -1) {
+         startedBy = "TestAccount";
+      }
 
-    });
+      return {
+         key: j.id,
+         job: j,
+         startedby: startedBy,
+         stream: stream
+      } as JobItem;
 
+   });
 
-    let columns: IColumn[] = [
-        { key: 'jobview_column1', name: 'Status', minWidth: 16, maxWidth: 16, onRenderHeader: () => null },
-        { key: 'jobview_column2', name: 'Change', minWidth: 64, maxWidth: 64 },
-        { key: 'jobview_column3', name: 'Job', minWidth: 220, maxWidth: 220 },
-        { key: 'jobview_column4', name: 'Labels', minWidth: 300, maxWidth: 300 },
-        { key: 'jobview_column5', name: 'Steps', minWidth: 324, maxWidth: 324 },
-        { key: 'jobview_column6', name: 'StartedBy', minWidth: 290, maxWidth: 290 },
-        { key: 'jobview_column7', name: 'Time', minWidth: 50, maxWidth: 50 },
 
-    ];
+   let columns: IColumn[] = [
+      { key: 'jobview_column1', name: 'Status', minWidth: 16, maxWidth: 16, onRenderHeader: () => null },
+      { key: 'jobview_column2', name: 'Change', minWidth: 64, maxWidth: 64 },
+      { key: 'jobview_column3', name: 'Job', minWidth: 220, maxWidth: 220 },
+      { key: 'jobview_column4', name: 'Labels', minWidth: 300, maxWidth: 300 },
+      { key: 'jobview_column5', name: 'Steps', minWidth: 324, maxWidth: 324 },
+      { key: 'jobview_column6', name: 'StartedBy', minWidth: 290, maxWidth: 290 },
+      { key: 'jobview_column7', name: 'Time', minWidth: 50, maxWidth: 50 },
 
-    columns = columns.map(c => {
+   ];
 
-        c.isResizable = true;
-        c.isPadded = false;
-        c.isMultiline = true;
-        c.isCollapsible = false;
-        c.headerClassName = customStyles.header;
+   columns = columns.map(c => {
 
-        c.styles = (props: any): any => {
-            props.cellStyleProps = { ...props.cellStyleProps };
-            props.cellStyleProps.cellLeftPadding = 4;
-            props.cellStyleProps.cellRightPadding = 0;
-        };
+      c.isResizable = true;
+      c.isPadded = false;
+      c.isMultiline = true;
+      c.isCollapsible = false;
+      c.headerClassName = customStyles.header;
 
-        return c;
+      c.styles = (props: any): any => {
+         props.cellStyleProps = { ...props.cellStyleProps };
+         props.cellStyleProps.cellLeftPadding = 4;
+         props.cellStyleProps.cellRightPadding = 0;
+      };
 
-    })
+      return c;
 
-    const renderSteps = (job: JobData) => {
+   })
 
-        type StepItem = {
-            step: StepData;
-            name: string;
-        }
+   const renderSteps = (job: JobData) => {
 
-        const jobId = job.id;
+      type StepItem = {
+         step: StepData;
+         name: string;
+      }
 
-        if (!job.batches) {
-            return null;
-        }
+      const jobId = job.id;
 
-        let steps: GetStepResponse[] = job.batches.map(b => b.steps).flat().filter(step => !!step.startTime && (step.outcome === JobStepOutcome.Warnings || step.outcome === JobStepOutcome.Failure));
+      if (!job.batches) {
+         return null;
+      }
 
-        if (!steps || !steps.length) {
-            return null;
-        }
+      let steps: GetStepResponse[] = job.batches.map(b => b.steps).flat().filter(step => !!step.startTime && (step.outcome === JobStepOutcome.Warnings || step.outcome === JobStepOutcome.Failure));
 
-        const onRenderCell = (stepItem: StepItem): JSX.Element => {
+      if (!steps || !steps.length) {
+         return null;
+      }
 
-            const step = stepItem.step;
+      const onRenderCell = (stepItem: StepItem): JSX.Element => {
 
-            const stepUrl = `/job/${jobId}?step=${step.id}`;
+         const step = stepItem.step;
 
-            return <Stack tokens={{ childrenGap: 12 }} key={`step_${step.id}_job_${job.id}_${stepItem.name}`}>
-                <Stack horizontal>
-                    <Link to={stepUrl} onClick={(ev) => { ev.stopPropagation(); }}><div style={{ cursor: "pointer" }}>
-                        <Stack horizontal>
-                            <StepStatusIcon step={step} style={{ fontSize: 10 }} />
-                            <Text styles={{ root: { fontSize: 10, paddingRight: 4, userSelect: "none" } }}>{`${stepItem.name}`}</Text>
-                        </Stack>
-                    </div></Link>
-                </Stack>
-            </Stack>;
-        };
+         const stepUrl = `/job/${jobId}?step=${step.id}`;
 
-        const andMore = (errors: number, warnings: number): JSX.Element => {
+         return <Stack tokens={{ childrenGap: 12 }} key={`step_${step.id}_job_${job.id}_${stepItem.name}`}>
+            <Stack horizontal>
+               <Link to={stepUrl} onClick={(ev) => { ev.stopPropagation(); }}><div style={{ cursor: "pointer" }}>
+                  <Stack horizontal>
+                     <StepStatusIcon step={step} style={{ fontSize: 10 }} />
+                     <Text styles={{ root: { fontSize: 10, paddingRight: 4, userSelect: "none" } }}>{`${stepItem.name}`}</Text>
+                  </Stack>
+               </div></Link>
+            </Stack>
+         </Stack>;
+      };
+
+      const andMore = (errors: number, warnings: number): JSX.Element => {
+
+         return <Stack tokens={{ childrenGap: 12 }} key={`job_${job.id}_andmore`}>
+            <Stack horizontal>
+               <Stack horizontal>
+                  <Link to={`/job/${jobId}`} onClick={(ev) => { ev.stopPropagation(); }}><div style={{ cursor: "pointer" }}>
+                     <Text styles={{ root: { fontSize: 10, paddingRight: 4, paddingLeft: 19, userSelect: "none" } }}>{`( +${errors + warnings} more )`}</Text>
+                  </div></Link>
+               </Stack>
+            </Stack>
+         </Stack>;
+      };
 
-            return <Stack tokens={{ childrenGap: 12 }} key={`job_${job.id}_andmore`}>
-                <Stack horizontal>
-                    <Stack horizontal>
-                        <Link to={`/job/${jobId}`} onClick={(ev) => { ev.stopPropagation(); }}><div style={{ cursor: "pointer" }}>
-                            <Text styles={{ root: { fontSize: 10, paddingRight: 4, paddingLeft: 19, userSelect: "none" } }}>{`( +${errors + warnings} more )`}</Text>
-                        </div></Link>
-                    </Stack>
-                </Stack>
-            </Stack>;
-        };
 
+      let items = steps.map(step => {
 
-        let items = steps.map(step => {
 
-            const batch = job.batches!.find(b => !!b.steps.find(s => s.id === step.id))!;
-            const groups = job?.graphRef?.groups;
-            if (!groups) {
-                return undefined;
-            }
-            const node = groups[batch.groupIdx].nodes[step.nodeIdx];
+         return {
+            step: step,
+            name: step.name
+         }
 
-            if (!node) {
-                return undefined;
-            }
+      })
 
-            return {
-                step: step,
-                name: node.name
-            }
+      items = items.filter(item => !!item);
 
-        })
+      items = items.sort((a, b) => {
 
-        items = items.filter(item => !!item);
+         const stepA = a!.step;
+         const stepB = b!.step;
 
-        items = items.sort((a, b) => {
+         const outA = stepA.outcome;
+         const outB = stepB.outcome;
 
-            const stepA = a!.step;
-            const stepB = b!.step;
+         if (outA === JobStepOutcome.Failure && outB === JobStepOutcome.Warnings) {
+            return -1;
+         }
 
-            const outA = stepA.outcome;
-            const outB = stepB.outcome;
+         if (outA === JobStepOutcome.Warnings && outB === JobStepOutcome.Failure) {
+            return 1;
+         }
 
-            if (outA === JobStepOutcome.Failure && outB === JobStepOutcome.Warnings) {
-                return -1;
-            }
+         return a!.name < b!.name ? -1 : 1;
 
-            if (outA === JobStepOutcome.Warnings && outB === JobStepOutcome.Failure) {
-                return 1;
-            }
+      });
 
-            return a!.name < b!.name ? -1 : 1;
+      let numErrors = 0;
+      let numWarnings = 0;
 
-        });
+      items.forEach(i => {
+         if (i?.step.outcome === JobStepOutcome.Warnings) {
+            numWarnings++;
+         } else {
+            numErrors++;
+         }
+      });
 
-        let numErrors = 0;
-        let numWarnings = 0;
+      items = items.slice(0, 10);
 
-        items.forEach(i => {
-            if (i?.step.outcome === JobStepOutcome.Warnings) {
-                numWarnings++;
-            } else {
-                numErrors++;
-            }
-        });
+      const render = items.map(s => onRenderCell(s!));
 
-        items = items.slice(0, 10);
+      let shownErrors = 0;
+      let shownWarnings = 0;
 
-        const render = items.map(s => onRenderCell(s!));
+      items.forEach(i => {
+         if (i?.step.outcome === JobStepOutcome.Warnings) {
+            shownWarnings++;
+         } else {
+            shownErrors++;
+         }
+      });
 
-        let shownErrors = 0;
-        let shownWarnings = 0;
+      if (shownErrors !== numErrors || numWarnings !== shownWarnings) {
+         render.push(andMore(numErrors - shownErrors, numWarnings - shownWarnings));
+      }
 
-        items.forEach(i => {
-            if (i?.step.outcome === JobStepOutcome.Warnings) {
-                shownWarnings++;
-            } else {
-                shownErrors++;
-            }
-        });
+      return render;
 
-        if (shownErrors !== numErrors || numWarnings !== shownWarnings) {
-            render.push(andMore(numErrors - shownErrors, numWarnings - shownWarnings));
-        }
+   }
 
-        return render;
+   const renderDetailStatus = (job: JobData): JSX.Element => {
 
-    }
+      if (job.state === JobState.Complete) {
+         return <div />;
+      }
 
-    const renderDetailStatus = (job: JobData): JSX.Element => {
+      return <Stack className="horde-no-darktheme" ><FontIcon iconName="FullCircle" style={{ color: getJobStateColor(job.state) }} /></Stack>;
+   };
 
-        if (job.state === JobState.Complete) {
-            return <div />;
-        }
 
-        return <Stack className="horde-no-darktheme" ><FontIcon iconName="FullCircle"  style={{ color: getJobStateColor(job.state) }}/></Stack>;
-    };
+   const renderItemInner = (item: JobItem, index?: number, column?: IColumn) => {
+
+      const job = item.job;
+
+      const style: any = {
+         margin: 0,
+         position: "absolute",
+         top: "50%",
+         left: "50%",
+         msTransform: "translate(-50%, -50%)",
+         transform: "translate(-50%, -50%)"
+      };
+
+      if (column!.name === "Space") {
+         return <div />
+      }
+
+      if (column!.name === "Labels") {
+         return <Stack verticalAlign="center" verticalFill={true}>{renderLabels(item)}</Stack>;
+
+      }
 
+      if (column!.name === "Change") {
 
-    const renderItemInner = (item: JobItem, index?: number, column?: IColumn) => {
+         return <ChangeButton job={job} />;
 
-        const job = item.job;
+      }
 
-        const style: any = {
-            margin: 0,
-            position: "absolute",
-            top: "50%",
-            left: "50%",
-            msTransform: "translate(-50%, -50%)",
-            transform: "translate(-50%, -50%)"
-        };
+      if (column!.name === "StartedBy") {
+         let startedBy = job.startedByUserInfo?.name ?? "Scheduler";
+         if (startedBy.toLowerCase().indexOf("testaccount") !== -1) {
+            startedBy = "TestAccount";
+         }
+
+         return <Stack verticalAlign="center" verticalFill={true} horizontalAlign="start">{startedBy}</Stack>
+
+      }
+      if (column!.name === "Steps") {
+         return <Stack verticalAlign="center" verticalFill={true} style={{ paddingLeft: 64 }}>{renderSteps(job)}</Stack>;
+      }
 
-        if (column!.name === "Space") {
-            return <div />
-        }
 
-        if (column!.name === "Labels") {
-            return <Stack verticalAlign="center" verticalFill={true}>{renderLabels(item)}</Stack>;
+      if (column!.name === "Job") {
 
-        }
+         return <Stack verticalAlign="center" verticalFill={true} horizontalAlign="center" style={{ whiteSpace: "normal" }}>{item.job.name}</Stack>
 
-        if (column!.name === "Change") {
+      }
 
-            return <ChangeButton job={job} />;
+      if (column!.name === "Time") {
 
-        }
+         const displayTime = moment(item.job!.createTime).tz(displayTimeZone());
 
-        if (column!.name === "StartedBy") {
-            let startedBy = job.startedByUserInfo?.name ?? "Scheduler";
-            if (startedBy.toLowerCase().indexOf("testaccount") !== -1) {
-                startedBy = "TestAccount";
-            }
+         const format = dashboard.display24HourClock ? "MMM Do, HH:mm z" : "MMM Do, h:mma z";
 
-            return <Stack verticalAlign="center" verticalFill={true} horizontalAlign="start">{startedBy}</Stack>
+         let displayTimeStr = displayTime.format(format);
 
-        }
-        if (column!.name === "Steps") {
-            return <Stack verticalAlign="center" verticalFill={true} style={{ paddingLeft: 64 }}>{renderSteps(job)}</Stack>;
-        }
+         return <Stack verticalAlign="center" verticalFill={true} horizontalAlign={"end"}>{displayTimeStr}</Stack>;
 
+      }
 
-        if (column!.name === "Job") {
+      if (column!.name === "Status") {
+         return <div style={style}>{renderDetailStatus(job)}</div>;
+      }
 
-            return <Stack verticalAlign="center" verticalFill={true} horizontalAlign="center" style={{ whiteSpace: "normal" }}>{item.job.name}</Stack>
+      return <Stack style={style} horizontalAlign="center">{item.job.state}</Stack>
+   }
 
-        }
+   const renderItem = (item: JobItem, index?: number, column?: IColumn) => {
 
-        if (column!.name === "Time") {
+      if (column?.name === "Change") {
+         return renderItemInner(item, index, column);
+      } else {
+         return <Link to={`/job/${(item as JobItem).job.id}`}>{renderItemInner(item, index, column)}</Link>;
+      }
+   }
 
-            const displayTime = moment(item.job!.createTime).tz(displayTimeZone());
+   const renderRow: IDetailsListProps['onRenderRow'] = (props) => {
 
-            const format = dashboard.display24HourClock ? "MMM Do, HH:mm z" : "MMM Do, h:mma z";
+      if (props) {
 
-            let displayTimeStr = displayTime.format(format);
+         const item = jobItems[props.itemIndex];
 
-            return <Stack verticalAlign="center" verticalFill={true} horizontalAlign={"end"}>{displayTimeStr}</Stack>;
+         return <JobOperationsContextMenu job={item.job}>
+            <DetailsRow styles={{ root: { paddingTop: 8, paddingBottom: 8, width: "100%" }, cell: { selectors: { "a, a:visited, a:active, a:hover": { color: modeColors.text } } } }} {...props} />
+         </JobOperationsContextMenu>
+      }
+      return null;
+   };
 
-        }
+   const JobLabel: React.FC<{ item: JobItem; label: GetLabelStateResponse }> = ({ item, label }) => {
 
-        if (column!.name === "Status") {
-            return <div style={style}>{renderDetailStatus(job)}</div>;
-        }
+      const defaultLabel = item.job.defaultLabel;
 
-        return <Stack style={style} horizontalAlign="center">{item.job.state}</Stack>
-    }
+      const aggregates = item.job.labels;
 
-    const renderItem = (item: JobItem, index?: number, column?: IColumn) => {
+      // note details may not be loaded here, as only initialized on callout for optimization (details.getLabelIndex(label.Name, label.Category);)
+      const jlabel = aggregates?.find((l, idx) => l.dashboardCategory === label.dashboardCategory && l.dashboardName === label.dashboardName && item.job.labels![idx]?.state !== LabelState.Unspecified);
+      let labelIdx = -1;
+      if (jlabel) {
+         labelIdx = aggregates?.indexOf(jlabel)!;
+      }
 
-        if (column?.name === "Change") {
-            return renderItemInner(item, index, column);
-        } else {
-            return <Link to={`/job/${(item as JobItem).job.id}`}>{renderItemInner(item, index, column)}</Link>;
-        }
-    }
+      let state: LabelState | undefined;
+      let outcome: LabelOutcome | undefined;
+      if (defaultLabel) {
+         state = defaultLabel.state;
+         outcome = defaultLabel.outcome;
+      } else {
+         state = item.job.labels![labelIdx]?.state;
+         outcome = item.job.labels![labelIdx]?.outcome;
+      }
 
-    const renderRow: IDetailsListProps['onRenderRow'] = (props) => {
+      const color = getLabelColor(state, outcome);
 
-        if (props) {
 
-            const item = jobItems[props.itemIndex];
+      let url = `/job/${item.job.id}`;
 
-           return <JobOperationsContextMenu job={item.job}>
-                <DetailsRow styles={{ root: { paddingTop: 8, paddingBottom: 8, width: "100%" }, cell: { selectors: { "a, a:visited, a:active, a:hover": { color: modeColors.text } } } }} {...props} />                
-            </JobOperationsContextMenu>
-        }
-        return null;
-    };
+      if (labelIdx >= 0) {
+         url = `/job/${item.job.id}?label=${labelIdx}`;
+      }
 
-    const JobLabel: React.FC<{ item: JobItem; label: LabelData }> = ({ item, label }) => {
+      const target = `label_${item.job.id}_${label.dashboardName}_${label.dashboardCategory}`.replace(/ /g, "");
+
+      return <Stack>
+         <div id={target}>
+            <Link to={url}><Stack className={hordeClasses.badgeNoIcon}>
+               <DefaultButton key={label.dashboardName} style={{ backgroundColor: color.primaryColor, fontSize: 9 }} text={label.dashboardName}>
+                  {!!color.secondaryColor && <div style={{
+                     borderLeft: "10px solid transparent",
+                     borderRight: `10px solid ${color.secondaryColor}`,
+                     borderBottom: "10px solid transparent",
+                     height: 0,
+                     width: 0,
+                     position: "absolute",
+                     right: 0,
+                     top: 0,
+                     zIndex: 1
+                  }} />}
+               </DefaultButton>
+            </Stack>
+            </Link>
+         </div>
+      </Stack>;
+   };
 
-        const aggregates = item.job.graphRef?.labels;
+   const renderLabels = (item: JobItem): JSX.Element => {
 
-        // note details may not be loaded here, as only initialized on callout for optimization (details.getLabelIndex(label.Name, label.Category);)
-        const jlabel = aggregates?.find((l, idx) => l.category === label.category && l.name === label.name && item.job.labels![idx]?.state !== LabelState.Unspecified);
-        let labelIdx = -1;
-        if (jlabel) {
-            labelIdx = aggregates?.indexOf(jlabel)!;
-        }
+      const job = item.job!;
+      const defaultLabel = job.defaultLabel;
+      const labels = job.labels ?? [];
 
-        let state: LabelState | undefined;
-        let outcome: LabelOutcome | undefined;
-        if (label.defaultLabel) {
-            state = label.defaultLabel.state;
-            outcome = label.defaultLabel.outcome;
-        } else {
-            state = item.job.labels![labelIdx]?.state;
-            outcome = item.job.labels![labelIdx]?.outcome;
-        }
+      if (!labels.length && !defaultLabel) {
+         return <div />;
+      }
 
-        const color = getLabelColor(state, outcome);
+      const view = labels.filter(a => {
+         const idx = labels.indexOf(a)!;
+         if (!job.labels) {
+            return false;
+         }
+         return job.labels[idx].state !== LabelState.Unspecified;
+      });
 
+      const catergories = new Set<string>();
+      view.forEach(label => catergories.add(label.dashboardCategory));
 
-        let url = `/job/${item.job.id}`;
+      const sorted = Array.from(catergories).sort((a, b) => { return a < b ? -1 : 1 });
 
-        if (labelIdx >= 0) {
-            url = `/job/${item.job.id}?label=${labelIdx}`;
-        }
+      if (!catergories.has("Other") && defaultLabel) {
 
-        const target = `label_${item.job.id}_${label.name}_${label.category}`.replace(/ /g, "");
+         defaultLabel.dashboardCategory = "Other";
+         defaultLabel.dashboardName = "Other";
+         view.push(defaultLabel);         
+         sorted.push("Other");
+      }
 
-        return <Stack>
-            <div id={target}>
-                <Link to={url}><Stack className={hordeClasses.badgeNoIcon}>
-                    <DefaultButton key={label.name} style={{ backgroundColor: color.primaryColor, fontSize: 9 }} text={label.name}>
-                        {!!color.secondaryColor && <div style={{
-                            borderLeft: "10px solid transparent",
-                            borderRight: `10px solid ${color.secondaryColor}`,
-                            borderBottom: "10px solid transparent",
-                            height: 0,
-                            width: 0,
-                            position: "absolute",
-                            right: 0,
-                            top: 0,
-                            zIndex: 1
-                        }} />}
-                    </DefaultButton>
-                </Stack>
-                </Link>
-            </div>
-        </Stack>;
-    };
+      let key = 0;
 
-    const renderLabels = (item: JobItem): JSX.Element => {
+      const labelStacks = sorted.map(cat => {
 
-        const job = item.job!;
-        const defaultLabel = job.defaultLabel;
-        const labels = job.graphRef?.labels ?? [];
+         const buttons = view.filter(label => label.dashboardCategory === cat).map(label => {
+            return <JobLabel key={`label_${item.job.id}_${label.dashboardName}_${key++}`} item={item} label={label} />;
+         });
 
-        if (!labels.length && !defaultLabel) {
+         if (!buttons.length) {
             return <div />;
-        }
-
-        const view = labels.filter(a => {
-            const idx = labels.indexOf(a)!;
-            if (!job.labels) {
-                return false;
-            }
-            return job.labels[idx].state !== LabelState.Unspecified;
-        });
-
-        const catergories = new Set<string>();
-        view.forEach(label => catergories.add(label.category));
-
-        const sorted = Array.from(catergories).sort((a, b) => { return a < b ? -1 : 1 });
-
-        if (!catergories.has("Other") && defaultLabel?.nodes?.length) {
-
-            const otherLabel = {
-                category: "Other",
-                name: "Other",
-                requiredNodes: [],
-                includedNodes: defaultLabel!.nodes,
-                defaultLabel: item.job.defaultLabel
-            };
-
-            view.push(otherLabel);
-            sorted.push("Other");
-        }
-
-        let key = 0;
-
-        const labelStacks = sorted.map(cat => {
-
-            const buttons = view.filter(label => label.category === cat).map(label => {
-                return <JobLabel key={`label_${item.job.id}_${label.name}_${key++}`} item={item} label={label} />;
-            });
-
-            if (!buttons.length) {
-                return <div />;
-            }
-
-            return <Stack tokens={{ childrenGap: 8 }}>
-                <Stack horizontal tokens={{ childrenGap: 8 }}>
-                    <Stack horizontalAlign="end" style={{ minWidth: 100, maxWidth: 100 }}>
-                        <Text variant="small" style={{}}>{cat}:</Text>
-                    </Stack>
-                    <Stack className="horde-no-darktheme" wrap horizontal horizontalAlign="start" tokens={{ childrenGap: 4 }}>
-                        {buttons}
-                    </Stack>                    
-                </Stack>
+         }
+
+         return <Stack tokens={{ childrenGap: 8 }}>
+            <Stack horizontal tokens={{ childrenGap: 8 }}>
+               <Stack horizontalAlign="end" style={{ minWidth: 100, maxWidth: 100 }}>
+                  <Text variant="small" style={{}}>{cat}:</Text>
+               </Stack>
+               <Stack className="horde-no-darktheme" wrap horizontal horizontalAlign="start" tokens={{ childrenGap: 4 }}>
+                  {buttons}
+               </Stack>
             </Stack>
-        })
+         </Stack>
+      })
 
 
-        return (
-            <Stack tokens={{ childrenGap: 4 }} styles={{ root: { paddingTop: 2 } }}>
-                {labelStacks}
-            </Stack>
+      return (
+         <Stack tokens={{ childrenGap: 4 }} styles={{ root: { paddingTop: 2 } }}>
+            {labelStacks}
+         </Stack>
 
-        );
+      );
 
-    };
+   };
 
-    if (jobFilter.templates && jobFilter.templates.length) {
+   if (jobFilter.templates && jobFilter.templates.length) {
 
-        let preflightStartedByUserId: string | undefined;
-        if (!filter.showOthersPreflights) {
-            preflightStartedByUserId = dashboard.userId;
-        }
+      let preflightStartedByUserId: string | undefined;
+      if (!filter.showOthersPreflights) {
+         preflightStartedByUserId = dashboard.userId;
+      }
 
-        if (jobHandler.filter(stream, jobFilter.templates?.map(t => t.id), preflightStartedByUserId)) {
+      if (jobHandler.filter(stream, jobFilter.templates?.map(t => t.id), preflightStartedByUserId)) {
 
-        }
-    }
+      }
+   }
 
-    let nojobs = !jobHandler.initial && !jobs.length;
+   let nojobs = !jobHandler.initial && !jobs.length;
 
-    return (
-        <Stack>
-            <Stack tokens={{ childrenGap: 0 }} style={{}}>
-                <Stack horizontalAlign="center" style={{ width: 1440, marginLeft: 4, boxShadow: "0 3px 3.6px 0 rgba(0,0,0,0.132), 0 0.3px 0.9px 0 rgba(0,0,0,0.108)" }}>
-                    <Stack style={{ paddingTop: 18, paddingBottom: 24 }}>
-                        <JobFilterPanel />
-                    </Stack>
-                </Stack>
-            </Stack>
-            <Stack tokens={{ childrenGap: 0 }} className={customStyles.detailsRow}>
-                <FocusZone direction={FocusZoneDirection.vertical}>
-                    <div className={detailClasses.container} style={{ height: 'calc(100vh - 352px)', position: 'relative', marginTop: 0 }} data-is-scrollable={true}>
-                        {<ScrollablePane scrollbarVisibility={ScrollbarVisibility.always} style={{ overflow: "visible" }}>
-                            <Stack style={{ width: 1440, marginLeft: 4, background: modeColors.content, boxShadow: "0 1.6px 3.6px 0 rgba(0,0,0,0.132), 0 0.3px 0.9px 0 rgba(0,0,0,0.108)"}}>
-                                <Stack>
-                                    <DetailsList
-                                        styles={{ root: { paddingLeft: 8, paddingRight: 8 } }}
-                                        compact={true}                                        
-                                        indentWidth={0}
-                                        items={jobItems}
-                                        columns={columns}
-                                        setKey="set"
-                                        selectionMode={SelectionMode.none}
-                                        layoutMode={DetailsListLayoutMode.fixedColumns}                                        
-                                        onRenderItemColumn={renderItem}
-                                        onRenderRow={renderRow}
-                                        onShouldVirtualize={() => true}
-                                    />
-                                </Stack>
-                            </Stack>
-
-                            {nojobs && <Stack style={{ width: 1364 }}>
-                                <Stack horizontalAlign="center" styles={{ root: { paddingTop: 20, paddingBottom: 20 } }}>
-                                    <Text variant="mediumPlus">No jobs found</Text>
-                                </Stack>
-                            </Stack>
-                            }
-
-                            {!nojobs && !jobItems.length && <Stack styles={{ root: { paddingTop: 20, paddingBottom: 20 } }}>
-                                <Stack style={{ width: 1440 }}><Spinner size={SpinnerSize.large} /></Stack>
-                            </Stack>
-                            }
-
-                        </ScrollablePane>
-                        }
-                    </div>
-                </FocusZone>
+   return (
+      <Stack>
+         <Stack tokens={{ childrenGap: 0 }} style={{}}>
+            <Stack horizontalAlign="center" style={{ width: 1440, marginLeft: 4, boxShadow: "0 3px 3.6px 0 rgba(0,0,0,0.132), 0 0.3px 0.9px 0 rgba(0,0,0,0.108)" }}>
+               <Stack style={{ paddingTop: 18, paddingBottom: 24 }}>
+                  <JobFilterPanel />
+               </Stack>
             </Stack>
-        </Stack>
-    );
+         </Stack>
+         <Stack tokens={{ childrenGap: 0 }} className={customStyles.detailsRow}>
+            <FocusZone direction={FocusZoneDirection.vertical}>
+               <div className={detailClasses.container} style={{ height: 'calc(100vh - 352px)', position: 'relative', marginTop: 0 }} data-is-scrollable={true}>
+                  {<ScrollablePane scrollbarVisibility={ScrollbarVisibility.always} style={{ overflow: "visible" }}>
+                     <Stack style={{ width: 1440, marginLeft: 4, background: modeColors.content, boxShadow: "0 1.6px 3.6px 0 rgba(0,0,0,0.132), 0 0.3px 0.9px 0 rgba(0,0,0,0.108)" }}>
+                        <Stack>
+                           <DetailsList
+                              styles={{ root: { paddingLeft: 8, paddingRight: 8 } }}
+                              compact={true}
+                              indentWidth={0}
+                              items={jobItems}
+                              columns={columns}
+                              setKey="set"
+                              selectionMode={SelectionMode.none}
+                              layoutMode={DetailsListLayoutMode.fixedColumns}
+                              onRenderItemColumn={renderItem}
+                              onRenderRow={renderRow}
+                              onShouldVirtualize={() => true}
+                           />
+                        </Stack>
+                     </Stack>
+
+                     {nojobs && <Stack style={{ width: 1364 }}>
+                        <Stack horizontalAlign="center" styles={{ root: { paddingTop: 20, paddingBottom: 20 } }}>
+                           <Text variant="mediumPlus">No jobs found</Text>
+                        </Stack>
+                     </Stack>
+                     }
+
+                     {!nojobs && !jobItems.length && <Stack styles={{ root: { paddingTop: 20, paddingBottom: 20 } }}>
+                        <Stack style={{ width: 1440 }}><Spinner size={SpinnerSize.large} /></Stack>
+                     </Stack>
+                     }
+
+                  </ScrollablePane>
+                  }
+               </div>
+            </FocusZone>
+         </Stack>
+      </Stack>
+   );
 
 });
 
 
 export const JobViewAll: React.FC<{ filter: JobFilterSimple }> = ({ filter }) => {
 
-    return (<JobViewAllInner filter={filter} />);
+   return (<JobViewAllInner filter={filter} />);
 
 };
diff --git a/Engine/Source/Programs/Horde/HordeDashboard/src/components/JobViewIncremental.tsx b/Engine/Source/Programs/Horde/HordeDashboard/src/components/JobViewIncremental.tsx
--- a/Engine/Source/Programs/Horde/HordeDashboard/src/components/JobViewIncremental.tsx
+++ b/Engine/Source/Programs/Horde/HordeDashboard/src/components/JobViewIncremental.tsx
@@ -6,7 +6,7 @@ import moment from 'moment-timezone';
 import React, { useEffect, useState } from 'react';
 import { Link, useParams } from 'react-router-dom';
 import { useBackend } from '../backend';
-import { GetChangeSummaryResponse, GetJobsTabLabelColumnResponse, GetJobsTabParameterColumnResponse, JobData, JobsTabColumnType, JobsTabData, JobState, LabelData, LabelOutcome, LabelState, StreamData } from '../backend/Api';
+import { GetChangeSummaryResponse, GetJobsTabLabelColumnResponse, GetJobsTabParameterColumnResponse, GetLabelStateResponse, JobData, JobsTabColumnType, JobsTabData, JobState, LabelOutcome, LabelState, StreamData } from '../backend/Api';
 import { CommitCache } from '../backend/CommitCache';
 import dashboard from '../backend/Dashboard';
 import { JobHandler } from '../backend/JobHandler';
@@ -340,7 +340,7 @@ const JobList: React.FC<{ tab: string; filter: JobFilterSimple, controller: Call
 
    const filterLabels = (job: JobData, category: string | undefined) => {
 
-      const labels = job.graphRef?.labels;
+      const labels = job.labels;
 
       if (!labels || !labels.length) {
          return [];
@@ -354,31 +354,30 @@ const JobList: React.FC<{ tab: string; filter: JobFilterSimple, controller: Call
          return job.labels[idx].state !== LabelState.Unspecified;
       });
 
-      const unassigned: LabelData[] = [];
+      const unassigned: GetLabelStateResponse[] = [];
 
       view.forEach(label => {
-         if (!label.category || !jobTab?.columns?.find(c => { return ((c as GetJobsTabLabelColumnResponse).category === label.category) })) {
+         if (!label.dashboardCategory || !jobTab?.columns?.find(c => { return ((c as GetJobsTabLabelColumnResponse).category === label.dashboardCategory) })) {
             unassigned.push(label);
          }
       });
 
-      return view.filter(label => (label.name && ((label.category === category) || ((category === "Other" || !category) && unassigned.indexOf(label) !== -1)))).sort((a, b) => a.name! < b.name! ? -1 : 1);
-
+      return view.filter(label => (label.dashboardName && ((label.dashboardCategory === category) || ((category === "Other" || !category) && unassigned.indexOf(label) !== -1)))).sort((a, b) => a.dashboardName! < b.dashboardName! ? -1 : 1);
    };
 
 
-   const JobLabel: React.FC<{ item: JobItem; column: GetJobsTabLabelColumnResponse; label: LabelData }> = ({ item, column, label }) => {
+   const JobLabel: React.FC<{ item: JobItem; column: GetJobsTabLabelColumnResponse; label: GetLabelStateResponse }> = ({ item, column, label }) => {
+
+      const defaultLabel = item.job.defaultLabel;
 
-      if (label.defaultLabel && column.category !== "Other" && column.heading !== "Other") {
+      if (label === defaultLabel && column.category !== "Other" && column.heading !== "Other") {
          return <div />;
       }
 
-
-
-      const aggregates = item.job.graphRef?.labels;
+      const aggregates = item.job.labels;
 
       // note details may not be loaded here, as only initialized on callout for optimization (details.getLabelIndex(label.Name, label.Category);)
-      const jlabel = aggregates?.find((l, idx) => l.category === label.category && l.name === label.name && item.job.labels![idx]?.state !== LabelState.Unspecified);
+      const jlabel = aggregates?.find((l, idx) => l.dashboardCategory === label.dashboardCategory && l.dashboardName === label.dashboardName && item.job.labels![idx]?.state !== LabelState.Unspecified);
       let labelIdx = -1;
       if (jlabel) {
          labelIdx = aggregates?.indexOf(jlabel)!;
@@ -386,9 +385,9 @@ const JobList: React.FC<{ tab: string; filter: JobFilterSimple, controller: Call
 
       let state: LabelState | undefined;
       let outcome: LabelOutcome | undefined;
-      if (label.defaultLabel) {
-         state = label.defaultLabel.state;
-         outcome = label.defaultLabel.outcome;
+      if (label === defaultLabel) {
+         state = defaultLabel.state;
+         outcome = defaultLabel.outcome;
       } else {
          state = item.job.labels![labelIdx]?.state;
          outcome = item.job.labels![labelIdx]?.outcome;
@@ -402,11 +401,11 @@ const JobList: React.FC<{ tab: string; filter: JobFilterSimple, controller: Call
          url = `/job/${item.job.id}?label=${labelIdx}`;
       }
 
-      const target = `label_${item.job.id}_${label.name}_${label.category}`.replace(/[^A-Za-z0-9]/g, "");
+      const target = `label_${item.job.id}_${label.dashboardName}_${label.dashboardCategory}`.replace(/[^A-Za-z0-9]/g, "");
 
       return <Stack>
          <div id={target} className="horde-no-darktheme">
-            <Link to={url}><Stack className={hordeClasses.badgeNoIcon}><DefaultButton key={label.name} style={{ backgroundColor: color.primaryColor }} text={label.name}
+            <Link to={url}><Stack className={hordeClasses.badgeNoIcon}><DefaultButton key={label.dashboardName} style={{ backgroundColor: color.primaryColor }} text={label.dashboardName}
                onMouseOver={(ev) => {
                   ev.stopPropagation();
                   controller.setState({ target: `#${target}`, label: label, jobId: item.job.id })
@@ -432,7 +431,7 @@ const JobList: React.FC<{ tab: string; filter: JobFilterSimple, controller: Call
 
    };
 
-   const renderLabels = (item: JobItem, labels: LabelData[], jobColumn: GetJobsTabLabelColumnResponse, column?: IColumn): JSX.Element => {
+   const renderLabels = (item: JobItem, labels: GetLabelStateResponse[], jobColumn: GetJobsTabLabelColumnResponse, column?: IColumn): JSX.Element => {
 
       const defaultLabel = item.job.defaultLabel;
 
@@ -440,23 +439,16 @@ const JobList: React.FC<{ tab: string; filter: JobFilterSimple, controller: Call
          return <div />;
       }
 
-      if (defaultLabel?.nodes?.length && (!jobColumn.category || jobColumn.category === "Other")) {
-
-         const otherLabel = {
-            category: "",
-            name: "Other",
-            requiredNodes: [],
-            includedNodes: defaultLabel!.nodes,
-            defaultLabel: item.job.defaultLabel
-         };
-
-         labels.push(otherLabel);
+      if (defaultLabel && (!jobColumn.category || jobColumn.category === "Other")) {
+         defaultLabel.dashboardCategory = "";
+         defaultLabel.dashboardName = "Other";
+         labels.push(defaultLabel);
       }
 
       let key = 0;
 
       const buttons = labels.map(label => {
-         return <JobLabel key={`label_${item.job.id}_${jobColumn.category}_${label.name}_${key++}`} item={item} column={jobColumn} label={label} />;
+         return <JobLabel key={`label_${item.job.id}_${jobColumn.category}_${label.dashboardName}_${key++}`} item={item} column={jobColumn} label={label} />;
       });
 
       return (
diff --git a/Engine/Source/Programs/Horde/HordeDashboard/src/components/LogSource.tsx b/Engine/Source/Programs/Horde/HordeDashboard/src/components/LogSource.tsx
--- a/Engine/Source/Programs/Horde/HordeDashboard/src/components/LogSource.tsx
+++ b/Engine/Source/Programs/Horde/HordeDashboard/src/components/LogSource.tsx
@@ -3,7 +3,7 @@
 import { action, makeObservable, observable } from "mobx";
 import moment, { Moment } from 'moment-timezone';
 import backend from '../backend';
-import { AgentData, BatchData, EventSeverity, GetArtifactResponseV2, GetLogEventResponse, GetLogFileResponse, IssueData, LeaseData, LogData, NodeData, StepData, StreamData } from "../backend/Api";
+import { AgentData, BatchData, EventSeverity, GetArtifactResponseV2, GetLogEventResponse, GetLogFileResponse, IssueData, LeaseData, LogData, StepData, StreamData } from "../backend/Api";
 import { getBatchSummaryMarkdown, getStepSummaryMarkdown, JobDetails } from "../backend/JobDetails";
 import { getLeaseElapsed, getStepPercent } from '../base/utilities/timeUtils';
 import { BreadcrumbItem } from './Breadcrumbs';
@@ -411,7 +411,7 @@ export class JobLogSource extends LogSource {
 
    get percentComplete(): number | undefined {
 
-      if (this.step && this.node) {
+      if (this.step) {
          return getStepPercent(this.step);
       }
 
@@ -498,10 +498,7 @@ export class JobLogSource extends LogSource {
 
       this.agentId = this.batch?.agentId ?? details.batchByStepId(this.step?.id!)?.agentId;
 
-      this.node = undefined;
-
       if (this.step) {
-         this.node = details.nodeByStepId(this.step.id);
          this.jobName = details.getStepName(this.step.id) ?? "Unknown Step Node";
       }
 
@@ -588,8 +585,7 @@ export class JobLogSource extends LogSource {
    agentId?: string;
    batch?: BatchData;
    step?: StepData;
-   node?: NodeData;
-
+   
    artifactsV2?: GetArtifactResponseV2[];
 
    jobDetails: JobDetails = new JobDetails(undefined, undefined, undefined, true);
diff --git a/Engine/Source/Programs/Horde/HordeDashboard/src/components/LogView.tsx b/Engine/Source/Programs/Horde/HordeDashboard/src/components/LogView.tsx
--- a/Engine/Source/Programs/Horde/HordeDashboard/src/components/LogView.tsx
+++ b/Engine/Source/Programs/Horde/HordeDashboard/src/components/LogView.tsx
@@ -233,6 +233,8 @@ const StepHistoryModal: React.FC<{ jobDetails: JobDetails, stepId: string | unde
       return null;
    }
 
+   const step = jobDetails.stepById(stepId);
+
    type HistoryItem = {
       ref: GetJobStepRefResponse;
    };
@@ -259,7 +261,7 @@ const StepHistoryModal: React.FC<{ jobDetails: JobDetails, stepId: string | unde
 
 
       if (column.name === "Name") {
-         return <Stack horizontal verticalAlign="center" tokens={{ childrenGap: 0, padding: 0 }} style={{ width: "100%", height: "100%" }} >{<StepRefStatusIcon stepRef={ref} />}<Text>{jobDetails.nodeByStepId(stepId)?.name}</Text></Stack>;
+         return <Stack horizontal verticalAlign="center" tokens={{ childrenGap: 0, padding: 0 }} style={{ width: "100%", height: "100%" }} >{<StepRefStatusIcon stepRef={ref} />}<Text>{step?.name}</Text></Stack>;
       }
 
       if (column.name === "Change") {
diff --git a/Engine/Source/Programs/Horde/HordeDashboard/src/components/NotificationDropdown.tsx b/Engine/Source/Programs/Horde/HordeDashboard/src/components/NotificationDropdown.tsx
--- a/Engine/Source/Programs/Horde/HordeDashboard/src/components/NotificationDropdown.tsx
+++ b/Engine/Source/Programs/Horde/HordeDashboard/src/components/NotificationDropdown.tsx
@@ -10,6 +10,7 @@ import { JobDetailsV2 } from './jobDetailsV2/JobDetailsViewCommon';
 import { useQuery } from "./JobDetailCommon";
 import { getHordeStyling } from "../styles/Styles";
 
+
 type NotificationType = "Job" | "Step" | "Label";
 type NotificationOutcome = "Warnings" | "Success" | "Failure";
 
@@ -125,10 +126,10 @@ class SubscriptionHandler  {
          stepName = details.getStepName(this.stepId);
       }
 
-      const label = details.labelByIndex(this.labelIdx);
+      const label = details.labelByIndex(this.labelIdx)?.stateResponse
 
-      const labelName = label?.name;
-      const categoryName = label?.category;
+      const labelName = label?.dashboardName;
+      const categoryName = label?.dashboardCategory;
 
       const streamId = details.stream!.id;
 
@@ -306,10 +307,10 @@ class SubscriptionHandler  {
          return (s.event.type === type) && (s.event.streamId === details.stream?.id) && (s.event.templateId === templateId);
       });
 
-      let label = details.labelByIndex(this.labelIdx);
+      let label = details.labelByIndex(this.labelIdx)?.stateResponse;
 
       if (label) {
-         subs = subs?.filter((s) => (s.event as LabelCompleteEventRecord).labelName === label!.name && (s.event as LabelCompleteEventRecord).categoryName === label!.category);
+         subs = subs?.filter((s) => (s.event as LabelCompleteEventRecord).labelName === label!.dashboardName && (s.event as LabelCompleteEventRecord).categoryName === label!.dashboardCategory);
       }
 
       if (this.stepId) {
@@ -479,7 +480,7 @@ export const NotificationDropdown: React.FC<{ jobDetails: JobDetailsV2 }> = obse
    }
    if (label) {
       type = "Label";
-      const aggregate = jobDetails.findLabel(label.name, label.category)!;
+      const aggregate = jobDetails.findLabel(label.stateResponse.dashboardName, label.stateResponse.dashboardCategory)!;
       complete = aggregate?.stateResponse?.state === LabelState.Complete;
       subscribed = subHandler.nlabel?.slack ?? false;
    }
@@ -633,4 +634,3 @@ export const NotificationDropdown: React.FC<{ jobDetails: JobDetailsV2 }> = obse
 
    return <CommandBar items={notificationItems} onReduceData={() => undefined} />
 });
-
diff --git a/Engine/Source/Programs/Horde/HordeDashboard/src/components/PoolView.tsx b/Engine/Source/Programs/Horde/HordeDashboard/src/components/PoolView.tsx
--- a/Engine/Source/Programs/Horde/HordeDashboard/src/components/PoolView.tsx
+++ b/Engine/Source/Programs/Horde/HordeDashboard/src/components/PoolView.tsx
@@ -126,7 +126,7 @@ class PoolHandler extends PollBase {
                   minCreateTime: minCreateTime,
                   count: 250
                }
-               return backend.getJobs(query, false);
+               return backend.getJobs(query);
 
             }).filter(r => !!r);
 
@@ -151,7 +151,7 @@ class PoolHandler extends PollBase {
 
          if (streamJobIds.size) {
 
-            const streamJobs = await backend.getJobs({ id: Array.from(streamJobIds), filter: "id,name,batches,change,createTime,streamId,preflightChange,graphHash" }, true);
+            const streamJobs = await backend.getJobs({ id: Array.from(streamJobIds), filter: "id,name,batches,change,createTime,streamId,preflightChange" });
 
             streamJobs.forEach(job => {
 
@@ -165,15 +165,8 @@ class PoolHandler extends PollBase {
                      return;
                   }
 
-                  const groups = job.graphRef?.groups;
 
-                  if (!groups || !groups[b.groupIdx]) {
-                     return;
-                  }
-
-                  const group = groups[b.groupIdx];
-
-                  const agentType = group?.agentType;
+                  const agentType = b?.agentType;
 
                   const stream = projectStore.streamById(job.streamId);
 
@@ -216,10 +209,10 @@ class PoolHandler extends PollBase {
 
             const query: JobQuery = {
                id: jobIds,
-               filter: "id,name,batches,change,createTime,streamId,preflightChange,graphHash",
+               filter: "id,name,batches,change,createTime,streamId,preflightChange",
             }
 
-            const jobs = await backend.getJobs(query, true);
+            const jobs = await backend.getJobs(query);
 
             jobs.forEach(job => {
 
@@ -386,7 +379,7 @@ const StepPanel: React.FC<{ stepState: StepState }> = ({ stepState }) => {
          step = handler.pendingSteps.get(agent.id);
       }
 
-      if (!job || !job.graphRef?.groups || !step) {
+      if (!job || !step) {
          return null;
       }
 
@@ -394,9 +387,6 @@ const StepPanel: React.FC<{ stepState: StepState }> = ({ stepState }) => {
       if (!batch) {
          return null;
       }
-      const group = job.graphRef.groups[batch.groupIdx];
-      const node = group.nodes[step.nodeIdx];
-
       if (column.name === "Time Active") {
          if (stepState === StepState.Pending) {
             return null;
@@ -433,7 +423,7 @@ const StepPanel: React.FC<{ stepState: StepState }> = ({ stepState }) => {
          if (jobName.indexOf("- Kicked By") !== -1) {
             jobName = jobName.split("- Kicked By")[0];
          }
-         let stepName = `${jobName} - ${node.name}`;
+         let stepName = `${jobName} - ${step.name}`;
          const stepUrl = `/job/${job.id}?step=${step.id}`;
 
          const stream = projectStore.streamById(job.streamId);
@@ -581,9 +571,7 @@ const BatchPanel: React.FC = () => {
          let step: GetStepResponse | undefined;
          if (batch.steps.length) {
             step = batch.steps[0];
-            const group = job.graphRef!.groups![batch.groupIdx];
-            const node = group.nodes[step.nodeIdx];
-            jobName += ` - ${node.name}`;
+            jobName += ` - ${step.name}`;
             if (stream) {
                jobName = `${stream.fullname ?? stream.id} - ` + jobName;
             }
diff --git a/Engine/Source/Programs/Horde/HordeDashboard/src/components/SchedulePane.tsx b/Engine/Source/Programs/Horde/HordeDashboard/src/components/SchedulePane.tsx
--- a/Engine/Source/Programs/Horde/HordeDashboard/src/components/SchedulePane.tsx
+++ b/Engine/Source/Programs/Horde/HordeDashboard/src/components/SchedulePane.tsx
@@ -54,7 +54,7 @@ export const SchedulePane: React.FC<{ templates: GetTemplateRefResponse[] }> = (
    })
 
    if (ids.size) {
-      backend.getJobsByIds(Array.from(ids.values()), { filter: "id,createTime,state" }, false).then(response => {
+      backend.getJobsByIds(Array.from(ids.values()), { filter: "id,createTime,state" }).then(response => {
          response.forEach(j => {
             scheduledJobs.set(j.id, j);
          });
diff --git a/Engine/Source/Programs/Horde/HordeDashboard/src/components/StatusIcon.tsx b/Engine/Source/Programs/Horde/HordeDashboard/src/components/StatusIcon.tsx
--- a/Engine/Source/Programs/Horde/HordeDashboard/src/components/StatusIcon.tsx
+++ b/Engine/Source/Programs/Horde/HordeDashboard/src/components/StatusIcon.tsx
@@ -4,7 +4,7 @@
 
 import { FontIcon, Stack } from '@fluentui/react';
 import React, { CSSProperties } from 'react';
-import { DeviceStatus, FindIssueResponse, GetAgentLeaseResponse, GetBatchResponse, GetIssueResponse, GetJobStepRefResponse, GetStepResponse, IssueSeverity, JobStepBatchError, JobStepOutcome, JobStepState, LabelOutcome, LabelState, LeaseOutcome, LeaseState } from '../backend/Api';
+import { DeviceStatus, FindIssueResponse, GetAgentLeaseResponse, GetBatchResponse, GetIssueResponse, GetJobStepRefResponse, GetLabelStateResponse, GetStepResponse, IssueSeverity, JobStepBatchError, JobStepOutcome, JobStepState, LabelOutcome, LabelState, LeaseOutcome, LeaseState } from '../backend/Api';
 import dashboard, { StatusColor } from "../backend/Dashboard";
 import { JobLabel } from '../backend/JobDetails';
 import { getStepStatusColor } from '../styles/colors';
@@ -55,14 +55,14 @@ export const getLabelIcon = (state: LabelState | undefined, outcome: LabelOutcom
 }
 
 
-export const LabelStatusIcon: React.FC<{ label: JobLabel, style?: CSSProperties }> = ({ label, style }) => {
+export const LabelStatusIcon: React.FC<{ label: GetLabelStateResponse, style?: CSSProperties }> = ({ label, style }) => {
 
    style = style ?? {};
    style.fontSize = style.fontSize ?? 13;
    style.paddingTop = style.paddingTop ?? 3;
    style.paddingRight = style.paddingRight ?? 8;
 
-   const icon = getLabelIcon(label.stateResponse.state, label.stateResponse.outcome);
+   const icon = getLabelIcon(label.state, label.outcome);
 
    return <StatusIcon iconName={icon.icon} style={{ ...style, color: icon.color }} />;
 }
diff --git a/Engine/Source/Programs/Horde/HordeDashboard/src/components/UserHome.tsx b/Engine/Source/Programs/Horde/HordeDashboard/src/components/UserHome.tsx
--- a/Engine/Source/Programs/Horde/HordeDashboard/src/components/UserHome.tsx
+++ b/Engine/Source/Programs/Horde/HordeDashboard/src/components/UserHome.tsx
@@ -7,9 +7,8 @@ import moment from "moment-timezone";
 import { default as React, useEffect, useState } from 'react';
 import { Link, Navigate, useLocation, useNavigate } from "react-router-dom";
 import backend, { useBackend } from "../backend";
-import { GetBisectTaskResponse, GetIssueResponse, GetStepResponse, JobData, JobQuery, JobState, JobStepOutcome, LabelData, LabelOutcome, LabelState, ProjectData, StepData, StreamData } from "../backend/Api";
+import { GetBisectTaskResponse, GetIssueResponse, GetLabelStateResponse, GetStepResponse, JobData, JobQuery, JobState, JobStepOutcome, LabelOutcome, LabelState, ProjectData, StepData, StreamData } from "../backend/Api";
 import dashboard, { StatusColor } from "../backend/Dashboard";
-import graphCache, { GraphQuery } from '../backend/GraphCache';
 import { PollBase } from '../backend/PollBase';
 import { useWindowSize } from '../base/utilities/hooks';
 import { displayTimeZone, getElapsedString, getShortNiceTime } from '../base/utilities/timeUtils';
@@ -37,7 +36,7 @@ const customStyles = mergeStyleSets({
    detailsRow: {
       selectors: {
          '.ms-DetailsRow': {
-            borderBottom: '0px',          
+            borderBottom: '0px',
             width: "100%"
          },
          '.ms-DetailsRow-cell': {
@@ -425,7 +424,7 @@ class UserJobsHandler {
 
          const jobIds = this.jobIds.slice(-maxJobs);
 
-         const mjobs = await backend.getJobsByIds(jobIds, query, false);
+         const mjobs = await backend.getJobsByIds(jobIds, query);
 
          // check for canceled after modified test
          if (this.canceled.has(cancelId)) {
@@ -435,10 +434,6 @@ class UserJobsHandler {
          const jobs: JobData[] = [];
 
          mjobs.forEach(j1 => {
-            const existing = this.jobs.find(j2 => j1.id === j2.id);
-            if (existing) {
-               j1.graphRef = existing.graphRef;
-            }
             jobs.push(j1);
          })
 
@@ -451,56 +446,13 @@ class UserJobsHandler {
             return timeA < timeB ? 1 : -1;
          });
 
-         const graphHashes = new Set<string>();
-
-         this.jobs.forEach(j => {
-
-            if (j.graphRef?.hash !== j.graphHash) {
-               j.graphRef = undefined;
-            }
-
-            if (graphHashes.size > 5) {
-               return;
-            }
-
-            if (j.graphHash && !j.graphRef) {
-               j.graphRef = graphCache.cache.get(j.graphHash);
-               if (!j.graphRef) {
-                  graphHashes.add(j.graphHash);
-               }
-            }
-         })
-
-         if (graphHashes.size) {
-
-            const graphQuery: GraphQuery[] = [];
-            Array.from(graphHashes.values()).forEach(h => {
-               const jobId = this.jobs.find(j => j.graphHash === h)!.id;
-               graphQuery.push({
-                  jobId: jobId,
-                  graphHash: h
-               })
-            })
-
-            const graphs = await graphCache.getGraphs(graphQuery);
-
-            graphs.forEach(graph => {
-
-               this.jobs.forEach(j => {
-                  if (graph.hash === j.graphHash) {
-                     j.graphRef = graph;
-                  }
-               })
-            })
-
-         }
 
          // check for canceled during graph request
          if (this.canceled.has(cancelId)) {
             return;
          }
 
-         if (graphHashes.size || mjobs.length || this.initial) {
+         if (mjobs.length || this.initial) {
             this.initial = false;
             this.setUpdated();
          }
@@ -664,25 +616,11 @@ const JobsPanel: React.FC<{ includeOtherPreflights: boolean }> = observer(({ inc
          </Stack>;
       };
 
-
       let items = steps.map(step => {
-
-         const batch = job.batches!.find(b => !!b.steps.find(s => s.id === step.id))!;
-         const groups = job?.graphRef?.groups;
-         if (!groups) {
-            return undefined;
-         }
-         const node = groups[batch.groupIdx].nodes[step.nodeIdx];
-
-         if (!node) {
-            return undefined;
-         }
-
          return {
             step: step,
-            name: node.name
+            name: step.name
          }
-
       })
 
       items = items.filter(item => !!item);
@@ -874,15 +812,15 @@ const JobsPanel: React.FC<{ includeOtherPreflights: boolean }> = observer(({ inc
 
       if (props) {
 
-         const item = jobItems[props.itemIndex];         
-         
+         const item = jobItems[props.itemIndex];
+
          let background: string | undefined;
          if (props.itemIndex % 2 === 0) {
-            background  =  dashboard.darktheme ? "#1D2021" : "#FAF9F9";
+            background = dashboard.darktheme ? "#1D2021" : "#FAF9F9";
          }
 
          return <JobOperationsContextMenu job={item.job}>
-            <DetailsRow styles={{ root: { paddingTop: 8, paddingBottom: 8, backgroundColor: background}, cell: { selectors: { "a, a:visited, a:activem, a:hover": { color: modeColors.text} } } }} {...props} />
+            <DetailsRow styles={{ root: { paddingTop: 8, paddingBottom: 8, backgroundColor: background }, cell: { selectors: { "a, a:visited, a:activem, a:hover": { color: modeColors.text } } } }} {...props} />
          </JobOperationsContextMenu>
       }
       return null;
@@ -899,12 +837,14 @@ const JobsPanel: React.FC<{ includeOtherPreflights: boolean }> = observer(({ inc
       return null;
    };
 
-   const JobLabel: React.FC<{ item: JobItem; label: LabelData }> = ({ item, label }) => {
 
-      const aggregates = item.job.graphRef?.labels;
+   const JobLabel: React.FC<{ item: JobItem; label: GetLabelStateResponse }> = ({ item, label }) => {
+
+      const defaultLabel = item.job.defaultLabel;
+      const aggregates = item.job.labels;
 
       // note details may not be loaded here, as only initialized on callout for optimization (details.getLabelIndex(label.Name, label.Category);)
-      const jlabel = aggregates?.find((l, idx) => l.category === label.category && l.name === label.name && item.job.labels![idx]?.state !== LabelState.Unspecified);
+      const jlabel = aggregates?.find((l, idx) => l.dashboardCategory === label.dashboardCategory && l.dashboardName === label.dashboardName && item.job.labels![idx]?.state !== LabelState.Unspecified);
       let labelIdx = -1;
       if (jlabel) {
          labelIdx = aggregates?.indexOf(jlabel)!;
@@ -912,9 +852,9 @@ const JobsPanel: React.FC<{ includeOtherPreflights: boolean }> = observer(({ inc
 
       let state: LabelState | undefined;
       let outcome: LabelOutcome | undefined;
-      if (label.defaultLabel) {
-         state = label.defaultLabel.state;
-         outcome = label.defaultLabel.outcome;
+      if (label === defaultLabel) {
+         state = defaultLabel.state;
+         outcome = defaultLabel.outcome;
       } else {
          state = item.job.labels![labelIdx]?.state;
          outcome = item.job.labels![labelIdx]?.outcome;
@@ -928,11 +868,11 @@ const JobsPanel: React.FC<{ includeOtherPreflights: boolean }> = observer(({ inc
          url = `/job/${item.job.id}?label=${labelIdx}`;
       }
 
-      const target = `label_${item.job.id}_${label.name}_${label.category}`.replace(/ /g, "");
+      const target = `label_${item.job.id}_${label.dashboardName}_${label.dashboardCategory}`.replace(/ /g, "");
 
       return <Stack>
          <div id={target} className="horde-no-darktheme">
-            <Link to={url}><Stack className={hordeClasses.badgeNoIcon}><DefaultButton key={label.name} style={{ backgroundColor: color.primaryColor, fontSize: 9 }} text={label.name}>
+            <Link to={url}><Stack className={hordeClasses.badgeNoIcon}><DefaultButton key={label.dashboardName} style={{ backgroundColor: color.primaryColor, fontSize: 9 }} text={label.dashboardName}>
                {!!color.secondaryColor && <div style={{
                   borderLeft: "10px solid transparent",
                   borderRight: `10px solid ${color.secondaryColor}`,
@@ -955,7 +895,8 @@ const JobsPanel: React.FC<{ includeOtherPreflights: boolean }> = observer(({ inc
 
       const job = item.job!;
       const defaultLabel = job.defaultLabel;
-      const labels = job.graphRef?.labels ?? [];
+
+      const labels = job.labels ?? [];
 
       if (!labels.length && !defaultLabel) {
          return <div />;
@@ -970,21 +911,14 @@ const JobsPanel: React.FC<{ includeOtherPreflights: boolean }> = observer(({ inc
       });
 
       const catergories = new Set<string>();
-      view.forEach(label => catergories.add(label.category));
+      view.forEach(label => catergories.add(label.dashboardCategory));
 
       const sorted = Array.from(catergories).sort((a, b) => { return a < b ? -1 : 1 });
 
-      if (!catergories.has("Other") && defaultLabel?.nodes?.length) {
+      if (!catergories.has("Other") && defaultLabel) {
 
-         const otherLabel = {
-            category: "Other",
-            name: "Other",
-            requiredNodes: [],
-            includedNodes: defaultLabel!.nodes,
-            defaultLabel: item.job.defaultLabel
-         };
-
-         view.push(otherLabel);
+         defaultLabel.dashboardCategory = defaultLabel.dashboardName = "Other";
+         view.push(defaultLabel);
          sorted.push("Other");
       }
 
@@ -992,8 +926,8 @@ const JobsPanel: React.FC<{ includeOtherPreflights: boolean }> = observer(({ inc
 
       const labelStacks = sorted.map(cat => {
 
-         const buttons = view.filter(label => label.category === cat).map(label => {
-            return <JobLabel key={`label_${item.job.id}_${label.name}_${key++}`} item={item} label={label} />;
+         const buttons = view.filter(label => label.dashboardCategory === cat).map(label => {
+            return <JobLabel key={`label_${item.job.id}_${label.dashboardName}_${key++}`} item={item} label={label} />;
          });
 
          if (!buttons.length) {
@@ -1115,7 +1049,7 @@ const BisectionPanel: React.FC = observer(() => {
    }, []);
 
    // subscribe
-   dashboard.subscribe();   
+   dashboard.subscribe();
 
    if (bisectionHandler.updated) { };
 
@@ -1173,7 +1107,7 @@ const UserHomeViewInner: React.FC = () => {
                </ScrollablePane>
             </div>
          </FocusZone>
-      </Stack> 
+      </Stack>
    </Stack>
 };
 
@@ -1184,7 +1118,7 @@ export const UserHomeView: React.FC = () => {
    if (dashboard.user?.dashboardFeatures?.showLandingPage) {
       return <Navigate to="/docs/Landing.md" replace={true} />
    }
-   
+
    const vw = Math.max(document.documentElement.clientWidth, window.innerWidth || 0);
 
    const { hordeClasses, modeColors } = getHordeStyling();
diff --git a/Engine/Source/Programs/Horde/HordeDashboard/src/components/jobDetailsV2/JobDetailLabels.tsx b/Engine/Source/Programs/Horde/HordeDashboard/src/components/jobDetailsV2/JobDetailLabels.tsx
--- a/Engine/Source/Programs/Horde/HordeDashboard/src/components/jobDetailsV2/JobDetailLabels.tsx
+++ b/Engine/Source/Programs/Horde/HordeDashboard/src/components/jobDetailsV2/JobDetailLabels.tsx
@@ -40,7 +40,7 @@ export const LabelsPanelV2: React.FC<{ jobDetails: JobDetailsV2, dataView: JobDa
    let labels = jobDetails.labels.filter(label => label.stateResponse.state !== LabelState.Unspecified);
 
    if (qlabel) {
-      labels = labels.filter(l => l.category === qlabel?.category && l.name === qlabel?.name);
+      labels = labels.filter(l => l.stateResponse.dashboardCategory === qlabel?.stateResponse.dashboardCategory && l.stateResponse.dashboardName === qlabel?.stateResponse.dashboardName);
    }
 
    const stateFilter = jobFilter.filterStates;
@@ -90,14 +90,14 @@ export const LabelsPanelV2: React.FC<{ jobDetails: JobDetailsV2, dataView: JobDa
    };
 
    const categories: Set<string> = new Set();
-   labels.forEach(label => { if (label.name) { categories.add(label.category!); } });
+   labels.forEach(label => { if (label.stateResponse.dashboardName ?? "") { categories.add(label.stateResponse.dashboardCategory!); } });
 
 
    let items = Array.from(categories.values()).map(c => {
       return {
          category: c,
-         labels: labels.filter(label => label.name && (label.category === c)).sort((a, b) => {
-            return a.name! < b.name! ? -1 : 1;
+         labels: labels.filter(label => label.stateResponse.dashboardName && (label.stateResponse.dashboardCategory === c)).sort((a, b) => {
+            return a.stateResponse.dashboardName! < b.stateResponse.dashboardName! ? -1 : 1;
          })
       } as LabelItem;
    }).filter(item => item.labels?.length).sort((a, b) => {
@@ -108,14 +108,14 @@ export const LabelsPanelV2: React.FC<{ jobDetails: JobDetailsV2, dataView: JobDa
 
    if (filter.search) {
       items = items.filter(i => {
-         i.labels = i.labels.filter(label => label.name.toLowerCase().indexOf(filter.search!.toLowerCase()) !== -1);
+         i.labels = i.labels.filter(label => (label.stateResponse.dashboardName ?? "").toLowerCase().indexOf(filter.search!.toLowerCase()) !== -1);
          return i.labels.length !== 0;
       });
    }
 
    if (filter.currentInput) {
       items = items.filter(i => {
-         i.labels = i.labels.filter(label => label.name.toLowerCase().indexOf(filter.currentInput!.toLowerCase()) !== -1);
+         i.labels = i.labels.filter(label => (label.stateResponse.dashboardName ?? "").toLowerCase().indexOf(filter.currentInput!.toLowerCase()) !== -1);
          return i.labels.length !== 0;
       });
    }   
@@ -154,18 +154,18 @@ export const LabelsPanelV2: React.FC<{ jobDetails: JobDetailsV2, dataView: JobDa
 
                   let filtered = false;
                   if (qlabel) {
-                     if (qlabel.category !== label.category || qlabel.name !== label.name) {
+                     if (qlabel.stateResponse.dashboardCategory!== label.stateResponse.dashboardCategory || qlabel.stateResponse.dashboardName !== label.stateResponse.dashboardName) {
                         filtered = true;
                      }
                   }
-                  return <Stack key={`labels_${label.category}_${label.name}`} className={hordeClasses.badge}>
+                  return <Stack key={`labels_${label.stateResponse.dashboardCategory}_${label.stateResponse.dashboardName}`} className={hordeClasses.badge}>
                      <DefaultButton
                         onClick={() => {
-                           if (qlabel?.category === label.category && qlabel?.name === label.name) {
+                           if (qlabel?.stateResponse.dashboardCategory === label.stateResponse.dashboardCategory && qlabel?.stateResponse.dashboardName === label.stateResponse.dashboardName) {
                               navigate(location.pathname)
                            } else {
 
-                              const idx = jobDetails.labelIndex(label.name, label.category);
+                              const idx = jobDetails.labelIndex(label.stateResponse.dashboardName, label.stateResponse.dashboardCategory);
                               if (idx >= 0) {
                                  navigate(location.pathname + `?label=${idx}`)
                               } else {
@@ -174,8 +174,8 @@ export const LabelsPanelV2: React.FC<{ jobDetails: JobDetailsV2, dataView: JobDa
                            }
 
                         }}
-                        key={label.name} style={{ backgroundColor: color.primaryColor, color: textColor, filter: filtered ? "brightness(0.70)" : undefined }}
-                        text={label.name!}>
+                        key={label.stateResponse.dashboardName  ?? ""} style={{ backgroundColor: color.primaryColor, color: textColor, filter: filtered ? "brightness(0.70)" : undefined }}
+                        text={label.stateResponse.dashboardName  ?? ""}>
                         {!!color.secondaryColor && <div style={{
                            borderLeft: "10px solid transparent",
                            borderRight: `10px solid ${color.secondaryColor}`,
diff --git a/Engine/Source/Programs/Horde/HordeDashboard/src/components/jobDetailsV2/JobDetailStepHistory.tsx b/Engine/Source/Programs/Horde/HordeDashboard/src/components/jobDetailsV2/JobDetailStepHistory.tsx
--- a/Engine/Source/Programs/Horde/HordeDashboard/src/components/jobDetailsV2/JobDetailStepHistory.tsx
+++ b/Engine/Source/Programs/Horde/HordeDashboard/src/components/jobDetailsV2/JobDetailStepHistory.tsx
@@ -146,9 +146,10 @@ export const StepHistoryPanel: React.FC<{ jobDetails: JobDetailsV2; stepId: stri
 
       const ref = item.ref;
 
+      const step = jobDetails.stepById(stepId);
 
       if (column.name === "Name") {
-         return <Stack horizontal>{<StepRefStatusIcon stepRef={ref} />}<Text>{jobDetails.nodeByStepId(stepId)?.name}</Text></Stack>;
+         return <Stack horizontal>{<StepRefStatusIcon stepRef={ref} />}<Text>{step?.name}</Text></Stack>;
       }
 
       if (column.name === "Change") {
diff --git a/Engine/Source/Programs/Horde/HordeDashboard/src/components/jobDetailsV2/JobDetailStepHistoryV2.tsx b/Engine/Source/Programs/Horde/HordeDashboard/src/components/jobDetailsV2/JobDetailStepHistoryV2.tsx
--- a/Engine/Source/Programs/Horde/HordeDashboard/src/components/jobDetailsV2/JobDetailStepHistoryV2.tsx
+++ b/Engine/Source/Programs/Horde/HordeDashboard/src/components/jobDetailsV2/JobDetailStepHistoryV2.tsx
@@ -643,6 +643,8 @@ const GraphTooltip: React.FC<{ dataView: StepHistoryDataView }> = observer(({ da
    const format = dashboard.display24HourClock ? "HH:mm:ss z" : "LT z";
    let displayTimeStr = displayTime.format('MMM Do') + ` at ${displayTime.format(format)}`;
 
+   const step = dataView.details?.stepById(dataView.stepId);
+
    return <div style={{
       position: "absolute",
       display: "block",
@@ -662,7 +664,7 @@ const GraphTooltip: React.FC<{ dataView: StepHistoryDataView }> = observer(({ da
       <Stack>
          <Link to={`/job/${ref.jobId}?step=${ref.stepId}`}><Stack horizontal>
             <StepRefStatusIcon stepRef={ref} />
-            <Text variant={textSize}>{dataView.details?.nodeByStepId(dataView.stepId)?.name}</Text>
+            <Text variant={textSize}>{step?.name}</Text>
          </Stack>
          </Link>
          <Stack style={{ paddingLeft: 2, paddingTop: 8 }} tokens={{ childrenGap: 8 }}>
diff --git a/Engine/Source/Programs/Horde/HordeDashboard/src/components/jobDetailsV2/JobDetailStepTrendsV2.tsx b/Engine/Source/Programs/Horde/HordeDashboard/src/components/jobDetailsV2/JobDetailStepTrendsV2.tsx
--- a/Engine/Source/Programs/Horde/HordeDashboard/src/components/jobDetailsV2/JobDetailStepTrendsV2.tsx
+++ b/Engine/Source/Programs/Horde/HordeDashboard/src/components/jobDetailsV2/JobDetailStepTrendsV2.tsx
@@ -634,6 +634,8 @@ const GraphTooltip: React.FC<{ dataView: StepTrendsDataView }> = observer(({ dat
    const format = dashboard.display24HourClock ? "HH:mm:ss z" : "LT z";
    let displayTimeStr = displayTime.format('MMM Do') + ` at ${displayTime.format(format)}`;
 
+   const step = dataView.details?.stepById(dataView.stepId);
+
    return <div style={{
       position: "absolute",
       display: "block",
@@ -653,7 +655,7 @@ const GraphTooltip: React.FC<{ dataView: StepTrendsDataView }> = observer(({ dat
       <Stack>
          <Link to={`/job/${ref.jobId}?step=${ref.stepId}`}><Stack horizontal>
             <StepRefStatusIcon stepRef={ref} />
-            <Text variant={textSize}>{dataView.details?.nodeByStepId(dataView.stepId)?.name}</Text>
+            <Text variant={textSize}>{step?.name}</Text>
          </Stack>
          </Link>
          <Stack style={{ paddingLeft: 2, paddingTop: 8 }} tokens={{ childrenGap: 8 }}>
diff --git a/Engine/Source/Programs/Horde/HordeDashboard/src/components/jobDetailsV2/JobDetailTimeline.tsx b/Engine/Source/Programs/Horde/HordeDashboard/src/components/jobDetailsV2/JobDetailTimeline.tsx
--- a/Engine/Source/Programs/Horde/HordeDashboard/src/components/jobDetailsV2/JobDetailTimeline.tsx
+++ b/Engine/Source/Programs/Horde/HordeDashboard/src/components/jobDetailsV2/JobDetailTimeline.tsx
@@ -177,9 +177,8 @@ class TimelineDataView extends JobDataView {
                b.steps = b.steps?.filter(s => {
                   if (!s.startTime || !s.finishTime) {
                      return false;
-                  }
-                  const node = jobDetails.nodeByStepId(s.id);
-                  return label.includedNodes.indexOf(node?.name ?? "") !== -1 || label.requiredNodes.indexOf(node?.name ?? "") !== -1;
+                  }                  
+                  return label.steps.indexOf(s?.id ?? "") !== -1;
                });
             });
 
@@ -760,8 +759,7 @@ const GraphTooltip: React.FC<{ dataView: TimelineDataView }> = observer(({ dataV
 
    const details = dataView.details!;
 
-   const group = details.groups[span.batch.groupIdx];
-   const pool = details.stream?.agentTypes[group?.agentType!];
+   const pool = details.stream?.agentTypes[span.batch?.agentType!];
 
    const agentPool = pool?.pool?.toUpperCase() ? pool.pool.toUpperCase() : "Unknown Pool";
    const agent = span.batch.agentId ? span.batch.agentId : "Unknown Agent";
diff --git a/Engine/Source/Programs/Horde/HordeDashboard/src/components/jobDetailsV2/JobDetailViewSteps.tsx b/Engine/Source/Programs/Horde/HordeDashboard/src/components/jobDetailsV2/JobDetailViewSteps.tsx
--- a/Engine/Source/Programs/Horde/HordeDashboard/src/components/jobDetailsV2/JobDetailViewSteps.tsx
+++ b/Engine/Source/Programs/Horde/HordeDashboard/src/components/jobDetailsV2/JobDetailViewSteps.tsx
@@ -2,7 +2,7 @@ import { CollapseAllVisibility, ConstrainMode, DetailsHeader, DetailsList, Detai
 import { observer } from "mobx-react-lite";
 import React, { useEffect, useState } from "react";
 import { Link, useParams } from "react-router-dom";
-import { BatchData, GetBatchResponse, JobStepBatchError, JobStepBatchState, JobStepOutcome, JobStepState, NodeData, StepData } from "../../backend/Api";
+import { BatchData, GetBatchResponse, JobStepBatchError, JobStepBatchState, JobStepOutcome, JobStepState, StepData } from "../../backend/Api";
 import dashboard, { StatusColor } from "../../backend/Dashboard";
 import { ISideRailLink } from "../../base/components/SideRail";
 import { getBatchInitElapsed, getNiceTime, getStepElapsed, getStepETA, getStepFinishTime, getStepPercent, getStepStartTime, getStepTimingDelta } from "../../base/utilities/timeUtils";
@@ -20,7 +20,6 @@ const depSideRail: ISideRailLink = { text: "Dependencies", url: "rail_dependenci
 type StepItem = {
    step?: StepData;
    batch?: BatchData;
-   node?: NodeData;
    agentId?: string;
    agentRow?: boolean;
    agentType?: string;
@@ -316,7 +315,6 @@ export const StepsPanelInner: React.FC<{ jobDetails: JobDetailsV2, depStepId?: s
          }
       };
 
-
       const stepName = jobDetails.getStepName(step.id) ?? "Unnown Step Name";
       const stepUrl = `/job/${jobId}?step=${step.id}`;
 
@@ -395,21 +393,11 @@ export const StepsPanelInner: React.FC<{ jobDetails: JobDetailsV2, depStepId?: s
 
    if (stepId) {
 
-      const nodes: NodeData[] = [];
-
       const getStepsRecursive = (stepId: string) => {
-
-         const stepNode = jobDetails.nodeByStepId(stepId);
-
-         if (!stepNode || nodes.find(n => stepNode === n)) {
-            return;
-         }
-
-         nodes.push(stepNode);
-
-         [stepNode.inputDependencies, stepNode.orderDependencies].flat().forEach(name => {
-            const s = jobDetails.stepByName(name);
-            if (s) {
+         const rstep = jobDetails.stepById(stepId);
+         [rstep.inputDependencies, rstep.orderDependencies].flat().forEach(id => {
+            const s = jobDetails.stepById(id);
+            if (s && !stepFilter.find(s => s.id === id)) {
                stepFilter.push(s);
                getStepsRecursive(s.id);
             }
@@ -423,14 +411,6 @@ export const StepsPanelInner: React.FC<{ jobDetails: JobDetailsV2, depStepId?: s
          stepFilter.push(step);
       }
 
-      /*
-      if (singleStep) {
-         const step = jobDetails.stepById(stepId);
-         if (step) {
-            stepFilter.push(step);
-         }
-      }*/
-
       if (!stepFilter.length) {
          return null;
       }
@@ -455,10 +435,9 @@ export const StepsPanelInner: React.FC<{ jobDetails: JobDetailsV2, depStepId?: s
             return;
          }
 
-         const g = jobDetails.groups[b.groupIdx];
-         const p = jobDetails.stream!.agentTypes[g.agentType];
+         const p = jobDetails.stream!.agentTypes[b.agentType];
 
-         if (!p || g.agentType !== agentType || p.pool !== agentPool) {
+         if (!p || b.agentType !== agentType || p.pool !== agentPool) {
             return;
          }
       }
@@ -491,9 +470,7 @@ export const StepsPanelInner: React.FC<{ jobDetails: JobDetailsV2, depStepId?: s
 
             filter = false;
 
-            const node = jobDetails.nodeByStepId(step.id);
-
-            if (label.includedNodes.indexOf(node?.name ?? "") !== -1) {
+            if (label.steps.indexOf(step?.id ?? "") !== -1) {
                filter = true;
             }
 
@@ -508,22 +485,20 @@ export const StepsPanelInner: React.FC<{ jobDetails: JobDetailsV2, depStepId?: s
          return;
       }
 
-      const group = jobDetails.groups[b.groupIdx];
-      const pool = jobDetails.stream?.agentTypes[group?.agentType!];
+      const pool = jobDetails.stream?.agentTypes[b?.agentType!];
 
       items.push({
          agentId: b.agentId,
          batch: b,
          agentRow: true,
-         agentType: group?.agentType.toUpperCase(),
+         agentType: b.agentType.toUpperCase(),
          agentPool: pool?.pool?.toUpperCase(),
       });
 
       steps.forEach(stepData => {
          const id = stepData.id;
          items.push({
-            step: jobDetails.stepById(id),
-            node: jobDetails.nodeByStepId(id),
+            step: jobDetails.stepById(id)
          });
       });
    });
@@ -533,17 +508,20 @@ export const StepsPanelInner: React.FC<{ jobDetails: JobDetailsV2, depStepId?: s
    }
 
    // get the current groups based on filtering
-   const groups: Set<number> = new Set();
+   const stepBatches: Set<string> = new Set();
    items.forEach(item => {
       const step = item.step;
       if (!step) {
          return;
       }
-      groups.add(jobDetails.getStepGroupIndex(step.id));
+
+      const b = jobDetails.jobData.batches.find(b => b.steps.findIndex(s => s.id === step.id) !== -1);
+      if (b)
+         stepBatches.add(b.id);
    });
 
    let batches = jobBatches.filter(b => {
-      if ((groups.size && !groups.has(b.groupIdx)) || b.steps.length || b.error === JobStepBatchError.None) {
+      if ((stepBatches.size && !stepBatches.has(b.id)) || b.steps.length || b.error === JobStepBatchError.None) {
          return false;
       }
       return true;
@@ -555,14 +533,13 @@ export const StepsPanelInner: React.FC<{ jobDetails: JobDetailsV2, depStepId?: s
          return;
       }
 
-      const group = jobDetails.groups[b.groupIdx];
-      const pool = jobDetails.stream?.agentTypes[group?.agentType!];
+      const pool = jobDetails.stream?.agentTypes[b?.agentType!];
 
       const nitem = {
          agentId: b.agentId,
          batch: b,
          agentRow: true,
-         agentType: group?.agentType.toUpperCase(),
+         agentType: b?.agentType.toUpperCase(),
          agentPool: pool?.pool?.toUpperCase(),
       };
 
@@ -692,8 +669,7 @@ export const getStepSummaryMarkdown = (jobDetails: JobDetailsV2, stepId: string)
          return undefined;
       }
 
-      const group = jobDetails.groups[batch!.groupIdx];
-      const agentType = group?.agentType;
+      const agentType = batch.agentType;
       const agentPool = jobDetails.stream?.agentTypes[agentType!]?.pool;
       return getBatchText({ batch: batch, agentType: agentType, agentPool: agentPool });
 
diff --git a/Engine/Source/Programs/Horde/HordeDashboard/src/components/jobDetailsV2/JobDetailViewV2.tsx b/Engine/Source/Programs/Horde/HordeDashboard/src/components/jobDetailsV2/JobDetailViewV2.tsx
--- a/Engine/Source/Programs/Horde/HordeDashboard/src/components/jobDetailsV2/JobDetailViewV2.tsx
+++ b/Engine/Source/Programs/Horde/HordeDashboard/src/components/jobDetailsV2/JobDetailViewV2.tsx
@@ -115,7 +115,7 @@ const JobBreadCrumbs: React.FC<{ jobDetails: JobDetailsV2 }> = observer(({ jobDe
    const label = jobDetails.labelByIndex(query.get("label") ?? undefined);
    if (label) {
       crumbItems.push({
-         text: label.category ? `${label.category}: ${label.name!}` : label.name!
+         text: label.stateResponse.dashboardCategory ? `${label.stateResponse.dashboardCategory}: ${label.stateResponse.dashboardName!}` : label.stateResponse.dashboardName!
       });
    }
 
diff --git a/Engine/Source/Programs/Horde/HordeDashboard/src/components/jobDetailsV2/JobDetailsViewCommon.tsx b/Engine/Source/Programs/Horde/HordeDashboard/src/components/jobDetailsV2/JobDetailsViewCommon.tsx
--- a/Engine/Source/Programs/Horde/HordeDashboard/src/components/jobDetailsV2/JobDetailsViewCommon.tsx
+++ b/Engine/Source/Programs/Horde/HordeDashboard/src/components/jobDetailsV2/JobDetailsViewCommon.tsx
@@ -5,7 +5,7 @@ import moment from "moment";
 import React, { useState } from "react";
 import { useLocation, useNavigate } from "react-router-dom";
 import backend from "../../backend";
-import { BatchData, GetArtifactResponseV2, GetJobTimingResponse, GetLabelResponse, GetTemplateRefResponse, GroupData, JobData, JobState, LabelState, NodeData, ReportPlacement, StepData, StreamData } from "../../backend/Api";
+import { BatchData, GetArtifactResponseV2, GetJobTimingResponse, GetLabelStateResponse, GetTemplateRefResponse, JobData, JobState, LabelState, ReportPlacement, StepData, StreamData } from "../../backend/Api";
 import { JobLabel } from "../../backend/JobDetails";
 import { PollBase } from "../../backend/PollBase";
 import { projectStore } from "../../backend/ProjectStore";
@@ -174,8 +174,6 @@ export class JobDetailsV2 extends PollBase {
       this.stream = undefined;
       this.template = undefined;
       this.labels = [];
-      this.groups = [];
-      this.nodes = [];
       this.batches = [];
       this.views = [];
       this.timing = undefined;
@@ -244,21 +242,6 @@ export class JobDetailsV2 extends PollBase {
 
    }
 
-   nodeByStepId(stepId: string | undefined): NodeData | undefined {
-
-      if (!stepId) {
-         return undefined;
-      }
-
-      const batch = this.batchByStepId(stepId);
-      const step = this.stepById(stepId);
-
-      if (!batch || !step) return undefined;
-
-      return this.groups[batch.groupIdx].nodes[step.nodeIdx];
-
-   }
-
    getSteps(): StepData[] {
       return this.batches.map(b => b.steps).flat();
    }
@@ -285,51 +268,43 @@ export class JobDetailsV2 extends PollBase {
       return step;
    }
 
-   getStepGroupIndex(stepId: string): number {
-
-      const node = this.nodeByStepId(stepId);
-      if (!node) {
-         return -1;
-      }
-
-      return this.groups.findIndex(g => !!g.nodes.find(n => n.name === node.name));
-   }
-
-
    getStepName(stepId: string | undefined, includeRetry: boolean = true): string {
 
       if (!stepId) {
          return "";
       }
 
-      const node = this.nodeByStepId(stepId);
-      if (!node) {
+      const step = this.stepById(stepId);
+
+      if (!step?.name) {
          return "";
       }
 
       const idx = this.getStepRetryNumber(stepId);
       if (!idx || !includeRetry) {
-         return node.name;
+         return step.name;
       }
 
-      return `${node.name} (${idx + 1})`;
+      return `${step.name} (${idx + 1})`;
 
    }
 
+
+
    stepRetries: Map<string, number> = new Map();
 
    getStepRetries(stepId: string): StepData[] {
 
-      const node = this.nodeByStepId(stepId);
-
-      if (!node) {
+      const step = this.stepById(stepId);
+      if (!step) {
          return [];
       }
 
+
       let steps: StepData[] = [];
 
       this.batches.forEach(b => {
-         b.steps.filter(s => this.nodeByStepId(s.id) === node).forEach(s => {
+         b.steps.filter(s => s.name === step.name).forEach(s => {
             steps.push(s);
          });
       })
@@ -400,26 +375,12 @@ export class JobDetailsV2 extends PollBase {
       }
 
       this.batches = jobData.batches ?? [];
-      this.groups = jobData.graphRef?.groups ?? [];
-      this.nodes = [];
-      if (this.groups) {
-         this.nodes = this.groups.map(g => g.nodes).flat();
-      }
-
-      let labels: GetLabelResponse[] = [];
-      if (jobData.graphRef?.labels && jobData.labels) {
-         labels = jobData.graphRef.labels;
-      }
 
+      let labels = jobData.labels ?? [];
       this.labels = labels.map((label, index) => {
          return {
-            category: label.category,
-            name: label.name,
-            includedNodes: label.includedNodes,
-            requiredNodes: label.requiredNodes,
             default: false,
-            stateResponse: jobData.labels![index],
-            internal: label,
+            stateResponse: label,
             timing: undefined
          }
       });
@@ -427,12 +388,9 @@ export class JobDetailsV2 extends PollBase {
       const defaultLabel = jobData.defaultLabel;
       if (defaultLabel) {
          this.labels.push({
-            category: "Other",
-            name: "Other",
-            includedNodes: defaultLabel.nodes,
-            requiredNodes: [],
             stateResponse: defaultLabel,
-            default: true
+            default: true,
+            timing: undefined
          })
       }
    }
@@ -454,7 +412,7 @@ export class JobDetailsV2 extends PollBase {
 
       const labels = this.labels;
 
-      if (!labels || (idx as number)>= labels.length) {
+      if (!labels || (idx as number) >= labels.length) {
          return undefined;
       }
 
@@ -475,9 +433,9 @@ export class JobDetailsV2 extends PollBase {
 
       const label = labels.find(label => {
          if (category) {
-            return category === label.category && name === label.name;
+            return category === label.stateResponse.dashboardCategory && name === label.stateResponse.dashboardName;
          }
-         return name === label.name;
+         return name === label.stateResponse.dashboardName;
       });
 
       return label;
@@ -499,7 +457,7 @@ export class JobDetailsV2 extends PollBase {
             return;
          }
 
-         step = b.steps.find(s => this.nodeByStepId(s.id)?.name === name);
+         step = b.steps.find(s => s.name === name);
 
       });
 
@@ -593,7 +551,7 @@ export class JobDetailsV2 extends PollBase {
 
       const initialRequest = !this.jobData;
 
-      requests.push(backend.getJob(this.jobId, undefined, true, true));
+      requests.push(backend.getJob(this.jobId, undefined, true));
 
       let results: any;
 
@@ -690,8 +648,6 @@ export class JobDetailsV2 extends PollBase {
    jobData?: JobData;
    stream?: StreamData;
    labels: JobLabel[] = [];
-   groups: GroupData[] = [];
-   nodes: NodeData[] = [];
    batches: BatchData[] = [];
    template?: GetTemplateRefResponse;
 
@@ -719,21 +675,13 @@ export class JobDetailsV2 extends PollBase {
 
       const steps: StepData[] = [];
 
-      const nodes: NodeData[] = [];
-
       const getStepsRecursive = (stepId: string) => {
 
-         const stepNode = this.nodeByStepId(stepId);
-
-         if (!stepNode || nodes.find(n => stepNode === n)) {
-            return;
-         }
-
-         nodes.push(stepNode);
+         const step = this.stepById(stepId);
 
-         [stepNode.inputDependencies, stepNode.orderDependencies].flat().forEach(name => {
-            const s = this.stepByName(name);
-            if (s) {
+         [step.inputDependencies, step.orderDependencies].flat().forEach(id => {
+            const s = this.stepById(id);
+            if (s && !steps.find(s => s.id === id)) {
                steps.push(s);
                getStepsRecursive(s.id);
             }
@@ -786,7 +734,7 @@ export type StateFilter = "All" | "Waiting" | "Ready" | "Skipped" | "Running" |
 
 // filter options for picker
 type FilterPickerItem = ITag & {
-   label?: GetLabelResponse;
+   label?: GetLabelStateResponse;
    step?: StepData;
    searchItem?: boolean;
    searchInput?: string;
@@ -800,7 +748,7 @@ export class JobDetailFilters {
       this.details = details;
    }
 
-   get label(): GetLabelResponse | undefined {
+   get label(): GetLabelStateResponse | undefined {
       return this._selected?.label;
    }
 
@@ -882,7 +830,7 @@ export class JobDetailFilters {
       this.setSelected(undefined, undefined, undefined);
    }
 
-   setSelected(step: StepData | undefined, label: GetLabelResponse | undefined, search?: string) {
+   setSelected(step: StepData | undefined, label: GetLabelStateResponse | undefined, search?: string) {
       /*
       console.log("setSelected Filter", step, label, search, this.filterItems.length);
       console.log(console.trace(""));
@@ -896,7 +844,7 @@ export class JobDetailFilters {
             return true;
          }
 
-         if (label && item.label?.category === label.category && item.label?.name === label.name) {
+         if (label && item.label?.dashboardCategory === label.dashboardCategory && item.label?.dashboardName === label.dashboardName) {
             return true;
          }
 
@@ -937,14 +885,14 @@ export class JobDetailFilters {
       };
 
       const categories: Set<string> = new Set();
-      labels.forEach(label => { if (label.name) { categories.add(label.category!); } });
+      labels.forEach(label => { if (label.stateResponse.dashboardName) { categories.add(label.stateResponse.dashboardCategory!); } });
 
 
       let labelItems = Array.from(categories.values()).map(c => {
          return {
             category: c,
-            labels: labels.filter(label => label.name && (label.category === c)).sort((a, b) => {
-               return a.name! < b.name! ? -1 : 1;
+            labels: labels.filter(label => label.stateResponse.dashboardName && (label.stateResponse.dashboardCategory === c)).sort((a, b) => {
+               return a.stateResponse.dashboardName! < b.stateResponse.dashboardName! ? -1 : 1;
             })
          } as LabelItem;
       }).filter(item => item.labels?.length).sort((a, b) => {
@@ -980,11 +928,11 @@ export class JobDetailFilters {
 
       labelItems.forEach(labelItem => {
          labelItem.labels.forEach(label => {
-            const count = labelDupes.get(label.name);
+            const count = labelDupes.get(label.stateResponse.dashboardName ?? "");
             if (!count) {
-               labelDupes.set(label.name, 1);
+               labelDupes.set(label.stateResponse.dashboardName ?? "", 1);
             } else {
-               labelDupes.set(label.name, count + 1);
+               labelDupes.set(label.stateResponse.dashboardName ?? "", count + 1);
             }
          });
       });
@@ -994,16 +942,16 @@ export class JobDetailFilters {
       labelItems.forEach(labelItem => {
          labelItem.labels.forEach(label => {
 
-            const count = labelDupes.get(label.name);
-            let name = `Label: ${label.name}`;
+            const count = labelDupes.get(label.stateResponse.dashboardName ?? "");
+            let name = `Label: ${label.stateResponse.dashboardName ?? ""}`;
             if (count! > 1) {
-               name = `Label: ${labelItem.category} - ${label.name}`
+               name = `Label: ${labelItem.category} - ${label.stateResponse.dashboardName ?? ""}`
             }
 
             filterItems.push({
-               key: `label_item_${labelItem.category}_${label.name}`,
+               key: `label_item_${labelItem.category}_${label.stateResponse.dashboardName ?? ""}`,
                name: name,
-               label: label
+               label: label.stateResponse
             })
          });
       });
@@ -1129,8 +1077,8 @@ export const JobFilterBar: React.FC<{ jobDetails: JobDetailsV2 }> = observer(({
       if (!isNaN(idx)) {
          const nlabel = jobDetails.labelByIndex(labelIdx);
          if (nlabel) {
-            if (!label || nlabel !== label) {
-               jobFilter.setSelected(undefined, nlabel);
+            if (!label || (nlabel?.stateResponse.dashboardCategory !== label.dashboardCategory || nlabel?.stateResponse.dashboardName !== label.dashboardName)) {
+               jobFilter.setSelected(undefined, nlabel.stateResponse);
             }
          }
       }
@@ -1222,7 +1170,7 @@ export const JobFilterBar: React.FC<{ jobDetails: JobDetailsV2 }> = observer(({
                let inner: JSX.Element | undefined;
                if (item.label) {
                   inner = <Stack horizontal verticalAlign="center" verticalFill={true}>
-                     <LabelStatusIcon label={item.label as JobLabel} />
+                     <LabelStatusIcon label={item.label} />
                      <Text>{item.name}</Text>
                   </Stack>
                }
@@ -1294,7 +1242,7 @@ export const JobFilterBar: React.FC<{ jobDetails: JobDetailsV2 }> = observer(({
 
                const label = (item as FilterPickerItem).label;
                if (label) {
-                  const idx = jobDetails.labelIndex(label.name, label.category);
+                  const idx = jobDetails.labelIndex(label.dashboardName ?? "", label.dashboardCategory ?? "");
                   if (idx >= 0) {
                      navigate(location.pathname + `?label=${idx}`);
                   } else {
diff --git a/Engine/Source/Programs/Horde/HordeDashboard/src/components/jobDetailsV2/JobOperationsBar.tsx b/Engine/Source/Programs/Horde/HordeDashboard/src/components/jobDetailsV2/JobOperationsBar.tsx
--- a/Engine/Source/Programs/Horde/HordeDashboard/src/components/jobDetailsV2/JobOperationsBar.tsx
+++ b/Engine/Source/Programs/Horde/HordeDashboard/src/components/jobDetailsV2/JobOperationsBar.tsx
@@ -360,9 +360,8 @@ const StepOperations: React.FC<{ jobDetails: JobDetailsV2, stepId: string }> = o
       return null;
    }
 
-   const node = jobDetails.nodeByStepId(stepId);
 
-   const canRunDisabled = !node?.allowRetry || !!step.retriedByUserInfo || !!jobData.abortedByUserInfo;
+   const canRunDisabled = !step?.allowRetry || !!step.retriedByUserInfo || !!jobData.abortedByUserInfo;
    const canTryFix = jobDetails.template?.allowPreflights;
    let canBisect = (step.outcome === JobStepOutcome.Failure || step.outcome === JobStepOutcome.Warnings) && !!jobDetails.template?.allowPreflights;
 
@@ -447,7 +446,7 @@ const StepOperations: React.FC<{ jobDetails: JobDetailsV2, stepId: string }> = o
 
       let args = "RunUAT.bat BuildGraph ";
 
-      args += `-Target="${node!.name}" `;
+      args += `-Target="${step!.name}" `;
 
       args += jobData!.arguments!.map(a => {
 
@@ -467,10 +466,10 @@ const StepOperations: React.FC<{ jobDetails: JobDetailsV2, stepId: string }> = o
    }
 
    return <Stack>
-      {!!shown.pauseShown && <PauseStepModal streamId={jobDetails.stream!.id} stepName={node!.name} templateName={jobDetails.template!.name} onClose={() => setShown({ pauseShown: false })} />}
+      {!!shown.pauseShown && <PauseStepModal streamId={jobDetails.stream!.id} stepName={step!.name} templateName={jobDetails.template!.name} onClose={() => setShown({ pauseShown: false })} />}
       {!!shown.retryShown && <StepRetryModal stepId={stepId} jobDetails={jobDetails} type={runType} show={true} onClose={() => { setShown({}); }} />}
       {!!shown.abortShown && <AbortJobModal stepId={stepId} jobDetails={jobDetails} show={true} onClose={() => { setShown({}); }} />}
-      {!!shown.bisectShown && <BisectionCreateModal jobId={jobId} nodeName={node?.name ?? "Unknown Node"} onClose={(response) => {
+      {!!shown.bisectShown && <BisectionCreateModal jobId={jobId} nodeName={step?.name ?? "Unknown Node"} onClose={(response) => {
          if (response) {
             jobDetails.bisectionUpdated();
          }
diff --git a/Engine/Source/Programs/Horde/HordeDashboard/src/components/jobDetailsV2/StepDetailSummary.tsx b/Engine/Source/Programs/Horde/HordeDashboard/src/components/jobDetailsV2/StepDetailSummary.tsx
--- a/Engine/Source/Programs/Horde/HordeDashboard/src/components/jobDetailsV2/StepDetailSummary.tsx
+++ b/Engine/Source/Programs/Horde/HordeDashboard/src/components/jobDetailsV2/StepDetailSummary.tsx
@@ -72,8 +72,7 @@ const getStepSummaryMarkdown = (jobDetails: JobDetailsV2, stepId: string): strin
 
       batchIncluded = true;
 
-      const group = jobDetails.groups[batch!.groupIdx];
-      const agentType = group?.agentType;
+      const agentType = batch?.agentType;
       const agentPool = jobDetails.stream?.agentTypes[agentType!]?.pool;
       return getBatchText({ batch: batch, agentType: agentType, agentPool: agentPool });
 
diff --git a/Engine/Source/Programs/Horde/HordeDashboard/src/components/jobDetailsV2/StepRetryModal.tsx b/Engine/Source/Programs/Horde/HordeDashboard/src/components/jobDetailsV2/StepRetryModal.tsx
--- a/Engine/Source/Programs/Horde/HordeDashboard/src/components/jobDetailsV2/StepRetryModal.tsx
+++ b/Engine/Source/Programs/Horde/HordeDashboard/src/components/jobDetailsV2/StepRetryModal.tsx
@@ -192,9 +192,9 @@ export const StepRetryModal: React.FC<{ stepId: string; jobDetails: JobDetailsV2
       return null;
    }
 
-   const node = jobDetails.nodeByStepId(stepId);
+   const step = jobDetails.stepById(stepId);
 
-   const headerText = type === StepRetryType.RunAgain ? `Run "${node?.name}" Again?` : "Preflight Step";
+   const headerText = type === StepRetryType.RunAgain ? `Run "${step?.name}" Again?` : "Preflight Step";
 
    let fixCL = 0;
 
@@ -235,11 +235,10 @@ export const StepRetryModal: React.FC<{ stepId: string; jobDetails: JobDetailsV2
       }
 
       const job = jobData;
-      const node = jobDetails.nodeByStepId(stepId);
 
       const args = [];
       args.push(`-Target=Setup Build`);
-      args.push(`-Target=${node!.name}`);
+      args.push(`-Target=${step!.name}`);
 
       job.arguments?.forEach(arg => {
 
diff --git a/Engine/Source/Programs/Horde/HordeDashboard/src/components/test/IssueStepReport.tsx b/Engine/Source/Programs/Horde/HordeDashboard/src/components/test/IssueStepReport.tsx
--- a/Engine/Source/Programs/Horde/HordeDashboard/src/components/test/IssueStepReport.tsx
+++ b/Engine/Source/Programs/Horde/HordeDashboard/src/components/test/IssueStepReport.tsx
@@ -64,14 +64,6 @@ class ToolHandler {
          return;
       }
       
-
-      if (!job.graphRef?.groups) {
-         console.error("Job has no groups");
-         this.loading = false;
-         this.setUpdated();
-         return;
-      }
-
       const batch = job.batches?.find(b => b.steps.find(s => s.id === stepId));
       if (!batch) {
          console.error("Batch not found for step id");
@@ -81,22 +73,7 @@ class ToolHandler {
       }
 
 
-      const group = job.graphRef?.groups[batch.groupIdx];
-      let stepIndex = -1;
-      batch.steps.forEach((s, index) => {
-         if (s.id === stepId) {
-            stepIndex = index;
-         }
-      })
-
-      if (stepIndex < 0) {
-         console.error("Bad step index");
-         this.loading = false;
-         this.setUpdated();
-         return;
-      }
-
-      const stepName = this.stepName = group.nodes[stepIndex].name;
+      const stepName = this.stepName = batch.steps.find(s => s.id === stepId)?.name;
 
       if (!stepName) {
          console.error("Bad step name");
