Horde: Reduce how long ephemeral and deleted agents are kept in database

Heavy use of auto-scaling can lead to excessive history of ephemeral agents being kept. In the case of AWS, each new instance has a unique agent ID. Many of the queries for agents (and in turn indices) are not optimized for a large collection.

diff --git a/Engine/Source/Programs/Horde/Horde.Server.Tests/Agents/AgentServiceTest.cs b/Engine/Source/Programs/Horde/Horde.Server.Tests/Agents/AgentServiceTest.cs
--- a/Engine/Source/Programs/Horde/Horde.Server.Tests/Agents/AgentServiceTest.cs
+++ b/Engine/Source/Programs/Horde/Horde.Server.Tests/Agents/AgentServiceTest.cs
@@ -191,7 +191,7 @@ public class AgentServiceTest : TestSetup
 		// Ephemeral agent is marked as deleted once its session is terminated
 		Assert.IsTrue((await AgentService.GetAgentAsync(agent.Id))!.Deleted);
 
-		await Clock.AdvanceAsync(TimeSpan.FromDays(8));
+		await Clock.AdvanceAsync(TimeSpan.FromHours(25));
 		await AgentService.TickAsync(CancellationToken.None);
 
 		// Once more time has passed, the ephemeral agent marked as deleted is removed from database
diff --git a/Engine/Source/Programs/Horde/Horde.Server/Agents/AgentService.cs b/Engine/Source/Programs/Horde/Horde.Server/Agents/AgentService.cs
--- a/Engine/Source/Programs/Horde/Horde.Server/Agents/AgentService.cs
+++ b/Engine/Source/Programs/Horde/Horde.Server/Agents/AgentService.cs
@@ -29,6 +29,7 @@ using HordeCommon.Rpc;
 using HordeCommon.Rpc.Messages;
 using Microsoft.Extensions.Hosting;
 using Microsoft.Extensions.Logging;
+using OpenTelemetry.Trace;
 using StackExchange.Redis;
 
 namespace Horde.Server.Agents
@@ -79,6 +80,7 @@ namespace Horde.Server.Agents
 		readonly IHostApplicationLifetime _applicationLifetime;
 		readonly IClock _clock;
 		readonly Meter _meter;
+		readonly Tracer _tracer;
 		readonly ILogger _logger;
 		readonly ITicker _ticker;
 
@@ -102,7 +104,7 @@ namespace Horde.Server.Agents
 		/// <summary>
 		/// Constructor
 		/// </summary>
-		public AgentService(IAgentCollection agents, ILeaseCollection leases, ISessionCollection sessions, AclService aclService, IDowntimeService downtimeService, IPoolCollection poolCollection, IEnumerable<ITaskSource> taskSources, RedisService redisService, IHostApplicationLifetime applicationLifetime, IClock clock, Meter meter, ILogger<AgentService> logger)
+		public AgentService(IAgentCollection agents, ILeaseCollection leases, ISessionCollection sessions, AclService aclService, IDowntimeService downtimeService, IPoolCollection poolCollection, IEnumerable<ITaskSource> taskSources, RedisService redisService, IHostApplicationLifetime applicationLifetime, IClock clock, Tracer tracer, Meter meter, ILogger<AgentService> logger)
 		{
 			Agents = agents;
 			_leases = leases;
@@ -117,6 +119,7 @@ namespace Horde.Server.Agents
 			_clock = clock;
 			_ticker = clock.AddTicker<AgentService>(TimeSpan.FromSeconds(30.0), TickAsync, logger);
 			_meter = meter;
+			_tracer = tracer;
 			_logger = logger;
 
 			_meter.CreateObservableGauge("horde.agent.count", () => _measurements);
@@ -1111,16 +1114,21 @@ namespace Horde.Server.Agents
 
 		private async Task DeleteExpiredEphemeralAgentsAsync(CancellationToken cancellationToken)
 		{
+			using TelemetrySpan span = _tracer.StartActiveSpan($"{nameof(AgentService)}.{nameof(DeleteExpiredEphemeralAgentsAsync)}");
+			int c = 0;
 			foreach (IAgent agent in await Agents.FindDeletedAsync(cancellationToken))
 			{
 				cancellationToken.ThrowIfCancellationRequested();
-				bool noStatusChangeDuringPeriod = _clock.UtcNow > agent.LastStatusChange + TimeSpan.FromDays(7);
-				if (agent.Status == AgentStatus.Stopped && agent.Ephemeral && noStatusChangeDuringPeriod)
+				bool noStatusChangeDuringPeriod = _clock.UtcNow > agent.LastStatusChange + TimeSpan.FromHours(24);
+				if (agent is { Status: AgentStatus.Stopped, Ephemeral: true } && noStatusChangeDuringPeriod)
 				{
 					_logger.LogDebug("Deleting ephemeral agent {Agent}", agent.Id);
 					await DeleteAgentAsync(agent, true, cancellationToken);
+					c++;
 				}
 			}
+
+			span.SetAttribute("NumAgentsDeleted", c);
 		}
 
 		private async Task CollectMetricsAsync(CancellationToken cancellationToken = default)
