Horde: Updating schema documentation

#rnx

diff --git a/Engine/Source/Programs/Horde/Docs/Config/Schema/Projects.md b/Engine/Source/Programs/Horde/Docs/Config/Schema/Projects.md
--- a/Engine/Source/Programs/Horde/Docs/Config/Schema/Projects.md
+++ b/Engine/Source/Programs/Horde/Docs/Config/Schema/Projects.md
@@ -173,7 +173,6 @@ Options for executing a job
 Name | Description
 ---- | -----------
 `executor` | `string`<br>Name of the executor to use
-`useNewTempStorage` | `boolean`<br>Whether to use the new temp storage backend
 `useWine` | `boolean`<br>Whether to execute using Wine emulation on Linux
 `runInSeparateProcess` | `boolean`<br>Executes the job lease in a separate process
 `workspaceMaterializer` | `string`<br>What workspace materializer to use in WorkspaceExecutor. Will override any value from workspace config.
diff --git a/Engine/Source/Programs/Horde/Docs/Config/Schema/Streams.md b/Engine/Source/Programs/Horde/Docs/Config/Schema/Streams.md
--- a/Engine/Source/Programs/Horde/Docs/Config/Schema/Streams.md
+++ b/Engine/Source/Programs/Horde/Docs/Config/Schema/Streams.md
@@ -59,7 +59,6 @@ Options for executing a job
 Name | Description
 ---- | -----------
 `executor` | `string`<br>Name of the executor to use
-`useNewTempStorage` | `boolean`<br>Whether to use the new temp storage backend
 `useWine` | `boolean`<br>Whether to execute using Wine emulation on Linux
 `runInSeparateProcess` | `boolean`<br>Executes the job lease in a separate process
 `workspaceMaterializer` | `string`<br>What workspace materializer to use in WorkspaceExecutor. Will override any value from workspace config.
