Horde: Add support for reserving a specific device by name
#rnx

diff --git a/Engine/Source/Programs/Horde/Horde.Server.Tests/Devices/DeviceServiceTests.cs b/Engine/Source/Programs/Horde/Horde.Server.Tests/Devices/DeviceServiceTests.cs
--- a/Engine/Source/Programs/Horde/Horde.Server.Tests/Devices/DeviceServiceTests.cs
+++ b/Engine/Source/Programs/Horde/Horde.Server.Tests/Devices/DeviceServiceTests.cs
@@ -245,7 +245,7 @@ namespace Horde.Server.Tests.Devices
 			return job;
 		}
 
-		static LegacyCreateReservationRequest SetupReservationTestAsync(IJob job, string poolId = "TestDevicePool1", string deviceType = "TestDevicePlatform1", JobStepId? stepId = null, string? modelId = null)
+		static LegacyCreateReservationRequest SetupReservationTestAsync(IJob job, string poolId = "TestDevicePool1", string deviceType = "TestDevicePlatform1", JobStepId? stepId = null, string? modelId = null, string? deviceName = null)
 		{
 			if (stepId == null)
 			{
@@ -265,6 +265,7 @@ namespace Horde.Server.Tests.Devices
 			request.Duration = "00:10:00";
 			request.JobId = job.Id.ToString();
 			request.StepId = stepId.ToString();
+			request.DeviceName = deviceName;
 
 			return request;
 		}
@@ -568,6 +569,26 @@ namespace Horde.Server.Tests.Devices
 			Assert.IsTrue(secondModel != firstModel && (secondModel == "TestDevicePlatform1_Model2" || secondModel == "TestDevicePlatform1_Model3"));
 		}
 
+		[TestMethod]
+		public async Task TestReservationDeviceNameAsync()
+		{
+			await SetupDevicesAsync();
+
+			IJob job = await SetupJobAsync();
+			LegacyCreateReservationRequest request = SetupReservationTestAsync(job, "TestDevicePool1", "TestDevicePlatform1", null, null, "testdevice1_platform1_testdevicepool1");
+
+			// create a reservation
+			GetLegacyReservationResponse reservation = ResultToValue(await DeviceController!.CreateDeviceReservationV1Async(request));
+			Assert.AreEqual(1, reservation.DeviceNames.Length);
+			Assert.AreEqual("testdevice1_platform1_testdevicepool1", reservation.DeviceNames[0], StringComparer.OrdinalIgnoreCase);
+			Assert.AreEqual("Test job", reservation.JobName);
+			Assert.AreEqual("abcd", reservation.StepId);
+
+			// get the device in the reservation, and make sure it is the right device
+			GetLegacyDeviceResponse device = ResultToValue(await DeviceController!.GetDeviceV1Async(reservation.DeviceNames[0]));
+			Assert.AreEqual(device.Name, "testdevice1_platform1_testdevicepool1", StringComparer.OrdinalIgnoreCase);
+		}
+
 		[TestMethod]
 		public async Task TestProblemDeviceAsync()
 		{
diff --git a/Engine/Source/Programs/Horde/Horde.Server/Devices/DeviceCollection.cs b/Engine/Source/Programs/Horde/Horde.Server/Devices/DeviceCollection.cs
--- a/Engine/Source/Programs/Horde/Horde.Server/Devices/DeviceCollection.cs
+++ b/Engine/Source/Programs/Horde/Horde.Server/Devices/DeviceCollection.cs
@@ -767,6 +767,11 @@ namespace Horde.Server.Devices
 						return false;
 					}
 
+					if (!String.IsNullOrEmpty(data.DeviceName) && !a.Name.Equals(data.DeviceName, StringComparison.OrdinalIgnoreCase))
+					{
+						return false;
+					}
+
 					return true;
 				});
 
diff --git a/Engine/Source/Programs/Horde/Horde.Server/Devices/DeviceMessages.cs b/Engine/Source/Programs/Horde/Horde.Server/Devices/DeviceMessages.cs
--- a/Engine/Source/Programs/Horde/Horde.Server/Devices/DeviceMessages.cs
+++ b/Engine/Source/Programs/Horde/Horde.Server/Devices/DeviceMessages.cs
@@ -851,6 +851,11 @@ namespace Horde.Server.Devices
 		/// </summary>
 		public string? StepId { get; set; } = null;
 
+		/// <summary>
+		/// A specific device to reserve
+		/// </summary>
+		public string? DeviceName { get; set; } = null;
+
 	}
 
 	/// <summary>
diff --git a/Engine/Source/Programs/Horde/Horde.Server/Devices/DevicesController.cs b/Engine/Source/Programs/Horde/Horde.Server/Devices/DevicesController.cs
--- a/Engine/Source/Programs/Horde/Horde.Server/Devices/DevicesController.cs
+++ b/Engine/Source/Programs/Horde/Horde.Server/Devices/DevicesController.cs
@@ -814,7 +814,7 @@ namespace Horde.Server.Devices
 					}
 				}
 
-				requestedDevices.Add(new DeviceRequestData(platform.Id, platformName, includeModels, excludeModels));
+				requestedDevices.Add(new DeviceRequestData(platform.Id, platformName, includeModels, excludeModels, request.DeviceName));
 			}
 
 			(IDeviceReservation? reservation, string? errorMessage, bool installRequired) = await _deviceService.TryCreateReservationAsync(poolIdValue, requestedDevices, request.Hostname, request.ReservationDetails, request.JobId == null ? null : JobId.Parse(request.JobId), request.StepId == null ? null : JobStepId.Parse(request.StepId));
diff --git a/Engine/Source/Programs/Horde/Horde.Server/Devices/IDeviceCollection.cs b/Engine/Source/Programs/Horde/Horde.Server/Devices/IDeviceCollection.cs
--- a/Engine/Source/Programs/Horde/Horde.Server/Devices/IDeviceCollection.cs
+++ b/Engine/Source/Programs/Horde/Horde.Server/Devices/IDeviceCollection.cs
@@ -36,15 +36,21 @@ namespace Horde.Server.Devices
 		/// </summary>
 		public List<string> ExcludeModels { get; set; }
 
+		/// <summary>
+		/// Specific device to reserve
+		/// </summary>
+		public string? DeviceName { get; set; } = null;
+
 		/// <summary>
 		/// Constructor
 		/// </summary>
-		public DeviceRequestData(DevicePlatformId platformId, string requestedPlatform, List<string>? includeModels = null, List<string>? excludeModels = null)
+		public DeviceRequestData(DevicePlatformId platformId, string requestedPlatform, List<string>? includeModels = null, List<string>? excludeModels = null, string? deviceName = null)
 		{
 			PlatformId = platformId;
 			IncludeModels = includeModels ?? new List<string>();
 			ExcludeModels = excludeModels ?? new List<string>();
 			RequestedPlatform = requestedPlatform;
+			DeviceName = deviceName;
 		}
 	}
 
