Horde: Add an index to the refs collection to prevent collection scans when finding refs to expire.

diff --git a/Engine/Source/Programs/Horde/Horde.Server/Storage/StorageService.cs b/Engine/Source/Programs/Horde/Horde.Server/Storage/StorageService.cs
--- a/Engine/Source/Programs/Horde/Horde.Server/Storage/StorageService.cs
+++ b/Engine/Source/Programs/Horde/Horde.Server/Storage/StorageService.cs
@@ -401,6 +401,7 @@ namespace Horde.Server.Storage
 			List<MongoIndex<RefInfo>> refIndexes = new List<MongoIndex<RefInfo>>();
 			refIndexes.Add(keys => keys.Ascending(x => x.NamespaceId).Ascending(x => x.Name), unique: true);
 			refIndexes.Add(keys => keys.Ascending(x => x.TargetBlobId));
+			refIndexes.Add(keys => keys.Descending(x => x.ExpiresAtUtc), sparse: true);
 			_refCollection = mongoService.GetCollection<RefInfo>("Storage.Refs", refIndexes);
 
 			_blobTicker = clock.AddSharedTicker("Storage:Blobs", TimeSpan.FromMinutes(5.0), TickBlobsAsync, _logger);
