Horde: Fix unimplemented methods on Job RPC controller.

#rnx

diff --git a/Engine/Source/Programs/Horde/Horde.Server/Jobs/JobRpcService.cs b/Engine/Source/Programs/Horde/Horde.Server/Jobs/JobRpcService.cs
--- a/Engine/Source/Programs/Horde/Horde.Server/Jobs/JobRpcService.cs
+++ b/Engine/Source/Programs/Horde/Horde.Server/Jobs/JobRpcService.cs
@@ -85,7 +85,7 @@ namespace Horde.Server.Jobs
 		/// <param name="request">Request arguments</param>
 		/// <param name="context">Context for the RPC call</param>
 		/// <returns>Information about the new agent</returns>
-		public Task<RpcGetStreamResponse> GetStreamAsync(RpcGetStreamRequest request, ServerCallContext context)
+		public override Task<RpcGetStreamResponse> GetStream(RpcGetStreamRequest request, ServerCallContext context)
 		{
 			StreamId streamIdValue = new StreamId(request.StreamId);
 
@@ -108,7 +108,7 @@ namespace Horde.Server.Jobs
 		/// <param name="request">Request arguments</param>
 		/// <param name="context">Context for the RPC call</param>
 		/// <returns>Information about the new agent</returns>
-		public async Task<RpcGetJobResponse> GetJobAsync(RpcGetJobRequest request, ServerCallContext context)
+		public override async Task<RpcGetJobResponse> GetJob(RpcGetJobRequest request, ServerCallContext context)
 		{
 			JobId jobIdValue = JobId.Parse(request.JobId);
 
@@ -131,7 +131,7 @@ namespace Horde.Server.Jobs
 		/// <param name="request">Request arguments</param>
 		/// <param name="context">Context for the RPC call</param>
 		/// <returns>Information about the new agent</returns>
-		public async Task<Empty> UpdateJobAsync(RpcUpdateJobRequest request, ServerCallContext context)
+		public override async Task<Empty> UpdateJob(RpcUpdateJobRequest request, ServerCallContext context)
 		{
 			JobId jobIdValue = JobId.Parse(request.JobId);
 
@@ -155,7 +155,7 @@ namespace Horde.Server.Jobs
 		/// <param name="request">Request arguments</param>
 		/// <param name="context">Context for the RPC call</param>
 		/// <returns>Information about the new agent</returns>
-		public async Task<RpcBeginBatchResponse> BeginBatchAsync(RpcBeginBatchRequest request, ServerCallContext context)
+		public override async Task<RpcBeginBatchResponse> BeginBatch(RpcBeginBatchRequest request, ServerCallContext context)
 		{
 			JobStepBatchId batchId = JobStepBatchId.Parse(request.BatchId);
 
@@ -208,7 +208,7 @@ namespace Horde.Server.Jobs
 		/// <param name="request">Request arguments</param>
 		/// <param name="context">Context for the RPC call</param>
 		/// <returns>Information about the new agent</returns>
-		public async Task<Empty> FinishBatchAsync(RpcFinishBatchRequest request, ServerCallContext context)
+		public override async Task<Empty> FinishBatch(RpcFinishBatchRequest request, ServerCallContext context)
 		{
 			IJob? job = await _jobService.GetJobAsync(JobId.Parse(request.JobId));
 			if (job == null)
@@ -231,7 +231,7 @@ namespace Horde.Server.Jobs
 		/// <param name="request">Request arguments</param>
 		/// <param name="context">Context for the RPC call</param>
 		/// <returns>Information about the new agent</returns>
-		public async Task<RpcBeginStepResponse> BeginStepAsync(RpcBeginStepRequest request, ServerCallContext context)
+		public override async Task<RpcBeginStepResponse> BeginStep(RpcBeginStepRequest request, ServerCallContext context)
 		{
 			Boxed<ILog?> log = new Boxed<ILog?>(null);
 			for (; ; )
@@ -542,7 +542,7 @@ namespace Horde.Server.Jobs
 		/// <param name="request">Request arguments</param>
 		/// <param name="context">Context for the RPC call</param>
 		/// <returns>Information about the new agent</returns>
-		public async Task<Empty> UpdateStepAsync(RpcUpdateStepRequest request, ServerCallContext context)
+		public override async Task<Empty> UpdateStep(RpcUpdateStepRequest request, ServerCallContext context)
 		{
 			IJob? job = await _jobService.GetJobAsync(JobId.Parse(request.JobId));
 			if (job == null)
@@ -578,7 +578,7 @@ namespace Horde.Server.Jobs
 		/// <param name="request">Request arguments</param>
 		/// <param name="context">Context for the RPC call</param>
 		/// <returns>Information about the step</returns>
-		public async Task<RpcGetStepResponse> GetStepAsync(RpcGetStepRequest request, ServerCallContext context)
+		public override async Task<RpcGetStepResponse> GetStep(RpcGetStepRequest request, ServerCallContext context)
 		{
 			IJob job = await GetJobAsync(JobId.Parse(request.JobId), context.CancellationToken);
 			IJobStepBatch batch = AuthorizeBatch(job, JobStepBatchId.Parse(request.BatchId), context);
@@ -598,7 +598,7 @@ namespace Horde.Server.Jobs
 		/// <param name="request">Request arguments</param>
 		/// <param name="context">Context for the RPC call</param>
 		/// <returns>Information about the new agent</returns>
-		public async Task<RpcUpdateGraphResponse> UpdateGraphAsync(RpcUpdateGraphRequest request, ServerCallContext context)
+		public override async Task<RpcUpdateGraphResponse> UpdateGraph(RpcUpdateGraphRequest request, ServerCallContext context)
 		{
 			JobId jobIdValue = JobId.Parse(request.JobId);
 			for (; ; )
@@ -687,7 +687,7 @@ namespace Horde.Server.Jobs
 		/// <param name="request">Request arguments</param>
 		/// <param name="context">Context for the RPC call</param>
 		/// <returns>Information about the new agent</returns>
-		public async Task<Empty> CreateEventsAsync(RpcCreateEventsRequest request, ServerCallContext context)
+		public override async Task<Empty> CreateEvents(RpcCreateEventsRequest request, ServerCallContext context)
 		{
 			if (!_globalConfig.Value.Authorize(LogAclAction.CreateEvent, context.GetHttpContext().User))
 			{
@@ -724,7 +724,7 @@ namespace Horde.Server.Jobs
 		/// <param name="reader">Request arguments</param>
 		/// <param name="context">Context for the RPC call</param>
 		/// <returns>Information about the new agent</returns>
-		public async Task<RpcUploadTestDataResponse> UploadTestDataAsync(IAsyncStreamReader<RpcUploadTestDataRequest> reader, ServerCallContext context)
+		public override async Task<RpcUploadTestDataResponse> UploadTestData(IAsyncStreamReader<RpcUploadTestDataRequest> reader, ServerCallContext context)
 		{
 			IJob? job = null;
 			IJobStep? jobStep = null;
@@ -888,7 +888,7 @@ namespace Horde.Server.Jobs
 		/// <param name="request"></param>
 		/// <param name="context"></param>
 		/// <returns></returns>
-		public async Task<RpcCreateReportResponse> CreateReportAsync(RpcCreateReportRequest request, ServerCallContext context)
+		public override async Task<RpcCreateReportResponse> CreateReport(RpcCreateReportRequest request, ServerCallContext context)
 		{
 			IJob job = await GetJobAsync(JobId.Parse(request.JobId), context.CancellationToken);
 			if (!_globalConfig.Value.TryGetStream(job.StreamId, out StreamConfig? streamConfig))
