Horde: Remove unnecessary using directives after moving methods into EpicGames.Core.

diff --git a/Engine/Source/Programs/Horde/Horde.Server/Agents/AgentService.cs b/Engine/Source/Programs/Horde/Horde.Server/Agents/AgentService.cs
--- a/Engine/Source/Programs/Horde/Horde.Server/Agents/AgentService.cs
+++ b/Engine/Source/Programs/Horde/Horde.Server/Agents/AgentService.cs
@@ -24,7 +24,6 @@ using Horde.Server.Auditing;
 using Horde.Server.Server;
 using Horde.Server.Tasks;
 using Horde.Server.Utilities;
-using HordeCommon;
 using HordeCommon.Rpc;
 using HordeCommon.Rpc.Messages;
 using Microsoft.Extensions.Hosting;
diff --git a/Engine/Source/Programs/Horde/Horde.Server/Agents/Fleet/FleetService.cs b/Engine/Source/Programs/Horde/Horde.Server/Agents/Fleet/FleetService.cs
--- a/Engine/Source/Programs/Horde/Horde.Server/Agents/Fleet/FleetService.cs
+++ b/Engine/Source/Programs/Horde/Horde.Server/Agents/Fleet/FleetService.cs
@@ -18,7 +18,6 @@ using Horde.Server.Jobs.Graphs;
 using Horde.Server.Server;
 using Horde.Server.Streams;
 using Horde.Server.Utilities;
-using HordeCommon;
 using Microsoft.Extensions.Caching.Memory;
 using Microsoft.Extensions.DependencyInjection;
 using Microsoft.Extensions.Hosting;
diff --git a/Engine/Source/Programs/Horde/Horde.Server/Agents/Pools/PoolUpdateService.cs b/Engine/Source/Programs/Horde/Horde.Server/Agents/Pools/PoolUpdateService.cs
--- a/Engine/Source/Programs/Horde/Horde.Server/Agents/Pools/PoolUpdateService.cs
+++ b/Engine/Source/Programs/Horde/Horde.Server/Agents/Pools/PoolUpdateService.cs
@@ -10,7 +10,6 @@ using EpicGames.Horde.Agents;
 using EpicGames.Horde.Agents.Pools;
 using Horde.Server.Server;
 using Horde.Server.Streams;
-using HordeCommon;
 using Microsoft.Extensions.Hosting;
 using Microsoft.Extensions.Logging;
 using Microsoft.Extensions.Options;
@@ -124,7 +123,7 @@ namespace Horde.Server.Agents.Pools
 			IEnumerable<PoolId> poolIds = agent.ExplicitPools.Concat(agent.DynamicPools);
 			return pools.Any(x => poolIds.Contains(x.Id) && x.EnableAutoscaling);
 		}
-		
+
 		/// <summary>
 		/// Find and conform any agents below the disk free conform threshold for a workspace
 		/// </summary>
@@ -135,7 +134,7 @@ namespace Horde.Server.Agents.Pools
 			using TelemetrySpan span = _tracer.StartActiveSpan($"{nameof(PoolUpdateService)}.{nameof(AutoConformAgentsAsync)}");
 			List<IAgent> agents = (await _agents.FindAsync(status: AgentStatus.Ok, enabled: true, cancellationToken: cancellationToken)).Where(x => !x.RequestShutdown).ToList();
 
-			long MegabytesToBytes (long v) => v * 1024 * 1024;
+			long MegabytesToBytes(long v) => v * 1024 * 1024;
 			foreach (IAgent agent in agents)
 			{
 				long? freeDiskSpace = agent.GetDiskFreeSpace();
diff --git a/Engine/Source/Programs/Horde/Horde.Server/Agents/Telemetry/AgentTelemetryCollection.cs b/Engine/Source/Programs/Horde/Horde.Server/Agents/Telemetry/AgentTelemetryCollection.cs
--- a/Engine/Source/Programs/Horde/Horde.Server/Agents/Telemetry/AgentTelemetryCollection.cs
+++ b/Engine/Source/Programs/Horde/Horde.Server/Agents/Telemetry/AgentTelemetryCollection.cs
@@ -9,7 +9,6 @@ using EpicGames.Core;
 using EpicGames.Horde.Agents;
 using Horde.Server.Server;
 using Horde.Server.Utilities;
-using HordeCommon;
 using Microsoft.Extensions.Hosting;
 using Microsoft.Extensions.Logging;
 using MongoDB.Bson;
diff --git a/Engine/Source/Programs/Horde/Horde.Server/Agents/Utilization/UtilizationDataService.cs b/Engine/Source/Programs/Horde/Horde.Server/Agents/Utilization/UtilizationDataService.cs
--- a/Engine/Source/Programs/Horde/Horde.Server/Agents/Utilization/UtilizationDataService.cs
+++ b/Engine/Source/Programs/Horde/Horde.Server/Agents/Utilization/UtilizationDataService.cs
@@ -11,7 +11,6 @@ using EpicGames.Horde.Agents.Pools;
 using Horde.Server.Agents.Fleet;
 using Horde.Server.Agents.Leases;
 using Horde.Server.Agents.Pools;
-using HordeCommon;
 using Microsoft.Extensions.Hosting;
 using Microsoft.Extensions.Logging;
 
diff --git a/Engine/Source/Programs/Horde/Horde.Server/Artifacts/ArtifactExpirationService.cs b/Engine/Source/Programs/Horde/Horde.Server/Artifacts/ArtifactExpirationService.cs
--- a/Engine/Source/Programs/Horde/Horde.Server/Artifacts/ArtifactExpirationService.cs
+++ b/Engine/Source/Programs/Horde/Horde.Server/Artifacts/ArtifactExpirationService.cs
@@ -11,7 +11,6 @@ using EpicGames.Horde.Storage;
 using EpicGames.Horde.Streams;
 using Horde.Server.Server;
 using Horde.Server.Storage;
-using HordeCommon;
 using Microsoft.Extensions.Hosting;
 using Microsoft.Extensions.Logging;
 using Microsoft.Extensions.Options;
diff --git a/Engine/Source/Programs/Horde/Horde.Server/Compute/ComputeService.cs b/Engine/Source/Programs/Horde/Horde.Server/Compute/ComputeService.cs
--- a/Engine/Source/Programs/Horde/Horde.Server/Compute/ComputeService.cs
+++ b/Engine/Source/Programs/Horde/Horde.Server/Compute/ComputeService.cs
@@ -27,7 +27,6 @@ using Horde.Server.Agents.Relay;
 using Horde.Server.Logs;
 using Horde.Server.Server;
 using Horde.Server.Utilities;
-using HordeCommon;
 using HordeCommon.Rpc.Tasks;
 using Microsoft.Extensions.Hosting;
 using Microsoft.Extensions.Logging;
diff --git a/Engine/Source/Programs/Horde/Horde.Server/Configuration/ConfigService.cs b/Engine/Source/Programs/Horde/Horde.Server/Configuration/ConfigService.cs
--- a/Engine/Source/Programs/Horde/Horde.Server/Configuration/ConfigService.cs
+++ b/Engine/Source/Programs/Horde/Horde.Server/Configuration/ConfigService.cs
@@ -16,7 +16,6 @@ using Horde.Server.Projects;
 using Horde.Server.Server;
 using Horde.Server.Streams;
 using Horde.Server.Users;
-using HordeCommon;
 using Microsoft.Extensions.Diagnostics.HealthChecks;
 using Microsoft.Extensions.Hosting;
 using Microsoft.Extensions.Logging;
diff --git a/Engine/Source/Programs/Horde/Horde.Server/Devices/DeviceReportService.cs b/Engine/Source/Programs/Horde/Horde.Server/Devices/DeviceReportService.cs
--- a/Engine/Source/Programs/Horde/Horde.Server/Devices/DeviceReportService.cs
+++ b/Engine/Source/Programs/Horde/Horde.Server/Devices/DeviceReportService.cs
@@ -10,7 +10,6 @@ using Horde.Server.Jobs;
 using Horde.Server.Notifications;
 using Horde.Server.Server;
 using Horde.Server.Utilities;
-using HordeCommon;
 using Microsoft.Extensions.Hosting;
 using Microsoft.Extensions.Logging;
 using Microsoft.Extensions.Options;
diff --git a/Engine/Source/Programs/Horde/Horde.Server/Devices/DeviceService.cs b/Engine/Source/Programs/Horde/Horde.Server/Devices/DeviceService.cs
--- a/Engine/Source/Programs/Horde/Horde.Server/Devices/DeviceService.cs
+++ b/Engine/Source/Programs/Horde/Horde.Server/Devices/DeviceService.cs
@@ -19,7 +19,6 @@ using Horde.Server.Projects;
 using Horde.Server.Server;
 using Horde.Server.Streams;
 using Horde.Server.Users;
-using HordeCommon;
 using Microsoft.Extensions.Hosting;
 using Microsoft.Extensions.Logging;
 using Microsoft.Extensions.Options;
diff --git a/Engine/Source/Programs/Horde/Horde.Server/Issues/External/JiraService.cs b/Engine/Source/Programs/Horde/Horde.Server/Issues/External/JiraService.cs
--- a/Engine/Source/Programs/Horde/Horde.Server/Issues/External/JiraService.cs
+++ b/Engine/Source/Programs/Horde/Horde.Server/Issues/External/JiraService.cs
@@ -14,7 +14,6 @@ using EpicGames.Core;
 using Horde.Server.Server;
 using Horde.Server.Streams;
 using Horde.Server.Users;
-using HordeCommon;
 using Microsoft.Extensions.Hosting;
 using Microsoft.Extensions.Logging;
 using Microsoft.Extensions.Options;
diff --git a/Engine/Source/Programs/Horde/Horde.Server/Issues/IssueService.cs b/Engine/Source/Programs/Horde/Horde.Server/Issues/IssueService.cs
--- a/Engine/Source/Programs/Horde/Horde.Server/Issues/IssueService.cs
+++ b/Engine/Source/Programs/Horde/Horde.Server/Issues/IssueService.cs
@@ -23,7 +23,6 @@ using Horde.Server.Logs;
 using Horde.Server.Server;
 using Horde.Server.Streams;
 using Horde.Server.Users;
-using HordeCommon;
 using Microsoft.Extensions.DependencyInjection;
 using Microsoft.Extensions.Hosting;
 using Microsoft.Extensions.Logging;
diff --git a/Engine/Source/Programs/Horde/Horde.Server/Issues/IssueTagService.cs b/Engine/Source/Programs/Horde/Horde.Server/Issues/IssueTagService.cs
--- a/Engine/Source/Programs/Horde/Horde.Server/Issues/IssueTagService.cs
+++ b/Engine/Source/Programs/Horde/Horde.Server/Issues/IssueTagService.cs
@@ -11,7 +11,6 @@ using Horde.Server.Commits;
 using Horde.Server.Server;
 using Horde.Server.Streams;
 using Horde.Server.Utilities;
-using HordeCommon;
 using Microsoft.Extensions.Hosting;
 using Microsoft.Extensions.Logging;
 using Microsoft.Extensions.Options;
diff --git a/Engine/Source/Programs/Horde/Horde.Server/Jobs/Bisect/BisectService.cs b/Engine/Source/Programs/Horde/Horde.Server/Jobs/Bisect/BisectService.cs
--- a/Engine/Source/Programs/Horde/Horde.Server/Jobs/Bisect/BisectService.cs
+++ b/Engine/Source/Programs/Horde/Horde.Server/Jobs/Bisect/BisectService.cs
@@ -12,7 +12,6 @@ using Horde.Server.Jobs.Graphs;
 using Horde.Server.Jobs.Templates;
 using Horde.Server.Server;
 using Horde.Server.Streams;
-using HordeCommon;
 using Microsoft.Extensions.Hosting;
 using Microsoft.Extensions.Logging;
 using Microsoft.Extensions.Options;
diff --git a/Engine/Source/Programs/Horde/Horde.Server/Jobs/JobExpirationService.cs b/Engine/Source/Programs/Horde/Horde.Server/Jobs/JobExpirationService.cs
--- a/Engine/Source/Programs/Horde/Horde.Server/Jobs/JobExpirationService.cs
+++ b/Engine/Source/Programs/Horde/Horde.Server/Jobs/JobExpirationService.cs
@@ -7,7 +7,6 @@ using System.Threading.Tasks;
 using EpicGames.Core;
 using Horde.Server.Server;
 using Horde.Server.Streams;
-using HordeCommon;
 using Microsoft.Extensions.Hosting;
 using Microsoft.Extensions.Logging;
 using Microsoft.Extensions.Options;
diff --git a/Engine/Source/Programs/Horde/Horde.Server/Jobs/TestData/TestDataService.cs b/Engine/Source/Programs/Horde/Horde.Server/Jobs/TestData/TestDataService.cs
--- a/Engine/Source/Programs/Horde/Horde.Server/Jobs/TestData/TestDataService.cs
+++ b/Engine/Source/Programs/Horde/Horde.Server/Jobs/TestData/TestDataService.cs
@@ -6,7 +6,6 @@ using System.Threading;
 using System.Threading.Tasks;
 using EpicGames.Core;
 using EpicGames.Horde.Streams;
-using HordeCommon;
 using Microsoft.Extensions.Hosting;
 using Microsoft.Extensions.Logging;
 using Microsoft.Extensions.Options;
diff --git a/Engine/Source/Programs/Horde/Horde.Server/Logs/LogTailService.cs b/Engine/Source/Programs/Horde/Horde.Server/Logs/LogTailService.cs
--- a/Engine/Source/Programs/Horde/Horde.Server/Logs/LogTailService.cs
+++ b/Engine/Source/Programs/Horde/Horde.Server/Logs/LogTailService.cs
@@ -10,7 +10,6 @@ using EpicGames.Core;
 using EpicGames.Horde.Logs;
 using EpicGames.Redis;
 using Horde.Server.Server;
-using HordeCommon;
 using Microsoft.Extensions.Hosting;
 using Microsoft.Extensions.Logging;
 using Microsoft.Extensions.Options;
diff --git a/Engine/Source/Programs/Horde/Horde.Server/Notifications/NotificationService.cs b/Engine/Source/Programs/Horde/Horde.Server/Notifications/NotificationService.cs
--- a/Engine/Source/Programs/Horde/Horde.Server/Notifications/NotificationService.cs
+++ b/Engine/Source/Programs/Horde/Horde.Server/Notifications/NotificationService.cs
@@ -24,7 +24,6 @@ using Horde.Server.Server;
 using Horde.Server.Streams;
 using Horde.Server.Users;
 using Horde.Server.Utilities;
-using HordeCommon;
 using Microsoft.Extensions.Caching.Memory;
 using Microsoft.Extensions.Hosting;
 using Microsoft.Extensions.Logging;
diff --git a/Engine/Source/Programs/Horde/Horde.Server/Notifications/Sinks/SlackNotificationSink.cs b/Engine/Source/Programs/Horde/Horde.Server/Notifications/Sinks/SlackNotificationSink.cs
--- a/Engine/Source/Programs/Horde/Horde.Server/Notifications/Sinks/SlackNotificationSink.cs
+++ b/Engine/Source/Programs/Horde/Horde.Server/Notifications/Sinks/SlackNotificationSink.cs
@@ -40,7 +40,6 @@ using Horde.Server.Logs;
 using Horde.Server.Server;
 using Horde.Server.Streams;
 using Horde.Server.Users;
-using HordeCommon;
 using Microsoft.AspNetCore.Hosting;
 using Microsoft.CodeAnalysis;
 using Microsoft.Extensions.Caching.Memory;
diff --git a/Engine/Source/Programs/Horde/Horde.Server/Perforce/PerforceServiceCache.cs b/Engine/Source/Programs/Horde/Horde.Server/Perforce/PerforceServiceCache.cs
--- a/Engine/Source/Programs/Horde/Horde.Server/Perforce/PerforceServiceCache.cs
+++ b/Engine/Source/Programs/Horde/Horde.Server/Perforce/PerforceServiceCache.cs
@@ -19,7 +19,6 @@ using Horde.Server.Server;
 using Horde.Server.Streams;
 using Horde.Server.Users;
 using Horde.Server.Utilities;
-using HordeCommon;
 using Microsoft.Extensions.Hosting;
 using Microsoft.Extensions.Logging;
 using Microsoft.Extensions.Options;
diff --git a/Engine/Source/Programs/Horde/Horde.Server/Replicators/ReplicationService.cs b/Engine/Source/Programs/Horde/Horde.Server/Replicators/ReplicationService.cs
--- a/Engine/Source/Programs/Horde/Horde.Server/Replicators/ReplicationService.cs
+++ b/Engine/Source/Programs/Horde/Horde.Server/Replicators/ReplicationService.cs
@@ -9,7 +9,6 @@ using EpicGames.Core;
 using EpicGames.Horde.Replicators;
 using Horde.Server.Server;
 using Horde.Server.Streams;
-using HordeCommon;
 using Microsoft.Extensions.Hosting;
 using Microsoft.Extensions.Logging;
 using Microsoft.Extensions.Options;
diff --git a/Engine/Source/Programs/Horde/Horde.Server/Server/ConsistencyService.cs b/Engine/Source/Programs/Horde/Horde.Server/Server/ConsistencyService.cs
--- a/Engine/Source/Programs/Horde/Horde.Server/Server/ConsistencyService.cs
+++ b/Engine/Source/Programs/Horde/Horde.Server/Server/ConsistencyService.cs
@@ -12,7 +12,6 @@ using EpicGames.Horde.Agents.Sessions;
 using Horde.Server.Agents;
 using Horde.Server.Agents.Leases;
 using Horde.Server.Agents.Sessions;
-using HordeCommon;
 using Microsoft.Extensions.Hosting;
 using Microsoft.Extensions.Logging;
 
diff --git a/Engine/Source/Programs/Horde/Horde.Server/Server/DowntimeService.cs b/Engine/Source/Programs/Horde/Horde.Server/Server/DowntimeService.cs
--- a/Engine/Source/Programs/Horde/Horde.Server/Server/DowntimeService.cs
+++ b/Engine/Source/Programs/Horde/Horde.Server/Server/DowntimeService.cs
@@ -5,7 +5,6 @@ using System.Linq;
 using System.Threading;
 using System.Threading.Tasks;
 using EpicGames.Core;
-using HordeCommon;
 using Microsoft.Extensions.Hosting;
 using Microsoft.Extensions.Logging;
 using Microsoft.Extensions.Options;
diff --git a/Engine/Source/Programs/Horde/Horde.Server/Server/LifetimeService.cs b/Engine/Source/Programs/Horde/Horde.Server/Server/LifetimeService.cs
--- a/Engine/Source/Programs/Horde/Horde.Server/Server/LifetimeService.cs
+++ b/Engine/Source/Programs/Horde/Horde.Server/Server/LifetimeService.cs
@@ -4,7 +4,6 @@ using System;
 using System.Threading;
 using System.Threading.Tasks;
 using EpicGames.Core;
-using HordeCommon;
 using Microsoft.Extensions.Hosting;
 using Microsoft.Extensions.Logging;
 using Microsoft.Extensions.Options;
diff --git a/Engine/Source/Programs/Horde/Horde.Server/Server/ServerStatusService.cs b/Engine/Source/Programs/Horde/Horde.Server/Server/ServerStatusService.cs
--- a/Engine/Source/Programs/Horde/Horde.Server/Server/ServerStatusService.cs
+++ b/Engine/Source/Programs/Horde/Horde.Server/Server/ServerStatusService.cs
@@ -7,7 +7,6 @@ using System.Text.Json;
 using System.Threading;
 using System.Threading.Tasks;
 using EpicGames.Core;
-using HordeCommon;
 using Microsoft.Extensions.Diagnostics.HealthChecks;
 using Microsoft.Extensions.Hosting;
 using Microsoft.Extensions.Logging;
@@ -187,7 +186,7 @@ public class ServerStatusService : IHostedService
 		}
 
 		string data = JsonSerializer.Serialize(status);
-		if(!_redis.ReadOnlyMode)
+		if (!_redis.ReadOnlyMode)
 		{
 			await redis.HashSetAsync(RedisHashKey(), id, data);
 		}
diff --git a/Engine/Source/Programs/Horde/Horde.Server/Storage/StorageService.cs b/Engine/Source/Programs/Horde/Horde.Server/Storage/StorageService.cs
--- a/Engine/Source/Programs/Horde/Horde.Server/Storage/StorageService.cs
+++ b/Engine/Source/Programs/Horde/Horde.Server/Storage/StorageService.cs
@@ -17,7 +17,6 @@ using EpicGames.Redis;
 using EpicGames.Redis.Utility;
 using Horde.Server.Server;
 using Horde.Server.Utilities;
-using HordeCommon;
 using Microsoft.Extensions.Caching.Memory;
 using Microsoft.Extensions.Configuration;
 using Microsoft.Extensions.DependencyInjection;
diff --git a/Engine/Source/Programs/Horde/Horde.Server/Tasks/ConformTaskSource.cs b/Engine/Source/Programs/Horde/Horde.Server/Tasks/ConformTaskSource.cs
--- a/Engine/Source/Programs/Horde/Horde.Server/Tasks/ConformTaskSource.cs
+++ b/Engine/Source/Programs/Horde/Horde.Server/Tasks/ConformTaskSource.cs
@@ -19,7 +19,6 @@ using Horde.Server.Logs;
 using Horde.Server.Perforce;
 using Horde.Server.Server;
 using Horde.Server.Utilities;
-using HordeCommon;
 using HordeCommon.Rpc.Messages;
 using HordeCommon.Rpc.Tasks;
 using Microsoft.Extensions.Hosting;
diff --git a/Engine/Source/Programs/Horde/Horde.Server/Telemetry/Sinks/MongoTelemetrySink.cs b/Engine/Source/Programs/Horde/Horde.Server/Telemetry/Sinks/MongoTelemetrySink.cs
--- a/Engine/Source/Programs/Horde/Horde.Server/Telemetry/Sinks/MongoTelemetrySink.cs
+++ b/Engine/Source/Programs/Horde/Horde.Server/Telemetry/Sinks/MongoTelemetrySink.cs
@@ -9,7 +9,6 @@ using EpicGames.Core;
 using EpicGames.Horde.Telemetry;
 using Horde.Server.Server;
 using Horde.Server.Utilities;
-using HordeCommon;
 using Microsoft.Extensions.Hosting;
 using Microsoft.Extensions.Logging;
 using Microsoft.Extensions.Options;
diff --git a/Engine/Source/Programs/Horde/Horde.Server/Telemetry/TelemetryManager.cs b/Engine/Source/Programs/Horde/Horde.Server/Telemetry/TelemetryManager.cs
--- a/Engine/Source/Programs/Horde/Horde.Server/Telemetry/TelemetryManager.cs
+++ b/Engine/Source/Programs/Horde/Horde.Server/Telemetry/TelemetryManager.cs
@@ -9,7 +9,6 @@ using System.Threading.Tasks;
 using EpicGames.Core;
 using EpicGames.Horde.Telemetry;
 using Horde.Server.Telemetry.Sinks;
-using HordeCommon;
 using Microsoft.Extensions.DependencyInjection;
 using Microsoft.Extensions.Hosting;
 using Microsoft.Extensions.Logging;
diff --git a/Engine/Source/Programs/Horde/Horde.Server/Utilities/Clock.cs b/Engine/Source/Programs/Horde/Horde.Server/Utilities/Clock.cs
--- a/Engine/Source/Programs/Horde/Horde.Server/Utilities/Clock.cs
+++ b/Engine/Source/Programs/Horde/Horde.Server/Utilities/Clock.cs
@@ -2,7 +2,6 @@
 
 using System;
 using System.Collections.Generic;
-using System.Diagnostics;
 using System.Threading;
 using System.Threading.Tasks;
 using EpicGames.Core;
@@ -177,7 +176,7 @@ namespace HordeCommon
 		public ITicker AddSharedTicker(string name, TimeSpan delay, Func<CancellationToken, ValueTask> tickAsync, ILogger logger)
 		{
 			RedisKey key = new RedisKey($"tick/{name}");
-			return ClockExtensions.AddTicker(this, name, delay / 4, token => TriggerSharedAsync(name, key, delay, tickAsync, token), logger);
+			return this.AddTicker(name, delay / 4, token => TriggerSharedAsync(name, key, delay, tickAsync, token), logger);
 		}
 
 		async ValueTask TriggerSharedAsync(string name, RedisKey key, TimeSpan interval, Func<CancellationToken, ValueTask> tickAsync, CancellationToken cancellationToken)
