Horde: Remove properties and agent version number from session responses.

diff --git a/Engine/Source/Programs/Horde/Horde.Server/Agents/AgentService.cs b/Engine/Source/Programs/Horde/Horde.Server/Agents/AgentService.cs
--- a/Engine/Source/Programs/Horde/Horde.Server/Agents/AgentService.cs
+++ b/Engine/Source/Programs/Horde/Horde.Server/Agents/AgentService.cs
@@ -394,7 +394,7 @@ namespace Horde.Server.Agents
 					}
 
 					// Create a new session document
-					ISession newSession = await _sessions.AddAsync(SessionIdUtils.GenerateNewId(), agent.Id, _clock.UtcNow, properties, resources, version, cancellationToken);
+					ISession newSession = await _sessions.AddAsync(SessionIdUtils.GenerateNewId(), agent.Id, _clock.UtcNow, cancellationToken);
 					DateTime sessionExpiresAt = utcNow + SessionExpiryTime;
 
 					// Get the new pools for the agent
@@ -852,7 +852,7 @@ namespace Horde.Server.Agents
 
 				// Update the session document
 				Agents.GetLogger(agent.Id).LogInformation("Terminated session {SessionId}", sessionId);
-				await _sessions.UpdateAsync(sessionId, finishTime, agent.Properties, agent.Resources, cancellationToken);
+				await _sessions.UpdateAsync(sessionId, finishTime, cancellationToken);
 				return agent;
 			}
 			return null;
diff --git a/Engine/Source/Programs/Horde/Horde.Server/Agents/AgentsController.cs b/Engine/Source/Programs/Horde/Horde.Server/Agents/AgentsController.cs
--- a/Engine/Source/Programs/Horde/Horde.Server/Agents/AgentsController.cs
+++ b/Engine/Source/Programs/Horde/Horde.Server/Agents/AgentsController.cs
@@ -331,7 +331,7 @@ namespace Horde.Server.Agents
 
 		static GetAgentSessionResponse CreateGetAgentSessionResponse(ISession session)
 		{
-			return new GetAgentSessionResponse(session.Id, session.StartTime, session.FinishTime, (session.Properties != null) ? new List<string>(session.Properties) : null, session.Version);
+			return new GetAgentSessionResponse(session.Id, session.StartTime, session.FinishTime);
 		}
 
 		/// <summary>
diff --git a/Engine/Source/Programs/Horde/Horde.Server/Agents/Sessions/ISession.cs b/Engine/Source/Programs/Horde/Horde.Server/Agents/Sessions/ISession.cs
--- a/Engine/Source/Programs/Horde/Horde.Server/Agents/Sessions/ISession.cs
+++ b/Engine/Source/Programs/Horde/Horde.Server/Agents/Sessions/ISession.cs
@@ -1,7 +1,6 @@
 // Copyright Epic Games, Inc. All Rights Reserved.
 
 using System;
-using System.Collections.Generic;
 using EpicGames.Horde.Agents;
 using EpicGames.Horde.Agents.Sessions;
 
@@ -31,15 +30,5 @@ namespace Horde.Server.Agents.Sessions
 		/// Finishing time for this session
 		/// </summary>
 		public DateTime? FinishTime { get; }
-
-		/// <summary>
-		/// Properties of this agent at the time the session started
-		/// </summary>
-		public IReadOnlyList<string>? Properties { get; }
-
-		/// <summary>
-		/// Version of the agent software
-		/// </summary>
-		public string? Version { get; }
 	}
 }
diff --git a/Engine/Source/Programs/Horde/Horde.Server/Agents/Sessions/ISessionCollection.cs b/Engine/Source/Programs/Horde/Horde.Server/Agents/Sessions/ISessionCollection.cs
--- a/Engine/Source/Programs/Horde/Horde.Server/Agents/Sessions/ISessionCollection.cs
+++ b/Engine/Source/Programs/Horde/Horde.Server/Agents/Sessions/ISessionCollection.cs
@@ -20,11 +20,8 @@ namespace Horde.Server.Agents.Sessions
 		/// <param name="id">The session id</param>
 		/// <param name="agentId">The agent this session is for</param>
 		/// <param name="startTime">Start time of this session</param>
-		/// <param name="properties">Properties of this agent at the time the session started</param>
-		/// <param name="resources">Resources which the agent has</param>
-		/// <param name="version">Version of the agent software</param>
 		/// <param name="cancellationToken">Cancellation token for the operation</param>
-		Task<ISession> AddAsync(SessionId id, AgentId agentId, DateTime startTime, IReadOnlyList<string>? properties, IReadOnlyDictionary<string, int>? resources, string? version, CancellationToken cancellationToken = default);
+		Task<ISession> AddAsync(SessionId id, AgentId agentId, DateTime startTime, CancellationToken cancellationToken = default);
 
 		/// <summary>
 		/// Gets information about a particular session
@@ -60,11 +57,9 @@ namespace Horde.Server.Agents.Sessions
 		/// </summary>
 		/// <param name="sessionId">The session to update</param>
 		/// <param name="finishTime">Time at which the session finished</param>
-		/// <param name="properties">The agent properties</param>
-		/// <param name="resources">Resources which the agent has</param>
 		/// <param name="cancellationToken">Cancellation token for the operation</param>
 		/// <returns>Async task</returns>
-		Task UpdateAsync(SessionId sessionId, DateTime? finishTime, IReadOnlyList<string>? properties, IReadOnlyDictionary<string, int>? resources, CancellationToken cancellationToken = default);
+		Task UpdateAsync(SessionId sessionId, DateTime? finishTime, CancellationToken cancellationToken = default);
 
 		/// <summary>
 		/// Delete a session from the collection
diff --git a/Engine/Source/Programs/Horde/Horde.Server/Agents/Sessions/SessionCollection.cs b/Engine/Source/Programs/Horde/Horde.Server/Agents/Sessions/SessionCollection.cs
--- a/Engine/Source/Programs/Horde/Horde.Server/Agents/Sessions/SessionCollection.cs
+++ b/Engine/Source/Programs/Horde/Horde.Server/Agents/Sessions/SessionCollection.cs
@@ -32,31 +32,17 @@ namespace Horde.Server.Agents.Sessions
 
 			public DateTime StartTime { get; set; }
 			public DateTime? FinishTime { get; set; }
-			public List<string>? Properties { get; set; }
-			public Dictionary<string, int>? Resources { get; set; }
-			public string? Version { get; set; }
-
-			IReadOnlyList<string>? ISession.Properties => Properties;
 
 			[BsonConstructor]
 			private SessionDocument()
 			{
 			}
 
-			public SessionDocument(SessionId id, AgentId agentId, DateTime startTime, IReadOnlyList<string>? properties, IReadOnlyDictionary<string, int>? resources, string? version)
+			public SessionDocument(SessionId id, AgentId agentId, DateTime startTime)
 			{
 				Id = id;
 				AgentId = agentId;
 				StartTime = startTime;
-				if (properties != null)
-				{
-					Properties = new List<string>(properties);
-				}
-				if (resources != null)
-				{
-					Resources = new Dictionary<string, int>(resources);
-				}
-				Version = version;
 			}
 		}
 
@@ -79,9 +65,9 @@ namespace Horde.Server.Agents.Sessions
 		}
 
 		/// <inheritdoc/>
-		public async Task<ISession> AddAsync(SessionId id, AgentId agentId, DateTime startTime, IReadOnlyList<string>? properties, IReadOnlyDictionary<string, int>? resources, string? version, CancellationToken cancellationToken = default)
+		public async Task<ISession> AddAsync(SessionId id, AgentId agentId, DateTime startTime, CancellationToken cancellationToken = default)
 		{
-			SessionDocument newSession = new SessionDocument(id, agentId, startTime, properties, resources, version);
+			SessionDocument newSession = new SessionDocument(id, agentId, startTime);
 			await _sessions.InsertOneAsync(newSession, null, cancellationToken);
 			return newSession;
 		}
@@ -119,21 +105,13 @@ namespace Horde.Server.Agents.Sessions
 		}
 
 		/// <inheritdoc/>
-		public Task UpdateAsync(SessionId sessionId, DateTime? finishTime, IReadOnlyList<string>? properties, IReadOnlyDictionary<string, int>? resources, CancellationToken cancellationToken = default)
+		public Task UpdateAsync(SessionId sessionId, DateTime? finishTime, CancellationToken cancellationToken = default)
 		{
 			List<UpdateDefinition<SessionDocument>> updates = new List<UpdateDefinition<SessionDocument>>();
 			if (finishTime != null)
 			{
 				updates.Add(Builders<SessionDocument>.Update.Set(x => x.FinishTime, finishTime));
 			}
-			if (properties != null)
-			{
-				updates.Add(Builders<SessionDocument>.Update.Set(x => x.Properties, new List<string>(properties)));
-			}
-			if (resources != null)
-			{
-				updates.Add(Builders<SessionDocument>.Update.Set(x => x.Resources, new Dictionary<string, int>(resources)));
-			}
 			return _sessions.FindOneAndUpdateAsync(x => x.Id == sessionId, Builders<SessionDocument>.Update.Combine(updates), cancellationToken: cancellationToken);
 		}
 
diff --git a/Engine/Source/Programs/Horde/Horde.Server/Server/ConsistencyService.cs b/Engine/Source/Programs/Horde/Horde.Server/Server/ConsistencyService.cs
--- a/Engine/Source/Programs/Horde/Horde.Server/Server/ConsistencyService.cs
+++ b/Engine/Source/Programs/Horde/Horde.Server/Server/ConsistencyService.cs
@@ -78,7 +78,7 @@ namespace Horde.Server.Server
 					if (agent == null || agent.SessionId != session.Id)
 					{
 						_logger.LogWarning("Forcing agent {AgentId} session {SessionId} to complete.", session.AgentId, session.Id);
-						await _sessionCollection.UpdateAsync(session.Id, utcNow, null, null, cancellationToken);
+						await _sessionCollection.UpdateAsync(session.Id, utcNow, cancellationToken);
 						sessionIdToInstance.Remove(session.Id);
 					}
 				}
diff --git a/Engine/Source/Programs/Horde/Horde.Server/Server/RpcService.cs b/Engine/Source/Programs/Horde/Horde.Server/Server/RpcService.cs
--- a/Engine/Source/Programs/Horde/Horde.Server/Server/RpcService.cs
+++ b/Engine/Source/Programs/Horde/Horde.Server/Server/RpcService.cs
@@ -505,17 +505,11 @@ namespace Horde.Server.Server
 		/// <param name="request">Request arguments</param>
 		/// <param name="context">Context for the RPC call</param>
 		/// <returns>An empty response</returns>
-		public override async Task<Empty> SendTelemetryEvents(RpcSendTelemetryEventsRequest request, ServerCallContext context)
+		public override Task<Empty> SendTelemetryEvents(RpcSendTelemetryEventsRequest request, ServerCallContext context)
 		{
-			ISession? session = null;
-
 			SessionId? sessionId = context.GetHttpContext().User.GetSessionClaim();
-			if (sessionId != null)
-			{
-				session = await _agentService.GetSessionAsync(sessionId.Value);
-			}
 
-			TelemetryRecordMeta agentMeta = new TelemetryRecordMeta("HordeAgent", session?.Version ?? "(Unknown)", ServerApp.DeploymentEnvironment, sessionId?.ToString() ?? "(Unknown)");
+			TelemetryRecordMeta agentMeta = new TelemetryRecordMeta("HordeAgent", "(Unknown)", ServerApp.DeploymentEnvironment, sessionId?.ToString() ?? "(Unknown)");
 			foreach (RpcWrappedTelemetryEvent wrappedEvent in request.Events)
 			{
 				OneofDescriptor oneofDescriptor = RpcWrappedTelemetryEvent.Descriptor.Oneofs[0];
@@ -525,7 +519,7 @@ namespace Horde.Server.Server
 				_telemetrySink.SendEvent(TelemetryStoreId.Default, agentMeta, wrappedValue);
 			}
 
-			return new Empty();
+			return Task.FromResult<Empty>(new Empty());
 		}
 	}
 }
diff --git a/Engine/Source/Programs/Shared/EpicGames.Horde/Agents/AgentMessages.cs b/Engine/Source/Programs/Shared/EpicGames.Horde/Agents/AgentMessages.cs
--- a/Engine/Source/Programs/Shared/EpicGames.Horde/Agents/AgentMessages.cs
+++ b/Engine/Source/Programs/Shared/EpicGames.Horde/Agents/AgentMessages.cs
@@ -81,9 +81,7 @@ namespace EpicGames.Horde.Agents
 	/// <param name="Id"> Unique id for this session </param>
 	/// <param name="StartTime"> Start time for this session </param>
 	/// <param name="FinishTime"> Finishing time for this session </param>
-	/// <param name="Properties"> Properties of this agent </param>
-	/// <param name="Version"> Version of the software running during this session </param>
-	public record GetAgentSessionResponse(SessionId Id, DateTime StartTime, DateTime? FinishTime, List<string>? Properties, string? Version);
+	public record GetAgentSessionResponse(SessionId Id, DateTime StartTime, DateTime? FinishTime);
 
 	/// <summary>
 	/// Information about a workspace synced on an agent
