Horde: Explicitly check that the expiry time is set when searching for refs to expire, so that DocDB will use an indexed query.

diff --git a/Engine/Source/Programs/Horde/Horde.Server/Storage/StorageService.cs b/Engine/Source/Programs/Horde/Horde.Server/Storage/StorageService.cs
--- a/Engine/Source/Programs/Horde/Horde.Server/Storage/StorageService.cs
+++ b/Engine/Source/Programs/Horde/Horde.Server/Storage/StorageService.cs
@@ -806,7 +806,9 @@ namespace Horde.Server.Storage
 		async ValueTask TickRefsAsync(CancellationToken cancellationToken)
 		{
 			DateTime utcNow = DateTime.UtcNow;
-			using (IAsyncCursor<RefInfo> cursor = await _refCollection.Find(x => x.ExpiresAtUtc < utcNow).ToCursorAsync(cancellationToken))
+
+			FilterDefinition<RefInfo> queryFilter = Builders<RefInfo>.Filter.Exists(x => x.ExpiresAtUtc) & Builders<RefInfo>.Filter.Lt(x => x.ExpiresAtUtc, utcNow);
+			using (IAsyncCursor<RefInfo> cursor = await _refCollection.Find(queryFilter).ToCursorAsync(cancellationToken))
 			{
 				List<DeleteOneModel<RefInfo>> requests = new List<DeleteOneModel<RefInfo>>();
 				while (await cursor.MoveNextAsync(cancellationToken))
