Makes the ThreadSanitizer LogEventMatcher and IssueHandler generic to Sanitizers in general and adds support for TSAN and ASAN. Updates the existing TSAN job to use the renamed commandlet and makes the job name clearer that a sanitizer report is being generated.

diff --git a/Engine/Build/EditorPerf.xml b/Engine/Build/EditorPerf.xml
--- a/Engine/Build/EditorPerf.xml
+++ b/Engine/Build/EditorPerf.xml
@@ -268,9 +268,9 @@
 	<Property Name="GenerateTSANReportNodeList" Value=""/>
 	<Macro Name="GenerateTSANReport" Arguments="NodeName;TestOutputDir;Phase;TestNameSuffix">
 		<Property Name="LogFilePath" Value="$(TestOutputDir)/$(Phase)$(TestNameSuffix)/Editor/EditorOutput.log"/>
-		<Property Name="GenerateTSANReportNodeName" Value="Generate $(NodeName) TSAN Report"/>
+		<Property Name="GenerateTSANReportNodeName" Value="Sanitizer Report - $(NodeName)"/>
 		<Node Name="$(GenerateTSANReportNodeName)">
-			<Command Name="LogThreadSanitizerOutput" Arguments="-logfile=&quot;$(LogFilePath)&quot;" />
+			<Command Name="LogSanitizerOutput" Arguments="-logfile=&quot;$(LogFilePath)&quot;" />
 		</Node>
 		<Property Name="GenerateTSANReportNodeList" Value="$(GenerateTSANReportNodeList); $(GenerateTSANReportNodeName)"/>
 	</Macro>
@@ -434,12 +434,14 @@
 							<Expand Name="GenerateTSANReport" NodeName="$(RunNodeLazyLoadZenName)" TestOutputDir="$(TestOutputDir)" Phase="$(Phase)" TestNameSuffix="ZenLoader" If="$(EnableLazyLoadZenLoaderTest) And $(WithTSAN)"/>
 						</Do>
 					</ForEach>
-					
-					<Property Name="RunNodeUploadPhaseTraces" Value="Upload $(OutputName) Traces"/>
-					<Node Name="$(RunNodeUploadPhaseTraces)">
-						<Expand Name="UploadTestTraces" Platform="$(Platform)" ProjectName="$(Project)"/>
-					</Node>
-					<Property Name="UploadNodeList" Value="$(RunNodeUploadPhaseTraces)"/>
+
+					<Do If="!$(WithTSAN)">
+						<Property Name="RunNodeUploadPhaseTraces" Value="Upload $(OutputName) Traces"/>
+						<Node Name="$(RunNodeUploadPhaseTraces)">
+							<Expand Name="UploadTestTraces" Platform="$(Platform)" ProjectName="$(Project)"/>
+						</Node>
+						<Property Name="UploadNodeList" Value="$(RunNodeUploadPhaseTraces)"/>
+					</Do>
 				</Agent>
 				<Label Category="$(RunCategory)" Name="$(Project) Editor Run" Requires="$(RunNodeRequires);$(RunNodeList);$(UploadNodeList)" UgsBadge="$(Project) Editor Run"/>
 				<Property Name="RunGraphNodes" Value="$(RunGraphNodes);$(RunNodeList);$(GenerateTSANReportNodeList);$(UploadNodeList)"/>
diff --git a/Engine/Source/Programs/AutomationTool.Tests/LogParserTests.cs b/Engine/Source/Programs/AutomationTool.Tests/LogParserTests.cs
--- a/Engine/Source/Programs/AutomationTool.Tests/LogParserTests.cs
+++ b/Engine/Source/Programs/AutomationTool.Tests/LogParserTests.cs
@@ -749,7 +749,7 @@ namespace AutomationTool.Tests
 
 			List<LogEvent> logEvents = Parse(lines);
 			Assert.AreEqual(20, logEvents[0].LineCount);
-			CheckEventGroup(logEvents, 0, 20, LogLevel.Error, KnownLogEvents.Sanitizer_Thread);
+			CheckEventGroup(logEvents, 0, 20, LogLevel.Warning, KnownLogEvents.Sanitizer_Thread);
 
 			Assert.AreEqual("/mnt/horde/++UE5/Sync/Engine/Source/Runtime/Core/Private/HAL/PThreadRunnableThread.h", logEvents[12].GetProperty("SourceFile").ToString());
 			Assert.AreEqual("90", logEvents[12].GetProperty("Line").ToString());
@@ -764,6 +764,77 @@ namespace AutomationTool.Tests
 			Assert.AreEqual("data race", logEvents[18].GetProperty("SummaryReason").ToString());
 		}
 
+		[TestMethod]
+		public void AddressSanitizerErrorMatcher()
+		{
+			string[] lines =
+			{
+				@"=================================================================",
+				@"==30562==ERROR: AddressSanitizer: heap-use-after-free on address 0x617002aa8418 at pc 0x7f98a08bd090 bp 0x7ffc30203af0 sp 0x7ffc30203ae8",
+				@"READ of size 8 at 0x617002aa8418 thread T0",
+				@"    #0 0x7f98a08bd08f in UObjectBase::GetFName() const /mnt/horde/FNR+Main+Inc/Sync/Engine/Source/Runtime/CoreUObject/Public/UObject/UObjectBase.h:166:10",
+				@"",
+				@"0x617002aa8418 is located 24 bytes inside of 712-byte region [0x617002aa8400,0x617002aa86c8)",
+				@"freed by thread T0 here:",
+				@"    #0 0x7f9a283e6b08 in __interceptor_free.part.5 /src/build/llvm-src/compiler-rt/lib/asan/asan_malloc_linux.cpp:52",
+				@"",
+				@"previously allocated by thread T0 here:",
+				@"    #0 0x7f9a283e730d in posix_memalign /src/build/llvm-src/compiler-rt/lib/asan/asan_malloc_linux.cpp:145",
+				@"",
+				@"SUMMARY: AddressSanitizer: heap-use-after-free /mnt/horde/FNR+Main+Inc/Sync/Engine/Source/Runtime/CoreUObject/Public/UObject/UObjectBase.h:166:10 in UObjectBase::GetFName() const",
+				@"Shadow bytes around the buggy address:",
+				@"  0x617002aa8180: 00 00 00 00 00 00 00 00 00 00 00 00 00 00 00 00",
+				@"  0x617002aa8200: 00 00 00 00 00 00 00 00 00 00 00 00 00 00 00 00",
+				@"  0x617002aa8280: 00 00 00 00 00 00 00 00 00 00 00 00 00 00 00 00",
+				@"  0x617002aa8300: 00 00 00 00 00 00 00 00 fa fa fa fa fa fa fa fa",
+				@"  0x617002aa8380: fa fa fa fa fa fa fa fa fa fa fa fa fa fa fa fa",
+				@"=>0x617002aa8400: fd fd fd[fd]fd fd fd fd fd fd fd fd fd fd fd fd",
+				@"  0x617002aa8480: fd fd fd fd fd fd fd fd fd fd fd fd fd fd fd fd",
+				@"  0x617002aa8500: fd fd fd fd fd fd fd fd fd fd fd fd fd fd fd fd",
+				@"  0x617002aa8580: fd fd fd fd fd fd fd fd fd fd fd fd fd fd fd fd",
+				@"  0x617002aa8600: fd fd fd fd fd fd fd fd fd fd fd fd fd fd fd fd",
+				@"  0x617002aa8680: fd fd fd fd fd fd fd fd fd fa fa fa fa fa fa fa",
+				@"Shadow byte legend (one shadow byte represents 8 application bytes):",
+				@"  Addressable:           00",
+				@"  Partially addressable: 01 02 03 04 05 06 07 ",
+				@"  Heap left redzone:       fa",
+				@"  Freed heap region:       fd",
+				@"  Stack left redzone:      f1",
+				@"  Stack mid redzone:       f2",
+				@"  Stack right redzone:     f3",
+				@"  Stack after return:      f5",
+				@"  Stack use after scope:   f8",
+				@"  Global redzone:          f9",
+				@"  Global init order:       f6",
+				@"  Poisoned by user:        f7",
+				@"  Container overflow:      fc",
+				@"  Array cookie:            ac",
+				@"  Intra object redzone:    bb",
+				@"  ASan internal:           fe",
+				@"  Left alloca redzone:     ca",
+				@"  Right alloca redzone:    cb",
+				@"==30562==ABORTING",
+			};
+
+			List<LogEvent> logEvents = Parse(lines);
+			Assert.AreEqual(45, logEvents[0].LineCount);
+			CheckEventGroup(logEvents, 0, 45, LogLevel.Error, KnownLogEvents.Sanitizer_Address);
+
+			LogValue? column;
+			Assert.AreEqual("/src/build/llvm-src/compiler-rt/lib/asan/asan_malloc_linux.cpp", logEvents[7].GetProperty("SourceFile").ToString());
+			Assert.AreEqual("52", logEvents[7].GetProperty("Line").ToString());
+			Assert.IsFalse(logEvents[7].TryGetProperty("Column", out column));
+			Assert.AreEqual("__interceptor_free.part.5", logEvents[7].GetProperty("Symbol").ToString());
+
+			// Summary
+			Assert.AreEqual("/mnt/horde/FNR+Main+Inc/Sync/Engine/Source/Runtime/CoreUObject/Public/UObject/UObjectBase.h", logEvents[12].GetProperty("SummarySourceFile").ToString());
+			Assert.AreEqual("166", logEvents[12].GetProperty("Line").ToString());
+			Assert.IsTrue(logEvents[12].TryGetProperty("Column", out column));
+			Assert.AreEqual("10", column.ToString());
+			Assert.AreEqual("UObjectBase::GetFName() const", logEvents[12].GetProperty("Symbol").ToString());
+			Assert.AreEqual("heap-use-after-free", logEvents[12].GetProperty("SummaryReason").ToString());
+		}
+
 		[TestMethod]
 		public void LocalizationChannelMatcher()
 		{
diff --git a/Engine/Source/Programs/AutomationTool/AutomationUtils/Matchers/SanitizerEventMatcher.cs b/Engine/Source/Programs/AutomationTool/AutomationUtils/Matchers/SanitizerEventMatcher.cs
new file mode 100644
--- /dev/null
+++ b/Engine/Source/Programs/AutomationTool/AutomationUtils/Matchers/SanitizerEventMatcher.cs
@@ -0,0 +1,141 @@
+// Copyright Epic Games, Inc. All Rights Reserved.
+
+using EpicGames.Core;
+using EpicGames.Horde.Issues.Handlers;
+using Microsoft.Extensions.Logging;
+using System.Text.RegularExpressions;
+
+#nullable enable
+
+namespace AutomationUtils.Matchers
+{
+	/// <summary>
+	/// Matcher for LLVM Sanitizer reports
+	/// 
+	/// </summary>
+	class SanitizerEventMatcher : ILogEventMatcher
+	{
+		// e.g.
+		// ASAN: =================================================================
+		// TSAN: ==================
+		static readonly Regex s_ReportStartPattern = new Regex(@"(^=================================================================$)|(^==================$)");
+
+		// e.g.
+		// ASAN: ==30562==ABORTING
+		// TSAN: ==================
+		static readonly Regex s_ReportEndPattern = new Regex(@"(^==\d+==(\s*\w+)?$)|(^==================$)");
+
+		// e.g.
+		// ASAN: ==30562==ERROR: AddressSanitizer: heap-use-after-free on address 0x617002aa8418 at pc 0x7f98a08bd090 bp 0x7ffc30203af0 sp 0x7ffc30203ae8
+		// TSAN: WARNING: ThreadSanitizer: data race (pid=6049)
+		static readonly Regex s_ReportLevelPattern = new Regex(@"^(==\d+==\s*)?(?<ReportLevel>WARNING|ERROR):\s+(?<SanitizerName>\w+)Sanitizer:\s*(.+)$");
+
+		// e.g. Note line number, column number and BuildId are optional
+		// ASAN: #0 0x7f98a08bd08f in UObjectBase::GetFName() const /mnt/somepath/file.h:90:10
+		// TSAN: #0 MyType<Type, OtherType>::Function(unsigned long*&) const /mnt/somepath/file.h:90:10 (BinaryName+0x2bc102e0) (BuildId: 8b17597a9f5444a0)
+		static readonly Regex s_StackTracePattern = new Regex(
+			@"^\s+\#[\d]+\s+(0x[0-9a-fA-F]+\s+in\s+)?" + // The optional capture group is to handle ASAN's slightly different stack trace syntax from TSAN
+			@"(?<Symbol>.+?)\s" +
+			@"(?<SourceFile>(\/|\w:).+?)" +
+			@"(:(?<Line>\d+?))?" +
+			@"(:(?<Column>\d+?))?" +
+			@"(\s+\(.+?\+(?<Address>0x[0-9a-fA-F]+?)\))?" +
+			@"(\s+?\(BuildId:\s+?(?<BuildId>.+?)\))?$");
+
+		// e.g.
+		// ASAN: SUMMARY: AddressSanitizer: heap-use-after-free /some/path.cpp:169:33 in UObjectBase::GetFName() const
+		// TSAN: SUMMARY: ThreadSanitizer: data race /some/path.cpp:169:33 in MyType::Foo()
+		static readonly Regex s_SummaryPattern = new Regex(
+			@"^SUMMARY:\s+(?<SanitizerName>\w+)Sanitizer:\s*" +
+			@"(?<SummaryReason>.+)\s" +
+			@"(?<SourceFile>(\/|\w:).+?)" +
+			@"(:(?<Line>\d+?))?" +
+			@"(:(?<Column>\d+?))?\s+?in\s+?" +
+			@"(?<Symbol>.+)$");
+
+		EventId ConvertSanitizerNameToEventId(string sanitizerName)
+		{
+			switch (sanitizerName)
+			{
+				case "Thread": return KnownLogEvents.Sanitizer_Thread;
+				case "Address": return KnownLogEvents.Sanitizer_Address;
+			}
+
+			return KnownLogEvents.Sanitizer;
+		}
+
+		LogLevel ConvertReportLevel(string reportLevel)
+		{
+			switch (reportLevel)
+			{
+				case "WARNING": return LogLevel.Warning;
+				case "ERROR": return LogLevel.Error;
+			}
+			return LogLevel.Error;
+		}
+
+		public LogEventMatch? Match(ILogCursor cursor)
+		{
+			Match? match;
+			if (cursor.TryMatch(s_ReportStartPattern, out match))
+			{
+				EventId sanitizerId = KnownLogEvents.Sanitizer;
+				LogLevel reportLevel = LogLevel.Information;
+
+				LogEventBuilder builder = new LogEventBuilder(cursor);
+				builder.MoveNext();
+
+				while (builder.Current.CurrentLine != null)
+				{
+
+					if (builder.Current.TryMatch(s_StackTracePattern, out match))
+					{
+						do
+						{
+							builder.AnnotateSymbol(match!.Groups["Symbol"]);
+							builder.AnnotateSourceFile(match.Groups["SourceFile"], "");
+							builder.TryAnnotate(match.Groups["Line"], LogEventMarkup.LineNumber);
+							builder.TryAnnotate(match.Groups["Column"], LogEventMarkup.ColumnNumber);
+							builder.TryAnnotate(match.Groups["Address"], "");
+							builder.TryAnnotate(match.Groups["BuildId"], "");
+
+							builder.MoveNext();
+						} while (builder.Current.TryMatch(s_StackTracePattern, out match));
+					}
+					else
+					{
+						if (builder.Current.TryMatch(s_ReportLevelPattern, out match))
+						{
+							sanitizerId = ConvertSanitizerNameToEventId(match.Groups["SanitizerName"].Value);
+							reportLevel = ConvertReportLevel(match.Groups["ReportLevel"].Value);
+
+							// Used by IssueHandler
+							builder.Annotate(match.Groups["SanitizerName"], SanitizerIssueHandler.SanitizerName);
+						}
+
+						if (builder.Current.TryMatch(s_SummaryPattern, out match))
+						{
+							// Used by IssueHandler
+							builder.Annotate(match.Groups["SummaryReason"], SanitizerIssueHandler.SummaryReason);
+							builder.Annotate("SummarySourceFile", match.Groups["SourceFile"], SanitizerIssueHandler.SummarySourceFile);
+
+							// Annotate the source file again so that we may get UGS link resolution
+							builder.AnnotateSourceFile(match.Groups["SourceFile"], "");
+							builder.TryAnnotate(match.Groups["Line"], LogEventMarkup.LineNumber);
+							builder.TryAnnotate(match.Groups["Column"], LogEventMarkup.ColumnNumber);
+							builder.AnnotateSymbol(match.Groups["Symbol"]);
+						}
+						else if (builder.Current.IsMatch(s_ReportEndPattern))
+						{
+							return builder.ToMatch(LogEventPriority.AboveNormal, reportLevel, sanitizerId);
+						}
+
+						builder.MoveNext();
+					}
+				}
+			}
+
+			return null;
+		}
+	}
+}
diff --git a/Engine/Source/Programs/AutomationTool/AutomationUtils/Matchers/ThreadSanitizerEventMatcher.cs b/Engine/Source/Programs/AutomationTool/AutomationUtils/Matchers/ThreadSanitizerEventMatcher.cs
deleted file mode 100644
diff --git a/Engine/Source/Programs/Horde/Horde.Server.Tests/Issues/IssueServiceTests.cs b/Engine/Source/Programs/Horde/Horde.Server.Tests/Issues/IssueServiceTests.cs
--- a/Engine/Source/Programs/Horde/Horde.Server.Tests/Issues/IssueServiceTests.cs
+++ b/Engine/Source/Programs/Horde/Horde.Server.Tests/Issues/IssueServiceTests.cs
@@ -700,7 +700,7 @@ namespace Horde.Server.Tests.Issues
 				@"=================="
 			};
 
-			IJob job = CreateJob(_mainStreamId, 120, "Generate SomeJobs TSAN Report", _graph);
+			IJob job = CreateJob(_mainStreamId, 120, "Sanitizer Report for 'SomeJob'", _graph);
 			await ParseEventsAsync(job, 0, 0, lines);
 			await UpdateCompleteStepAsync(job, 0, 0, JobStepOutcome.Failure);
 
@@ -751,7 +751,7 @@ namespace Horde.Server.Tests.Issues
 				@"==================",
 			};
 
-			IJob job = CreateJob(_mainStreamId, 120, "Generate SomeJobs TSAN Report", _graph);
+			IJob job = CreateJob(_mainStreamId, 120, "Sanitizer Report for 'SomeJob'", _graph);
 			await ParseEventsAsync(job, 0, 0, lines);
 			await UpdateCompleteStepAsync(job, 0, 0, JobStepOutcome.Failure);
 
@@ -771,6 +771,102 @@ namespace Horde.Server.Tests.Issues
 			Assert.IsTrue(issue2.Fingerprints[0].Type.Contains("Paths.cpp", StringComparison.OrdinalIgnoreCase)); 
 		}
 
+		[TestMethod]
+		public async Task AddressSanitizerDeduplicateSimilarIssuesTestAsync()
+		{
+			// We have three unique errors reported, but we key issues on the file since file lines change over time without fixes, so only two unique issues should be found
+			string[] lines =
+			{
+				@"=================================================================",
+				@"==30562==ERROR: AddressSanitizer: heap-use-after-free on address 0x617002aa8418 at pc 0x7f98a08bd090 bp 0x7ffc30203af0 sp 0x7ffc30203ae8",
+				@"READ of size 8 at 0x617002aa8418 thread T0",
+				@"    #0 0x7f98a08bd08f in UObjectBase::GetFName() const /mnt/horde/FNR+Main+Inc/Sync/Engine/Source/Runtime/CoreUObject/Public/UObject/UObjectBase.h:166:10",
+				@"",
+				@"0x617002aa8418 is located 24 bytes inside of 712-byte region [0x617002aa8400,0x617002aa86c8)",
+				@"freed by thread T0 here:",
+				@"    #0 0x7f9a283e6b08 in __interceptor_free.part.5 /src/build/llvm-src/compiler-rt/lib/asan/asan_malloc_linux.cpp:52",
+				@"",
+				@"previously allocated by thread T0 here:",
+				@"    #0 0x7f9a283e730d in posix_memalign /src/build/llvm-src/compiler-rt/lib/asan/asan_malloc_linux.cpp:145",
+				@"",
+				@"SUMMARY: AddressSanitizer: heap-use-after-free /mnt/horde/FNR+Main+Inc/Sync/Engine/Source/Runtime/CoreUObject/Public/UObject/UObjectBase.h:166:10 in UObjectBase::GetFName() const",
+				@"==30562==ABORTING",
+				@"=================================================================",
+				@"==30562==ERROR: AddressSanitizer: heap-use-after-free on address 0x617002aa8418 at pc 0x7f98a08bd090 bp 0x7ffc30203af0 sp 0x7ffc30203ae8",
+				@"SUMMARY: AddressSanitizer: heap-use-after-free /mnt/horde/FNR+Main+Inc/Sync/Engine/Source/Runtime/CoreUObject/Public/UObject/UObjectBase.h:166:10 in UObjectBase::GetFName() const",
+				@"==30562==ABORTING",
+				@"=================================================================",
+				@"==30562==ERROR: AddressSanitizer: heap-use-after-free on address 0x617002aa8418 at pc 0x7f98a08bd090 bp 0x7ffc30203af0 sp 0x7ffc30203ae8",
+				@"SUMMARY: AddressSanitizer: heap-use-after-free /mnt/horde/FNR+Main+Inc/Sync/Engine/Source/Runtime/CoreUObject/Public/UObject/SomeOtherFile.h:166:10 in UObjectBase::GetFName() const",
+				@"==30562==ABORTING"
+			};
+
+			IJob job = CreateJob(_mainStreamId, 120, "Sanitizer Report for 'SomeJob'", _graph);
+			await ParseEventsAsync(job, 0, 0, lines);
+			await UpdateCompleteStepAsync(job, 0, 0, JobStepOutcome.Failure);
+
+			IReadOnlyList<IIssue> issues = await IssueCollection.FindIssuesAsync();
+			Assert.AreEqual(2, issues.Count);
+
+			IIssue issue1 = issues[0];
+			Assert.AreEqual(issue1.Fingerprints.Count, 1);
+			Assert.IsTrue(issue1.Fingerprints[0].Type.Contains("AddressSanitizer", StringComparison.OrdinalIgnoreCase));
+			Assert.IsTrue(issue1.Fingerprints[0].Type.Contains("heap-use-after-free", StringComparison.OrdinalIgnoreCase));
+			Assert.IsTrue(issue1.Fingerprints[0].Type.Contains("SomeOtherFile.h", StringComparison.OrdinalIgnoreCase));
+
+			IIssue issue2 = issues[1];
+			Assert.AreEqual(issue2.Fingerprints.Count, 1);
+			Assert.IsTrue(issue2.Fingerprints[0].Type.Contains("AddressSanitizer", StringComparison.OrdinalIgnoreCase));
+			Assert.IsTrue(issue2.Fingerprints[0].Type.Contains("heap-use-after-free", StringComparison.OrdinalIgnoreCase));
+			Assert.IsTrue(issue2.Fingerprints[0].Type.Contains("UObjectBase.h", StringComparison.OrdinalIgnoreCase));
+		}
+
+		[TestMethod]
+		public async Task AddressSanitizerEnsureUniqueIssuesForSameFileTestAsync()
+		{
+			// We have three unique errors reported, but we key issues on the file since file lines change over time without fixes, so only two unique issues should be found
+			string[] lines =
+			{
+				@"=================================================================",
+				@"==30562==ERROR: AddressSanitizer: heap-use-after-free on address 0x617002aa8418 at pc 0x7f98a08bd090 bp 0x7ffc30203af0 sp 0x7ffc30203ae8",
+				@"READ of size 8 at 0x617002aa8418 thread T0",
+				@"    #0 0x7f98a08bd08f in UObjectBase::GetFName() const /mnt/horde/FNR+Main+Inc/Sync/Engine/Source/Runtime/CoreUObject/Public/UObject/UObjectBase.h:166:10",
+				@"",
+				@"0x617002aa8418 is located 24 bytes inside of 712-byte region [0x617002aa8400,0x617002aa86c8)",
+				@"freed by thread T0 here:",
+				@"    #0 0x7f9a283e6b08 in __interceptor_free.part.5 /src/build/llvm-src/compiler-rt/lib/asan/asan_malloc_linux.cpp:52",
+				@"",
+				@"previously allocated by thread T0 here:",
+				@"    #0 0x7f9a283e730d in posix_memalign /src/build/llvm-src/compiler-rt/lib/asan/asan_malloc_linux.cpp:145",
+				@"",
+				@"SUMMARY: AddressSanitizer: heap-use-after-free /mnt/horde/FNR+Main+Inc/Sync/Engine/Source/Runtime/CoreUObject/Public/UObject/UObjectBase.h:166:10 in UObjectBase::GetFName() const",
+				@"==30562==ABORTING",
+				@"=================================================================",
+				@"==30562==ERROR: AddressSanitizer: stack-buffer-overflow on address 0x617002aa8418 at pc 0x7f98a08bd090 bp 0x7ffc30203af0 sp 0x7ffc30203ae8",
+				@"SUMMARY: AddressSanitizer: stack-buffer-overflow /mnt/horde/FNR+Main+Inc/Sync/Engine/Source/Runtime/CoreUObject/Public/UObject/UObjectBase.h:166:10 in UObjectBase::GetFName() const",
+				@"==30562==ABORTING",
+			};
+
+			IJob job = CreateJob(_mainStreamId, 120, "Sanitizer Report for 'SomeJob'", _graph);
+			await ParseEventsAsync(job, 0, 0, lines);
+			await UpdateCompleteStepAsync(job, 0, 0, JobStepOutcome.Failure);
+
+			IReadOnlyList<IIssue> issues = await IssueCollection.FindIssuesAsync();
+			Assert.AreEqual(2, issues.Count);
+
+			IIssue issue1 = issues[0];
+			Assert.AreEqual(issue1.Fingerprints.Count, 1);
+			Assert.IsTrue(issue1.Fingerprints[0].Type.Contains("AddressSanitizer", StringComparison.OrdinalIgnoreCase));
+			Assert.IsTrue(issue1.Fingerprints[0].Type.Contains("stack-buffer-overflow", StringComparison.OrdinalIgnoreCase));
+			Assert.IsTrue(issue1.Fingerprints[0].Type.Contains("UObjectBase.h", StringComparison.OrdinalIgnoreCase));
+
+			IIssue issue2 = issues[1];
+			Assert.AreEqual(issue2.Fingerprints.Count, 1);
+			Assert.IsTrue(issue2.Fingerprints[0].Type.Contains("AddressSanitizer", StringComparison.OrdinalIgnoreCase));
+			Assert.IsTrue(issue2.Fingerprints[0].Type.Contains("heap-use-after-free", StringComparison.OrdinalIgnoreCase));
+			Assert.IsTrue(issue2.Fingerprints[0].Type.Contains("UObjectBase.h", StringComparison.OrdinalIgnoreCase));
+		}
+
 		[TestMethod]
 		public async Task AutoSdkWarningTestAsync()
 		{
diff --git a/Engine/Source/Programs/Shared/EpicGames.Core/KnownLogEvents.cs b/Engine/Source/Programs/Shared/EpicGames.Core/KnownLogEvents.cs
--- a/Engine/Source/Programs/Shared/EpicGames.Core/KnownLogEvents.cs
+++ b/Engine/Source/Programs/Shared/EpicGames.Core/KnownLogEvents.cs
@@ -341,6 +341,11 @@ namespace EpicGames.Core
 		/// </summary>
 		public static EventId Sanitizer_Thread{ get; } = new EventId(801);
 
+		/// <summary>
+		/// Address Sanitizer Error
+		/// </summary>
+		public static EventId Sanitizer_Address { get; } = new EventId(802);
+
 		/// <summary>
 		/// Horde error codes
 		/// </summary>
diff --git a/Engine/Source/Programs/Shared/EpicGames.Horde/Issues/Handlers/SanitizerIssueHandler.cs b/Engine/Source/Programs/Shared/EpicGames.Horde/Issues/Handlers/SanitizerIssueHandler.cs
new file mode 100644
--- /dev/null
+++ b/Engine/Source/Programs/Shared/EpicGames.Horde/Issues/Handlers/SanitizerIssueHandler.cs
@@ -0,0 +1,84 @@
+// Copyright Epic Games, Inc. All Rights Reserved.
+
+using System.Collections.Generic;
+using System.Linq;
+using System.Text.Json;
+using EpicGames.Core;
+using Microsoft.Extensions.Logging;
+
+namespace EpicGames.Horde.Issues.Handlers
+{
+	/// <summary>
+	/// Instance of a specific Thread Sanitizer error
+	/// </summary>
+	[IssueHandler(Priority = 10)]
+	public class SanitizerIssueHandler : IssueHandler
+	{
+		static readonly Utf8String s_summaryReasonAnnotation = new Utf8String("SummaryReason");
+		static readonly Utf8String s_summarySourceFileAnnotation = new Utf8String("SummarySourceFile");
+		static readonly Utf8String s_sanitizerNameAnnotation = new Utf8String("SanitizerName");
+		readonly List<IssueEventGroup> _issues = new List<IssueEventGroup>();
+
+		/// <summary>
+		/// Log value describing the thread sanitizer error summary reason
+		/// </summary>
+		public static LogValue SanitizerName => new LogValue(s_sanitizerNameAnnotation, "");
+
+		/// <summary>
+		/// Log value describing the thread sanitizer error summary reason
+		/// </summary>
+		public static LogValue SummaryReason => new LogValue(s_summaryReasonAnnotation, "");
+
+		/// <summary>
+		/// Log value describing the thread sanitizer error summary source file
+		/// </summary>
+		public static LogValue SummarySourceFile => new LogValue(s_summarySourceFileAnnotation, "");
+
+		static bool FindAnnotation(IssueEvent issueEvent, Utf8String annotation, ref string result)
+		{
+			foreach (JsonProperty property in issueEvent.Lines.SelectMany(x => x.FindPropertiesOfType(annotation)))
+			{
+				JsonElement value;
+				if (property.Value.TryGetProperty(LogEventPropertyName.Text.Span, out value) && value.ValueKind == JsonValueKind.String)
+				{
+					result = value.GetString() ?? "";
+					return true;
+				}
+			}
+			return false;
+		}
+
+		/// <inheritdoc/>
+		public override bool HandleEvent(IssueEvent issueEvent)
+		{
+			EventId? eventId = issueEvent.EventId;
+			if (eventId == KnownLogEvents.Sanitizer_Thread || eventId == KnownLogEvents.Sanitizer_Address)
+			{
+				string summaryReason = "error";
+				string summaryFile = "";
+				string sanitizerName = "";
+
+				FindAnnotation(issueEvent, s_summaryReasonAnnotation, ref summaryReason);
+				FindAnnotation(issueEvent, s_summarySourceFileAnnotation, ref summaryFile);
+				FindAnnotation(issueEvent, s_sanitizerNameAnnotation, ref sanitizerName);
+				string fingerprint = $"{sanitizerName}Sanitizer:{summaryReason}:{summaryFile}";
+
+				IssueEventGroup issue = new IssueEventGroup(fingerprint, $"{{Meta:{s_sanitizerNameAnnotation}}}Sanitizer found '{{Meta:{s_summaryReasonAnnotation}}}' in {{Meta:{s_summarySourceFileAnnotation}}}", IssueChangeFilter.Code);
+				issue.Keys.Add(sanitizerName, IssueKeyType.None);
+				issue.Metadata.Add(s_sanitizerNameAnnotation.ToString(), sanitizerName);
+				issue.Keys.Add(summaryReason, IssueKeyType.Note);
+				issue.Metadata.Add(s_summaryReasonAnnotation.ToString(), summaryReason);
+				issue.Keys.AddSourceFile(summaryFile, IssueKeyType.File);
+				issue.Metadata.Add(s_summarySourceFileAnnotation.ToString(), summaryFile);
+				issue.Events.Add(issueEvent);
+
+				_issues.Add(issue);
+				return true;
+			}
+			return false;
+		}
+
+		/// <inheritdoc/>
+		public override IEnumerable<IssueEventGroup> GetIssues() => _issues;
+	}
+}
diff --git a/Engine/Source/Programs/Shared/EpicGames.Horde/Issues/Handlers/ThreadSanitizerIssueHandler.cs b/Engine/Source/Programs/Shared/EpicGames.Horde/Issues/Handlers/ThreadSanitizerIssueHandler.cs
deleted file mode 100644
