Horde: Add spans to monitor time spent ticking blobs and refs.

#rnx

diff --git a/Engine/Source/Programs/Horde/Horde.Server/Storage/StorageService.cs b/Engine/Source/Programs/Horde/Horde.Server/Storage/StorageService.cs
--- a/Engine/Source/Programs/Horde/Horde.Server/Storage/StorageService.cs
+++ b/Engine/Source/Programs/Horde/Horde.Server/Storage/StorageService.cs
@@ -598,6 +598,8 @@ namespace Horde.Server.Storage
 		/// <param name="cancellationToken">Cancellation token for the operation</param>
 		async ValueTask TickBlobsAsync(CancellationToken cancellationToken)
 		{
+			using TelemetrySpan span = _tracer.StartActiveSpan($"{nameof(StorageService)}.{nameof(TickBlobsAsync)}");
+
 			GcState gcState = await _gcState.GetAsync(cancellationToken);
 			DateTime utcNow = _clock.UtcNow;
 
@@ -807,6 +809,8 @@ namespace Horde.Server.Storage
 		{
 			DateTime utcNow = DateTime.UtcNow;
 
+			using TelemetrySpan span = _tracer.StartActiveSpan($"{nameof(StorageService)}.{nameof(TickRefsAsync)}");
+
 			FilterDefinition<RefInfo> queryFilter = Builders<RefInfo>.Filter.Exists(x => x.ExpiresAtUtc) & Builders<RefInfo>.Filter.Lt(x => x.ExpiresAtUtc, utcNow);
 			using (IAsyncCursor<RefInfo> cursor = await _refCollection.Find(queryFilter).ToCursorAsync(cancellationToken))
 			{
