Horde: Update release notes.

#rnx

diff --git a/Engine/Source/Programs/Horde/Docs/ReleaseNotes.md b/Engine/Source/Programs/Horde/Docs/ReleaseNotes.md
--- a/Engine/Source/Programs/Horde/Docs/ReleaseNotes.md
+++ b/Engine/Source/Programs/Horde/Docs/ReleaseNotes.md
@@ -1,5 +1,16 @@
 # Release Notes
 
+## 2024-05-01
+
+* Explicitly check that the expiry time is set when searching for refs to expire, so that DocDB will use an indexed query. (33367777)
+* Add CheckConnectionAsync to HordeHttpClient. (33353008)
+* Add an index to the refs collection to prevent collection scans when finding refs to expire. (33338629)
+* Reduce how long ephemeral and deleted agents are kept in database. Heavy use of auto-scaling can lead to excessive history of ephemeral agents being kept. In the case of AWS, each new instance has a unique agent ID. Many of the queries for agents (and in turn indices) are not optimized for a large collection. (33317167)
+* Disable file modified errors in job executor. (33299280)
+* Remove legacy job methods from horde_rpc interface. (33298879)
+* Remove support for legacy artifact uploads. (33298589)
+* Enable agents sending CPU/RAM usage metrics to the server. (33290135, 33289584)
+
 ## 2024-04-25
 
 * Allow setting project-wide workspace types. Will automatically be inserted to the workspace type list of each stream belonging to the project. This also allows defining base types which can be inherited from. (33213109)
