Horde: Add more context to mongodb exceptions when caching commits.

#rnx

diff --git a/Engine/Source/Programs/Horde/Horde.Server/Perforce/PerforceServiceCache.cs b/Engine/Source/Programs/Horde/Horde.Server/Perforce/PerforceServiceCache.cs
--- a/Engine/Source/Programs/Horde/Horde.Server/Perforce/PerforceServiceCache.cs
+++ b/Engine/Source/Programs/Horde/Horde.Server/Perforce/PerforceServiceCache.cs
@@ -576,8 +576,16 @@ namespace Horde.Server.Perforce
 
 		async Task AddCachedCommitAsync(CachedCommitDoc commitDoc, CancellationToken cancellationToken)
 		{
-			FilterDefinition<CachedCommitDoc> filter = Builders<CachedCommitDoc>.Filter.Expr(x => x.StreamId == commitDoc.StreamId && x.Number == commitDoc.Number);
-			await _commits.ReplaceOneAsync(filter, commitDoc, new ReplaceOptions { IsUpsert = true }, cancellationToken);
+			try
+			{
+				FilterDefinition<CachedCommitDoc> filter = Builders<CachedCommitDoc>.Filter.Expr(x => x.StreamId == commitDoc.StreamId && x.Number == commitDoc.Number);
+				await _commits.ReplaceOneAsync(filter, commitDoc, new ReplaceOptions { IsUpsert = true }, cancellationToken);
+			}
+			catch (Exception ex)
+			{
+				_logger.LogError(ex, "Unable to add cached commit for {StreamId} CL {Number}: {Message}", commitDoc.StreamId, commitDoc.Number, ex.Message);
+				throw;
+			}
 		}
 
 #pragma warning disable CA1308 // Normalize strings to uppercase
