Horde: Prevent jobs terminating with an error if temp storage artifacts have been modified.

#rnx

diff --git a/Engine/Source/Programs/Horde/Horde.Agent/Execution/JobExecutor.cs b/Engine/Source/Programs/Horde/Horde.Agent/Execution/JobExecutor.cs
--- a/Engine/Source/Programs/Horde/Horde.Agent/Execution/JobExecutor.cs
+++ b/Engine/Source/Programs/Horde/Horde.Agent/Execution/JobExecutor.cs
@@ -960,7 +960,6 @@ namespace Horde.Agent.Execution
 				}
 
 				logger.LogInformation("Build product(s) from a previous step have been modified:\n{FileList}", modifiedFileList);
-				return false;
 			}
 
 			// Determine all the output files which are required to be copied to temp storage (because they're referenced by nodes in another agent)
