Horde: Adjust agent store polling
#rnx

diff --git a/Engine/Source/Programs/Horde/HordeDashboard/package.json b/Engine/Source/Programs/Horde/HordeDashboard/package.json
--- a/Engine/Source/Programs/Horde/HordeDashboard/package.json
+++ b/Engine/Source/Programs/Horde/HordeDashboard/package.json
@@ -1,7 +1,7 @@
 {
   "name": "horde-dashboard",
   "private": true,
-  "version": "0.0.1",
+  "version": "0.0.2",
   "type": "module",  
   "scripts": {
     "dev": "vite",
diff --git a/Engine/Source/Programs/Horde/HordeDashboard/src/backend/AgentStore.ts b/Engine/Source/Programs/Horde/HordeDashboard/src/backend/AgentStore.ts
--- a/Engine/Source/Programs/Horde/HordeDashboard/src/backend/AgentStore.ts
+++ b/Engine/Source/Programs/Horde/HordeDashboard/src/backend/AgentStore.ts
@@ -35,6 +35,8 @@ export class AgentStore {
     // don't need this for now
     modifiedAfterDate: Date = new Date(0);
 
+    inflight?: boolean;
+
     @action
     private _setPools(data: PoolData[]) {
         this._pools = data;
@@ -117,6 +119,13 @@ export class AgentStore {
     async update(slim = false): Promise<void> {
         return new Promise<void>((resolve, reject) => {
 
+            if (this.inflight !== undefined && slim) {                
+                resolve();
+                return;
+            }
+
+            this.inflight = true;
+
             const filter = "id,name,sessionId,sessionExpiresAt,online,enabled,ephemeral,comment,version,forceVersion,pools,capabilities,leases,acl,updateTime,deleted,pendingConform,pendingFullConform,conformAttemptCount,lastConformTime,nextConformTime,lastShutdownReason,pendingShutdown,status";
             const promises: any[] = [];
             promises.push(backend.getAgents({ includeDeleted: false, modifiedAfter: this.modifiedAfterDate?.toISOString(), filter:filter    }));
@@ -132,6 +141,8 @@ export class AgentStore {
                 resolve();
             }).catch(reason => {
                 console.error(`Error getting agents and pools: ${reason}`);
+            }).finally(() => {
+                this.inflight = undefined;
             });
         });
     }
diff --git a/Engine/Source/Programs/Horde/HordeDashboard/src/components/AgentView.tsx b/Engine/Source/Programs/Horde/HordeDashboard/src/components/AgentView.tsx
--- a/Engine/Source/Programs/Horde/HordeDashboard/src/components/AgentView.tsx
+++ b/Engine/Source/Programs/Horde/HordeDashboard/src/components/AgentView.tsx
@@ -1950,7 +1950,7 @@ export const AgentViewInner: React.FC<{ agentId?: string, poolId?: string, searc
    useEffect(() => {
       const interval = setInterval(() => {
          agentStore.update(true)
-      }, 5000);
+      }, 15000);
       return () => clearInterval(interval);
    }, []);
 
