Horde: Add trace scopes around agent service tick methods.

#rnx

diff --git a/Engine/Source/Programs/Horde/Horde.Server/Agents/AgentService.cs b/Engine/Source/Programs/Horde/Horde.Server/Agents/AgentService.cs
--- a/Engine/Source/Programs/Horde/Horde.Server/Agents/AgentService.cs
+++ b/Engine/Source/Programs/Horde/Horde.Server/Agents/AgentService.cs
@@ -1081,6 +1081,8 @@ namespace Horde.Server.Agents
 		/// <returns>Async task</returns>
 		internal async ValueTask TickAsync(CancellationToken stoppingToken)
 		{
+			using TelemetrySpan span = _tracer.StartActiveSpan($"{nameof(AgentService)}.{nameof(TickAsync)}");
+
 			await TerminateExpiredSessionsAsync(stoppingToken);
 			await DeleteExpiredEphemeralAgentsAsync(stoppingToken);
 			await CollectMetricsAsync(stoppingToken);
@@ -1088,6 +1090,8 @@ namespace Horde.Server.Agents
 
 		private async Task TerminateExpiredSessionsAsync(CancellationToken cancellationToken)
 		{
+			using TelemetrySpan span = _tracer.StartActiveSpan($"{nameof(AgentService)}.{nameof(TerminateExpiredSessionsAsync)}");
+
 			while (!cancellationToken.IsCancellationRequested)
 			{
 				// Find all the agents which are ready to be expired
@@ -1132,6 +1136,8 @@ namespace Horde.Server.Agents
 
 		private async Task CollectMetricsAsync(CancellationToken cancellationToken = default)
 		{
+			using TelemetrySpan span = _tracer.StartActiveSpan($"{nameof(AgentService)}.{nameof(CollectMetricsAsync)}");
+
 			IReadOnlyList<IAgent> agentList = await Agents.FindAsync(cancellationToken: cancellationToken);
 			int numAgentsTotal = agentList.Count;
 			int numAgentsTotalEnabled = agentList.Count(a => a.Enabled);
