Horde: Add job compound indexes over (StreamId, TemplateId, CL number) and (StreamId, CreateTimeUtc). Remove rarely used indexes.

#rnx

diff --git a/Engine/Source/Programs/Horde/Horde.Server/Jobs/JobCollection.cs b/Engine/Source/Programs/Horde/Horde.Server/Jobs/JobCollection.cs
--- a/Engine/Source/Programs/Horde/Horde.Server/Jobs/JobCollection.cs
+++ b/Engine/Source/Programs/Horde/Horde.Server/Jobs/JobCollection.cs
@@ -708,7 +708,6 @@ namespace Horde.Server.Jobs
 
 		readonly IMongoCollection<JobDocument> _jobs;
 		readonly MongoIndex<JobDocument> _createTimeIndex;
-		readonly MongoIndex<JobDocument> _updateTimeIndex;
 		readonly MongoIndex<JobDocument> _streamThenTemplateThenCreationTimeIndex;
 		readonly MongoIndex<JobDocument> _startedByBisectTaskIdIndex;
 		readonly ITelemetrySink _telemetrySink;
@@ -733,15 +732,13 @@ namespace Horde.Server.Jobs
 			_logger = logger;
 
 			List<MongoIndex<JobDocument>> indexes = new List<MongoIndex<JobDocument>>();
-			indexes.Add(keys => keys.Ascending(x => x.StreamId));
+			indexes.Add(keys => keys.Ascending(x => x.StreamId).Descending(x => x.CreateTimeUtc));
 			indexes.Add(_streamThenTemplateThenCreationTimeIndex = MongoIndex.Create<JobDocument>(keys => keys.Ascending(x => x.StreamId).Ascending(x => x.TemplateId).Descending(x => x.CreateTimeUtc)));
+			indexes.Add(MongoIndex.Create<JobDocument>(keys => keys.Ascending(x => x.StreamId).Ascending(x => x.TemplateId).Descending(x => x.Change)));
 			indexes.Add(keys => keys.Ascending(x => x.Change));
 			indexes.Add(keys => keys.Ascending(x => x.PreflightChange));
 			indexes.Add(_createTimeIndex = MongoIndex.Create<JobDocument>(keys => keys.Descending(x => x.CreateTimeUtc)));
-			indexes.Add(_updateTimeIndex = MongoIndex.Create<JobDocument>(keys => keys.Descending(x => x.UpdateTimeUtc)));
-			indexes.Add(keys => keys.Ascending(x => x.Name));
 			indexes.Add(keys => keys.Ascending(x => x.StartedByUserId));
-			indexes.Add(keys => keys.Ascending(x => x.TemplateId));
 			indexes.Add(keys => keys.Descending(x => x.SchedulePriority));
 			indexes.Add(_startedByBisectTaskIdIndex = MongoIndex.Create<JobDocument>(keys => keys.Descending(x => x.StartedByBisectTaskId), sparse: true));
 			_jobs = mongoService.GetCollection<JobDocument>("Jobs", indexes);
@@ -866,7 +863,14 @@ namespace Horde.Server.Jobs
 			}
 			if (templates != null)
 			{
-				filter &= filterBuilder.In(x => x.TemplateId, templates);
+				if (templates.Length == 1)
+				{
+					filter &= filterBuilder.Eq(x => x.TemplateId, templates[0]);
+				}
+				else
+				{
+					filter &= filterBuilder.In(x => x.TemplateId, templates);
+				}
 			}
 			if (minChange != null)
 			{
