Horde: Add resources to session interface and response object. Since these are static, it can be used to reduce the size of data stored in the agent document in the future.

diff --git a/Engine/Source/Programs/Horde/Horde.Server/Agents/AgentsController.cs b/Engine/Source/Programs/Horde/Horde.Server/Agents/AgentsController.cs
--- a/Engine/Source/Programs/Horde/Horde.Server/Agents/AgentsController.cs
+++ b/Engine/Source/Programs/Horde/Horde.Server/Agents/AgentsController.cs
@@ -312,7 +312,7 @@ namespace Horde.Server.Agents
 		/// <returns>Sessions </returns>
 		[HttpGet]
 		[Route("/api/v1/agents/{agentId}/sessions")]
-		public async Task<ActionResult<List<GetAgentSessionResponse>>> FindSessionsAsync(AgentId agentId, [FromQuery] DateTimeOffset? startTime, [FromQuery] DateTimeOffset? finishTime, [FromQuery] int index = 0, [FromQuery] int count = 50)
+		public async Task<ActionResult<List<GetSessionResponse>>> FindSessionsAsync(AgentId agentId, [FromQuery] DateTimeOffset? startTime, [FromQuery] DateTimeOffset? finishTime, [FromQuery] int index = 0, [FromQuery] int count = 50)
 		{
 			if (!_globalConfig.Value.Authorize(SessionAclAction.ViewSession, User))
 			{
@@ -326,12 +326,12 @@ namespace Horde.Server.Agents
 			}
 
 			List<ISession> sessions = await _agentService.FindSessionsAsync(agentId, startTime?.UtcDateTime, finishTime?.UtcDateTime, index, count);
-			return sessions.ConvertAll(x => CreateGetAgentSessionResponse(x));
+			return sessions.ConvertAll(x => CreateGetSessionResponse(x));
 		}
 
-		static GetAgentSessionResponse CreateGetAgentSessionResponse(ISession session)
+		static GetSessionResponse CreateGetSessionResponse(ISession session)
 		{
-			return new GetAgentSessionResponse(session.Id, session.StartTime, session.FinishTime, (session.Properties != null) ? new List<string>(session.Properties) : null, session.Version);
+			return new GetSessionResponse(session.Id, session.StartTime, session.FinishTime, new List<string>(session.Properties), new Dictionary<string, int>(session.Resources), session.Version);
 		}
 
 		/// <summary>
@@ -342,7 +342,7 @@ namespace Horde.Server.Agents
 		/// <returns>Sessions </returns>
 		[HttpGet]
 		[Route("/api/v1/agents/{agentId}/sessions/{sessionId}")]
-		public async Task<ActionResult<GetAgentSessionResponse>> GetSessionAsync(AgentId agentId, SessionId sessionId)
+		public async Task<ActionResult<GetSessionResponse>> GetSessionAsync(AgentId agentId, SessionId sessionId)
 		{
 			if (!_globalConfig.Value.Authorize(SessionAclAction.ViewSession, User))
 			{
@@ -361,7 +361,7 @@ namespace Horde.Server.Agents
 				return NotFound();
 			}
 
-			return CreateGetAgentSessionResponse(session);
+			return CreateGetSessionResponse(session);
 		}
 
 		/// <summary>
diff --git a/Engine/Source/Programs/Horde/Horde.Server/Agents/Sessions/ISession.cs b/Engine/Source/Programs/Horde/Horde.Server/Agents/Sessions/ISession.cs
--- a/Engine/Source/Programs/Horde/Horde.Server/Agents/Sessions/ISession.cs
+++ b/Engine/Source/Programs/Horde/Horde.Server/Agents/Sessions/ISession.cs
@@ -35,11 +35,16 @@ namespace Horde.Server.Agents.Sessions
 		/// <summary>
 		/// Properties of this agent at the time the session started
 		/// </summary>
-		public IReadOnlyList<string>? Properties { get; }
+		public IReadOnlyList<string> Properties { get; }
+
+		/// <summary>
+		/// Resources reported by the agent at the start of the session
+		/// </summary>
+		public IReadOnlyDictionary<string, int> Resources { get; }
 
 		/// <summary>
 		/// Version of the agent software
 		/// </summary>
-		public string? Version { get; }
+		public string Version { get; }
 	}
 }
diff --git a/Engine/Source/Programs/Horde/Horde.Server/Agents/Sessions/SessionCollection.cs b/Engine/Source/Programs/Horde/Horde.Server/Agents/Sessions/SessionCollection.cs
--- a/Engine/Source/Programs/Horde/Horde.Server/Agents/Sessions/SessionCollection.cs
+++ b/Engine/Source/Programs/Horde/Horde.Server/Agents/Sessions/SessionCollection.cs
@@ -32,11 +32,12 @@ namespace Horde.Server.Agents.Sessions
 
 			public DateTime StartTime { get; set; }
 			public DateTime? FinishTime { get; set; }
-			public List<string>? Properties { get; set; }
-			public Dictionary<string, int>? Resources { get; set; }
-			public string? Version { get; set; }
+			public List<string> Properties { get; set; } = new List<string>();
+			public Dictionary<string, int> Resources { get; set; } = new Dictionary<string, int>();
+			public string Version { get; set; } = String.Empty;
 
-			IReadOnlyList<string>? ISession.Properties => Properties;
+			IReadOnlyList<string> ISession.Properties => Properties;
+			IReadOnlyDictionary<string, int> ISession.Resources => Resources;
 
 			[BsonConstructor]
 			private SessionDocument()
@@ -56,7 +57,10 @@ namespace Horde.Server.Agents.Sessions
 				{
 					Resources = new Dictionary<string, int>(resources);
 				}
-				Version = version;
+				if (version != null)
+				{
+					Version = version;
+				}
 			}
 		}
 
diff --git a/Engine/Source/Programs/Shared/EpicGames.Horde/Agents/AgentMessages.cs b/Engine/Source/Programs/Shared/EpicGames.Horde/Agents/AgentMessages.cs
--- a/Engine/Source/Programs/Shared/EpicGames.Horde/Agents/AgentMessages.cs
+++ b/Engine/Source/Programs/Shared/EpicGames.Horde/Agents/AgentMessages.cs
@@ -75,16 +75,6 @@ namespace EpicGames.Horde.Agents
 	/// <param name="State"> State of the lease (for AgentLeases) </param>
 	public record GetAgentLeaseResponse(LeaseId Id, LeaseId? ParentId, AgentId? AgentId, double? AgentRate, string Name, LogId? LogId, DateTime StartTime, DateTime? FinishTime, bool Executing, Dictionary<string, string>? Details, LeaseOutcome? Outcome, LeaseState? State);
 
-	/// <summary>
-	/// Information about an agent session
-	/// </summary>
-	/// <param name="Id"> Unique id for this session </param>
-	/// <param name="StartTime"> Start time for this session </param>
-	/// <param name="FinishTime"> Finishing time for this session </param>
-	/// <param name="Properties"> Properties of this agent </param>
-	/// <param name="Version"> Version of the software running during this session </param>
-	public record GetAgentSessionResponse(SessionId Id, DateTime StartTime, DateTime? FinishTime, List<string>? Properties, string? Version);
-
 	/// <summary>
 	/// Information about a workspace synced on an agent
 	/// </summary>
diff --git a/Engine/Source/Programs/Shared/EpicGames.Horde/Agents/Sessions/SessionMessages.cs b/Engine/Source/Programs/Shared/EpicGames.Horde/Agents/Sessions/SessionMessages.cs
new file mode 100644
--- /dev/null
+++ b/Engine/Source/Programs/Shared/EpicGames.Horde/Agents/Sessions/SessionMessages.cs
@@ -0,0 +1,18 @@
+// Copyright Epic Games, Inc. All Rights Reserved.
+
+using System;
+using System.Collections.Generic;
+
+namespace EpicGames.Horde.Agents.Sessions
+{
+	/// <summary>
+	/// Information about a session
+	/// </summary>
+	/// <param name="Id">Unique id for this session</param>
+	/// <param name="StartTime">Start time for this session</param>
+	/// <param name="FinishTime">Finishing time for this session</param>
+	/// <param name="Properties">Properties of this agent</param>
+	/// <param name="Resources">Resources held by the agent at the start of the session</param>
+	/// <param name="Version"> Version of the software running during this session </param>
+	public record GetSessionResponse(SessionId Id, DateTime StartTime, DateTime? FinishTime, List<string>? Properties, Dictionary<string, int>? Resources, string? Version);
+}
