Horde: Add a sparse index over agent expiry times. This is currently falling back to a collection scan.

diff --git a/Engine/Source/Programs/Horde/Horde.Server/Agents/AgentCollection.cs b/Engine/Source/Programs/Horde/Horde.Server/Agents/AgentCollection.cs
--- a/Engine/Source/Programs/Horde/Horde.Server/Agents/AgentCollection.cs
+++ b/Engine/Source/Programs/Horde/Horde.Server/Agents/AgentCollection.cs
@@ -248,6 +248,7 @@ namespace Horde.Server.Agents
 		{
 			List<MongoIndex<AgentDocument>> indexes = new List<MongoIndex<AgentDocument>>();
 			indexes.Add(keys => keys.Ascending(x => x.Deleted).Ascending(x => x.Id).Ascending(x => x.Pools));
+			indexes.Add(keys => keys.Ascending(x => x.SessionExpiresAt), sparse: true);
 
 			_agents = mongoService.GetCollection<AgentDocument>("Agents", indexes);
 			_redisService = redisService;
@@ -380,7 +381,8 @@ namespace Horde.Server.Agents
 		/// <inheritdoc/>
 		public async Task<IReadOnlyList<IAgent>> FindExpiredAsync(DateTime utcNow, int maxAgents, CancellationToken cancellationToken)
 		{
-			List<AgentDocument> documents = await _agents.Find(x => x.SessionId.HasValue && !(x.SessionExpiresAt > utcNow)).Limit(maxAgents).ToListAsync(cancellationToken);
+			FilterDefinition<AgentDocument> filter = Builders<AgentDocument>.Filter.Exists(x => x.SessionExpiresAt) & Builders<AgentDocument>.Filter.Lt(x => x.SessionExpiresAt, utcNow);
+			List<AgentDocument> documents = await _agents.Find(filter).Limit(maxAgents).ToListAsync(cancellationToken);
 			return documents.ConvertAll(x => CreateAgentObject(x));
 		}
 
