Horde: Add a debug endpoint to repair a particular job.

#rnx

diff --git a/Engine/Source/Programs/Horde/Horde.Server/Server/DebugController.cs b/Engine/Source/Programs/Horde/Horde.Server/Server/DebugController.cs
--- a/Engine/Source/Programs/Horde/Horde.Server/Server/DebugController.cs
+++ b/Engine/Source/Programs/Horde/Horde.Server/Server/DebugController.cs
@@ -20,6 +20,7 @@ using EpicGames.Horde.Jobs;
 using EpicGames.Horde.Logs;
 using Google.Protobuf;
 using Horde.Common.Rpc;
+using Horde.Server.Agents.Leases;
 using Horde.Server.Agents.Relay;
 using Horde.Server.Configuration;
 using Horde.Server.Jobs;
@@ -261,6 +262,7 @@ namespace Horde.Server.Server
 		private readonly JobService _jobService;
 		private readonly JobTaskSource _jobTaskSource;
 		private readonly ILogCollection _logCollection;
+		private readonly ILeaseCollection _leaseCollection;
 		private readonly IOptionsSnapshot<GlobalConfig> _globalConfig;
 		private readonly ILogger<SecureDebugController> _logger;
 
@@ -273,7 +275,7 @@ namespace Horde.Server.Server
 			AgentRelayService agentRelayService,
 			JobService jobService,
 			JobTaskSource jobTaskSource,
-			ILogCollection logCollection, IOptionsSnapshot<GlobalConfig> globalConfig, ILogger<SecureDebugController> logger)
+			ILogCollection logCollection, ILeaseCollection leaseCollection, IOptionsSnapshot<GlobalConfig> globalConfig, ILogger<SecureDebugController> logger)
 		{
 			_mongoService = mongoService;
 			_configService = configService;
@@ -281,6 +283,7 @@ namespace Horde.Server.Server
 			_agentRelayService = agentRelayService;
 			_jobTaskSource = jobTaskSource;
 			_logCollection = logCollection;
+			_leaseCollection = leaseCollection;
 			_globalConfig = globalConfig;
 			_logger = logger;
 		}
@@ -781,6 +784,105 @@ namespace Horde.Server.Server
 			return groupedTimings.Values.ToList();
 		}
 
+		/// <summary>
+		/// Repairs any inconsistencies with jobs created in a particular time range 
+		/// </summary>
+		[HttpPost]
+		[Route("/api/v1/debug/repair-jobs")]
+		public async Task<ActionResult> RepairJobsAsync([FromQuery] DateTime? minTime, [FromQuery] DateTime? maxTime)
+		{
+			if (!_globalConfig.Value.Authorize(ServerAclAction.Debug, User))
+			{
+				return Forbid(ServerAclAction.Debug);
+			}
+
+			if (minTime == null || maxTime == null)
+			{
+				return BadRequest();
+			}
+
+			IReadOnlyList<IJob> jobs = await _jobService.FindJobsAsync(minCreateTime: minTime, maxCreateTime: maxTime);
+			foreach (IJob job in jobs)
+			{
+				StreamConfig? streamConfig;
+				if (_globalConfig.Value.TryGetStream(job.StreamId, out streamConfig))
+				{
+					_logger.LogInformation("Checking job {JobId}", job.Id);
+					await TryRepairJobAsync(streamConfig, job, HttpContext.RequestAborted);
+				}
+			}
+
+			_logger.LogInformation("Finished repair");
+			return Ok();
+		}
+
+		/// <summary>
+		/// Repairs any inconsistencies with a particular job 
+		/// </summary>
+		/// <param name="jobId">Id of the job to find</param>
+		[HttpPost]
+		[Route("/api/v1/debug/repair-job/{jobId}")]
+		public async Task<ActionResult> RepairJobAsync(JobId jobId)
+		{
+			IJob? job = await _jobService.GetJobAsync(jobId);
+			if (job == null || job.TemplateHash == null)
+			{
+				return NotFound(jobId);
+			}
+
+			StreamConfig? streamConfig;
+			if (!_globalConfig.Value.TryGetStream(job.StreamId, out streamConfig))
+			{
+				return NotFound(job.StreamId);
+			}
+			if (!_globalConfig.Value.Authorize(ServerAclAction.Debug, User))
+			{
+				return Forbid(ServerAclAction.Debug);
+			}
+
+			for (; ; )
+			{
+				IJob? newJob = await TryRepairJobAsync(streamConfig, job, HttpContext.RequestAborted);
+				if (newJob != null)
+				{
+					return Ok();
+				}
+
+				newJob = await _jobService.GetJobAsync(job.Id, HttpContext.RequestAborted);
+				if (newJob == null)
+				{
+					return NotFound();
+				}
+			}
+		}
+
+		async Task<IJob?> TryRepairJobAsync(StreamConfig streamConfig, IJob job, CancellationToken cancellationToken)
+		{
+			// Check the lease has not already completed. Workaround for issue where jobs collection came out of sync with leases collection due to Mongo timeouts while updating indexes.
+			List<JobStepBatchId> batchIds = job.Batches.Select(x => x.Id).ToList();
+			foreach (JobStepBatchId batchId in batchIds)
+			{
+				IJobStepBatch? batch;
+				if (job.TryGetBatch(batchId, out batch) && batch.LeaseId.HasValue && batch.State == JobStepBatchState.Running)
+				{
+					ILease? lease = await _leaseCollection.GetAsync(batch.LeaseId.Value, cancellationToken);
+					if (lease != null && lease.FinishTime.HasValue)
+					{
+						_logger.LogWarning("Job {JobId} batch {BatchId} is out of sync with lease {LeaseId}", job.Id, batch.Id, lease.Id);
+
+						IJob? newJob = await _jobService.UpdateBatchAsync(job, batch.Id, streamConfig, newState: JobStepBatchState.Complete, cancellationToken: cancellationToken);
+						if (newJob == null)
+						{
+							return null;
+						}
+
+						job = newJob;
+					}
+				}
+			}
+			return job;
+		}
+
 		/// <summary>
 		/// Populate the database with test data
 		/// </summary>
