Horde: Add agent telemetry
#rnx

diff --git a/Engine/Source/Programs/Horde/HordeDashboard/src/backend/Api.ts b/Engine/Source/Programs/Horde/HordeDashboard/src/backend/Api.ts
--- a/Engine/Source/Programs/Horde/HordeDashboard/src/backend/Api.ts
+++ b/Engine/Source/Programs/Horde/HordeDashboard/src/backend/Api.ts
@@ -2097,6 +2097,16 @@ export type GetAgentTypeResponse = {
 
 }
 
+export type GetAgentTelemetrySampleResponse = {
+	time: Date;
+	userCpu: number;
+	idleCpu: number;
+	systemCpu: number;
+	freeRam: number;
+	usedRam: number;
+	totalRam: number;
+}
+
 /**Information about a workspace type */
 export type GetWorkspaceTypeResponse = {
 
diff --git a/Engine/Source/Programs/Horde/HordeDashboard/src/backend/Backend.ts b/Engine/Source/Programs/Horde/HordeDashboard/src/backend/Backend.ts
--- a/Engine/Source/Programs/Horde/HordeDashboard/src/backend/Backend.ts
+++ b/Engine/Source/Programs/Horde/HordeDashboard/src/backend/Backend.ts
@@ -1,7 +1,7 @@
 // Copyright Epic Games, Inc. All Rights Reserved.
 
 import templateCache from '../backend/TemplateCache';
-import { AccountClaimMessage, AgentData, AgentQuery, ApproveAgentsRequest, ArtifactData, AuditLogEntry, AuditLogQuery, BatchUpdatePoolRequest, ChangeSummaryData, CreateAccountRequest, CreateBisectTaskRequest, CreateBisectTaskResponse, CreateDeviceRequest, CreateDeviceResponse, CreateExternalIssueRequest, CreateExternalIssueResponse, CreateJobRequest, CreateJobResponse, CreateNoticeRequest, CreatePoolRequest, CreateServiceAccountRequest, CreateServiceAccountResponse, CreateSoftwareResponse, CreateSubscriptionRequest, CreateSubscriptionResponse, CreateZipRequest, DashboardPreference, DevicePoolTelemetryQuery, DeviceTelemetryQuery, EventData, FindArtifactsResponse, FindIssueResponse, FindJobTimingsResponse, GetAccountResponse, GetArtifactDirectoryResponse, GetArtifactZipRequest, GetBisectTaskResponse, GetDashboardConfigResponse, GetDevicePlatformResponse, GetDevicePoolResponse, GetDevicePoolTelemetryResponse, GetDeviceReservationResponse, GetDeviceResponse, GetDeviceTelemetryResponse, GetExternalIssueProjectResponse, GetExternalIssueResponse, GetIssueStreamResponse, GetJobsTabResponse, GetJobStepRefResponse, GetJobStepTraceResponse, GetJobTimingResponse, GetLogEventResponse, GetLogFileResponse, GetNoticeResponse, GetNotificationResponse, GetPendingAgentsResponse, GetPerforceServerStatusResponse, GetPoolResponse, GetPoolSummaryResponse, GetServerInfoResponse, GetServerSettingsResponse, GetServiceAccountResponse, GetSoftwareResponse, GetSubscriptionResponse, GetTelemetryMetricsResponse, GetTestDataDetailsResponse, GetTestDataRefResponse, GetTestMetaResponse, GetTestResponse, GetTestsRequest, GetTestStreamResponse, GetToolSummaryResponse, GetUserResponse, GetUtilizationTelemetryResponse, GlobalConfig, IssueData, IssueQuery, IssueQueryV2, JobData, JobQuery, JobsTabColumnType, JobStepOutcome, JobStreamQuery, JobTimingsQuery, LeaseData, LogData, LogEventQuery, LogLineData, MetricsQuery, PoolData, PoolQuery, PreflightConfigResponse, ProjectData, ScheduleData, ScheduleQuery, SearchLogFileResponse, ServerStatusResponse, ServerUpdateResponse, SessionData, StreamData, TabType, TestData, UpdateAccountRequest, UpdateAgentRequest, UpdateBisectTaskRequest, UpdateCurrentAccountRequest, UpdateDeviceRequest, UpdateGlobalConfigRequest, UpdateIssueRequest, UpdateJobRequest, UpdateLeaseRequest, UpdateNoticeRequest, UpdateNotificationsRequest, UpdatePoolRequest, UpdateServerSettingsRequest, UpdateServiceAccountRequest, UpdateServiceAccountResponse, UpdateStepRequest, UpdateStepResponse, UpdateTemplateRefRequest, UpdateUserRequest, UsersQuery } from './Api';
+import { AccountClaimMessage, AgentData, AgentQuery, ApproveAgentsRequest, ArtifactData, AuditLogEntry, AuditLogQuery, BatchUpdatePoolRequest, ChangeSummaryData, CreateAccountRequest, CreateBisectTaskRequest, CreateBisectTaskResponse, CreateDeviceRequest, CreateDeviceResponse, CreateExternalIssueRequest, CreateExternalIssueResponse, CreateJobRequest, CreateJobResponse, CreateNoticeRequest, CreatePoolRequest, CreateServiceAccountRequest, CreateServiceAccountResponse, CreateSoftwareResponse, CreateSubscriptionRequest, CreateSubscriptionResponse, CreateZipRequest, DashboardPreference, DevicePoolTelemetryQuery, DeviceTelemetryQuery, EventData, FindArtifactsResponse, FindIssueResponse, FindJobTimingsResponse, GetAccountResponse, GetAgentTelemetrySampleResponse, GetArtifactDirectoryResponse, GetArtifactZipRequest, GetBisectTaskResponse, GetDashboardConfigResponse, GetDevicePlatformResponse, GetDevicePoolResponse, GetDevicePoolTelemetryResponse, GetDeviceReservationResponse, GetDeviceResponse, GetDeviceTelemetryResponse, GetExternalIssueProjectResponse, GetExternalIssueResponse, GetIssueStreamResponse, GetJobsTabResponse, GetJobStepRefResponse, GetJobStepTraceResponse, GetJobTimingResponse, GetLogEventResponse, GetLogFileResponse, GetNoticeResponse, GetNotificationResponse, GetPendingAgentsResponse, GetPerforceServerStatusResponse, GetPoolResponse, GetPoolSummaryResponse, GetServerInfoResponse, GetServerSettingsResponse, GetServiceAccountResponse, GetSoftwareResponse, GetSubscriptionResponse, GetTelemetryMetricsResponse, GetTestDataDetailsResponse, GetTestDataRefResponse, GetTestMetaResponse, GetTestResponse, GetTestsRequest, GetTestStreamResponse, GetToolSummaryResponse, GetUserResponse, GetUtilizationTelemetryResponse, GlobalConfig, IssueData, IssueQuery, IssueQueryV2, JobData, JobQuery, JobsTabColumnType, JobStepOutcome, JobStreamQuery, JobTimingsQuery, LeaseData, LogData, LogEventQuery, LogLineData, MetricsQuery, PoolData, PoolQuery, PreflightConfigResponse, ProjectData, ScheduleData, ScheduleQuery, SearchLogFileResponse, ServerStatusResponse, ServerUpdateResponse, SessionData, StreamData, TabType, TestData, UpdateAccountRequest, UpdateAgentRequest, UpdateBisectTaskRequest, UpdateCurrentAccountRequest, UpdateDeviceRequest, UpdateGlobalConfigRequest, UpdateIssueRequest, UpdateJobRequest, UpdateLeaseRequest, UpdateNoticeRequest, UpdateNotificationsRequest, UpdatePoolRequest, UpdateServerSettingsRequest, UpdateServiceAccountRequest, UpdateServiceAccountResponse, UpdateStepRequest, UpdateStepResponse, UpdateTemplateRefRequest, UpdateUserRequest, UsersQuery } from './Api';
 import dashboard, { Dashboard } from './Dashboard';
 import { ChallengeStatus, Fetch } from './Fetch';
 import { projectStore } from './ProjectStore';
@@ -263,6 +263,22 @@ export class Backend {
         });
     }
 
+    getAgentTelemetry(agentId: string, minTime: Date, maxTime: Date): Promise<GetAgentTelemetrySampleResponse[]> {
+        const params: string[] = [];
+        params.push("minTime=" + minTime.toISOString());
+        params.push("maxTime=" + maxTime.toISOString());
+        const paramString = params.join("&");
+        return new Promise<GetAgentTelemetrySampleResponse[]>((resolve, reject) => {
+            this.backend.get(`/api/v1/agents/${agentId}/telemetry?${paramString}`).then((response) => {
+                const data = response.data.samples as GetAgentTelemetrySampleResponse[];
+                data.forEach((d) => {
+                    d.time = new Date(d.time);
+                });
+                resolve(data);
+            }).catch(reason => { reject(reason); });
+        });
+    }
+
     retryJobStep(jobId: string, batchId: string, stepId: string): Promise<any> {
 
         return new Promise<any>((resolve, reject) => {
diff --git a/Engine/Source/Programs/Horde/HordeDashboard/src/components/HistoryModal.tsx b/Engine/Source/Programs/Horde/HordeDashboard/src/components/HistoryModal.tsx
--- a/Engine/Source/Programs/Horde/HordeDashboard/src/components/HistoryModal.tsx
+++ b/Engine/Source/Programs/Horde/HordeDashboard/src/components/HistoryModal.tsx
@@ -1,15 +1,16 @@
 // Copyright Epic Games, Inc. All Rights Reserved.
 
-import { Checkbox, ConstrainMode, DefaultButton, DetailsHeader, DetailsList, DetailsListLayoutMode, DetailsRow, Dialog, DialogFooter, DialogType, GroupedList, GroupHeader, ICheckbox, IColumn, IContextualMenuItem, IContextualMenuProps, IDetailsHeaderProps, IDetailsHeaderStyles, IDetailsListProps, IGroup, ITextField, ITooltipHostStyles, mergeStyleSets, Modal, Pivot, PivotItem, PrimaryButton, ScrollablePane, ScrollbarVisibility, Selection, SelectionMode, Spinner, SpinnerSize, Stack, Sticky, StickyPositionType, Text, TextField } from "@fluentui/react";
+import { Checkbox, ComboBox, ConstrainMode, DefaultButton, DetailsHeader, DetailsList, DetailsListLayoutMode, DetailsRow, Dialog, DialogFooter, DialogType, GroupedList, GroupHeader, ICheckbox, IColumn, IComboBoxOption, IContextualMenuItem, IContextualMenuProps, IDetailsHeaderProps, IDetailsHeaderStyles, IDetailsListProps, IGroup, ITextField, ITooltipHostStyles, Label, mergeStyleSets, Modal, Pivot, PivotItem, PrimaryButton, ScrollablePane, ScrollbarVisibility, Selection, SelectionMode, Spinner, SpinnerSize, Stack, Sticky, StickyPositionType, Text, TextField } from "@fluentui/react";
 import { action, makeObservable, observable } from "mobx";
 import { observer } from "mobx-react-lite";
-import React, { useState } from "react";
+import React, { useEffect, useState } from "react";
 import { Link } from "react-router-dom";
+import { Sparklines, SparklinesLine, SparklinesReferenceLine } from "react-sparklines";
 import backend from "../backend";
 import { agentStore } from "../backend/AgentStore";
-import { AgentData, GetAgentLeaseResponse, GetAgentSessionResponse, JobStepBatchError, JobStepOutcome, LeaseData, SessionData, UpdateAgentRequest } from "../backend/Api";
+import { AgentData, GetAgentLeaseResponse, GetAgentSessionResponse, GetAgentTelemetrySampleResponse, JobStepBatchError, JobStepOutcome, LeaseData, SessionData, UpdateAgentRequest } from "../backend/Api";
 import dashboard from "../backend/Dashboard";
-import { getShortNiceTime } from "../base/utilities/timeUtils";
+import { getNiceTime, getShortNiceTime } from "../base/utilities/timeUtils";
 import { getHordeStyling } from "../styles/Styles";
 import { getHordeTheme } from "../styles/theme";
 import { BatchStatusIcon, LeaseStatusIcon, StepStatusIcon } from "./StatusIcon";
@@ -125,7 +126,7 @@ class HistoryModalState {
       this.agentItemCount = 0;
       this.devicesItemCount = 0;
       this.workspaceItemCount = 0;
-      const items:InfoPanelItem[] = [
+      const items: InfoPanelItem[] = [
          {
             key: "overview",
             name: "Overview",
@@ -162,7 +163,7 @@ class HistoryModalState {
                const name: string = count > 1 ? `${workspace.stream} (${count})` : workspace.stream;
 
                (workspace as any)._hackName = name;
-;
+               ;
                items.push({ key: `workspace${workspaceIdx}`, name: name, selected: false, data: name });
                this.workspaceItemCount++;
             }
@@ -268,8 +269,24 @@ class HistoryModalState {
       });
    }
 
+   UpdateTelemetry(minutes: number) {
+      backend.getAgentTelemetry(this.selectedAgent!.id, new Date(Date.now() - 1000 * 60 * minutes), new Date()).then(data => {
+         this.setTelemetryData(data);
+      })
+   }
+
+   @action
+   setTelemetryData(newData: GetAgentTelemetrySampleResponse[]) {
+      this.currentData = newData;
+   }
+
    @action
    appendData(newData: any[]) {
+
+      if (this.mode === "telemetry") {
+         return;
+      }
+
       // if there's any data, there might be more data next time, so add another callback.
       if (newData.length > 0) {
          newData.push(null);
@@ -302,6 +319,174 @@ class HistoryModalState {
 
 const state = new HistoryModalState();
 
+type TimeSelection = {
+   text: string;
+   key: string;
+   minutes: number;
+}
+
+const timeSelections: TimeSelection[] = [
+   {
+      text: "Past 1 Hour", key: "time_1_hour", minutes: 60
+   },
+   {
+      text: "Past 2 Hours", key: "time_2_hours", minutes: 60 * 2
+   },
+   {
+      text: "Past 4 Hours", key: "time_4_hours", minutes: 60 * 4
+   },
+   {
+      text: "Past 1 Day", key: "time_1_day", minutes: 60 * 24
+   },
+   {
+      text: "Past 2 Days", key: "time_2_days", minutes: 60 * 24 * 2
+   },
+   {
+      text: "Past 1 Week", key: "time_1_week", minutes: 60 * 24 * 7
+   },
+   {
+      text: "Past 2 Weeks", key: "time_2_weeks", minutes: 60 * 24 * 7 * 2
+   },
+   {
+      text: "Past Month", key: "time_4_weeks", minutes: 60 * 24 * 7 * 4
+   }
+]
+
+const cpuSelections: IComboBoxOption[] = [
+   {
+      text: "Combined", key: "combined"
+   },
+   {
+      text: "User", key: "user"
+   },
+   {
+      text: "System", key: "system"
+   }
+]
+
+const TelemetryPanel: React.FC<{}> = observer(({ }) => {
+
+   const [timeKey, setTimeKey] = useState("time_2_weeks");
+   const [cpuKey, setCPUKey] = useState("combined");
+
+   useEffect(() => {
+
+      const time = timeSelections.find(t => t.key === timeKey);
+      if (time) {
+         state.UpdateTelemetry(time.minutes);
+      }
+
+      return () => {
+      };
+
+   }, []);
+
+
+   // subscribe
+   if (state.currentData) { }
+
+   const data = (state.currentData ?? []) as GetAgentTelemetrySampleResponse[];
+
+   if (!data.length) {
+      return <Stack horizontalAlign="center">
+         <Spinner size={SpinnerSize.large} />
+      </Stack>;
+   }
+
+   const cpuData = data.map(d => {
+      if (cpuKey == "user") return d.userCpu;
+      if (cpuKey == "system") return d.systemCpu;
+      if (cpuKey == "combined") return d.userCpu + d.systemCpu;
+   })
+
+   const min = data[0].time;
+   const max = data[data.length - 1].time;
+
+   const minTime = getShortNiceTime(min, true);
+   const maxTime = getShortNiceTime(max, true);
+
+   const ramData = (state.currentData ?? []).map(m => {
+      const d = m as GetAgentTelemetrySampleResponse;
+      return d.usedRam / d.totalRam;
+   })
+
+   let cpuText = "";
+   let ramText = "";
+
+   let cpuLabel = "";
+   if (cpuKey == "user") cpuLabel = "CPU - User";
+   if (cpuKey == "system") cpuLabel = "CPU - System";
+   if (cpuKey == "combined") cpuLabel = "CPU - Combined";
+
+
+   const device = state.selectedAgent.capabilities?.devices[0];
+   device.properties?.forEach(v => {
+      if (v.startsWith("CPU=")) {
+         cpuText = `${v.replace("CPU=", "")}`;
+      }
+
+      if (v.startsWith("RAM=")) {
+         ramText = `${v.replace("RAM=", "")} GB`;
+      }
+   })
+
+
+
+   return <Stack tokens={{ childrenGap: 24 }} style={{ paddingBottom: 32, paddingLeft: 8 }}>
+      <Stack horizontal verticalAlign="center" verticalFill>
+         <Text variant="mediumPlus">{maxTime} - {minTime}</Text>
+         <Stack grow />
+         <Stack horizontal tokens={{ childrenGap: 18 }}>
+
+            <ComboBox
+               label={"CPU"}
+               styles={{ root: { width: 128 } }}
+               options={cpuSelections}
+               selectedKey={cpuKey}
+               onChange={(ev, option, index, value) => {
+                  setCPUKey(option.key as string);
+               }}
+            />
+
+            <ComboBox
+               label="Time"
+               styles={{ root: { width: 180 } }}
+               options={timeSelections}
+               selectedKey={timeKey}
+               onChange={(ev, option, index, value) => {
+                  const select = option as TimeSelection;
+                  setTimeKey(select.key);
+                  state.UpdateTelemetry(select.minutes);
+               }}
+            />
+         </Stack>
+      </Stack>
+      <Stack verticalAlign="center" verticalFill>
+         <Stack horizontal>
+            <Label>{cpuLabel}</Label>
+            <Stack grow></Stack>
+            {!!cpuText && <Label>{cpuText}</Label>}
+         </Stack>
+         <Sparklines width={1100} height={128} data={cpuData} max={100} style={{ backgroundColor: dashboard.darktheme ? "#060709" : "#F3F2F1", padding: 8, border: "solid 1px #181A1B" }}>
+            <SparklinesLine color={dashboard.darktheme ? "lightblue" : "#1E90FF"} />
+            <SparklinesReferenceLine type="avg" />
+         </Sparklines>
+      </Stack>
+      <Stack verticalAlign="center" verticalFill>
+         <Stack horizontal>
+            <Label>RAM</Label>
+            <Stack grow></Stack>
+            {!!ramText && <Label>{ramText}</Label>}
+         </Stack>
+         <Sparklines width={1100} height={128} data={ramData} max={1} style={{ backgroundColor: dashboard.darktheme ? "#060709" : "#F3F2F1", padding: 8, border: "solid 1px #181A1B" }}>
+            <SparklinesLine color={dashboard.darktheme ? "lightblue" : "#1E90FF"} />
+            <SparklinesReferenceLine type="avg" />
+         </Sparklines>
+      </Stack>
+
+   </Stack>
+})
+
 export const HistoryModal: React.FC<{ agentId: string | undefined, onDismiss: (...args: any[]) => any; }> = observer(({ agentId, onDismiss }) => {
 
    const [selectedAgent, setSelectedAgent] = useState<string | undefined>(undefined);
@@ -550,26 +735,31 @@ export const HistoryModal: React.FC<{ agentId: string | undefined, onDismiss: (.
             }
          ];
       }
-      columns.push({
-         key: 'startTime',
-         name: 'Start Time',
-         minWidth: 200,
-         maxWidth: 200,
-         isResizable: false,
-         isSorted: false,
-         isSortedDescending: false,
-         onColumnClick: onColumnClick
-      });
-      columns.push({
-         key: 'endTime',
-         name: 'Finish Time',
-         minWidth: 200,
-         maxWidth: 200,
-         isResizable: false,
-         isSorted: false,
-         isSortedDescending: false,
-         onColumnClick: onColumnClick
-      });
+
+      if (state.mode === "leases" || state.mode === "sessions") {
+
+
+         columns.push({
+            key: 'startTime',
+            name: 'Start Time',
+            minWidth: 200,
+            maxWidth: 200,
+            isResizable: false,
+            isSorted: false,
+            isSortedDescending: false,
+            onColumnClick: onColumnClick
+         });
+         columns.push({
+            key: 'endTime',
+            name: 'Finish Time',
+            minWidth: 200,
+            maxWidth: 200,
+            isResizable: false,
+            isSorted: false,
+            isSortedDescending: false,
+            onColumnClick: onColumnClick
+         });
+      }
 
       if (state.mode === "leases") {
          columns.find(col => col.key === state.sortedLeaseColumn)!.isSorted = true;
@@ -852,12 +1042,10 @@ export const HistoryModal: React.FC<{ agentId: string | undefined, onDismiss: (.
                         linkFormat="links"
                         defaultSelectedKey={state.mode ?? "info"}
                      >
-                        <PivotItem headerText="Info" itemKey="info">
-                        </PivotItem>
-                        <PivotItem headerText="Sessions" itemKey="sessions">
-                        </PivotItem>
-                        <PivotItem headerText="Leases" itemKey="leases">
-                        </PivotItem>
+                        <PivotItem headerText="Info" itemKey="info" />
+                        <PivotItem headerText="Sessions" itemKey="sessions" />
+                        <PivotItem headerText="Leases" itemKey="leases" />
+                        <PivotItem headerText="Telemetry" itemKey="telemetry" />
                      </Pivot>
                   </Stack>
                </Stack>
@@ -886,85 +1074,84 @@ export const HistoryModal: React.FC<{ agentId: string | undefined, onDismiss: (.
                </Stack>}
             </Stack>
             <Stack styles={{ root: { paddingTop: 30 } }}>
-               {
-                  state.mode === "info" ?
-                     <Stack horizontal tokens={{ childrenGap: 20 }}>
-                        <Stack styles={{ root: { width: 300 } }}>
-                           <GroupedList
-                              items={state.infoItems}
-                              compact={true}
-                              onRenderCell={onRenderBuilderInfoCell}
-                              groups={groups}
-                              selection={new Selection()}
-                              selectionMode={SelectionMode.none}
-                              groupProps={{
-                                 showEmptyGroups: true,
-                                 onRenderHeader: (props) => {
-                                    return <Link to="" onClick={(ev) => { ev.preventDefault(); props!.onToggleCollapse!(props!.group!); }}> <GroupHeader {...props} /></Link>;
+               {state.mode === "info" && <Stack horizontal tokens={{ childrenGap: 20 }}>
+                  <Stack styles={{ root: { width: 300 } }}>
+                     <GroupedList
+                        items={state.infoItems}
+                        compact={true}
+                        onRenderCell={onRenderBuilderInfoCell}
+                        groups={groups}
+                        selection={new Selection()}
+                        selectionMode={SelectionMode.none}
+                        groupProps={{
+                           showEmptyGroups: true,
+                           onRenderHeader: (props) => {
+                              return <Link to="" onClick={(ev) => { ev.preventDefault(); props!.onToggleCollapse!(props!.group!); }}> <GroupHeader {...props} /></Link>;
+                           },
+                           headerProps: {
+                              styles: {
+                                 title: {
+                                    fontFamily: "Horde Open Sans Semibold",
+                                    color: modeColors.text,
+                                    paddingLeft: 0
+
                                  },
-                                 headerProps: {
-                                    styles: {
-                                       title: {
-                                          fontFamily: "Horde Open Sans Semibold",
-                                          color: modeColors.text,
-                                          paddingLeft: 0
-
-                                       },
-                                       headerCount: {
-                                          display: 'none'
-                                       }
-                                    }
+                                 headerCount: {
+                                    display: 'none'
                                  }
-                              }}
-                           />
-                        </Stack>
-                        <Stack styles={{ root: { width: '100%' } }}>
-                           <Stack.Item className={hordeClasses.relativeModalSmall}>
-                              <ScrollablePane scrollbarVisibility={ScrollbarVisibility.auto} styles={{ contentContainer: { overflowX: 'hidden' } }}>
-                                 <DetailsList
-                                    className={historyStyles.detailsList}
-                                    compact={true}
-                                    items={state.infoSubItems}
-                                    columns={[
-                                       { key: 'column1', name: 'Name', fieldName: 'name', minWidth: 200, maxWidth: 200, isResizable: false },
-                                       { key: 'column2', name: 'Value', fieldName: 'value', minWidth: 200, isResizable: false }
-                                    ]}
-                                    layoutMode={DetailsListLayoutMode.justified}
-                                    onRenderDetailsHeader={onRenderInfoDetailsHeader}
-                                    constrainMode={ConstrainMode.unconstrained}
-                                    selectionMode={SelectionMode.none}
-                                    onRenderRow={(props) => {
-
-                                       if (props) {
-                                          return <DetailsRow styles={{ cell: { whiteSpace: "pre-line", overflowWrap: "break-word" } }} {...props} />
-                                       }
-
-                                       return null;
-                                    }}
-                                 />
-                              </ScrollablePane>
-                           </Stack.Item>
-                        </Stack>
-                     </Stack>
-                     :
+                              }
+                           }
+                        }}
+                     />
+                  </Stack>
+                  <Stack styles={{ root: { width: '100%' } }}>
                      <Stack.Item className={hordeClasses.relativeModalSmall}>
                         <ScrollablePane scrollbarVisibility={ScrollbarVisibility.auto} styles={{ contentContainer: { overflowX: 'hidden' } }}>
                            <DetailsList
                               className={historyStyles.detailsList}
                               compact={true}
-                              items={state.currentData}
-                              columns={generateColumns()}
-                              onRenderDetailsHeader={onRenderDetailsHeader}
-                              onRenderItemColumn={onRenderHistoryItem}
+                              items={state.infoSubItems}
+                              columns={[
+                                 { key: 'column1', name: 'Name', fieldName: 'name', minWidth: 200, maxWidth: 200, isResizable: false },
+                                 { key: 'column2', name: 'Value', fieldName: 'value', minWidth: 200, isResizable: false }
+                              ]}
                               layoutMode={DetailsListLayoutMode.justified}
+                              onRenderDetailsHeader={onRenderInfoDetailsHeader}
                               constrainMode={ConstrainMode.unconstrained}
                               selectionMode={SelectionMode.none}
-                              listProps={{ renderedWindowsAhead: 1, renderedWindowsBehind: 1 }}
-                              onRenderMissingItem={() => { state.nullItemTrigger(); return <div></div> }}
+                              onRenderRow={(props) => {
+
+                                 if (props) {
+                                    return <DetailsRow styles={{ cell: { whiteSpace: "pre-line", overflowWrap: "break-word" } }} {...props} />
+                                 }
+
+                                 return null;
+                              }}
                            />
                         </ScrollablePane>
                      </Stack.Item>
+                  </Stack>
+               </Stack>
                }
+               {(state.mode === "leases" || state.mode === "sessions") &&
+                  <Stack.Item className={hordeClasses.relativeModalSmall}>
+                     <ScrollablePane scrollbarVisibility={ScrollbarVisibility.auto} styles={{ contentContainer: { overflowX: 'hidden' } }}>
+                        <DetailsList
+                           className={historyStyles.detailsList}
+                           compact={true}
+                           items={state.currentData}
+                           columns={generateColumns()}
+                           onRenderDetailsHeader={onRenderDetailsHeader}
+                           onRenderItemColumn={onRenderHistoryItem}
+                           layoutMode={DetailsListLayoutMode.justified}
+                           constrainMode={ConstrainMode.unconstrained}
+                           selectionMode={SelectionMode.none}
+                           listProps={{ renderedWindowsAhead: 1, renderedWindowsBehind: 1 }}
+                           onRenderMissingItem={() => { state.nullItemTrigger(); return <div></div> }}
+                        />
+                     </ScrollablePane>
+                  </Stack.Item>}
+               {state.mode === "telemetry" && <TelemetryPanel />}
             </Stack>
          </Stack>
          <DialogFooter>
