Horde: Bump dashboard version
#skipci
#rnx

diff --git a/Engine/Source/Programs/Horde/HordeDashboard/package.json b/Engine/Source/Programs/Horde/HordeDashboard/package.json
--- a/Engine/Source/Programs/Horde/HordeDashboard/package.json
+++ b/Engine/Source/Programs/Horde/HordeDashboard/package.json
@@ -1,7 +1,7 @@
 {
   "name": "horde-dashboard",
   "private": true,
-  "version": "0.0.2",
+  "version": "0.0.3",
   "type": "module",  
   "scripts": {
     "dev": "vite",
