Horde: Adding more context to agent telemetry, and tooltip
#rnx
#skipci

diff --git a/Engine/Source/Programs/Horde/HordeDashboard/src/backend/Backend.ts b/Engine/Source/Programs/Horde/HordeDashboard/src/backend/Backend.ts
--- a/Engine/Source/Programs/Horde/HordeDashboard/src/backend/Backend.ts
+++ b/Engine/Source/Programs/Horde/HordeDashboard/src/backend/Backend.ts
@@ -1,7 +1,7 @@
 // Copyright Epic Games, Inc. All Rights Reserved.
 
 import templateCache from '../backend/TemplateCache';
-import { AccountClaimMessage, AgentData, AgentQuery, ApproveAgentsRequest, ArtifactData, AuditLogEntry, AuditLogQuery, BatchUpdatePoolRequest, ChangeSummaryData, CreateAccountRequest, CreateBisectTaskRequest, CreateBisectTaskResponse, CreateDeviceRequest, CreateDeviceResponse, CreateExternalIssueRequest, CreateExternalIssueResponse, CreateJobRequest, CreateJobResponse, CreateNoticeRequest, CreatePoolRequest, CreateServiceAccountRequest, CreateServiceAccountResponse, CreateSoftwareResponse, CreateSubscriptionRequest, CreateSubscriptionResponse, CreateZipRequest, DashboardPreference, DevicePoolTelemetryQuery, DeviceTelemetryQuery, EventData, FindArtifactsResponse, FindIssueResponse, FindJobTimingsResponse, GetAccountResponse, GetAgentTelemetrySampleResponse, GetArtifactDirectoryResponse, GetArtifactZipRequest, GetBisectTaskResponse, GetDashboardConfigResponse, GetDevicePlatformResponse, GetDevicePoolResponse, GetDevicePoolTelemetryResponse, GetDeviceReservationResponse, GetDeviceResponse, GetDeviceTelemetryResponse, GetExternalIssueProjectResponse, GetExternalIssueResponse, GetIssueStreamResponse, GetJobsTabResponse, GetJobStepRefResponse, GetJobStepTraceResponse, GetJobTimingResponse, GetLogEventResponse, GetLogFileResponse, GetNoticeResponse, GetNotificationResponse, GetPendingAgentsResponse, GetPerforceServerStatusResponse, GetPoolResponse, GetPoolSummaryResponse, GetServerInfoResponse, GetServerSettingsResponse, GetServiceAccountResponse, GetSoftwareResponse, GetSubscriptionResponse, GetTelemetryMetricsResponse, GetTestDataDetailsResponse, GetTestDataRefResponse, GetTestMetaResponse, GetTestResponse, GetTestsRequest, GetTestStreamResponse, GetToolSummaryResponse, GetUserResponse, GetUtilizationTelemetryResponse, GlobalConfig, IssueData, IssueQuery, IssueQueryV2, JobData, JobQuery, JobsTabColumnType, JobStepOutcome, JobStreamQuery, JobTimingsQuery, LeaseData, LogData, LogEventQuery, LogLineData, MetricsQuery, PoolData, PoolQuery, PreflightConfigResponse, ProjectData, ScheduleData, ScheduleQuery, SearchLogFileResponse, ServerStatusResponse, ServerUpdateResponse, SessionData, StreamData, TabType, TestData, UpdateAccountRequest, UpdateAgentRequest, UpdateBisectTaskRequest, UpdateCurrentAccountRequest, UpdateDeviceRequest, UpdateGlobalConfigRequest, UpdateIssueRequest, UpdateJobRequest, UpdateLeaseRequest, UpdateNoticeRequest, UpdateNotificationsRequest, UpdatePoolRequest, UpdateServerSettingsRequest, UpdateServiceAccountRequest, UpdateServiceAccountResponse, UpdateStepRequest, UpdateStepResponse, UpdateTemplateRefRequest, UpdateUserRequest, UsersQuery } from './Api';
+import { AccountClaimMessage, AgentData, AgentQuery, ApproveAgentsRequest, ArtifactData, AuditLogEntry, AuditLogQuery, BatchUpdatePoolRequest, ChangeSummaryData, CreateAccountRequest, CreateBisectTaskRequest, CreateBisectTaskResponse, CreateDeviceRequest, CreateDeviceResponse, CreateExternalIssueRequest, CreateExternalIssueResponse, CreateJobRequest, CreateJobResponse, CreateNoticeRequest, CreatePoolRequest, CreateServiceAccountRequest, CreateServiceAccountResponse, CreateSoftwareResponse, CreateSubscriptionRequest, CreateSubscriptionResponse, CreateZipRequest, DashboardPreference, DevicePoolTelemetryQuery, DeviceTelemetryQuery, EventData, FindArtifactsResponse, FindIssueResponse, FindJobTimingsResponse, GetAccountResponse, GetAgentLeaseResponse, GetAgentSessionResponse, GetAgentTelemetrySampleResponse, GetArtifactDirectoryResponse, GetArtifactZipRequest, GetBisectTaskResponse, GetDashboardConfigResponse, GetDevicePlatformResponse, GetDevicePoolResponse, GetDevicePoolTelemetryResponse, GetDeviceReservationResponse, GetDeviceResponse, GetDeviceTelemetryResponse, GetExternalIssueProjectResponse, GetExternalIssueResponse, GetIssueStreamResponse, GetJobsTabResponse, GetJobStepRefResponse, GetJobStepTraceResponse, GetJobTimingResponse, GetLogEventResponse, GetLogFileResponse, GetNoticeResponse, GetNotificationResponse, GetPendingAgentsResponse, GetPerforceServerStatusResponse, GetPoolResponse, GetPoolSummaryResponse, GetServerInfoResponse, GetServerSettingsResponse, GetServiceAccountResponse, GetSoftwareResponse, GetSubscriptionResponse, GetTelemetryMetricsResponse, GetTestDataDetailsResponse, GetTestDataRefResponse, GetTestMetaResponse, GetTestResponse, GetTestsRequest, GetTestStreamResponse, GetToolSummaryResponse, GetUserResponse, GetUtilizationTelemetryResponse, GlobalConfig, IssueData, IssueQuery, IssueQueryV2, JobData, JobQuery, JobsTabColumnType, JobStepOutcome, JobStreamQuery, JobTimingsQuery, LeaseData, LogData, LogEventQuery, LogLineData, MetricsQuery, PoolData, PoolQuery, PreflightConfigResponse, ProjectData, ScheduleData, ScheduleQuery, SearchLogFileResponse, ServerStatusResponse, ServerUpdateResponse, SessionData, StreamData, TabType, TestData, UpdateAccountRequest, UpdateAgentRequest, UpdateBisectTaskRequest, UpdateCurrentAccountRequest, UpdateDeviceRequest, UpdateGlobalConfigRequest, UpdateIssueRequest, UpdateJobRequest, UpdateLeaseRequest, UpdateNoticeRequest, UpdateNotificationsRequest, UpdatePoolRequest, UpdateServerSettingsRequest, UpdateServiceAccountRequest, UpdateServiceAccountResponse, UpdateStepRequest, UpdateStepResponse, UpdateTemplateRefRequest, UpdateUserRequest, UsersQuery } from './Api';
 import dashboard, { Dashboard } from './Dashboard';
 import { ChallengeStatus, Fetch } from './Fetch';
 import { projectStore } from './ProjectStore';
@@ -263,6 +263,23 @@ export class Backend {
         });
     }
 
+    getAgentLeases(agentId: string, minTime: Date, maxTime: Date): Promise<GetAgentLeaseResponse[]> {
+        const params: string[] = [];
+        params.push("startTime=" + minTime.toISOString());
+        params.push("finishime=" + maxTime.toISOString());
+        const paramString = params.join("&");
+        return new Promise<GetAgentLeaseResponse[]>((resolve, reject) => {
+            this.backend.get(`/api/v1/agents/${agentId}/leases?${paramString}`).then((response) => {
+                (response.data as GetAgentLeaseResponse[]).forEach(s => {
+                    s.startTime = s.startTime ? new Date(s.startTime) : undefined;
+                    s.finishTime = s.finishTime ? new Date(s.finishTime) : undefined;
+                })
+                resolve(response.data as GetAgentLeaseResponse[]);
+            }).catch(reason => { reject(reason); });
+        });
+    }
+
+
     getAgentTelemetry(agentId: string, minTime: Date, maxTime: Date): Promise<GetAgentTelemetrySampleResponse[]> {
         const params: string[] = [];
         params.push("minTime=" + minTime.toISOString());
diff --git a/Engine/Source/Programs/Horde/HordeDashboard/src/components/HistoryModal.tsx b/Engine/Source/Programs/Horde/HordeDashboard/src/components/HistoryModal.tsx
--- a/Engine/Source/Programs/Horde/HordeDashboard/src/components/HistoryModal.tsx
+++ b/Engine/Source/Programs/Horde/HordeDashboard/src/components/HistoryModal.tsx
@@ -1,6 +1,6 @@
 // Copyright Epic Games, Inc. All Rights Reserved.
 
-import { Checkbox, ComboBox, ConstrainMode, DefaultButton, DetailsHeader, DetailsList, DetailsListLayoutMode, DetailsRow, Dialog, DialogFooter, DialogType, GroupedList, GroupHeader, ICheckbox, IColumn, IComboBoxOption, IContextualMenuItem, IContextualMenuProps, IDetailsHeaderProps, IDetailsHeaderStyles, IDetailsListProps, IGroup, ITextField, ITooltipHostStyles, Label, mergeStyleSets, Modal, Pivot, PivotItem, PrimaryButton, ScrollablePane, ScrollbarVisibility, Selection, SelectionMode, Spinner, SpinnerSize, Stack, Sticky, StickyPositionType, Text, TextField } from "@fluentui/react";
+import { Checkbox, ComboBox, ConstrainMode, DefaultButton, DetailsHeader, DetailsList, DetailsListLayoutMode, DetailsRow, Dialog, DialogFooter, DialogType, FontIcon, GroupedList, GroupHeader, ICheckbox, IColumn, IComboBoxOption, IContextualMenuItem, IContextualMenuProps, IDetailsHeaderProps, IDetailsHeaderStyles, IDetailsListProps, IGroup, ITextField, ITooltipHostStyles, Label, mergeStyleSets, Modal, Pivot, PivotItem, PrimaryButton, ScrollablePane, ScrollbarVisibility, Selection, SelectionMode, Spinner, SpinnerSize, Stack, Sticky, StickyPositionType, Text, TextField } from "@fluentui/react";
 import { action, makeObservable, observable } from "mobx";
 import { observer } from "mobx-react-lite";
 import React, { useEffect, useState } from "react";
@@ -16,6 +16,15 @@ import { getHordeTheme } from "../styles/theme";
 import { BatchStatusIcon, LeaseStatusIcon, StepStatusIcon } from "./StatusIcon";
 
 
+type LeaseTooltip = {
+   lease?: GetAgentLeaseResponse;
+   time?: Date;
+   id?: string;
+   x?: number;
+   y?: number;
+   frozen?: boolean;
+}
+
 type InfoPanelItem = {
    key: string;
    name: string;
@@ -55,6 +64,9 @@ class HistoryModalState {
    @observable
    selectedAgentUpdated = 0;
 
+   @observable
+   currentLeaseTooltip?: LeaseTooltip;
+
    get selectedAgent(): AgentData | undefined {
       // subscribe in any observers
       if (this.selectedAgentUpdated) { }
@@ -110,6 +122,11 @@ class HistoryModalState {
       }
    }
 
+   @action
+   setLeaseTooltip(tooltip?: LeaseTooltip) {
+      this.currentLeaseTooltip = tooltip;
+   }
+
    setInfoItemSelected(item: InfoPanelItem) {
       this.infoItems.forEach(infoItem => {
          if (infoItem.key === item.key) {
@@ -270,14 +287,52 @@ class HistoryModalState {
    }
 
    UpdateTelemetry(minutes: number) {
-      backend.getAgentTelemetry(this.selectedAgent!.id, new Date(Date.now() - 1000 * 60 * minutes), new Date()).then(data => {
-         this.setTelemetryData(data);
+
+      const requests = [backend.getAgentTelemetry(this.selectedAgent!.id, new Date(Date.now() - 1000 * 60 * minutes), new Date()), backend.getAgentLeases(this.selectedAgent!.id, new Date(Date.now() - 1000 * 60 * minutes), new Date())]
+      Promise.all(requests).then(data => {
+         this.setTelemetryData((data[0] as GetAgentTelemetrySampleResponse[]).reverse(), data[1] as GetAgentLeaseResponse[]);
       })
    }
 
    @action
-   setTelemetryData(newData: GetAgentTelemetrySampleResponse[]) {
-      this.currentData = newData;
+   setTelemetryData(telemetryData: GetAgentTelemetrySampleResponse[], sessionData: GetAgentLeaseResponse[]) {
+
+      const linear: GetAgentTelemetrySampleResponse[] = [];
+
+      if (telemetryData.length) {
+
+         let ctime: number = 0;
+
+         telemetryData.forEach(t => {
+
+            if (!ctime) {
+               ctime = t.time.getTime();
+               linear.push(t);
+            } else {
+               const ntime = t.time.getTime();
+               // dependant on agent sample rate!               
+               while ((ntime - ctime) > 60000) {
+
+                  ctime += 30000;
+                  linear.push({
+                     time: new Date(ctime),
+                     userCpu: 0,
+                     idleCpu: 0,
+                     systemCpu: 0,
+                     freeRam: 0,
+                     usedRam: 0,
+                     totalRam: 1
+                  });
+               }
+
+               linear.push(t);
+               ctime = ntime;
+            }
+         });
+      }
+
+      this.currentData = [linear, sessionData];
+
    }
 
    @action
@@ -319,6 +374,108 @@ class HistoryModalState {
 
 const state = new HistoryModalState();
 
+
+const TelemetryTooltip: React.FC<{ id: string }> = observer(({ id }) => {
+
+   const { modeColors } = getHordeStyling();
+
+   const tooltip = state.currentLeaseTooltip;
+
+   if (!tooltip?.lease || tooltip.id !== id) {
+      return null;
+   }
+
+   let tipX = tooltip.x;
+   let offsetX = 32;
+   let translateX = "0%";
+
+   if (tipX > 600) {
+      offsetX = -32;
+      translateX = "-100%";
+   }
+
+   const translateY = "-50%";
+
+   const textSize = 11;
+   const titleWidth = 48;
+
+   const dataElement = (title: string, value: string, link?: string, agent?: string, linkTarget?: string) => {
+
+      let valueElement = <Stack>
+         <Stack style={{ fontSize: textSize }}>
+            {value}
+         </Stack>
+      </Stack>
+
+      if (link) {
+         valueElement = <Link className="horde-link" to={link} target={linkTarget}>
+            {valueElement}
+         </Link>
+      }
+
+      let component = <Stack style={{ cursor: (link || agent) ? "pointer" : undefined }} onClick={() => {
+      }}>
+         <Stack horizontal>
+            <Stack style={{ width: titleWidth }}>
+               {!!title && <Text style={{ fontFamily: "Horde Open Sans SemiBold", fontSize: textSize }} >{title}</Text>}
+            </Stack>
+            {valueElement}
+         </Stack>
+      </Stack>
+
+      return component;
+   }
+
+   const elements: JSX.Element[] = [];
+
+   const lease = tooltip.lease;
+
+   if (lease?.details?.type === "Job") {
+      elements.push(dataElement("Job:", tooltip.lease?.name ?? "Unknown Job Lease", `/job/${lease.details.jobId}`));
+   } else {
+      elements.push(dataElement("Lease:", tooltip.lease?.name ?? "Unknown Lease"));
+   }
+
+   if (lease?.details?.logId) {
+      elements.push(dataElement("Log:", `Lease Log`, `/log/${lease.details.logId}`));
+   }
+
+
+
+   if (tooltip.time) {
+      elements.push(dataElement("Time:", getNiceTime(tooltip.time)));
+   }
+
+
+   return <div style={{
+      position: "absolute",
+      display: "block",
+      top: `${tooltip.y}px`,
+      left: `${tooltip.x + offsetX}px`,
+      backgroundColor: modeColors.background,
+      zIndex: 1,
+      border: "solid",
+      borderWidth: "1px",
+      borderRadius: "3px",
+      width: "max-content",
+      borderColor: dashboard.darktheme ? "#413F3D" : "#2D3F5F",
+      pointerEvents: tooltip.frozen ? undefined : "none",
+      transform: `translate(${translateX}, ${translateY})`
+   }}>
+      <Stack horizontal>
+         <Stack tokens={{ childrenGap: 6, padding: 12 }}>
+            {elements}
+         </Stack>
+         {!!tooltip.frozen && <Stack style={{ paddingLeft: 32, paddingRight: 12, paddingTop: 12, cursor: "pointer" }} onClick={() => { state.setLeaseTooltip(undefined) }}>
+            <FontIcon iconName="Cancel" />
+         </Stack>}
+         {!tooltip.frozen && <Stack style={{ paddingLeft: 32, paddingRight: 12 + 16, paddingTop: 12 }} />}
+
+      </Stack>
+   </div>
+
+});
+
 type TimeSelection = {
    text: string;
    key: string;
@@ -385,14 +542,18 @@ const TelemetryPanel: React.FC<{}> = observer(({ }) => {
    // subscribe
    if (state.currentData) { }
 
-   const data = (state.currentData ?? []) as GetAgentTelemetrySampleResponse[];
+   const currentData = (state.currentData ?? []);
 
-   if (!data.length) {
+   if (!currentData.length) {
       return <Stack horizontalAlign="center">
          <Spinner size={SpinnerSize.large} />
       </Stack>;
    }
 
+   const sparkWidth = 1100;
+   const data = (currentData[0] as GetAgentTelemetrySampleResponse[]);
+   const leases = ((currentData.length > 1) ? currentData[1] : []) as GetAgentLeaseResponse[];
+
    const cpuData = data.map(d => {
       if (cpuKey == "user") return d.userCpu;
       if (cpuKey == "system") return d.systemCpu;
@@ -405,8 +566,9 @@ const TelemetryPanel: React.FC<{}> = observer(({ }) => {
    const minTime = getShortNiceTime(min, true);
    const maxTime = getShortNiceTime(max, true);
 
-   const ramData = (state.currentData ?? []).map(m => {
-      const d = m as GetAgentTelemetrySampleResponse;
+
+
+   const ramData = data.map(d => {
       return d.usedRam / d.totalRam;
    })
 
@@ -430,11 +592,35 @@ const TelemetryPanel: React.FC<{}> = observer(({ }) => {
       }
    })
 
+   function relativeCoords(event: any, id: string) {
 
+      if (state.currentLeaseTooltip?.frozen) {
+         return;
+      }
+
+      let bounds = event.target.getBoundingClientRect();
+      let x = event.clientX - bounds.left;
+      let y = event.clientY - bounds.top;
+
+      // interpolate time
+
+      const percent = ((x - 16) / sparkWidth);
+
+      const time = new Date(min.getTime() + ((max.getTime() - min.getTime()) * percent));
+
+      let lease: GetAgentLeaseResponse | undefined;
+      leases.forEach(s => {
+         if (time >= s.startTime && (!s.finishTime || time <= s.finishTime)) {
+            lease = s;
+         }
+      });
+
+      state.setLeaseTooltip({ lease: lease, x: x, y: y, id: id, time: time });
+   };
 
    return <Stack tokens={{ childrenGap: 24 }} style={{ paddingBottom: 32, paddingLeft: 8 }}>
       <Stack horizontal verticalAlign="center" verticalFill>
-         <Text variant="mediumPlus">{maxTime} - {minTime}</Text>
+         <Text variant="mediumPlus">{minTime} - {maxTime}</Text>
          <Stack grow />
          <Stack horizontal tokens={{ childrenGap: 18 }}>
 
@@ -467,10 +653,23 @@ const TelemetryPanel: React.FC<{}> = observer(({ }) => {
             <Stack grow></Stack>
             {!!cpuText && <Label>{cpuText}</Label>}
          </Stack>
-         <Sparklines width={1100} height={128} data={cpuData} max={100} style={{ backgroundColor: dashboard.darktheme ? "#060709" : "#F3F2F1", padding: 8, border: "solid 1px #181A1B" }}>
-            <SparklinesLine color={dashboard.darktheme ? "lightblue" : "#1E90FF"} />
-            <SparklinesReferenceLine type="avg" />
-         </Sparklines>
+         <div onMouseLeave={() => state.setLeaseTooltip(undefined)} onMouseMove={(ev) => { relativeCoords(ev, "cpu") }} onClick={() => {
+
+            if (!state.currentLeaseTooltip?.frozen) {
+               if (state.currentLeaseTooltip) {
+                  state.setLeaseTooltip({ ...state.currentLeaseTooltip, frozen: true });
+               }
+            }
+         }}>
+            <div style={{ position: "relative" }}>
+               <TelemetryTooltip id={"cpu"} />
+            </div>
+
+            <Sparklines width={sparkWidth} height={128} data={cpuData} max={100} style={{ backgroundColor: dashboard.darktheme ? "#060709" : "#F3F2F1", padding: 8, border: "solid 1px #181A1B" }}>
+               <SparklinesLine color={dashboard.darktheme ? "lightblue" : "#1E90FF"} />
+               <SparklinesReferenceLine type="avg" />
+            </Sparklines>
+         </div>
       </Stack>
       <Stack verticalAlign="center" verticalFill>
          <Stack horizontal>
@@ -478,10 +677,22 @@ const TelemetryPanel: React.FC<{}> = observer(({ }) => {
             <Stack grow></Stack>
             {!!ramText && <Label>{ramText}</Label>}
          </Stack>
-         <Sparklines width={1100} height={128} data={ramData} max={1} style={{ backgroundColor: dashboard.darktheme ? "#060709" : "#F3F2F1", padding: 8, border: "solid 1px #181A1B" }}>
-            <SparklinesLine color={dashboard.darktheme ? "lightblue" : "#1E90FF"} />
-            <SparklinesReferenceLine type="avg" />
-         </Sparklines>
+         <div onMouseLeave={() => state.setLeaseTooltip(undefined)} onMouseMove={(ev) => { relativeCoords(ev, "ram") }} onClick={() => {
+
+            if (!state.currentLeaseTooltip?.frozen) {
+               if (state.currentLeaseTooltip) {
+                  state.setLeaseTooltip({ ...state.currentLeaseTooltip, frozen: true });
+               }
+            }
+         }}>
+            <div style={{ position: "relative" }}>
+               <TelemetryTooltip id="ram" />
+            </div>
+            <Sparklines width={sparkWidth} height={128} data={ramData} max={1} style={{ backgroundColor: dashboard.darktheme ? "#060709" : "#F3F2F1", padding: 8, border: "solid 1px #181A1B" }}>
+               <SparklinesLine color={dashboard.darktheme ? "lightblue" : "#1E90FF"} />
+               <SparklinesReferenceLine type="avg" />
+            </Sparklines>
+         </div>
       </Stack>
 
    </Stack>
