Horde: Add steps to agent telemetry
#rnx
#submitool bypass
#skipci

diff --git a/Engine/Source/Programs/Horde/HordeDashboard/src/components/HistoryModal.tsx b/Engine/Source/Programs/Horde/HordeDashboard/src/components/HistoryModal.tsx
--- a/Engine/Source/Programs/Horde/HordeDashboard/src/components/HistoryModal.tsx
+++ b/Engine/Source/Programs/Horde/HordeDashboard/src/components/HistoryModal.tsx
@@ -8,7 +8,7 @@ import { Link } from "react-router-dom";
 import { Sparklines, SparklinesLine, SparklinesReferenceLine } from "react-sparklines";
 import backend from "../backend";
 import { agentStore } from "../backend/AgentStore";
-import { AgentData, GetAgentLeaseResponse, GetAgentSessionResponse, GetAgentTelemetrySampleResponse, JobStepBatchError, JobStepOutcome, LeaseData, SessionData, UpdateAgentRequest } from "../backend/Api";
+import { AgentData, GetAgentLeaseResponse, GetAgentSessionResponse, GetAgentTelemetrySampleResponse, GetJobResponse, JobStepBatchError, JobStepOutcome, LeaseData, SessionData, UpdateAgentRequest } from "../backend/Api";
 import dashboard from "../backend/Dashboard";
 import { getNiceTime, getShortNiceTime } from "../base/utilities/timeUtils";
 import { getHordeStyling } from "../styles/Styles";
@@ -286,10 +286,12 @@ class HistoryModalState {
       });
    }
 
+   jobData = new Map<string, GetJobResponse | boolean>();
+
    UpdateTelemetry(minutes: number) {
 
       const requests = [backend.getAgentTelemetry(this.selectedAgent!.id, new Date(Date.now() - 1000 * 60 * minutes), new Date()), backend.getAgentLeases(this.selectedAgent!.id, new Date(Date.now() - 1000 * 60 * minutes), new Date())]
-      Promise.all(requests).then(data => {
+      Promise.all(requests).then((data) => {
          this.setTelemetryData((data[0] as GetAgentTelemetrySampleResponse[]).reverse(), data[1] as GetAgentLeaseResponse[]);
       })
    }
@@ -432,12 +434,28 @@ const TelemetryTooltip: React.FC<{ id: string }> = observer(({ id }) => {
 
    if (lease?.details?.type === "Job") {
       elements.push(dataElement("Job:", tooltip.lease?.name ?? "Unknown Job Lease", `/job/${lease.details.jobId}`));
+      const job = state.jobData.get(lease.details.jobId);
+      if (job && job !== true && job.batches) {
+         const step = job.batches.map(b => b.agentId === state.selectedAgent?.id ? b.steps : []).flat().find(s => {
+            if (s.startTime && tooltip.time) {
+               if (tooltip.time >= new Date(s.startTime) && (!s.finishTime || tooltip.time <= new Date(s.finishTime))) {
+                  return true;
+               }
+            }
+            return false;
+         });
+         if (step) {
+            elements.push(dataElement("Step:", step.name ?? "Unknown Job Step", `/job/${lease.details.jobId}?step=${step.id}`));
+            if (step.logId)
+               elements.push(dataElement("Log:", `Step`, `/log/${step.logId}`));
+         }
+      }
    } else {
       elements.push(dataElement("Lease:", tooltip.lease?.name ?? "Unknown Lease"));
    }
 
    if (lease?.details?.logId) {
-      elements.push(dataElement("Log:", `Lease Log`, `/log/${lease.details.logId}`));
+      elements.push(dataElement("Log:", `Lease`, `/log/${lease.details.logId}`));
    }
 
 
@@ -615,6 +633,19 @@ const TelemetryPanel: React.FC<{}> = observer(({ }) => {
          }
       });
 
+      if (lease?.details?.type === "Job") {
+         const jobId = lease.details.jobId;
+         if (jobId) {
+            if (!state.jobData.has(jobId)) {
+               state.jobData.set(jobId, true);
+               backend.getJob(jobId, { filter: "id,batches" }).then(job => {
+                  state.jobData.set(jobId, job);
+                  state.setLeaseTooltip(state.currentLeaseTooltip ? { ...state.currentLeaseTooltip } : undefined);
+               })
+            }
+         }
+      }
+
       state.setLeaseTooltip({ lease: lease, x: x, y: y, id: id, time: time });
    };
 
