Horde: Guard against lack of agent telemetry data
#rnx
#skipci
#skipci

diff --git a/Engine/Source/Programs/Horde/HordeDashboard/src/components/HistoryModal.tsx b/Engine/Source/Programs/Horde/HordeDashboard/src/components/HistoryModal.tsx
--- a/Engine/Source/Programs/Horde/HordeDashboard/src/components/HistoryModal.tsx
+++ b/Engine/Source/Programs/Horde/HordeDashboard/src/components/HistoryModal.tsx
@@ -568,10 +568,15 @@ const TelemetryPanel: React.FC<{}> = observer(({ }) => {
       </Stack>;
    }
 
+
    const sparkWidth = 1100;
    const data = (currentData[0] as GetAgentTelemetrySampleResponse[]);
    const leases = ((currentData.length > 1) ? currentData[1] : []) as GetAgentLeaseResponse[];
 
+   if (!data.length) {
+      return null;
+   }
+
    const cpuData = data.map(d => {
       if (cpuKey == "user") return d.userCpu;
       if (cpuKey == "system") return d.systemCpu;
