Horde: Report OpenTelemetry metrics for deleted agent count

#rnx

diff --git a/Engine/Source/Programs/Horde/Horde.Server/Agents/AgentService.cs b/Engine/Source/Programs/Horde/Horde.Server/Agents/AgentService.cs
--- a/Engine/Source/Programs/Horde/Horde.Server/Agents/AgentService.cs
+++ b/Engine/Source/Programs/Horde/Horde.Server/Agents/AgentService.cs
@@ -1140,6 +1140,7 @@ namespace Horde.Server.Agents
 
 			IReadOnlyList<IAgent> agentList = await Agents.FindAsync(cancellationToken: cancellationToken);
 			int numAgentsTotal = agentList.Count;
+			int numAgentsTotalDeleted = agentList.Count(a => a.Deleted);
 			int numAgentsTotalEnabled = agentList.Count(a => a.Enabled);
 			int numAgentsTotalDisabled = agentList.Count(a => !a.Enabled);
 			int numAgentsTotalOk = agentList.Count(a => a.Enabled && a.Status == AgentStatus.Ok);
@@ -1148,17 +1149,18 @@ namespace Horde.Server.Agents
 			int numAgentsTotalBusy = agentList.Count(a => a.Enabled && a.Status == AgentStatus.Busy);
 			int numAgentsTotalUnspecified = agentList.Count(a => a.Enabled && a.Status == AgentStatus.Unspecified);
 
-			List<Measurement<int>> newMeasurements = new()
-			{
+			List<Measurement<int>> newMeasurements =
+			[
 				new Measurement<int>(numAgentsTotal),
+				new Measurement<int>(numAgentsTotalDeleted, new KeyValuePair<string, object?>("status", "deleted")),
 				new Measurement<int>(numAgentsTotalEnabled, new KeyValuePair<string, object?>("status", "enabled")),
 				new Measurement<int>(numAgentsTotalDisabled, new KeyValuePair<string, object?>("status", "disabled")),
 				new Measurement<int>(numAgentsTotalOk, new KeyValuePair<string, object?>("status", "ok")),
 				new Measurement<int>(numAgentsTotalStopping, new KeyValuePair<string, object?>("status", "stopping")),
 				new Measurement<int>(numAgentsTotalUnhealthy, new KeyValuePair<string, object?>("status", "unhealthy")),
 				new Measurement<int>(numAgentsTotalBusy, new KeyValuePair<string, object?>("status", "paused")),
-				new Measurement<int>(numAgentsTotalUnspecified, new KeyValuePair<string, object?>("status", "unspecified")),
-			};
+				new Measurement<int>(numAgentsTotalUnspecified, new KeyValuePair<string, object?>("status", "unspecified"))
+			];
 
 			_measurements = newMeasurements;
 		}
