Horde: Set the appropriate default parameters when starting a scheduled build.

diff --git a/Engine/Source/Programs/Horde/Horde.Server/Jobs/JobService.cs b/Engine/Source/Programs/Horde/Horde.Server/Jobs/JobService.cs
--- a/Engine/Source/Programs/Horde/Horde.Server/Jobs/JobService.cs
+++ b/Engine/Source/Programs/Horde/Horde.Server/Jobs/JobService.cs
@@ -1372,7 +1372,8 @@ namespace Horde.Server.Jobs
 					CreateJobOptions options = new CreateJobOptions(templateRefConfig);
 					options.PreflightChange = newJob.PreflightChange;
 					options.PreflightDescription = newJob.PreflightDescription;
-					options.Arguments.AddRange(template.GetDefaultArguments(true));
+					template.GetDefaultParameters(options.Parameters, true);
+					template.GetArgumentsForParameters(null, options.Parameters, options.Arguments);
 
 					await CreateJobAsync(chainedJobId, streamConfig, jobTrigger.TemplateRefId, template.Hash, triggerGraph, templateRefConfig.Name, change.Value, codeChange, options);
 					return newJob;
diff --git a/Engine/Source/Programs/Horde/Horde.Server/Jobs/JobsController.cs b/Engine/Source/Programs/Horde/Horde.Server/Jobs/JobsController.cs
--- a/Engine/Source/Programs/Horde/Horde.Server/Jobs/JobsController.cs
+++ b/Engine/Source/Programs/Horde/Horde.Server/Jobs/JobsController.cs
@@ -179,8 +179,6 @@ namespace Horde.Server.Jobs
 			int codeChange = lastCodeCommit?.Number ?? change;
 
 			// New properties for the job
-			List<string> arguments = create.Arguments ?? template.GetDefaultArguments(false);
-
 			bool? updateIssues = null;
 			if (template.UpdateIssues)
 			{
@@ -200,21 +198,23 @@ namespace Horde.Server.Jobs
 			options.AutoSubmit = create.AutoSubmit;
 			options.UpdateIssues = updateIssues;
 			options.Claims.AddRange(User.Claims.Select(x => new AclClaimConfig(x)));
-			options.Arguments.AddRange(arguments);
 			options.JobOptions ??= create.JobOptions;
 
+			// Find all the default parameters, and override any settings with the values in the request
+			template.GetDefaultParameters(options.Parameters, false);
+
 			if (create.Parameters != null)
 			{
 				foreach ((ParameterId parameter, string value) in create.Parameters)
 				{
 					options.Parameters[parameter] = value;
 				}
-				foreach (Parameter parameter in template.Parameters)
-				{
-					parameter.GetArguments(options.Parameters, false, options.Arguments);
-				}
 			}
 
+			// Build the final arguments list from the combined parameter set
+			template.GetArgumentsForParameters(create.Arguments, options.Parameters, options.Arguments);
+
+			// Merge the environment variables
 			foreach ((string key, string value) in environment)
 			{
 				options.Environment[key] = value;
diff --git a/Engine/Source/Programs/Horde/Horde.Server/Jobs/Schedules/ScheduleService.cs b/Engine/Source/Programs/Horde/Horde.Server/Jobs/Schedules/ScheduleService.cs
--- a/Engine/Source/Programs/Horde/Horde.Server/Jobs/Schedules/ScheduleService.cs
+++ b/Engine/Source/Programs/Horde/Horde.Server/Jobs/Schedules/ScheduleService.cs
@@ -468,7 +468,8 @@ namespace Horde.Server.Jobs.Schedules
 
 				CreateJobOptions options = new CreateJobOptions(templateRef.Config);
 				options.Priority = template.Priority;
-				options.Arguments.AddRange(template.GetDefaultArguments(true));
+				template.GetDefaultParameters(options.Parameters, true);
+				template.GetArgumentsForParameters(null, options.Parameters, options.Arguments);
 
 				IJob newJob = await _jobService.CreateJobAsync(null, stream.Config, templateId, template.Hash, graph, template.Name, change, codeChange, options, cancellationToken);
 				_logger.LogInformation("Started new job for {StreamId} template {TemplateId} at CL {Change} (Code CL {CodeChange}): {JobId}", stream.Id, templateId, change, codeChange, newJob.Id);
diff --git a/Engine/Source/Programs/Horde/Horde.Server/Jobs/Templates/ITemplate.cs b/Engine/Source/Programs/Horde/Horde.Server/Jobs/Templates/ITemplate.cs
--- a/Engine/Source/Programs/Horde/Horde.Server/Jobs/Templates/ITemplate.cs
+++ b/Engine/Source/Programs/Horde/Horde.Server/Jobs/Templates/ITemplate.cs
@@ -22,14 +22,14 @@ namespace Horde.Server.Jobs.Templates
 		/// <param name="parameters">Map of parameter id to value</param>
 		/// <param name="scheduledBuild">Whether this is a scheduled build</param>
 		/// <param name="arguments">Receives command line arguments for the job</param>
-		public abstract void GetArguments(Dictionary<ParameterId, string> parameters, bool scheduledBuild, List<string> arguments);
+		public abstract void GetArguments(IReadOnlyDictionary<ParameterId, string> parameters, bool scheduledBuild, List<string> arguments);
 
 		/// <summary>
 		/// Gets the default arguments for this parameter and its children
 		/// </summary>
-		/// <param name="defaultArguments">List of default arguments</param>
+		/// <param name="parameters">List of default parameters</param>
 		/// <param name="scheduledBuild">Whether the arguments are being queried for a scheduled build</param>
-		public abstract void GetDefaultArguments(List<string> defaultArguments, bool scheduledBuild);
+		public abstract void GetDefaultParameters(Dictionary<ParameterId, string> parameters, bool scheduledBuild);
 
 		/// <summary>
 		/// Convert this parameter to data for serialization
@@ -81,7 +81,7 @@ namespace Horde.Server.Jobs.Templates
 		}
 
 		/// <inheritdoc/>
-		public override void GetArguments(Dictionary<ParameterId, string> parameters, bool scheduledBuild, List<string> arguments)
+		public override void GetArguments(IReadOnlyDictionary<ParameterId, string> parameters, bool scheduledBuild, List<string> arguments)
 		{
 			foreach (Parameter child in Children)
 			{
@@ -90,11 +90,11 @@ namespace Horde.Server.Jobs.Templates
 		}
 
 		/// <inheritdoc/>
-		public override void GetDefaultArguments(List<string> defaultArguments, bool scheduledBuild)
+		public override void GetDefaultParameters(Dictionary<ParameterId, string> parameters, bool scheduledBuild)
 		{
 			foreach (Parameter child in Children)
 			{
-				child.GetDefaultArguments(defaultArguments, scheduledBuild);
+				child.GetDefaultParameters(parameters, scheduledBuild);
 			}
 		}
 
@@ -196,7 +196,7 @@ namespace Horde.Server.Jobs.Templates
 		}
 
 		/// <inheritdoc/>
-		public override void GetArguments(Dictionary<ParameterId, string> parameters, bool scheduledBuild, List<string> arguments)
+		public override void GetArguments(IReadOnlyDictionary<ParameterId, string> parameters, bool scheduledBuild, List<string> arguments)
 		{
 			if (parameters.TryGetValue(Id, out string? value))
 			{
@@ -209,9 +209,9 @@ namespace Horde.Server.Jobs.Templates
 		}
 
 		/// <inheritdoc/>
-		public override void GetDefaultArguments(List<string> defaultArguments, bool scheduledBuild)
+		public override void GetDefaultParameters(Dictionary<ParameterId, string> parameters, bool scheduledBuild)
 		{
-			defaultArguments.Add(Argument + (scheduledBuild ? (ScheduleOverride ?? Default) : Default));
+			parameters[Id] = scheduledBuild ? (ScheduleOverride ?? Default) : Default;
 		}
 
 		/// <summary>
@@ -373,7 +373,7 @@ namespace Horde.Server.Jobs.Templates
 		}
 
 		/// <inheritdoc/>
-		public override void GetArguments(Dictionary<ParameterId, string> parameters, bool scheduledBuild, List<string> arguments)
+		public override void GetArguments(IReadOnlyDictionary<ParameterId, string> parameters, bool scheduledBuild, List<string> arguments)
 		{
 			foreach (ListParameterItem item in Items)
 			{
@@ -383,12 +383,12 @@ namespace Horde.Server.Jobs.Templates
 		}
 
 		/// <inheritdoc/>
-		public override void GetDefaultArguments(List<string> defaultArguments, bool scheduledBuild)
+		public override void GetDefaultParameters(Dictionary<ParameterId, string> parameters, bool scheduledBuild)
 		{
 			foreach (ListParameterItem item in Items)
 			{
 				bool value = scheduledBuild ? (item.ScheduleOverride ?? item.Default) : item.Default;
-				GetCommandLineArgumentsForItem(item, value, defaultArguments);
+				parameters[item.Id] = value.ToString();
 			}
 		}
 
@@ -522,7 +522,7 @@ namespace Horde.Server.Jobs.Templates
 		/// </summary>
 		/// <param name="parameterId">The parameter id</param>
 		/// <param name="parameters">Map from parameter id to value</param>
-		public static bool? GetValue(ParameterId parameterId, Dictionary<ParameterId, string> parameters)
+		public static bool? GetValue(ParameterId parameterId, IReadOnlyDictionary<ParameterId, string> parameters)
 		{
 			if (parameters.TryGetValue(parameterId, out string? stringValue))
 			{
@@ -539,17 +539,17 @@ namespace Horde.Server.Jobs.Templates
 		}
 
 		/// <inheritdoc/>
-		public override void GetArguments(Dictionary<ParameterId, string> parameters, bool scheduledBuild, List<string> arguments)
+		public override void GetArguments(IReadOnlyDictionary<ParameterId, string> parameters, bool scheduledBuild, List<string> arguments)
 		{
 			bool value = GetValue(Id, parameters) ?? (scheduledBuild ? (ScheduleOverride ?? Default) : Default);
 			GetCommandLineArgumentsForItem(value, arguments);
 		}
 
 		/// <inheritdoc/>
-		public override void GetDefaultArguments(List<string> defaultArguments, bool scheduledBuild)
+		public override void GetDefaultParameters(Dictionary<ParameterId, string> parameters, bool scheduledBuild)
 		{
 			bool value = scheduledBuild ? (ScheduleOverride ?? Default) : Default;
-			GetCommandLineArgumentsForItem(value, defaultArguments);
+			parameters[Id] = value.ToString();
 		}
 
 		void GetCommandLineArgumentsForItem(bool value, List<string> arguments)
@@ -658,18 +658,29 @@ namespace Horde.Server.Jobs.Templates
 	/// </summary>
 	public static class TemplateExtensions
 	{
+		/// <summary>
+		/// Gets the full argument list for a template
+		/// </summary>
+		public static void GetArgumentsForParameters(this ITemplate template, IReadOnlyList<string>? baseArguments, IReadOnlyDictionary<ParameterId, string> parameters, List<string> arguments)
+		{
+			arguments.AddRange(baseArguments ?? template.Arguments);
+
+			foreach (Parameter parameter in template.Parameters)
+			{
+				parameter.GetArguments(parameters, false, arguments);
+			}
+		}
+
 		/// <summary>
 		/// Gets the arguments for default options in this template. Does not include the standard template arguments.
 		/// </summary>
 		/// <returns>List of default arguments</returns>
-		public static List<string> GetDefaultArguments(this ITemplate template, bool scheduledBuild)
+		public static void GetDefaultParameters(this ITemplate template, Dictionary<ParameterId, string> parameters, bool scheduledBuild)
 		{
-			List<string> defaultArguments = new List<string>(template.Arguments);
 			foreach (Parameter parameter in template.Parameters)
 			{
-				parameter.GetDefaultArguments(defaultArguments, scheduledBuild);
+				parameter.GetDefaultParameters(parameters, scheduledBuild);
 			}
-			return defaultArguments;
 		}
 	}
 }
