Horde: Move expired and ephemeral agent clean up to shared ticker

diff --git a/Engine/Source/Programs/Horde/Horde.Server.Tests/Agents/AgentServiceTest.cs b/Engine/Source/Programs/Horde/Horde.Server.Tests/Agents/AgentServiceTest.cs
--- a/Engine/Source/Programs/Horde/Horde.Server.Tests/Agents/AgentServiceTest.cs
+++ b/Engine/Source/Programs/Horde/Horde.Server.Tests/Agents/AgentServiceTest.cs
@@ -186,13 +186,13 @@ public class AgentServiceTest : TestSetup
 
 		// Let background task run for purging outdated sessions, which will terminate session for our agent
 		await Clock.AdvanceAsync(TimeSpan.FromHours(1));
-		await AgentService.TickAsync(CancellationToken.None);
+		await AgentService.TickSharedAsync(CancellationToken.None);
 
 		// Ephemeral agent is marked as deleted once its session is terminated
 		Assert.IsTrue((await AgentService.GetAgentAsync(agent.Id))!.Deleted);
 
 		await Clock.AdvanceAsync(TimeSpan.FromHours(25));
-		await AgentService.TickAsync(CancellationToken.None);
+		await AgentService.TickSharedAsync(CancellationToken.None);
 
 		// Once more time has passed, the ephemeral agent marked as deleted is removed from database
 		Assert.IsNull(await AgentService.GetAgentAsync(agent.Id));
diff --git a/Engine/Source/Programs/Horde/Horde.Server/Agents/AgentService.cs b/Engine/Source/Programs/Horde/Horde.Server/Agents/AgentService.cs
--- a/Engine/Source/Programs/Horde/Horde.Server/Agents/AgentService.cs
+++ b/Engine/Source/Programs/Horde/Horde.Server/Agents/AgentService.cs
@@ -82,6 +82,7 @@ namespace Horde.Server.Agents
 		readonly Tracer _tracer;
 		readonly ILogger _logger;
 		readonly ITicker _ticker;
+		readonly ITicker _sharedTicker;
 
 		readonly RedisStringKey<AgentRateTable> _agentRateTableData = new RedisStringKey<AgentRateTable>("agent-rates");
 		readonly RedisService _redisService;
@@ -116,7 +117,8 @@ namespace Horde.Server.Agents
 			_applicationLifetime = applicationLifetime;
 			_redisService = redisService;
 			_clock = clock;
-			_ticker = clock.AddTicker<AgentService>(TimeSpan.FromSeconds(30.0), TickAsync, logger);
+			_ticker = clock.AddTicker($"{nameof(AgentService)}.{nameof(TickAsync)}", TimeSpan.FromSeconds(30.0), TickAsync, logger);
+			_sharedTicker = clock.AddSharedTicker($"{nameof(AgentService)}.{nameof(TickSharedAsync)}", TimeSpan.FromSeconds(30.0), TickSharedAsync, logger);
 			_meter = meter;
 			_tracer = tracer;
 			_logger = logger;
@@ -128,6 +130,7 @@ namespace Horde.Server.Agents
 		public async Task StartAsync(CancellationToken cancellationToken)
 		{
 			await _ticker.StartAsync();
+			await _sharedTicker.StartAsync();
 			_subscription = await Agents.SubscribeToUpdateEventsAsync(OnAgentUpdate);
 		}
 
@@ -140,6 +143,7 @@ namespace Horde.Server.Agents
 				_subscription = null;
 			}
 			await _ticker.StopAsync();
+			await _sharedTicker.StopAsync();
 		}
 
 		/// <inheritdoc/>
@@ -148,6 +152,7 @@ namespace Horde.Server.Agents
 			await _agentRateTable.DisposeAsync();
 			await _poolsList.DisposeAsync();
 			await _ticker.DisposeAsync();
+			await _sharedTicker.DisposeAsync();
 		}
 
 		/// <summary>
@@ -1074,24 +1079,25 @@ namespace Horde.Server.Agents
 			}
 		}
 
-		/// <summary>
-		/// Terminate any sessions for agents that are offline
-		/// </summary>
-		/// <param name="stoppingToken">Token indicating the service is shutting down</param>
-		/// <returns>Async task</returns>
 		internal async ValueTask TickAsync(CancellationToken stoppingToken)
 		{
 			using TelemetrySpan span = _tracer.StartActiveSpan($"{nameof(AgentService)}.{nameof(TickAsync)}");
+			await CollectMetricsAsync(stoppingToken);
+		}
+
+		internal async ValueTask TickSharedAsync(CancellationToken stoppingToken)
+		{
+			using TelemetrySpan span = _tracer.StartActiveSpan($"{nameof(AgentService)}.{nameof(TickSharedAsync)}");
 
 			await TerminateExpiredSessionsAsync(stoppingToken);
 			await DeleteExpiredEphemeralAgentsAsync(stoppingToken);
-			await CollectMetricsAsync(stoppingToken);
 		}
 
 		private async Task TerminateExpiredSessionsAsync(CancellationToken cancellationToken)
 		{
 			using TelemetrySpan span = _tracer.StartActiveSpan($"{nameof(AgentService)}.{nameof(TerminateExpiredSessionsAsync)}");
 
+			int c = 0;
 			while (!cancellationToken.IsCancellationRequested)
 			{
 				// Find all the agents which are ready to be expired
@@ -1106,6 +1112,7 @@ namespace Horde.Server.Agents
 					_logger.LogDebug("Terminating session {SessionId} for agent {Agent}", expiredAgent.SessionId, expiredAgent.Id);
 					await TryTerminateSessionAsync(expiredAgent, cancellationToken);
 				}
+				c += expiredAgents.Count;
 
 				// Try again if we didn't fetch everything
 				if (expiredAgents.Count < MaxAgents)
@@ -1113,6 +1120,7 @@ namespace Horde.Server.Agents
 					break;
 				}
 			}
+			span.SetAttribute("NumAgentsTerminated", c);
 		}
 
 		private async Task DeleteExpiredEphemeralAgentsAsync(CancellationToken cancellationToken)
