Horde: Add an explicit error when the server URL does not have a valid scheme or host name.

diff --git a/Engine/Source/Programs/Horde/Horde.Agent/AgentApp.cs b/Engine/Source/Programs/Horde/Horde.Agent/AgentApp.cs
--- a/Engine/Source/Programs/Horde/Horde.Agent/AgentApp.cs
+++ b/Engine/Source/Programs/Horde/Horde.Agent/AgentApp.cs
@@ -164,6 +164,13 @@ namespace Horde.Agent
 			services.AddOptions<AgentSettings>().Configure(options => configSection.Bind(options)).ValidateDataAnnotations();
 
 			ServerProfile serverProfile = settings.GetCurrentServerProfile();
+			if (String.IsNullOrEmpty(serverProfile.Url.Scheme) || String.IsNullOrEmpty(serverProfile.Url.Host))
+			{
+				ILogger logger = loggerFactory.CreateLogger(typeof(AgentApp));
+				logger.LogError("\"{Url}\" is an invalid server url. The specified url must have a valid scheme and host name.", serverProfile.Url);
+				return 1;
+			}
+
 			ConfigureTracing(serverProfile.Environment, AgentApp.Version);
 
 			Logging.SetEnv(serverProfile.Environment);
