Horde: Keep a cached list of current agents in AgentService

diff --git a/Engine/Source/Programs/Horde/Horde.Server.Tests/Agents/Pools/PoolUpdateServiceTest.cs b/Engine/Source/Programs/Horde/Horde.Server.Tests/Agents/Pools/PoolUpdateServiceTest.cs
--- a/Engine/Source/Programs/Horde/Horde.Server.Tests/Agents/Pools/PoolUpdateServiceTest.cs
+++ b/Engine/Source/Programs/Horde/Horde.Server.Tests/Agents/Pools/PoolUpdateServiceTest.cs
@@ -25,7 +25,7 @@ public class PoolUpdateServiceTest : TestSetup
 	public PoolUpdateServiceTest()
 	{
 		UpdateConfig(x => x.Pools.Clear());
-		_pus = new(AgentCollection, PoolCollection, Clock, GlobalConfig, Tracer, new NullLogger<PoolUpdateService>());
+		_pus = new PoolUpdateService(AgentService, PoolCollection, Clock, GlobalConfig, Tracer, new NullLogger<PoolUpdateService>());
 	}
 
 	[TestInitialize]
diff --git a/Engine/Source/Programs/Horde/Horde.Server/Agents/AgentService.cs b/Engine/Source/Programs/Horde/Horde.Server/Agents/AgentService.cs
--- a/Engine/Source/Programs/Horde/Horde.Server/Agents/AgentService.cs
+++ b/Engine/Source/Programs/Horde/Horde.Server/Agents/AgentService.cs
@@ -47,7 +47,7 @@ namespace Horde.Server.Agents
 	}
 
 	/// <summary>
-	/// Wraps funtionality for manipulating agents
+	/// Wraps functionality for manipulating agents
 	/// </summary>
 	public sealed class AgentService : IHostedService, IAsyncDisposable
 	{
@@ -65,8 +65,6 @@ namespace Horde.Server.Agents
 		/// Time after which a session will be renewed
 		/// </summary>
 		public static readonly TimeSpan SessionRenewTime = TimeSpan.FromSeconds(50);
-		readonly AclService _aclService;
-		readonly IDowntimeService _downtimeService;
 
 		/// <summary>
 		/// Collection of agent documents
@@ -75,55 +73,72 @@ namespace Horde.Server.Agents
 
 		readonly ILeaseCollection _leases;
 		readonly ISessionCollection _sessions;
+		readonly AclService _aclService;
+		readonly IDowntimeService _downtimeService;
 		readonly ITaskSource[] _taskSources;
+		readonly RedisService _redisService;
 		readonly IHostApplicationLifetime _applicationLifetime;
 		readonly IClock _clock;
-		readonly Meter _meter;
 		readonly Tracer _tracer;
 		readonly ILogger _logger;
 		readonly ITicker _ticker;
 		readonly ITicker _sharedTicker;
 
-		readonly RedisStringKey<AgentRateTable> _agentRateTableData = new RedisStringKey<AgentRateTable>("agent-rates");
-		readonly RedisService _redisService;
+		readonly RedisStringKey<AgentRateTable> _agentRateTableData = new ("agent-rates");
 
-		// Lazily updated costs for different agent types
-		readonly AsyncCachedValue<AgentRateTable?> _agentRateTable;
+		/// <summary>Lazily updated costs for different agent types</summary>
+		readonly AsyncCachedValue<AgentRateTable?> _cachedRates;
 
-		// Lazily updated list of current pools
-		readonly AsyncCachedValue<IReadOnlyList<IPoolConfig>> _poolsList;
+		/// <summary>Lazily updated list of current pools</summary>
+		readonly AsyncCachedValue<IReadOnlyList<IPoolConfig>> _cachedPools;
+		
+		/// <summary>Manually updated cached list of current agents</summary>
+		IReadOnlyDictionary<AgentId, IAgent>? _cachedAgents;
 
-		// All the agents currently performing a long poll for work on this server
-		readonly Dictionary<AgentId, CancellationTokenSource> _waitingAgents = new Dictionary<AgentId, CancellationTokenSource>();
+		/// <summary>All the agents currently performing a long poll for work on this server</summary>
+		readonly Dictionary<AgentId, CancellationTokenSource> _waitingAgents = new ();
 
+		/// <summary>OpenTelemetry measurements (gauges)</summary>
 		IEnumerable<Measurement<int>> _measurements = new List<Measurement<int>>();
 
-		// Subscription for update events
+		/// <summary>Subscription for update events</summary>
 		IAsyncDisposable? _subscription;
 
 		/// <summary>
 		/// Constructor
 		/// </summary>
-		public AgentService(IAgentCollection agents, ILeaseCollection leases, ISessionCollection sessions, AclService aclService, IDowntimeService downtimeService, IPoolCollection poolCollection, IEnumerable<ITaskSource> taskSources, RedisService redisService, IHostApplicationLifetime applicationLifetime, IClock clock, Tracer tracer, Meter meter, ILogger<AgentService> logger)
+		public AgentService(
+			IAgentCollection agents,
+			ILeaseCollection leases,
+			ISessionCollection sessions,
+			AclService aclService,
+			IDowntimeService downtimeService,
+			IPoolCollection poolCollection,
+			IEnumerable<ITaskSource> taskSources,
+			RedisService redisService,
+			IHostApplicationLifetime applicationLifetime,
+			IClock clock,
+			Tracer tracer,
+			Meter meter,
+			ILogger<AgentService> logger)
 		{
 			Agents = agents;
 			_leases = leases;
 			_sessions = sessions;
 			_aclService = aclService;
 			_downtimeService = downtimeService;
-			_agentRateTable = new AsyncCachedValue<AgentRateTable?>(_ => redisService.GetDatabase().StringGetAsync(_agentRateTableData), TimeSpan.FromSeconds(2.0));//.FromMinutes(5.0));
-			_poolsList = new AsyncCachedValue<IReadOnlyList<IPoolConfig>>(ctx => poolCollection.GetConfigsAsync(ctx), TimeSpan.FromSeconds(30.0));
+			_cachedRates = new AsyncCachedValue<AgentRateTable?>(_ => redisService.GetDatabase().StringGetAsync(_agentRateTableData), TimeSpan.FromMinutes(2));
+			_cachedPools = new AsyncCachedValue<IReadOnlyList<IPoolConfig>>(poolCollection.GetConfigsAsync, TimeSpan.FromSeconds(30.0));
 			_taskSources = taskSources.ToArray();
 			_applicationLifetime = applicationLifetime;
 			_redisService = redisService;
 			_clock = clock;
 			_ticker = clock.AddTicker($"{nameof(AgentService)}.{nameof(TickAsync)}", TimeSpan.FromSeconds(30.0), TickAsync, logger);
 			_sharedTicker = clock.AddSharedTicker($"{nameof(AgentService)}.{nameof(TickSharedAsync)}", TimeSpan.FromSeconds(30.0), TickSharedAsync, logger);
-			_meter = meter;
 			_tracer = tracer;
 			_logger = logger;
 
-			_meter.CreateObservableGauge("horde.agent.count", () => _measurements);
+			meter.CreateObservableGauge("horde.agent.count", () => _measurements);
 		}
 
 		/// <inheritdoc/>
@@ -149,8 +164,8 @@ namespace Horde.Server.Agents
 		/// <inheritdoc/>
 		public async ValueTask DisposeAsync()
 		{
-			await _agentRateTable.DisposeAsync();
-			await _poolsList.DisposeAsync();
+			await _cachedRates.DisposeAsync();
+			await _cachedPools.DisposeAsync();
 			await _ticker.DisposeAsync();
 			await _sharedTicker.DisposeAsync();
 		}
@@ -261,6 +276,28 @@ namespace Horde.Server.Agents
 		{
 			return Agents.FindAsync(poolId, modifiedAfter, property, null, null, includeDeleted, index, count, cancellationToken);
 		}
+		
+		/// <summary>
+		/// Get all agents from local in-memory cache
+		/// </summary>
+		/// <param name="cancellationToken">Cancellation token for the operation</param>
+		/// <returns>List of all agents</returns>
+		public async Task<IReadOnlyList<IAgent>> GetCachedAgentsAsync(CancellationToken cancellationToken)
+		{
+			if (_cachedAgents == null)
+			{
+				// First time this method is called, cached agents may not be initialized
+				// This could lead to multiple concurrent agent refreshes (dog-piling), but that's fine as it should only happen once
+				await RefreshCachedAgentsAsync(cancellationToken);
+			}
+
+			if (_cachedAgents == null)
+			{
+				throw new Exception("Cached agents not initialized");
+			}
+
+			return _cachedAgents.Values.ToList();
+		}
 
 		/// <summary>
 		/// Update the current workspaces for an agent.
@@ -311,7 +348,7 @@ namespace Horde.Server.Agents
 		{
 			List<PoolId> newDynamicPools = new List<PoolId>();
 
-			IReadOnlyList<IPoolConfig> pools = await _poolsList.GetAsync(cancellationToken);
+			IReadOnlyList<IPoolConfig> pools = await _cachedPools.GetAsync(cancellationToken);
 			foreach (IPoolConfig pool in pools)
 			{
 				if (pool.Condition != null && agent.SatisfiesCondition(pool.Condition))
@@ -1056,7 +1093,7 @@ namespace Horde.Server.Agents
 				double rate = 0.0;
 
 				// Get the rate table
-				AgentRateTable? rateTable = await _agentRateTable.GetAsync(cancellationToken);
+				AgentRateTable? rateTable = await _cachedRates.GetAsync(cancellationToken);
 				if (rateTable != null && rateTable.Entries.Count > 0)
 				{
 					IAgent? agent = await GetAgentAsync(agentId, cancellationToken);
@@ -1082,6 +1119,7 @@ namespace Horde.Server.Agents
 		internal async ValueTask TickAsync(CancellationToken stoppingToken)
 		{
 			using TelemetrySpan span = _tracer.StartActiveSpan($"{nameof(AgentService)}.{nameof(TickAsync)}");
+			await RefreshCachedAgentsAsync(stoppingToken);
 			await CollectMetricsAsync(stoppingToken);
 		}
 
@@ -1142,11 +1180,29 @@ namespace Horde.Server.Agents
 			span.SetAttribute("NumAgentsDeleted", c);
 		}
 
+		/// <summary>
+		/// Updates the in-memory cache of current agents
+		/// Called by ticker to avoid blocking reads. Updates in the background, in favor of slightly more stale agents
+		/// </summary>
+		/// <param name="cancellationToken">Cancellation token</param>
+		private async Task RefreshCachedAgentsAsync(CancellationToken cancellationToken = default)
+		{
+			using TelemetrySpan span = _tracer.StartActiveSpan($"{nameof(AgentService)}.{nameof(RefreshCachedAgentsAsync)}");
+			Dictionary<AgentId, IAgent> agents = new();
+			IReadOnlyList<IAgent> agentList = await Agents.FindAsync(cancellationToken: cancellationToken);
+			foreach (IAgent agent in agentList)
+			{
+				agents[agent.Id] = agent;
+			}
+			
+			Interlocked.Exchange(ref _cachedAgents, agents);
+		}
+
 		private async Task CollectMetricsAsync(CancellationToken cancellationToken = default)
 		{
 			using TelemetrySpan span = _tracer.StartActiveSpan($"{nameof(AgentService)}.{nameof(CollectMetricsAsync)}");
 
-			IReadOnlyList<IAgent> agentList = await Agents.FindAsync(cancellationToken: cancellationToken);
+			IReadOnlyList<IAgent> agentList = await GetCachedAgentsAsync(cancellationToken);
 			int numAgentsTotal = agentList.Count;
 			int numAgentsTotalDeleted = agentList.Count(a => a.Deleted);
 			int numAgentsTotalEnabled = agentList.Count(a => a.Enabled);
diff --git a/Engine/Source/Programs/Horde/Horde.Server/Agents/Pools/PoolUpdateService.cs b/Engine/Source/Programs/Horde/Horde.Server/Agents/Pools/PoolUpdateService.cs
--- a/Engine/Source/Programs/Horde/Horde.Server/Agents/Pools/PoolUpdateService.cs
+++ b/Engine/Source/Programs/Horde/Horde.Server/Agents/Pools/PoolUpdateService.cs
@@ -22,7 +22,7 @@ namespace Horde.Server.Agents.Pools
 	/// </summary>
 	public sealed class PoolUpdateService : IHostedService, IAsyncDisposable
 	{
-		readonly IAgentCollection _agents;
+		readonly AgentService _agentService;
 		readonly IPoolCollection _pools;
 		readonly IClock _clock;
 		readonly IOptionsMonitor<GlobalConfig> _globalConfig;
@@ -35,9 +35,9 @@ namespace Horde.Server.Agents.Pools
 		/// <summary>
 		/// Constructor
 		/// </summary>
-		public PoolUpdateService(IAgentCollection agents, IPoolCollection pools, IClock clock, IOptionsMonitor<GlobalConfig> globalConfig, Tracer tracer, ILogger<PoolUpdateService> logger)
+		public PoolUpdateService(AgentService agentService, IPoolCollection pools, IClock clock, IOptionsMonitor<GlobalConfig> globalConfig, Tracer tracer, ILogger<PoolUpdateService> logger)
 		{
-			_agents = agents;
+			_agentService = agentService;
 			_pools = pools;
 			_clock = clock;
 			_globalConfig = globalConfig;
@@ -82,8 +82,7 @@ namespace Horde.Server.Agents.Pools
 			using TelemetrySpan span = _tracer.StartActiveSpan($"{nameof(PoolUpdateService)}.{nameof(ShutdownDisabledAgentsAsync)}");
 
 			IReadOnlyList<IPoolConfig> pools = await _pools.GetConfigsAsync(cancellationToken);
-			IEnumerable<IAgent> disabledAgents = await _agents.FindAsync(enabled: false, cancellationToken: cancellationToken);
-			disabledAgents = disabledAgents.Where(x => IsAgentAutoScaled(x, pools));
+			IEnumerable<IAgent> disabledAgents = (await _agentService.GetCachedAgentsAsync(cancellationToken)).Where(x => !x.Enabled && IsAgentAutoScaled(x, pools));
 
 			int c = 0;
 			foreach (IAgent agent in disabledAgents)
@@ -132,7 +131,8 @@ namespace Horde.Server.Agents.Pools
 		internal async ValueTask AutoConformAgentsAsync(CancellationToken cancellationToken)
 		{
 			using TelemetrySpan span = _tracer.StartActiveSpan($"{nameof(PoolUpdateService)}.{nameof(AutoConformAgentsAsync)}");
-			List<IAgent> agents = (await _agents.FindAsync(status: AgentStatus.Ok, enabled: true, cancellationToken: cancellationToken)).Where(x => !x.RequestShutdown).ToList();
+			List<IAgent> agents = (await _agentService.GetCachedAgentsAsync(cancellationToken))
+				.Where(x => x is { Status: AgentStatus.Ok, Enabled: true, RequestShutdown: false }).ToList();
 
 			long MegabytesToBytes(long v) => v * 1024 * 1024;
 			foreach (IAgent agent in agents)
