Horde: Reduce how long ephemeral and deleted agents are kept in database to 1 hour

#rnx

diff --git a/Engine/Source/Programs/Horde/Horde.Server/Agents/AgentService.cs b/Engine/Source/Programs/Horde/Horde.Server/Agents/AgentService.cs
--- a/Engine/Source/Programs/Horde/Horde.Server/Agents/AgentService.cs
+++ b/Engine/Source/Programs/Horde/Horde.Server/Agents/AgentService.cs
@@ -1168,7 +1168,7 @@ namespace Horde.Server.Agents
 			foreach (IAgent agent in await Agents.FindDeletedAsync(cancellationToken))
 			{
 				cancellationToken.ThrowIfCancellationRequested();
-				bool noStatusChangeDuringPeriod = _clock.UtcNow > agent.LastStatusChange + TimeSpan.FromHours(24);
+				bool noStatusChangeDuringPeriod = _clock.UtcNow > agent.LastStatusChange + TimeSpan.FromHours(1);
 				if (agent is { Status: AgentStatus.Stopped, Ephemeral: true } && noStatusChangeDuringPeriod)
 				{
 					_logger.LogDebug("Deleting ephemeral agent {Agent}", agent.Id);
