Horde: Add command-based tracer for MongoDB

Old tracer operates at the collection level, this instead listens for events emitted by the Mongo client.

diff --git a/Engine/Source/Programs/Horde/Horde.Server.Tests/DatabaseIntegrationTest.cs b/Engine/Source/Programs/Horde/Horde.Server.Tests/DatabaseIntegrationTest.cs
--- a/Engine/Source/Programs/Horde/Horde.Server.Tests/DatabaseIntegrationTest.cs
+++ b/Engine/Source/Programs/Horde/Horde.Server.Tests/DatabaseIntegrationTest.cs
@@ -260,14 +260,14 @@ namespace Horde.Server.Tests
 				if (_mongoService == null)
 				{
 					RedisService redisService = GetRedisServiceSingleton();
-
 					_mongoInstance = new MongoInstance();
 
 					ServerSettings ss = new ServerSettings();
 					ss.MongoDatabaseName = _mongoInstance.DatabaseName;
 					ss.MongoConnectionString = _mongoInstance.ConnectionString;
 
-					_mongoService = new MongoService(Options.Create(ss), redisService, OpenTelemetryTracers.Horde, _loggerFactory.CreateLogger<MongoService>(), _loggerFactory);
+					MongoCommandTracer mongoTracer = new (OpenTelemetryTracers.MongoDb, _loggerFactory.CreateLogger<MongoCommandTracer>());
+					_mongoService = new MongoService(Options.Create(ss), redisService, mongoTracer, OpenTelemetryTracers.Horde, _loggerFactory.CreateLogger<MongoService>(), _loggerFactory);
 				}
 			}
 			return _mongoService;
diff --git a/Engine/Source/Programs/Horde/Horde.Server/Server/MongoService.cs b/Engine/Source/Programs/Horde/Horde.Server/Server/MongoService.cs
--- a/Engine/Source/Programs/Horde/Horde.Server/Server/MongoService.cs
+++ b/Engine/Source/Programs/Horde/Horde.Server/Server/MongoService.cs
@@ -270,7 +270,7 @@ namespace Horde.Server.Server
 		/// <summary>
 		/// Constructor
 		/// </summary>
-		public MongoService(IOptions<ServerSettings> settingsSnapshot, RedisService redisService, Tracer tracer, ILogger<MongoService> logger, ILoggerFactory loggerFactory)
+		public MongoService(IOptions<ServerSettings> settingsSnapshot, RedisService redisService, MongoCommandTracer mongoTracer, Tracer tracer, ILogger<MongoService> logger, ILoggerFactory loggerFactory)
 		{
 			if (s_existingInstance != null)
 			{
@@ -333,6 +333,11 @@ namespace Horde.Server.Server
 					{
 						clusterBuilder.Subscribe<CommandStartedEvent>(ev => TraceMongoCommand(ev.Command));
 					}
+
+					if (settingsSnapshot.Value.OpenTelemetry.Enabled)
+					{
+						mongoTracer.Register(clusterBuilder);
+					}
 				};
 
 				mongoSettings.LinqProvider = LinqProvider.V2;
diff --git a/Engine/Source/Programs/Horde/Horde.Server/Startup.cs b/Engine/Source/Programs/Horde/Horde.Server/Startup.cs
--- a/Engine/Source/Programs/Horde/Horde.Server/Startup.cs
+++ b/Engine/Source/Programs/Horde/Horde.Server/Startup.cs
@@ -554,6 +554,7 @@ namespace Horde.Server
 			services.AddSingleton<ConsistencyService>();
 			services.AddSingleton<RequestTrackerService>();
 			services.AddSingleton<ComputeService>();
+			services.AddSingleton<MongoCommandTracer>();
 			services.AddSingleton<MongoService>();
 			services.AddSingleton<GlobalsService>();
 			services.AddSingleton<ICommitService, CommitService>();
diff --git a/Engine/Source/Programs/Horde/Horde.Server/Utilities/MongoCommandTracer.cs b/Engine/Source/Programs/Horde/Horde.Server/Utilities/MongoCommandTracer.cs
new file mode 100644
--- /dev/null
+++ b/Engine/Source/Programs/Horde/Horde.Server/Utilities/MongoCommandTracer.cs
@@ -0,0 +1,111 @@
+// Copyright Epic Games, Inc. All Rights Reserved.
+
+using System;
+using System.Collections.Concurrent;
+using Microsoft.Extensions.Logging;
+using MongoDB.Bson;
+using MongoDB.Driver.Core.Configuration;
+using MongoDB.Driver.Core.Events;
+using OpenTelemetry.Trace;
+
+namespace Horde.Server.Utilities;
+
+/// <summary>
+/// OpenTelemetry-based tracer listening for MongoDB command events
+/// </summary>
+/// <param name="tracer">Tracer</param>
+/// <param name="logger">Logger</param>
+public class MongoCommandTracer(Tracer tracer, ILogger<MongoCommandTracer> logger)
+{
+	private readonly ConcurrentDictionary<int, TelemetrySpan> _activeSpans = new();
+
+	/// <summary>
+	/// Registers event listeners from MongoDB's client
+	/// </summary>
+	/// <param name="clusterBuilder">A MongoDB cluster builder</param>
+	public void Register(ClusterBuilder clusterBuilder)
+	{
+		clusterBuilder.Subscribe<CommandStartedEvent>(OnEvent);
+		clusterBuilder.Subscribe<CommandSucceededEvent>(OnEvent);
+		clusterBuilder.Subscribe<CommandFailedEvent>(OnEvent);
+	}
+	
+	private void OnEvent(CommandStartedEvent ev)
+	{
+		string? GetString(string fieldName)
+		{
+			return ev.Command.TryGetValue(fieldName, out BsonValue value) ? value.ToString() : null;
+		}
+		
+		try
+		{
+			// TODO: Nest getMore under the originating span as it's really a continuation of a find command
+			(string? collectionName, string? statement) = ev.CommandName switch
+			{
+				"find" => (GetString("find"), GetString("filter")),
+				"update" => (GetString("update"), GetString("updates")),
+				"delete" => (GetString("delete"), GetString("deletes")),
+				"insert" => (GetString("insert"), null),
+				"findAndModify" => (GetString("findAndModify"), GetString("query")),
+				"aggregate" => (GetString("aggregate"), null),
+				"getMore" => (GetString("collection"), null),
+				"listIndexes" => (GetString("listIndexes"), null),
+				"createIndexes" => (GetString("createIndexes"), null),
+				_ => (null, null)
+			};
+
+			string name = $"mongodb.{ev.CommandName}";
+			if (collectionName != null)
+			{
+				name = $"{collectionName}.{ev.CommandName}";
+			}
+			
+			// OpenTelemetry MongoDB conventions https://opentelemetry.io/docs/specs/semconv/database/mongodb/
+			TelemetrySpan span = tracer.StartMongoDbSpan<object>(name);
+			span.SetAttribute("db.system", "mongodb");
+			span.SetAttribute("db.name", ev.DatabaseNamespace.ToString());
+			span.SetAttribute("db.operationId", ev.OperationId ?? -1);
+			span.SetAttribute("db.requestId", ev.RequestId);
+			span.SetAttribute("db.serviceId", ev.ServiceId?.ToString());
+			
+			if (collectionName != null)
+			{
+				span.SetAttribute("db.mongodb.collection", collectionName);
+			}
+			
+			if (statement != null)
+			{
+				span.SetAttribute("db.statement", statement.Length > 200 ? statement[..200] : statement);
+			}
+
+			_activeSpans.TryAdd(ev.RequestId, span);
+		}
+		catch (Exception e)
+		{
+			logger.LogError(e, "Unhandled exception when capturing MongoDB span");
+		}
+	}
+	
+	private void OnEvent(CommandSucceededEvent ev)
+	{
+		if (_activeSpans.TryRemove(ev.RequestId, out TelemetrySpan? span))
+		{
+			span.SetStatus(Status.Ok);
+			span.SetAttribute("db.durationMs", ev.Duration.TotalMilliseconds);
+			span.End();
+			span.Dispose();
+		}
+	}
+	
+	private void OnEvent(CommandFailedEvent ev)
+	{
+		if (_activeSpans.TryRemove(ev.RequestId, out TelemetrySpan? span))
+		{
+			span.SetStatus(Status.Error);
+			span.SetAttribute("db.durationMs", ev.Duration.TotalMilliseconds);
+			span.RecordException(ev.Failure);
+			span.End();
+			span.Dispose();
+		}
+	}
+}
