Horde: Remove unnecessary OpenTelemetry span prefix for MongoDB

#rnx

diff --git a/Engine/Source/Programs/Horde/Horde.Server/Utilities/MongoCommandTracer.cs b/Engine/Source/Programs/Horde/Horde.Server/Utilities/MongoCommandTracer.cs
--- a/Engine/Source/Programs/Horde/Horde.Server/Utilities/MongoCommandTracer.cs
+++ b/Engine/Source/Programs/Horde/Horde.Server/Utilities/MongoCommandTracer.cs
@@ -54,7 +54,7 @@ public class MongoCommandTracer(Tracer tracer, ILogger<MongoCommandTracer> logge
 				_ => (null, null)
 			};
 
-			string name = $"mongodb.{ev.CommandName}";
+			string name = ev.CommandName;
 			if (collectionName != null)
 			{
 				name = $"{collectionName}.{ev.CommandName}";
diff --git a/Engine/Source/Programs/Horde/Horde.Server/Utilities/OpenTelemetry.cs b/Engine/Source/Programs/Horde/Horde.Server/Utilities/OpenTelemetry.cs
--- a/Engine/Source/Programs/Horde/Horde.Server/Utilities/OpenTelemetry.cs
+++ b/Engine/Source/Programs/Horde/Horde.Server/Utilities/OpenTelemetry.cs
@@ -319,11 +319,10 @@ public static class OpenTelemetrySpanExtensions
 	{
 		_ = tracer;
 
-		string name = "mongodb." + spanName;
 		TelemetrySpan span = OpenTelemetryTracers.MongoDb
-			.StartActiveSpan(name, parentContext: Tracer.CurrentSpan.Context)
+			.StartActiveSpan(spanName, parentContext: Tracer.CurrentSpan.Context)
 			.SetAttribute("type", "db")
-			.SetAttribute("operation.name", name)
+			.SetAttribute("operation.name", spanName)
 			.SetAttribute("service.name", OpenTelemetryTracers.MongoDbName);
 
 		if (collection != null)
