Horde: Optimize agent assignment for compute tasks by filtering by pool

diff --git a/Engine/Source/Programs/Horde/Horde.Server.Tests/Compute/ComputeServiceTest.cs b/Engine/Source/Programs/Horde/Horde.Server.Tests/Compute/ComputeServiceTest.cs
--- a/Engine/Source/Programs/Horde/Horde.Server.Tests/Compute/ComputeServiceTest.cs
+++ b/Engine/Source/Programs/Horde/Horde.Server.Tests/Compute/ComputeServiceTest.cs
@@ -225,6 +225,26 @@ namespace Horde.Server.Tests.Compute
 			ComputeResource? resource2 = await ComputeService.TryAllocateResourceAsync(CreateParams("11.0.0.2"), CancellationToken.None);
 			Assert.AreEqual(agent.Id, resource2!.AgentId);
 		}
+		
+		[TestMethod]
+		public async Task Assignment_PoolInRequirementsAsync()
+		{
+			AllocateResourceParams CreateParams(string pool)
+			{
+				return new AllocateResourceParams(new ClusterId("default"), ComputeProtocol.Latest, new Requirements { Pool = pool });
+			}
+
+			IAgent agent1 = await CreateAgentAsync(new PoolId("foo"), properties: ["ComputeIp=11.0.0.1", "ComputePort=5000"]);
+			IAgent agent2 = await CreateAgentAsync(new PoolId("bar"), properties: ["ComputeIp=11.0.0.2", "ComputePort=5000"]);
+			
+			ComputeResource? resource1 = await ComputeService.TryAllocateResourceAsync(CreateParams("foo"), CancellationToken.None);
+			Assert.AreEqual(agent1.Id, resource1!.AgentId);
+			
+			ComputeResource? resource2 = await ComputeService.TryAllocateResourceAsync(CreateParams("bar"), CancellationToken.None);
+			Assert.AreEqual(agent2.Id, resource2!.AgentId);
+			
+			Assert.IsNull(await ComputeService.TryAllocateResourceAsync(CreateParams("does-not-exist"), CancellationToken.None));
+		}
 
 		[TestMethod]
 		public async Task Connection_Direct_IpConnection_Async()
@@ -284,14 +304,14 @@ namespace Horde.Server.Tests.Compute
 			Assert.AreEqual(new ComputeResourcePort(10004, 13000), cr.Ports["myOtherPort"]);
 		}
 
-		private async Task<ComputeService> CreateComputeServiceAsync(string? tunnelAddress)
+		private async Task<ComputeService> CreateComputeServiceAsync(string? tunnelAddress, string pool)
 		{
 			ServerSettings ss = new() { ComputeTunnelAddress = tunnelAddress };
 			ComputeService cs = new(AgentCollection, LogCollection, AgentService, AgentRelayService, GetRedisServiceSingleton(),
 				new TestOptionsMonitor<ServerSettings>(ss), GlobalConfig, Clock, Tracer, Meter,
 				NullLogger<ComputeService>.Instance);
 			List<string> props = new() { "ComputeIp=11.0.0.1", "ComputePort=5000" };
-			await CreateAgentAsync(new PoolId("foo"), properties: props);
+			await CreateAgentAsync(new PoolId(pool), properties: props);
 			return cs;
 		}
 
@@ -300,11 +320,11 @@ namespace Horde.Server.Tests.Compute
 			Dictionary<string, int>? ports = null,
 			bool usePublicIp = false,
 			string[]? relayIps = null,
-			string? tunnelAddress = null
-			)
+			string? tunnelAddress = null,
+			string pool = "myPool")
 		{
-			await using ComputeService cs = await CreateComputeServiceAsync(tunnelAddress);
-			AllocateResourceParams arp = new(_cluster1, ComputeProtocol.Latest, new Requirements())
+			await using ComputeService cs = await CreateComputeServiceAsync(tunnelAddress, pool);
+			AllocateResourceParams arp = new(_cluster1, ComputeProtocol.Latest, new Requirements() { Pool = pool})
 			{
 				ConnectionMode = connectionMode,
 				Ports = ports ?? new Dictionary<string, int>(),
diff --git a/Engine/Source/Programs/Horde/Horde.Server/Compute/ComputeService.cs b/Engine/Source/Programs/Horde/Horde.Server/Compute/ComputeService.cs
--- a/Engine/Source/Programs/Horde/Horde.Server/Compute/ComputeService.cs
+++ b/Engine/Source/Programs/Horde/Horde.Server/Compute/ComputeService.cs
@@ -13,6 +13,7 @@ using System.Threading.Tasks;
 using EpicGames.Core;
 using EpicGames.Horde.Agents;
 using EpicGames.Horde.Agents.Leases;
+using EpicGames.Horde.Agents.Pools;
 using EpicGames.Horde.Compute;
 using EpicGames.Horde.Compute.Clients;
 using EpicGames.Horde.Compute.Transports;
@@ -382,7 +383,8 @@ namespace Horde.Server.Compute
 
 			try
 			{
-				IReadOnlyList<IAgent> agents = await _agentCollection.FindAsync(cancellationToken: cancellationToken);
+				PoolId? poolId = arp.Requirements.Pool != null ? new PoolId(arp.Requirements.Pool) : null;
+				IReadOnlyList<IAgent> agents = await _agentCollection.FindAsync(poolId, cancellationToken: cancellationToken);
 				foreach (IAgent agent in agents)
 				{
 					Dictionary<string, int> assignedResources = new Dictionary<string, int>();
@@ -395,7 +397,7 @@ namespace Horde.Server.Compute
 						ComputeProtocol protocol = ComputeProtocol.Initial;
 						foreach (string value in agent.GetPropertyValues("ComputeProtocol"))
 						{
-							if (int.TryParse(value, out int versionInt))
+							if (Int32.TryParse(value, out int versionInt))
 							{
 								protocol = (ComputeProtocol)Math.Min((int)arp.Protocol, versionInt);
 								break;
