Horde: Enable new build dialog
#rnx

diff --git a/Engine/Source/Programs/Horde/HordeDashboard/package.json b/Engine/Source/Programs/Horde/HordeDashboard/package.json
--- a/Engine/Source/Programs/Horde/HordeDashboard/package.json
+++ b/Engine/Source/Programs/Horde/HordeDashboard/package.json
@@ -1,7 +1,7 @@
 {
   "name": "horde-dashboard",
   "private": true,
-  "version": "0.0.3",
+  "version": "0.0.4",
   "type": "module",  
   "scripts": {
     "dev": "vite",
diff --git a/Engine/Source/Programs/Horde/HordeDashboard/src/components/StreamView.tsx b/Engine/Source/Programs/Horde/HordeDashboard/src/components/StreamView.tsx
--- a/Engine/Source/Programs/Horde/HordeDashboard/src/components/StreamView.tsx
+++ b/Engine/Source/Programs/Horde/HordeDashboard/src/components/StreamView.tsx
@@ -356,10 +356,7 @@ const StreamViewInner: React.FC = observer(() => {
    const stream = projectStore.streamById(streamId);
    const project = stream?.project;
 
-   let newBuildVersion: string = query.get("newbuildversion") ? query.get("newbuildversion") : "1";
-   if (getSiteConfig().environment === "dev") {
-      newBuildVersion = "2";
-   }
+   let newBuildVersion: string = query.get("newbuildversion") ? query.get("newbuildversion") : "2";
 
    if (!stream || !project) {
       console.error("Bad stream or project id in StreamView");
