Horde: Use new build dialog on preflights
#rnx

diff --git a/Engine/Source/Programs/Horde/HordeDashboard/package.json b/Engine/Source/Programs/Horde/HordeDashboard/package.json
--- a/Engine/Source/Programs/Horde/HordeDashboard/package.json
+++ b/Engine/Source/Programs/Horde/HordeDashboard/package.json
@@ -1,7 +1,7 @@
 {
   "name": "horde-dashboard",
   "private": true,
-  "version": "0.0.4",
+  "version": "0.0.5",
   "type": "module",  
   "scripts": {
     "dev": "vite",
diff --git a/Engine/Source/Programs/Horde/HordeDashboard/src/components/Preflight.tsx b/Engine/Source/Programs/Horde/HordeDashboard/src/components/Preflight.tsx
--- a/Engine/Source/Programs/Horde/HordeDashboard/src/components/Preflight.tsx
+++ b/Engine/Source/Programs/Horde/HordeDashboard/src/components/Preflight.tsx
@@ -39,9 +39,6 @@ export const PreflightRedirector: React.FC = () => {
    // whether a template is specified
    const templateId = !query.get("templateId") ? "" : query.get("templateId")!;
 
-   // explicit version otherwise keys off there being parameters
-   let version = !query.get("version") ? "" : query.get("version")!;
-
    const parameters: string[] = [];
 
    query.forEach((value, key) => {
@@ -50,9 +47,7 @@ export const PreflightRedirector: React.FC = () => {
       }
    })
 
-   if (parameters.length) {
-      version = "2";
-   }
+   const version = "2";
 
    if (!change) {
       setError("No preflight change specified");
