Horde: Resurrect job options and add parameters
#rnx

diff --git a/Engine/Source/Programs/Horde/HordeDashboard/src/components/build/NewBuildV2.tsx b/Engine/Source/Programs/Horde/HordeDashboard/src/components/build/NewBuildV2.tsx
--- a/Engine/Source/Programs/Horde/HordeDashboard/src/components/build/NewBuildV2.tsx
+++ b/Engine/Source/Programs/Horde/HordeDashboard/src/components/build/NewBuildV2.tsx
@@ -1256,7 +1256,32 @@ const AdvancedPanel: React.FC = observer(() => {
       });
    }
 
-   const height = options.estimateHeight();
+   let height = options.estimateHeight();
+
+   const a = options.jobDetails?.jobData?.arguments ?? [];
+   const b = options.jobDetails?.jobData?.parameters ?? {};
+
+   let farguments = "";
+   let fargumentsHeight = 0;
+
+   if (a.length > 0) {
+      farguments = (JSON.stringify(a, null, 2));
+      const flines = (farguments.match(/\n/g) || '').length + 1
+      fargumentsHeight = Math.min(flines * 18, 240)
+      height += fargumentsHeight + 8
+   }
+
+   let fparameters = "";
+   let fparametersHeight = 0;
+
+   if ((Object.keys(b).length)) {
+      fparameters = (JSON.stringify(b, null, 2));
+      const plines = (fparameters.match(/\n/g) || '').length + 1
+      debugger;
+      fparametersHeight = Math.min(plines * 18, 240)
+      height += fparametersHeight + 8
+
+   }
 
    return <Stack style={{
       height: height,
@@ -1291,6 +1316,8 @@ const AdvancedPanel: React.FC = observer(() => {
                options.setChanged();
             }} />
          </Stack>
+         {!!farguments && <TextField style={{ height: fargumentsHeight }} key={"key_adv_job_arguments"} defaultValue={farguments} readOnly={true} label="Job Arguments" multiline resizable={false} />}
+         {!!fparameters && <TextField style={{ height: fparametersHeight }} key={"key_adv_job_parameters]"} defaultValue={fparameters} readOnly={true} label="Job Parameters" multiline resizable={false} />}
          <Stack>
             <Checkbox disabled={options.readOnly}
                label="Template Editor"
