Horde: Don't show args when running again
#rnx

diff --git a/Engine/Source/Programs/Horde/HordeDashboard/package.json b/Engine/Source/Programs/Horde/HordeDashboard/package.json
--- a/Engine/Source/Programs/Horde/HordeDashboard/package.json
+++ b/Engine/Source/Programs/Horde/HordeDashboard/package.json
@@ -1,7 +1,7 @@
 {
   "name": "horde-dashboard",
   "private": true,
-  "version": "0.0.5",
+  "version": "0.0.6",
   "type": "module",  
   "scripts": {
     "dev": "vite",
diff --git a/Engine/Source/Programs/Horde/HordeDashboard/src/components/build/NewBuildV2.tsx b/Engine/Source/Programs/Horde/HordeDashboard/src/components/build/NewBuildV2.tsx
--- a/Engine/Source/Programs/Horde/HordeDashboard/src/components/build/NewBuildV2.tsx
+++ b/Engine/Source/Programs/Horde/HordeDashboard/src/components/build/NewBuildV2.tsx
@@ -1264,7 +1264,7 @@ const AdvancedPanel: React.FC = observer(() => {
    let farguments = "";
    let fargumentsHeight = 0;
 
-   if (a.length > 0) {
+   if (options.readOnly && a.length > 0) {
       farguments = (JSON.stringify(a, null, 2));
       const flines = (farguments.match(/\n/g) || '').length + 1
       fargumentsHeight = Math.min(flines * 18, 240)
@@ -1274,7 +1274,7 @@ const AdvancedPanel: React.FC = observer(() => {
    let fparameters = "";
    let fparametersHeight = 0;
 
-   if ((Object.keys(b).length)) {
+   if (options.readOnly &&  (Object.keys(b).length)) {
       fparameters = (JSON.stringify(b, null, 2));
       const plines = (fparameters.match(/\n/g) || '').length + 1
       debugger;
