Horde: New bulid fixes
#rnx

diff --git a/Engine/Source/Programs/Horde/HordeDashboard/package.json b/Engine/Source/Programs/Horde/HordeDashboard/package.json
--- a/Engine/Source/Programs/Horde/HordeDashboard/package.json
+++ b/Engine/Source/Programs/Horde/HordeDashboard/package.json
@@ -1,7 +1,7 @@
 {
   "name": "horde-dashboard",
   "private": true,
-  "version": "0.0.6",
+  "version": "0.0.7",
   "type": "module",  
   "scripts": {
     "dev": "vite",
diff --git a/Engine/Source/Programs/Horde/HordeDashboard/src/components/StreamView.tsx b/Engine/Source/Programs/Horde/HordeDashboard/src/components/StreamView.tsx
--- a/Engine/Source/Programs/Horde/HordeDashboard/src/components/StreamView.tsx
+++ b/Engine/Source/Programs/Horde/HordeDashboard/src/components/StreamView.tsx
@@ -356,7 +356,8 @@ const StreamViewInner: React.FC = observer(() => {
    const stream = projectStore.streamById(streamId);
    const project = stream?.project;
 
-   let newBuildVersion: string = query.get("newbuildversion") ? query.get("newbuildversion") : "2";
+   //let newBuildVersion: string = query.get("newbuildversion") ? query.get("newbuildversion") : "2";
+   const newBuildVersion:string = "2"
 
    if (!stream || !project) {
       console.error("Bad stream or project id in StreamView");
diff --git a/Engine/Source/Programs/Horde/HordeDashboard/src/components/jobDetailsV2/JobOperationsBar.tsx b/Engine/Source/Programs/Horde/HordeDashboard/src/components/jobDetailsV2/JobOperationsBar.tsx
--- a/Engine/Source/Programs/Horde/HordeDashboard/src/components/jobDetailsV2/JobOperationsBar.tsx
+++ b/Engine/Source/Programs/Horde/HordeDashboard/src/components/jobDetailsV2/JobOperationsBar.tsx
@@ -44,10 +44,7 @@ export const JobOperations: React.FC<{ jobDetails: JobDetailsV2 }> = observer(({
 
    // show new build version 2 if we have new parameters store
    const hasJobParameters = !!(Object.keys(jobDetails?.jobData?.parameters ?? {}).length)
-   if (hasJobParameters) {
-      newBuildVersion = "2";
-   }
-
+   newBuildVersion = "2";
 
    // subscribe
    if (dashboard.updated) { }
@@ -61,7 +58,7 @@ export const JobOperations: React.FC<{ jobDetails: JobDetailsV2 }> = observer(({
    }
 
    const abortDisabled = jobData.state === JobState.Complete;
-   const runAgainDisabled = newBuildVersion === "2" && !hasJobParameters;
+   const runAgainDisabled = !hasJobParameters;
 
    const failedSteps = jobDetails.getSteps().filter(s => {
 
