Horde: More new build fixes
#rnx

diff --git a/Engine/Source/Programs/Horde/HordeDashboard/src/components/build/NewBuildV2.tsx b/Engine/Source/Programs/Horde/HordeDashboard/src/components/build/NewBuildV2.tsx
--- a/Engine/Source/Programs/Horde/HordeDashboard/src/components/build/NewBuildV2.tsx
+++ b/Engine/Source/Programs/Horde/HordeDashboard/src/components/build/NewBuildV2.tsx
@@ -409,9 +409,11 @@ class BuildOptions {
 
       }
 
-      template.parameters.forEach(p => {
-         estimateHeight(p);
-      });
+      if (this.mode === "Basic") {
+         template.parameters.forEach(p => {
+            estimateHeight(p);
+         });
+      }
 
       if (!!this.preflightChange) {
          estimatedHeight += 32;
@@ -1268,19 +1270,17 @@ const AdvancedPanel: React.FC = observer(() => {
       farguments = (JSON.stringify(a, null, 2));
       const flines = (farguments.match(/\n/g) || '').length + 1
       fargumentsHeight = Math.min(flines * 18, 240)
-      height += fargumentsHeight + 8
+      height += (fargumentsHeight + 8)
    }
 
    let fparameters = "";
    let fparametersHeight = 0;
 
-   if (options.readOnly &&  (Object.keys(b).length)) {
+   if (options.readOnly && (Object.keys(b).length)) {
       fparameters = (JSON.stringify(b, null, 2));
       const plines = (fparameters.match(/\n/g) || '').length + 1
-      debugger;
       fparametersHeight = Math.min(plines * 18, 240)
-      height += fparametersHeight + 8
-
+      height += (fparametersHeight + 8)
    }
 
    return <Stack style={{
