Horde: New Build dialog layout tweaks
#rnx
#skipci

diff --git a/Engine/Source/Programs/Horde/HordeDashboard/src/components/build/NewBuildV2.tsx b/Engine/Source/Programs/Horde/HordeDashboard/src/components/build/NewBuildV2.tsx
--- a/Engine/Source/Programs/Horde/HordeDashboard/src/components/build/NewBuildV2.tsx
+++ b/Engine/Source/Programs/Horde/HordeDashboard/src/components/build/NewBuildV2.tsx
@@ -1269,7 +1269,7 @@ const AdvancedPanel: React.FC = observer(() => {
    if (options.readOnly && a.length > 0) {
       farguments = (JSON.stringify(a, null, 2));
       const flines = (farguments.match(/\n/g) || '').length + 1
-      fargumentsHeight = Math.min(flines * 18, 240)
+      fargumentsHeight = Math.min(flines * 18, 240) + 16
       height += (fargumentsHeight + 8)
    }
 
@@ -1279,7 +1279,7 @@ const AdvancedPanel: React.FC = observer(() => {
    if (options.readOnly && (Object.keys(b).length)) {
       fparameters = (JSON.stringify(b, null, 2));
       const plines = (fparameters.match(/\n/g) || '').length + 1
-      fparametersHeight = Math.min(plines * 18, 240)
+      fparametersHeight = Math.min(plines * 18, 240) + 16
       height += (fparametersHeight + 8)
    }
 
