Horde: Adding copy arguments to clipboard
#rnx

diff --git a/Engine/Source/Programs/Horde/HordeDashboard/package.json b/Engine/Source/Programs/Horde/HordeDashboard/package.json
--- a/Engine/Source/Programs/Horde/HordeDashboard/package.json
+++ b/Engine/Source/Programs/Horde/HordeDashboard/package.json
@@ -1,7 +1,7 @@
 {
   "name": "horde-dashboard",
   "private": true,
-  "version": "0.0.7",
+  "version": "0.0.8",
   "type": "module",  
   "scripts": {
     "dev": "vite",
diff --git a/Engine/Source/Programs/Horde/HordeDashboard/src/components/build/NewBuildV2.tsx b/Engine/Source/Programs/Horde/HordeDashboard/src/components/build/NewBuildV2.tsx
--- a/Engine/Source/Programs/Horde/HordeDashboard/src/components/build/NewBuildV2.tsx
+++ b/Engine/Source/Programs/Horde/HordeDashboard/src/components/build/NewBuildV2.tsx
@@ -1316,6 +1316,19 @@ const AdvancedPanel: React.FC = observer(() => {
                options.setChanged();
             }} />
          </Stack>
+         {!!options.jobDetails?.jobData?.arguments?.length && <DefaultButton text="Copy Job Arguments to Clipboard" style={{ width: 240 }} onClick={() => copyToClipboard(
+            options.jobDetails.jobData.arguments.map(arg => {
+               if (arg.indexOf("=") !== -1) {
+                  const components = arg.split("=");
+                  if (components[1].indexOf(" ") === -1) {
+                     return arg;
+                  }
+                  return `${components[0]}="${components[1]}"`;
+               } else {
+                  return arg;
+               }
+            }).join(" ")
+         )}/>}
          {!!farguments && <TextField style={{ height: fargumentsHeight }} key={"key_adv_job_arguments"} defaultValue={farguments} readOnly={true} label="Job Arguments" multiline resizable={false} />}
          {!!fparameters && <TextField style={{ height: fparametersHeight }} key={"key_adv_job_parameters]"} defaultValue={fparameters} readOnly={true} label="Job Parameters" multiline resizable={false} />}
          <Stack>
