Horde: Fix spacing due to copy job arguments button
#rnx

diff --git a/Engine/Source/Programs/Horde/HordeDashboard/src/components/build/NewBuildV2.tsx b/Engine/Source/Programs/Horde/HordeDashboard/src/components/build/NewBuildV2.tsx
--- a/Engine/Source/Programs/Horde/HordeDashboard/src/components/build/NewBuildV2.tsx
+++ b/Engine/Source/Programs/Horde/HordeDashboard/src/components/build/NewBuildV2.tsx
@@ -1283,6 +1283,11 @@ const AdvancedPanel: React.FC = observer(() => {
       height += (fparametersHeight + 8)
    }
 
+   let showArgumentClipboardButton = !!options.jobDetails?.jobData?.arguments?.length;
+   if (showArgumentClipboardButton) {
+      height += 32;
+   }
+
    return <Stack style={{
       height: height,
       position: 'relative',
@@ -1316,7 +1321,7 @@ const AdvancedPanel: React.FC = observer(() => {
                options.setChanged();
             }} />
          </Stack>
-         {!!options.jobDetails?.jobData?.arguments?.length && <DefaultButton text="Copy Job Arguments to Clipboard" style={{ width: 240 }} onClick={() => copyToClipboard(
+         {showArgumentClipboardButton && <DefaultButton text="Copy Job Arguments to Clipboard" style={{ width: 240 }} onClick={() => copyToClipboard(
             options.jobDetails.jobData.arguments.map(arg => {
                if (arg.indexOf("=") !== -1) {
                   const components = arg.split("=");
