Horde: Add support for agent queries via replica read

diff --git a/Engine/Source/Programs/Horde/Horde.Server/Agents/AgentCollection.cs b/Engine/Source/Programs/Horde/Horde.Server/Agents/AgentCollection.cs
--- a/Engine/Source/Programs/Horde/Horde.Server/Agents/AgentCollection.cs
+++ b/Engine/Source/Programs/Horde/Horde.Server/Agents/AgentCollection.cs
@@ -329,7 +329,7 @@ namespace Horde.Server.Agents
 		}
 
 		/// <inheritdoc/>
-		public async Task<IReadOnlyList<IAgent>> FindAsync(PoolId? poolId, DateTime? modifiedAfter, string? property, AgentStatus? status, bool? enabled, bool includeDeleted, int? index, int? count, CancellationToken cancellationToken)
+		public async Task<IReadOnlyList<IAgent>> FindAsync(PoolId? poolId, DateTime? modifiedAfter, string? property, AgentStatus? status, bool? enabled, bool includeDeleted, int? index, int? count, bool consistentRead, CancellationToken cancellationToken)
 		{
 			FilterDefinitionBuilder<AgentDocument> filterBuilder = new FilterDefinitionBuilder<AgentDocument>();
 
@@ -364,7 +364,8 @@ namespace Horde.Server.Agents
 				filter &= filterBuilder.Eq(x => x.Enabled, enabled.Value);
 			}
 
-			IFindFluent<AgentDocument, AgentDocument> search = _agents.Find(filter);
+			IMongoCollection<AgentDocument> collection = consistentRead ? _agents : _agents.WithReadPreference(ReadPreference.SecondaryPreferred);
+			IFindFluent<AgentDocument, AgentDocument> search = collection.Find(filter);
 			if (index != null)
 			{
 				search = search.Skip(index.Value);
diff --git a/Engine/Source/Programs/Horde/Horde.Server/Agents/AgentService.cs b/Engine/Source/Programs/Horde/Horde.Server/Agents/AgentService.cs
--- a/Engine/Source/Programs/Horde/Horde.Server/Agents/AgentService.cs
+++ b/Engine/Source/Programs/Horde/Horde.Server/Agents/AgentService.cs
@@ -274,7 +274,7 @@ namespace Horde.Server.Agents
 		/// <returns>List of agents matching the given criteria</returns>
 		public Task<IReadOnlyList<IAgent>> FindAgentsAsync(PoolId? poolId, DateTime? modifiedAfter, string? property, bool includeDeleted, int? index, int? count, CancellationToken cancellationToken)
 		{
-			return Agents.FindAsync(poolId, modifiedAfter, property, null, null, includeDeleted, index, count, cancellationToken);
+			return Agents.FindAsync(poolId, modifiedAfter, property, null, null, includeDeleted, index, count, true, cancellationToken);
 		}
 		
 		/// <summary>
diff --git a/Engine/Source/Programs/Horde/Horde.Server/Agents/IAgentCollection.cs b/Engine/Source/Programs/Horde/Horde.Server/Agents/IAgentCollection.cs
--- a/Engine/Source/Programs/Horde/Horde.Server/Agents/IAgentCollection.cs
+++ b/Engine/Source/Programs/Horde/Horde.Server/Agents/IAgentCollection.cs
@@ -60,9 +60,10 @@ namespace Horde.Server.Agents
 		/// <param name="includeDeleted">Whether agents marked as deleted should be included</param>
 		/// <param name="index">Index of the first result</param>
 		/// <param name="count">Number of results to return</param>
+		/// <param name="consistentRead">If the database read should be made to the replica server</param>
 		/// <param name="cancellationToken">Cancellation token for the operation</param>
 		/// <returns>List of agents matching the given criteria</returns>
-		Task<IReadOnlyList<IAgent>> FindAsync(PoolId? poolId = null, DateTime? modifiedAfter = null, string? property = null, AgentStatus? status = null, bool? enabled = null, bool includeDeleted = false, int? index = null, int? count = null, CancellationToken cancellationToken = default);
+		Task<IReadOnlyList<IAgent>> FindAsync(PoolId? poolId = null, DateTime? modifiedAfter = null, string? property = null, AgentStatus? status = null, bool? enabled = null, bool includeDeleted = false, int? index = null, int? count = null, bool consistentRead = true, CancellationToken cancellationToken = default);
 
 		/// <summary>
 		/// Finds all agents with an expired session
