Horde: Enable HTTP compression for server-sent responses

diff --git a/Engine/Source/Programs/Horde/Horde.Server/Startup.cs b/Engine/Source/Programs/Horde/Horde.Server/Startup.cs
--- a/Engine/Source/Programs/Horde/Horde.Server/Startup.cs
+++ b/Engine/Source/Programs/Horde/Horde.Server/Startup.cs
@@ -463,7 +463,12 @@ namespace Horde.Server
 #pragma warning restore CA2000 // Dispose objects before losing scope
 			services.AddSingleton<RedisService>(sp => redisService);
 			services.AddDataProtection().PersistKeysToStackExchangeRedis(() => redisService.GetDatabase(), "aspnet-data-protection");
-
+			
+			services.AddResponseCompression(options =>
+			{
+				options.EnableForHttps = true;
+			});
+			
 			if (settings.CorsEnabled)
 			{
 				services.AddCors(options =>
@@ -1261,6 +1266,7 @@ namespace Horde.Server
 		public static void Configure(IApplicationBuilder app, IWebHostEnvironment env, Microsoft.Extensions.Hosting.IHostApplicationLifetime lifetime, IOptions<ServerSettings> settings)
 		{
 			app.UseForwardedHeaders();
+			app.UseResponseCompression();
 
 			// Used for allowing auth cookies in combination with OpenID Connect auth (for example, Google Auth did not work with these unset)
 			app.UseCookiePolicy(new CookiePolicyOptions()
