Horde: Trap exceptions when parsing invalid workflow ids from node annotations.

diff --git a/Engine/Source/Programs/Horde/Horde.Server/Jobs/Graphs/NodeAnnotations.cs b/Engine/Source/Programs/Horde/Horde.Server/Jobs/Graphs/NodeAnnotations.cs
--- a/Engine/Source/Programs/Horde/Horde.Server/Jobs/Graphs/NodeAnnotations.cs
+++ b/Engine/Source/Programs/Horde/Horde.Server/Jobs/Graphs/NodeAnnotations.cs
@@ -107,7 +107,15 @@ namespace Horde.Server.Jobs.Graphs
 				{
 					return null;
 				}
-				return new WorkflowId(workflowName);
+
+				try
+				{
+					return new WorkflowId(workflowName);
+				}
+				catch
+				{
+					return null;
+				}
 			}
 			set => SetStringValue(WorkflowKeyName, value?.ToString());
 		}
