Horde: Fix agent tests.

diff --git a/Engine/Source/Programs/Horde/Horde.Agent.Tests/Execution/WorkspaceExecutorTest.cs b/Engine/Source/Programs/Horde/Horde.Agent.Tests/Execution/WorkspaceExecutorTest.cs
--- a/Engine/Source/Programs/Horde/Horde.Agent.Tests/Execution/WorkspaceExecutorTest.cs
+++ b/Engine/Source/Programs/Horde/Horde.Agent.Tests/Execution/WorkspaceExecutorTest.cs
@@ -10,6 +10,7 @@ using EpicGames.Horde.Jobs;
 using EpicGames.Horde.Streams;
 using Horde.Agent.Execution;
 using Horde.Agent.Services;
+using Horde.Agent.Utility;
 using HordeCommon.Rpc;
 using HordeCommon.Rpc.Tasks;
 using Microsoft.Extensions.Logging;
@@ -53,7 +54,7 @@ public sealed class WorkspaceExecutorTest : IAsyncDisposable
 		_workspace.SetFile(1, "foo/bar/baz.h", "baz");
 
 		RpcBeginBatchResponse batch = new RpcBeginBatchResponse { Change = 1 };
-		JobExecutorOptions executorOptions = new JobExecutorOptions(_session, null!, _jobId, _batchId, batch, null!, new RpcJobOptions());
+		JobExecutorOptions executorOptions = new JobExecutorOptions(_session.ServerUrl, _session.WorkingDir, _session.RpcConnection, null, null!, _jobId, _batchId, batch, null!, new RpcJobOptions());
 		_executor = new(executorOptions, _workspace, null, NullLogger.Instance);
 	}
 
diff --git a/Engine/Source/Programs/Horde/Horde.Agent.Tests/WorkerServiceTest.cs b/Engine/Source/Programs/Horde/Horde.Agent.Tests/WorkerServiceTest.cs
--- a/Engine/Source/Programs/Horde/Horde.Agent.Tests/WorkerServiceTest.cs
+++ b/Engine/Source/Programs/Horde/Horde.Agent.Tests/WorkerServiceTest.cs
@@ -282,9 +282,9 @@ namespace Horde.Agent.Tests
 			fakeSession.Setup(x => x.SessionId).Returns(new EpicGames.Horde.Agents.Sessions.SessionId(default));
 			fakeSession.Setup(x => x.RpcConnection).Returns(rpcConnection);
 			fakeSession.Setup(x => x.GrpcChannel).Returns(grpcChannel);
-			fakeSession.Setup(x => x.TerminateProcessesAsync(It.IsAny<TerminateCondition>(), It.IsAny<ILogger>(), It.IsAny<CancellationToken>())).Returns(Task.CompletedTask);
+			fakeSession.Setup(x => x.ProcessNamesToTerminate).Returns((IReadOnlyDictionary<string, TerminateCondition>)new Dictionary<string, TerminateCondition>());
 			fakeSession.Setup(x => x.DisposeAsync()).Returns(new ValueTask());
-			fakeSession.Setup(x => x.WorkingDir).Returns(DirectoryReference.GetCurrentDirectory());
+			fakeSession.Setup(x => x.WorkingDir).Returns(DirectoryReference.Combine(DirectoryReference.GetCurrentDirectory(), Guid.NewGuid().ToString()));
 			return fakeSession.Object;
 		}
 	}
diff --git a/Engine/Source/Programs/Horde/Horde.Agent/Execution/JobExecutor.cs b/Engine/Source/Programs/Horde/Horde.Agent/Execution/JobExecutor.cs
--- a/Engine/Source/Programs/Horde/Horde.Agent/Execution/JobExecutor.cs
+++ b/Engine/Source/Programs/Horde/Horde.Agent/Execution/JobExecutor.cs
@@ -63,7 +63,7 @@ namespace Horde.Agent.Execution
 		public Uri ServerUrl { get; }
 		public DirectoryReference WorkingDir { get; }
 		public IRpcConnection RpcConnection { get; }
-		public IReadOnlyDictionary<string, TerminateCondition> ProcessNamesToTerminate { get; }
+		public IReadOnlyDictionary<string, TerminateCondition>? ProcessNamesToTerminate { get; }
 		public HttpStorageClientFactory StorageFactory { get; }
 		public JobId JobId { get; }
 		public JobStepBatchId BatchId { get; }
@@ -71,7 +71,7 @@ namespace Horde.Agent.Execution
 		public string Token { get; }
 		public RpcJobOptions JobOptions { get; }
 
-		public JobExecutorOptions(Uri serverUrl, DirectoryReference workingDir, IRpcConnection rpcConnection, IReadOnlyDictionary<string, TerminateCondition> processNamesToTerminate, HttpStorageClientFactory storageFactory, JobId jobId, JobStepBatchId batchId, RpcBeginBatchResponse batch, string token, RpcJobOptions jobOptions)
+		public JobExecutorOptions(Uri serverUrl, DirectoryReference workingDir, IRpcConnection rpcConnection, IReadOnlyDictionary<string, TerminateCondition>? processNamesToTerminate, HttpStorageClientFactory storageFactory, JobId jobId, JobStepBatchId batchId, RpcBeginBatchResponse batch, string token, RpcJobOptions jobOptions)
 		{
 			ServerUrl = serverUrl;
 			WorkingDir = workingDir;
diff --git a/Engine/Source/Programs/Horde/Horde.Agent/Utility/TerminateProcessHelper.cs b/Engine/Source/Programs/Horde/Horde.Agent/Utility/TerminateProcessHelper.cs
--- a/Engine/Source/Programs/Horde/Horde.Agent/Utility/TerminateProcessHelper.cs
+++ b/Engine/Source/Programs/Horde/Horde.Agent/Utility/TerminateProcessHelper.cs
@@ -50,7 +50,7 @@ namespace Horde.Agent.Utility
 		/// <summary>
 		/// Terminate processes matching certain criteria
 		/// </summary>
-		public static Task TerminateProcessesAsync(TerminateCondition condition, DirectoryReference workingDir, IReadOnlyDictionary<string, TerminateCondition> processNamesToTerminate, ILogger logger, CancellationToken cancellationToken)
+		public static Task TerminateProcessesAsync(TerminateCondition condition, DirectoryReference workingDir, IReadOnlyDictionary<string, TerminateCondition>? processNamesToTerminate, ILogger logger, CancellationToken cancellationToken)
 		{
 			// Terminate child processes from any previous runs
 			ProcessUtils.TerminateProcesses(x => ShouldTerminateProcess(x, condition, workingDir, processNamesToTerminate), logger, cancellationToken);
@@ -60,19 +60,22 @@ namespace Horde.Agent.Utility
 		/// <summary>
 		/// Callback for determining whether a process should be terminated
 		/// </summary>
-		static bool ShouldTerminateProcess(FileReference imageFile, TerminateCondition condition, DirectoryReference workingDir, IReadOnlyDictionary<string, TerminateCondition> processNamesToTerminate)
+		static bool ShouldTerminateProcess(FileReference imageFile, TerminateCondition condition, DirectoryReference workingDir, IReadOnlyDictionary<string, TerminateCondition>? processNamesToTerminate)
 		{
 			if (imageFile.IsUnderDirectory(workingDir))
 			{
 				return true;
 			}
 
-			string fileName = imageFile.GetFileName();
-			if (processNamesToTerminate.TryGetValue(fileName, out TerminateCondition terminateFlags))
+			if (processNamesToTerminate != null)
 			{
-				if (terminateFlags == TerminateCondition.None || (terminateFlags & condition) != 0)
+				string fileName = imageFile.GetFileName();
+				if (processNamesToTerminate.TryGetValue(fileName, out TerminateCondition terminateFlags))
 				{
-					return true;
+					if (terminateFlags == TerminateCondition.None || (terminateFlags & condition) != 0)
+					{
+						return true;
+					}
 				}
 			}
 
