Horde: Move settings for different executors into their own files.

diff --git a/Engine/Source/Programs/Horde/Horde.Agent/AgentSettings.cs b/Engine/Source/Programs/Horde/Horde.Agent/AgentSettings.cs
--- a/Engine/Source/Programs/Horde/Horde.Agent/AgentSettings.cs
+++ b/Engine/Source/Programs/Horde/Horde.Agent/AgentSettings.cs
@@ -3,6 +3,7 @@
 using System.ComponentModel.DataAnnotations;
 using EpicGames.Core;
 using EpicGames.Horde;
+using Horde.Agent.Execution;
 using Horde.Agent.Utility;
 using Microsoft.Extensions.Configuration;
 
@@ -80,33 +81,6 @@ namespace Horde.Agent
 		}
 	}
 
-	/// <summary>
-	/// Settings for the local executor
-	/// </summary>
-	public class LocalExecutorSettings
-	{
-		/// <summary>
-		/// Path to the local workspace to use with the local executor
-		/// </summary>
-		public string? WorkspaceDir { get; set; }
-
-		/// <summary>
-		/// Whether to actually execute steps, or just do job setup
-		/// </summary>
-		public bool RunSteps { get; set; } = true;
-	}
-
-	/// <summary>
-	/// Settings for the perforce executor
-	/// </summary>
-	public class PerforceExecutorSettings
-	{
-		/// <summary>
-		/// Whether to run conform jobs
-		/// </summary>
-		public bool RunConform { get; set; } = true;
-	}
-
 	/// <summary>
 	/// Specifies a process to terminate
 	/// </summary>
diff --git a/Engine/Source/Programs/Horde/Horde.Agent/Execution/LocalExecutor.cs b/Engine/Source/Programs/Horde/Horde.Agent/Execution/LocalExecutor.cs
--- a/Engine/Source/Programs/Horde/Horde.Agent/Execution/LocalExecutor.cs
+++ b/Engine/Source/Programs/Horde/Horde.Agent/Execution/LocalExecutor.cs
@@ -7,6 +7,25 @@ using Microsoft.Extensions.Options;
 
 namespace Horde.Agent.Execution
 {
+	/// <summary>
+	/// Settings for the local executor
+	/// </summary>
+	public class LocalExecutorSettings
+	{
+		/// <summary>
+		/// Path to the local workspace to use with the local executor
+		/// </summary>
+		public string? WorkspaceDir { get; set; }
+
+		/// <summary>
+		/// Whether to actually execute steps, or just do job setup
+		/// </summary>
+		public bool RunSteps { get; set; } = true;
+	}
+
+	/// <summary>
+	/// Job executor which runs out of a local workspace
+	/// </summary>
 	class LocalExecutor : JobExecutor
 	{
 		private readonly LocalExecutorSettings _settings;
diff --git a/Engine/Source/Programs/Horde/Horde.Agent/Execution/PerforceExecutor.cs b/Engine/Source/Programs/Horde/Horde.Agent/Execution/PerforceExecutor.cs
--- a/Engine/Source/Programs/Horde/Horde.Agent/Execution/PerforceExecutor.cs
+++ b/Engine/Source/Programs/Horde/Horde.Agent/Execution/PerforceExecutor.cs
@@ -15,6 +15,20 @@ using OpenTracing.Util;
 
 namespace Horde.Agent.Execution
 {
+	/// <summary>
+	/// Settings for the perforce executor
+	/// </summary>
+	public class PerforceExecutorSettings
+	{
+		/// <summary>
+		/// Whether to run conform jobs
+		/// </summary>
+		public bool RunConform { get; set; } = true;
+	}
+
+	/// <summary>
+	/// Executor which syncs workspaces from Perforce
+	/// </summary>
 	class PerforceExecutor : JobExecutor
 	{
 		public const string Name = "Perforce";
