Horde: Add missing entry point for driver app.

#rnx

diff --git a/Engine/Source/Programs/Horde/Horde.Agent.Driver/DriverApp.cs b/Engine/Source/Programs/Horde/Horde.Agent.Driver/DriverApp.cs
new file mode 100644
--- /dev/null
+++ b/Engine/Source/Programs/Horde/Horde.Agent.Driver/DriverApp.cs
@@ -0,0 +1,25 @@
+using System;
+using System.Collections.Generic;
+using System.Linq;
+using System.Reflection;
+using System.Text;
+using System.Threading.Tasks;
+using EpicGames.Core;
+using Microsoft.Extensions.DependencyInjection;
+
+namespace Horde.Agent.Driver
+{
+	class DriverApp
+	{
+		public static async Task<int> Main(string[] args)
+		{
+			CommandLineArguments arguments = new CommandLineArguments(args);
+
+			IServiceCollection services = new ServiceCollection();
+			services.AddCommandsFromAssembly(Assembly.GetExecutingAssembly());
+
+			await using ServiceProvider serviceProvider = services.BuildServiceProvider();
+			return await CommandHost.RunAsync(arguments, serviceProvider, null);
+		}
+	}
+}
