Horde: Add missing copyright notice.

#rnx

diff --git a/Engine/Source/Programs/Horde/Horde.Agent.Driver/DriverApp.cs b/Engine/Source/Programs/Horde/Horde.Agent.Driver/DriverApp.cs
--- a/Engine/Source/Programs/Horde/Horde.Agent.Driver/DriverApp.cs
+++ b/Engine/Source/Programs/Horde/Horde.Agent.Driver/DriverApp.cs
@@ -1,3 +1,5 @@
+// Copyright Epic Games, Inc. All Rights Reserved.
+
 using System;
 using System.Collections.Generic;
 using System.Linq;
