Horde: Track getMore commands from MongoDB under same parent trace span

Find and other commands batches fetching of documents as getMore commands. Keeping track of these allow for a proper overview when viewing traces.

#rnx

diff --git a/Engine/Source/Programs/Horde/Horde.Server/Utilities/MongoCommandTracer.cs b/Engine/Source/Programs/Horde/Horde.Server/Utilities/MongoCommandTracer.cs
--- a/Engine/Source/Programs/Horde/Horde.Server/Utilities/MongoCommandTracer.cs
+++ b/Engine/Source/Programs/Horde/Horde.Server/Utilities/MongoCommandTracer.cs
@@ -2,6 +2,7 @@
 
 using System;
 using System.Collections.Concurrent;
+using System.Collections.Generic;
 using Microsoft.Extensions.Logging;
 using MongoDB.Bson;
 using MongoDB.Driver.Core.Configuration;
@@ -13,11 +14,68 @@ namespace Horde.Server.Utilities;
 /// <summary>
 /// OpenTelemetry-based tracer listening for MongoDB command events
 /// </summary>
-/// <param name="tracer">Tracer</param>
-/// <param name="logger">Logger</param>
-public class MongoCommandTracer(Tracer tracer, ILogger<MongoCommandTracer> logger)
+public class MongoCommandTracer
 {
-	private readonly ConcurrentDictionary<int, TelemetrySpan> _activeSpans = new();
+	private readonly ConcurrentDictionary<int, TelemetrySpan> _requestToSpans = new();
+	private readonly ConcurrentDictionary<long, OperationSpanEntry> _operationToSpans = new();
+	private readonly Dictionary<string, MongoCommand> _strToCommand = new();
+	private readonly Tracer _tracer;
+	private readonly ILogger<MongoCommandTracer> _logger;
+	private readonly ConcurrentDictionary<string, byte> _unhandledCommands = new();
+
+	private record OperationSpanEntry(TelemetrySpan Span, int RequestId);
+	private record MongoCommand(MongoCommandType Type, string Name, string? StatementFieldName);
+
+	private enum MongoCommandType
+	{
+		Aggregate,
+		CreateIndexes,
+		Delete,
+		Find,
+		FindAndModify,
+		GetLastError,
+		GetMore,
+		Insert,
+		IsMaster,
+		ListIndexes,
+		Ping,
+		SaslContinue,
+		SaslStart,
+		Update,
+	}
+
+	/// <summary>
+	/// Constructor
+	/// </summary>
+	/// <param name="tracer">Tracer</param>
+	/// <param name="logger">Logger</param>
+	public MongoCommandTracer(Tracer tracer, ILogger<MongoCommandTracer> logger)
+	{
+		_tracer = tracer;
+		_logger = logger;
+
+		List<MongoCommand> commands = [
+			new MongoCommand(MongoCommandType.Aggregate, "aggregate", "aggregate"),
+			new MongoCommand(MongoCommandType.CreateIndexes, "createIndexes", "createIndexes"),
+			new MongoCommand(MongoCommandType.Delete, "delete", "deletes"),
+			new MongoCommand(MongoCommandType.Find, "find", "filter"),
+			new MongoCommand(MongoCommandType.FindAndModify, "findAndModify", "query"),
+			new MongoCommand(MongoCommandType.GetLastError, "getLastError", null),
+			new MongoCommand(MongoCommandType.GetMore, "getMore", "collection"),
+			new MongoCommand(MongoCommandType.Insert, "insert", null),
+			new MongoCommand(MongoCommandType.IsMaster, "isMaster", null),
+			new MongoCommand(MongoCommandType.ListIndexes, "listIndexes", "listIndexes"),
+			new MongoCommand(MongoCommandType.Ping, "ping", null),
+			new MongoCommand(MongoCommandType.SaslContinue, "saslContinue", null),
+			new MongoCommand(MongoCommandType.SaslStart, "saslStart", null),
+			new MongoCommand(MongoCommandType.Update, "update", "updates"),
+		];
+
+		foreach (MongoCommand cmd in commands)
+		{
+			_strToCommand[cmd.Name] = cmd;
+		}
+	}
 
 	/// <summary>
 	/// Registers event listeners from MongoDB's client
@@ -39,34 +97,40 @@ public class MongoCommandTracer(Tracer tracer, ILogger<MongoCommandTracer> logge
 		
 		try
 		{
-			// TODO: Nest getMore under the originating span as it's really a continuation of a find command
-			(string? collectionName, string? statement) = ev.CommandName switch
+			MongoCommand? command = ResolveCommand(ev.CommandName);
+			if (command == null)
 			{
-				"find" => (GetString("find"), GetString("filter")),
-				"update" => (GetString("update"), GetString("updates")),
-				"delete" => (GetString("delete"), GetString("deletes")),
-				"insert" => (GetString("insert"), null),
-				"findAndModify" => (GetString("findAndModify"), GetString("query")),
-				"aggregate" => (GetString("aggregate"), null),
-				"getMore" => (GetString("collection"), null),
-				"listIndexes" => (GetString("listIndexes"), null),
-				"createIndexes" => (GetString("createIndexes"), null),
-				_ => (null, null)
-			};
-
-			string name = ev.CommandName;
+				return;
+			}
+
+			string name = command.Name;
+			string? collectionName = GetString(command.Name);
+			string? statement = command.StatementFieldName != null ? GetString(command.StatementFieldName) : null;
+			
 			if (collectionName != null)
 			{
-				name = $"{collectionName}.{ev.CommandName}";
+				name = $"{collectionName}.{command.Name}";
+			}
+
+			OperationSpanEntry? parentSpanEntry = null;
+			if (command.Type == MongoCommandType.GetMore && ev.OperationId != null)
+			{
+				_operationToSpans.TryGetValue(ev.OperationId.Value, out parentSpanEntry);
 			}
 			
+			SpanAttributes sa = new();
+			sa.Add("type", "db");
+			sa.Add("operation.name", name);
+			sa.Add("service.name", OpenTelemetryTracers.MongoDbName);
+			
 			// OpenTelemetry MongoDB conventions https://opentelemetry.io/docs/specs/semconv/database/mongodb/
-			TelemetrySpan span = tracer.StartMongoDbSpan<object>(name);
-			span.SetAttribute("db.system", "mongodb");
-			span.SetAttribute("db.name", ev.DatabaseNamespace.ToString());
-			span.SetAttribute("db.operationId", ev.OperationId ?? -1);
-			span.SetAttribute("db.requestId", ev.RequestId);
-			span.SetAttribute("db.serviceId", ev.ServiceId?.ToString());
+			sa.Add("db.system", "mongodb");
+			sa.Add("db.name", ev.DatabaseNamespace.ToString());
+			sa.Add("db.operationId", ev.OperationId ?? -1);
+			sa.Add("db.requestId", ev.RequestId);
+			sa.Add("db.serviceId", ev.ServiceId?.ToString());
+			
+			TelemetrySpan span = _tracer.StartActiveSpan(name, SpanKind.Client, parentContext: parentSpanEntry?.Span.Context ?? Tracer.CurrentSpan.Context, sa);
 			
 			if (collectionName != null)
 			{
@@ -78,34 +142,112 @@ public class MongoCommandTracer(Tracer tracer, ILogger<MongoCommandTracer> logge
 				span.SetAttribute("db.statement", statement.Length > 200 ? statement[..200] : statement);
 			}
 
-			_activeSpans.TryAdd(ev.RequestId, span);
+			_requestToSpans.TryAdd(ev.RequestId, span);
+			if (command.Type == MongoCommandType.Find && ev.OperationId != null)
+			{
+				_operationToSpans.TryAdd(ev.OperationId.Value, new OperationSpanEntry(span, ev.RequestId));
+			}
 		}
 		catch (Exception e)
 		{
-			logger.LogError(e, "Unhandled exception when capturing MongoDB span");
+			_logger.LogError(e, "Unhandled exception when capturing MongoDB span");
 		}
 	}
 	
 	private void OnEvent(CommandSucceededEvent ev)
 	{
-		if (_activeSpans.TryRemove(ev.RequestId, out TelemetrySpan? span))
+		MongoCommand? command = ResolveCommand(ev.CommandName);
+		if (command == null)
 		{
-			span.SetStatus(Status.Ok);
-			span.SetAttribute("db.durationMs", ev.Duration.TotalMilliseconds);
-			span.End();
-			span.Dispose();
+			return;
+		}
+
+		long? waitedMs = null;
+		if (command.Type is MongoCommandType.Find or MongoCommandType.GetMore && ev.OperationId != null)
+		{
+			waitedMs = GetBsonDocumentLong(ev.Reply, "waitedMS");
+			long? cursorId = GetCursorId(ev.Reply);
+			bool hasMoreDocuments = cursorId is > 0;
+
+			if (hasMoreDocuments && command.Type == MongoCommandType.Find)
+			{
+				// Don't end the span for "find", keep it open and wait for subsequent "getMore" reply
+				return;
+			}
+
+			if (!hasMoreDocuments && command.Type == MongoCommandType.GetMore)
+			{
+				if (_operationToSpans.TryGetValue(ev.OperationId.Value, out OperationSpanEntry? opSpanEntry))
+				{
+					EndSpan(opSpanEntry.RequestId, ev.OperationId.Value, Status.Ok);
+				}
+			}
 		}
+		
+		EndSpan(ev.RequestId, opId: null, Status.Ok, ev.Duration, waitedMs: waitedMs);
 	}
 	
 	private void OnEvent(CommandFailedEvent ev)
 	{
-		if (_activeSpans.TryRemove(ev.RequestId, out TelemetrySpan? span))
+		EndSpan(ev.RequestId, ev.OperationId, Status.Ok, ev.Duration, ev.Failure);
+	}
+
+	private MongoCommand? ResolveCommand(string name)
+	{
+		if (_strToCommand.TryGetValue(name, out MongoCommand? command))
 		{
-			span.SetStatus(Status.Error);
-			span.SetAttribute("db.durationMs", ev.Duration.TotalMilliseconds);
-			span.RecordException(ev.Failure);
+			return command;
+		}
+
+		// Ensure logging is done once per command by storing which type has been logged
+		if (_unhandledCommands.TryAdd(name, 0))
+		{
+			_logger.LogInformation("Trace handling for MongoDB command {Command} is not implemented. Ignoring", name);
+		}
+
+		return null;
+	}
+	
+	private void EndSpan(int reqId, long? opId, Status status, TimeSpan? duration = null, Exception? exception = null, long? waitedMs = null)
+	{
+		if (_requestToSpans.TryRemove(reqId, out TelemetrySpan? span))
+		{
+			span.SetStatus(status);
+			
+			if (duration != null)
+			{
+				span.SetAttribute("db.durationMs", duration.Value.TotalMilliseconds);
+			}
+			
+			if (exception != null)
+			{
+				span.RecordException(exception);
+			}
+			
+			if (waitedMs != null)
+			{
+				span.SetAttribute("db.waitedMs", waitedMs.Value);
+			}
+			
 			span.End();
 			span.Dispose();
 		}
+
+		if (opId != null)
+		{
+			_operationToSpans.TryRemove(opId.Value, out OperationSpanEntry? _);
+		}
+	}
+	
+	private static long? GetBsonDocumentLong(BsonDocument document, string key)
+	{
+		return document.TryGetPropertyValue(key, BsonType.Int64, out BsonValue? value) ? value.AsInt64 : null;
+	}
+	
+	private static long? GetCursorId(BsonDocument replyDocument)
+	{
+		return replyDocument.TryGetPropertyValue("cursor", BsonType.Document, out BsonValue? cursorDoc)
+			? GetBsonDocumentLong(cursorDoc.AsBsonDocument, "id")
+			: null;
 	}
 }
