Horde: Remove unnecessary using directives.

#rnx

diff --git a/Engine/Source/Programs/Horde/Horde.Agent.Driver/DriverApp.cs b/Engine/Source/Programs/Horde/Horde.Agent.Driver/DriverApp.cs
--- a/Engine/Source/Programs/Horde/Horde.Agent.Driver/DriverApp.cs
+++ b/Engine/Source/Programs/Horde/Horde.Agent.Driver/DriverApp.cs
@@ -1,11 +1,6 @@
 // Copyright Epic Games, Inc. All Rights Reserved.
 
-using System;
-using System.Collections.Generic;
-using System.Linq;
 using System.Reflection;
-using System.Text;
-using System.Threading.Tasks;
 using EpicGames.Core;
 using Microsoft.Extensions.DependencyInjection;
 
