Horde: Add Horde.Agent.Driver to container when building docker image.

#rnx

diff --git a/Engine/Source/Programs/Horde/Horde.Server/Dockerfile b/Engine/Source/Programs/Horde/Horde.Server/Dockerfile
--- a/Engine/Source/Programs/Horde/Horde.Server/Dockerfile
+++ b/Engine/Source/Programs/Horde/Horde.Server/Dockerfile
@@ -38,6 +38,7 @@ COPY Source/Programs/Shared/EpicGames.Redis/*.csproj ./Source/Programs/Shared/Ep
 COPY Source/Programs/Shared/EpicGames.Serialization/*.csproj ./Source/Programs/Shared/EpicGames.Serialization/
 COPY Source/Programs/Shared/EpicGames.Serialization.Tests/*.csproj ./Source/Programs/Shared/EpicGames.Serialization.Tests/
 COPY Source/Programs/Horde/Horde.Agent/*.csproj ./Source/Programs/Horde/Horde.Agent/
+COPY Source/Programs/Horde/Horde.Agent.Driver/*.csproj ./Source/Programs/Horde/Horde.Agent.Driver/
 COPY Source/Programs/Horde/Horde.Server/*.csproj ./Source/Programs/Horde/Horde.Server/
 COPY Source/Programs/Horde/Horde.Server.Tests/*.csproj ./Source/Programs/Horde/Horde.Server.Tests/
 RUN dotnet restore Source/Programs/Horde/Horde.Server/Horde.Server.csproj -p:WithHordeStorage=false
@@ -54,6 +55,7 @@ RUN dotnet publish Source/Programs/Horde/Horde.Server/Horde.Server.csproj -c Rel
 
 # Run tests
 COPY Source/Programs/Horde/Horde.Agent/ Source/Programs/Horde/Horde.Agent/
+COPY Source/Programs/Horde/Horde.Agent.Driver/ Source/Programs/Horde/Horde.Agent.Driver/
 COPY Source/Programs/Horde/Horde.Server.Tests/ Source/Programs/Horde/Horde.Server.Tests/
 COPY Source/Programs/Shared/EpicGames.Serialization.Tests/ Source/Programs/Shared/EpicGames.Serialization.Tests/
 
