Horde: Remove GrpcChannel from sesson object. The supplied RpcConnection can be used for the same purpose.

#rnx

diff --git a/Engine/Source/Programs/Horde/Horde.Agent.Tests/Execution/WorkspaceExecutorTest.cs b/Engine/Source/Programs/Horde/Horde.Agent.Tests/Execution/WorkspaceExecutorTest.cs
--- a/Engine/Source/Programs/Horde/Horde.Agent.Tests/Execution/WorkspaceExecutorTest.cs
+++ b/Engine/Source/Programs/Horde/Horde.Agent.Tests/Execution/WorkspaceExecutorTest.cs
@@ -10,7 +10,6 @@ using EpicGames.Horde.Jobs;
 using EpicGames.Horde.Streams;
 using Horde.Agent.Execution;
 using Horde.Agent.Services;
-using Horde.Agent.Utility;
 using HordeCommon.Rpc;
 using HordeCommon.Rpc.Tasks;
 using Microsoft.Extensions.Logging;
@@ -47,7 +46,7 @@ public sealed class WorkspaceExecutorTest : IAsyncDisposable
 		_server.AddStream(_streamId, "//Foo/Main");
 		_server.AddAgentType(_streamId, AgentType);
 		_server.AddJob(_jobId, _streamId, 1, 0);
-		_session = FakeServerSessionFactory.CreateSession(_server.GetConnection(), _server.GetGrpcChannel());
+		_session = FakeServerSessionFactory.CreateSession(_server.GetConnection());
 
 		_autoSdkWorkspace.SetFile(1, "HostWin64/Android/base.h", "base");
 		_workspace.SetFile(1, "main.cpp", "main");
diff --git a/Engine/Source/Programs/Horde/Horde.Agent.Tests/WorkerServiceTest.cs b/Engine/Source/Programs/Horde/Horde.Agent.Tests/WorkerServiceTest.cs
--- a/Engine/Source/Programs/Horde/Horde.Agent.Tests/WorkerServiceTest.cs
+++ b/Engine/Source/Programs/Horde/Horde.Agent.Tests/WorkerServiceTest.cs
@@ -162,7 +162,7 @@ namespace Horde.Agent.Tests
 			await using RpcConnectionStub rpcConnection = new RpcConnectionStub(null!, null!, client);
 
 			await using FakeHordeRpcServer fakeServer = new();
-			await using ISession session = FakeServerSessionFactory.CreateSession(rpcConnection, fakeServer.GetGrpcChannel());
+			await using ISession session = FakeServerSessionFactory.CreateSession(rpcConnection);
 
 			client.BeginStepResponses.Enqueue(new RpcBeginStepResponse { Name = "stepName1", StepId = _stepId1.ToString() });
 			client.BeginStepResponses.Enqueue(new RpcBeginStepResponse { Name = "stepName2", StepId = _stepId2.ToString() });
@@ -252,7 +252,7 @@ namespace Horde.Agent.Tests
 			cts.CancelAfter(20000);
 
 			await using FakeHordeRpcServer fakeServer = new();
-			await using ISession session = FakeServerSessionFactory.CreateSession(fakeServer.GetConnection(), fakeServer.GetGrpcChannel());
+			await using ISession session = FakeServerSessionFactory.CreateSession(fakeServer.GetConnection());
 
 			LeaseManager manager = new LeaseManager(session, serviceProvider);
 
@@ -271,17 +271,16 @@ namespace Horde.Agent.Tests
 
 		public Task<ISession> CreateAsync(CancellationToken cancellationToken)
 		{
-			return Task.FromResult(CreateSession(_fakeServer.GetConnection(), _fakeServer.GetGrpcChannel()));
+			return Task.FromResult(CreateSession(_fakeServer.GetConnection()));
 		}
 
-		public static ISession CreateSession(IRpcConnection rpcConnection, GrpcChannel grpcChannel)
+		public static ISession CreateSession(IRpcConnection rpcConnection)
 		{
 			Mock<ISession> fakeSession = new Mock<ISession>(MockBehavior.Strict);
 			fakeSession.Setup(x => x.ServerUrl).Returns(new Uri("https://localhost:9999"));
 			fakeSession.Setup(x => x.AgentId).Returns(new EpicGames.Horde.Agents.AgentId("LocalAgent"));
 			fakeSession.Setup(x => x.SessionId).Returns(new EpicGames.Horde.Agents.Sessions.SessionId(default));
 			fakeSession.Setup(x => x.RpcConnection).Returns(rpcConnection);
-			fakeSession.Setup(x => x.GrpcChannel).Returns(grpcChannel);
 			fakeSession.Setup(x => x.ProcessNamesToTerminate).Returns((IReadOnlyDictionary<string, TerminateCondition>)new Dictionary<string, TerminateCondition>());
 			fakeSession.Setup(x => x.DisposeAsync()).Returns(new ValueTask());
 			fakeSession.Setup(x => x.WorkingDir).Returns(DirectoryReference.Combine(DirectoryReference.GetCurrentDirectory(), Guid.NewGuid().ToString()));
diff --git a/Engine/Source/Programs/Horde/Horde.Agent/Commands/Execution/ExecuteJobCommand.cs b/Engine/Source/Programs/Horde/Horde.Agent/Commands/Execution/ExecuteJobCommand.cs
--- a/Engine/Source/Programs/Horde/Horde.Agent/Commands/Execution/ExecuteJobCommand.cs
+++ b/Engine/Source/Programs/Horde/Horde.Agent/Commands/Execution/ExecuteJobCommand.cs
@@ -4,7 +4,6 @@ using EpicGames.Core;
 using EpicGames.Horde.Agents;
 using EpicGames.Horde.Agents.Leases;
 using EpicGames.Horde.Agents.Sessions;
-using Grpc.Net.Client;
 using Horde.Agent.Leases.Handlers;
 using Horde.Agent.Services;
 using Horde.Agent.Utility;
@@ -51,8 +50,7 @@ namespace Horde.Agent.Commands.Execution
 			Dictionary<string, TerminateCondition> processNamesToTerminate = _settings.GetProcessesToTerminateMap();
 
 			await using RpcConnection rpcConnection = new RpcConnection(ctx => _grpcService.CreateGrpcChannelAsync(executeTask.Token, ctx), logger);
-			using GrpcChannel channel = await _grpcService.CreateGrpcChannelAsync(executeTask.Token, CancellationToken.None);
-			await using Session session = new Session(serverProfile.Url, AgentId, SessionId, executeTask.Token, rpcConnection, channel, WorkingDir, processNamesToTerminate);
+			await using Session session = new Session(serverProfile.Url, AgentId, SessionId, executeTask.Token, rpcConnection, WorkingDir, processNamesToTerminate);
 
 			await _jobHandler.ExecuteInternalAsync(session, LeaseId, executeTask, logger, CancellationToken.None);
 			return 0;
diff --git a/Engine/Source/Programs/Horde/Horde.Agent/Leases/LeaseManager.cs b/Engine/Source/Programs/Horde/Horde.Agent/Leases/LeaseManager.cs
--- a/Engine/Source/Programs/Horde/Horde.Agent/Leases/LeaseManager.cs
+++ b/Engine/Source/Programs/Horde/Horde.Agent/Leases/LeaseManager.cs
@@ -230,7 +230,6 @@ namespace Horde.Agent.Leases
 		async Task<SessionResult> HandleSessionAsync(bool shutdownAfterFinishedLease, CancellationToken stoppingToken)
 		{
 			IRpcConnection rpcCon = _session.RpcConnection;
-			HordeRpc.HordeRpcClient rpcClient = new(_session.GrpcChannel);
 
 			// Terminate any remaining child processes from other instances
 			await _session.TerminateProcessesAsync(TerminateCondition.BeforeSession, _logger, stoppingToken);
@@ -243,7 +242,7 @@ namespace Horde.Agent.Leases
 			await using BackgroundTask updateCapsTask = BackgroundTask.StartNew(ctx => UpdateCapabilitiesBackgroundAsync(_session.WorkingDir, ctx));
 
 			// Run another background task to send telemetry data
-			await using BackgroundTask telemetryTask = BackgroundTask.StartNew(ctx => SendTelemetryAsync(rpcClient, ctx));
+			await using BackgroundTask telemetryTask = BackgroundTask.StartNew(ctx => SendTelemetryAsync(_session.RpcConnection, ctx));
 
 			// Loop until we're ready to exit
 			Stopwatch updateCapabilitiesTimer = Stopwatch.StartNew();
@@ -309,7 +308,7 @@ namespace Horde.Agent.Leases
 				using (stopping ? (CancellationTokenRegistration?)null : stoppingToken.Register(() => _updateLeasesEvent.Set()))
 				{
 					// Update the state with the server
-					RpcUpdateSessionResponse? updateSessionResponse = await UpdateSessionAsync(rpcClient, updateSessionRequest, waitTask);
+					RpcUpdateSessionResponse? updateSessionResponse = await UpdateSessionAsync(_session.RpcConnection, updateSessionRequest, waitTask);
 
 					lock (_lockObject)
 					{
@@ -414,16 +413,18 @@ namespace Horde.Agent.Leases
 		/// <summary>
 		/// Wrapper for <see cref="UpdateSessionInternalAsync"/> which filters/logs exceptions
 		/// </summary>
-		/// <param name="rpcClient">The RPC client connection</param>
+		/// <param name="rpcConnection">The RPC client connection</param>
 		/// <param name="updateSessionRequest">The session update request</param>
 		/// <param name="waitTask">Task which can be used to jump out of the update early</param>
 		/// <returns>Response from the call</returns>
-		async Task<RpcUpdateSessionResponse?> UpdateSessionAsync(HordeRpc.HordeRpcClient rpcClient, RpcUpdateSessionRequest updateSessionRequest, Task waitTask)
+		async Task<RpcUpdateSessionResponse?> UpdateSessionAsync(IRpcConnection rpcConnection, RpcUpdateSessionRequest updateSessionRequest, Task waitTask)
 		{
+			using IRpcClientRef<HordeRpc.HordeRpcClient> rpcClientRef = await rpcConnection.GetClientRefAsync<HordeRpc.HordeRpcClient>(CancellationToken.None);
+
 			RpcUpdateSessionResponse? updateSessionResponse = null;
 			try
 			{
-				updateSessionResponse = await UpdateSessionInternalAsync(rpcClient, updateSessionRequest, waitTask);
+				updateSessionResponse = await UpdateSessionInternalAsync(rpcClientRef.Client, updateSessionRequest, waitTask);
 				_updateSessionFailures = 0;
 			}
 			catch (RpcException ex)
@@ -592,7 +593,7 @@ namespace Horde.Agent.Leases
 		/// <summary>
 		/// Periodically send agent telemetry to the server
 		/// </summary>
-		async Task SendTelemetryAsync(HordeRpc.HordeRpcClient rpcClient, CancellationToken cancellationToken)
+		async Task SendTelemetryAsync(IRpcConnection rpcConnection, CancellationToken cancellationToken)
 		{
 			while (!cancellationToken.IsCancellationRequested)
 			{
@@ -617,7 +618,9 @@ namespace Horde.Agent.Leases
 							request.FreeRam = memMetrics.Available / 1024;
 							request.UsedRam = memMetrics.Used / 1024;
 						}
-						await rpcClient.UploadTelemetryAsync(request, new CallOptions(cancellationToken: cancellationToken));
+
+						CallOptions callOptions = new CallOptions(cancellationToken: cancellationToken);
+						await rpcConnection.InvokeAsync<HordeRpc.HordeRpcClient, Empty>(client => client.UploadTelemetryAsync(request, callOptions), cancellationToken);
 					}
 				}
 				catch (OperationCanceledException) when (cancellationToken.IsCancellationRequested)
diff --git a/Engine/Source/Programs/Horde/Horde.Agent/Services/SessionFactory.cs b/Engine/Source/Programs/Horde/Horde.Agent/Services/SessionFactory.cs
--- a/Engine/Source/Programs/Horde/Horde.Agent/Services/SessionFactory.cs
+++ b/Engine/Source/Programs/Horde/Horde.Agent/Services/SessionFactory.cs
@@ -47,11 +47,6 @@ namespace Horde.Agent.Services
 		/// </summary>
 		IRpcConnection RpcConnection { get; }
 
-		/// <summary>
-		/// A gRPC channel authenticated for this session
-		/// </summary>
-		GrpcChannel GrpcChannel { get; }
-
 		/// <summary>
 		/// Working directory for sandboxes etc..
 		/// </summary>
@@ -96,9 +91,6 @@ namespace Horde.Agent.Services
 		/// <inheritdoc/>
 		public IRpcConnection RpcConnection { get; }
 
-		/// <inheritdoc/>
-		public GrpcChannel GrpcChannel { get; }
-
 		/// <inheritdoc/>
 		public DirectoryReference WorkingDir { get; }
 
@@ -108,14 +100,13 @@ namespace Horde.Agent.Services
 		/// <summary>
 		/// Constructor
 		/// </summary>
-		public Session(Uri serverUrl, AgentId agentId, SessionId sessionId, string token, IRpcConnection rpcConnection, GrpcChannel grpcChannel, DirectoryReference workingDir, IReadOnlyDictionary<string, TerminateCondition> processNamesToTerminate)
+		public Session(Uri serverUrl, AgentId agentId, SessionId sessionId, string token, IRpcConnection rpcConnection, DirectoryReference workingDir, IReadOnlyDictionary<string, TerminateCondition> processNamesToTerminate)
 		{
 			ServerUrl = serverUrl;
 			AgentId = agentId;
 			SessionId = sessionId;
 			Token = token;
 			RpcConnection = rpcConnection;
-			GrpcChannel = grpcChannel;
 			WorkingDir = workingDir;
 			ProcessNamesToTerminate = processNamesToTerminate;
 		}
@@ -232,13 +223,11 @@ namespace Horde.Agent.Services
 				}
 			}
 
-			Func<CancellationToken, Task<GrpcChannel>> createGrpcChannelAsync = ctx => grpcService.CreateGrpcChannelAsync(createSessionResponse.Token, ctx);
-
 			// Open a connection to the server
 #pragma warning disable CA2000 // False positive; ownership is transferred to new Session object.
+			Func<CancellationToken, Task<GrpcChannel>> createGrpcChannelAsync = ctx => grpcService.CreateGrpcChannelAsync(createSessionResponse.Token, ctx);
 			IRpcConnection rpcConnection = new RpcConnection(createGrpcChannelAsync, logger);
-			GrpcChannel sessionGrpcChannel = await grpcService.CreateGrpcChannelAsync(createSessionResponse.Token, cancellationToken);
-			return new Session(serverProfile.Url, new AgentId(createSessionResponse.AgentId), SessionId.Parse(createSessionResponse.SessionId), createSessionResponse.Token, rpcConnection, sessionGrpcChannel, workingDir, currentSettings.GetProcessesToTerminateMap());
+			return new Session(serverProfile.Url, new AgentId(createSessionResponse.AgentId), SessionId.Parse(createSessionResponse.SessionId), createSessionResponse.Token, rpcConnection, workingDir, currentSettings.GetProcessesToTerminateMap());
 #pragma warning restore CA2000
 		}
 
