Horde: Fix invalid collection name for getMore MongoDB traces

Also add two missing commands to known list.

#rnx

diff --git a/Engine/Source/Programs/Horde/Horde.Server/Utilities/MongoCommandTracer.cs b/Engine/Source/Programs/Horde/Horde.Server/Utilities/MongoCommandTracer.cs
--- a/Engine/Source/Programs/Horde/Horde.Server/Utilities/MongoCommandTracer.cs
+++ b/Engine/Source/Programs/Horde/Horde.Server/Utilities/MongoCommandTracer.cs
@@ -24,19 +24,21 @@ public class MongoCommandTracer
 	private readonly ConcurrentDictionary<string, byte> _unhandledCommands = new();
 
 	private record OperationSpanEntry(TelemetrySpan Span, int RequestId);
-	private record MongoCommand(MongoCommandType Type, string Name, string? StatementFieldName);
+	private record MongoCommand(MongoCommandType Type, string Name, string? StatementFieldName, string? CollectionFieldName = null);
 
 	private enum MongoCommandType
 	{
 		Aggregate,
 		CreateIndexes,
 		Delete,
+		Distinct,
 		Find,
 		FindAndModify,
 		GetLastError,
 		GetMore,
 		Insert,
 		IsMaster,
+		KillCursors,
 		ListIndexes,
 		Ping,
 		SaslContinue,
@@ -57,18 +59,20 @@ public class MongoCommandTracer
 		List<MongoCommand> commands = [
 			new MongoCommand(MongoCommandType.Aggregate, "aggregate", "aggregate"),
 			new MongoCommand(MongoCommandType.CreateIndexes, "createIndexes", "createIndexes"),
-			new MongoCommand(MongoCommandType.Delete, "delete", "deletes"),
-			new MongoCommand(MongoCommandType.Find, "find", "filter"),
-			new MongoCommand(MongoCommandType.FindAndModify, "findAndModify", "query"),
+			new MongoCommand(MongoCommandType.Delete, "delete", "deletes", "delete"),
+			new MongoCommand(MongoCommandType.Distinct, "distinct", null),
+			new MongoCommand(MongoCommandType.Find, "find", "filter", "find"),
+			new MongoCommand(MongoCommandType.FindAndModify, "findAndModify", "query", "findAndModify"),
 			new MongoCommand(MongoCommandType.GetLastError, "getLastError", null),
-			new MongoCommand(MongoCommandType.GetMore, "getMore", "collection"),
-			new MongoCommand(MongoCommandType.Insert, "insert", null),
+			new MongoCommand(MongoCommandType.GetMore, "getMore", null, "collection"),
+			new MongoCommand(MongoCommandType.Insert, "insert", null, "insert"),
 			new MongoCommand(MongoCommandType.IsMaster, "isMaster", null),
+			new MongoCommand(MongoCommandType.KillCursors, "killCursors", null),
 			new MongoCommand(MongoCommandType.ListIndexes, "listIndexes", "listIndexes"),
 			new MongoCommand(MongoCommandType.Ping, "ping", null),
 			new MongoCommand(MongoCommandType.SaslContinue, "saslContinue", null),
 			new MongoCommand(MongoCommandType.SaslStart, "saslStart", null),
-			new MongoCommand(MongoCommandType.Update, "update", "updates"),
+			new MongoCommand(MongoCommandType.Update, "update", "updates", "update"),
 		];
 
 		foreach (MongoCommand cmd in commands)
@@ -104,7 +108,7 @@ public class MongoCommandTracer
 			}
 
 			string name = command.Name;
-			string? collectionName = GetString(command.Name);
+			string? collectionName = GetString(command.CollectionFieldName ?? command.Name);
 			string? statement = command.StatementFieldName != null ? GetString(command.StatementFieldName) : null;
 			
 			if (collectionName != null)
