Horde: Prevent default parameters being appended to jobs when an explicit argument list is specified.

diff --git a/Engine/Source/Programs/Horde/Horde.Server/Jobs/JobService.cs b/Engine/Source/Programs/Horde/Horde.Server/Jobs/JobService.cs
--- a/Engine/Source/Programs/Horde/Horde.Server/Jobs/JobService.cs
+++ b/Engine/Source/Programs/Horde/Horde.Server/Jobs/JobService.cs
@@ -1373,7 +1373,7 @@ namespace Horde.Server.Jobs
 					options.PreflightChange = newJob.PreflightChange;
 					options.PreflightDescription = newJob.PreflightDescription;
 					template.GetDefaultParameters(options.Parameters, true);
-					template.GetArgumentsForParameters(null, options.Parameters, options.Arguments);
+					template.GetArgumentsForParameters(options.Parameters, options.Arguments);
 
 					await CreateJobAsync(chainedJobId, streamConfig, jobTrigger.TemplateRefId, template.Hash, triggerGraph, templateRefConfig.Name, change.Value, codeChange, options);
 					return newJob;
diff --git a/Engine/Source/Programs/Horde/Horde.Server/Jobs/JobsController.cs b/Engine/Source/Programs/Horde/Horde.Server/Jobs/JobsController.cs
--- a/Engine/Source/Programs/Horde/Horde.Server/Jobs/JobsController.cs
+++ b/Engine/Source/Programs/Horde/Horde.Server/Jobs/JobsController.cs
@@ -200,19 +200,33 @@ namespace Horde.Server.Jobs
 			options.Claims.AddRange(User.Claims.Select(x => new AclClaimConfig(x)));
 			options.JobOptions ??= create.JobOptions;
 
-			// Find all the default parameters, and override any settings with the values in the request
-			template.GetDefaultParameters(options.Parameters, false);
-
-			if (create.Parameters != null)
+			if (create.Arguments != null && create.Arguments.Count > 0)
 			{
-				foreach ((ParameterId parameter, string value) in create.Parameters)
+				// Use the specific argument list specified in the request
+				options.Arguments.AddRange(create.Arguments);
+
+				// Make sure we're not trying to specify parameters as well
+				if (create.Parameters != null && create.Parameters.Count > 0)
 				{
-					options.Parameters[parameter] = value;
+					return BadRequest("Cannot specify argument list and parameter list at the same time.");
 				}
 			}
+			else
+			{
+				// Find all the default parameters, and override any settings with the values in the request
+				template.GetDefaultParameters(options.Parameters, false);
+
+				if (create.Parameters != null)
+				{
+					foreach ((ParameterId parameter, string value) in create.Parameters)
+					{
+						options.Parameters[parameter] = value;
+					}
+				}
 
-			// Build the final arguments list from the combined parameter set
-			template.GetArgumentsForParameters(create.Arguments, options.Parameters, options.Arguments);
+				// Build the final arguments list from the combined parameter set
+				template.GetArgumentsForParameters(options.Parameters, options.Arguments);
+			}
 
 			// Merge the environment variables
 			foreach ((string key, string value) in environment)
diff --git a/Engine/Source/Programs/Horde/Horde.Server/Jobs/Schedules/ScheduleService.cs b/Engine/Source/Programs/Horde/Horde.Server/Jobs/Schedules/ScheduleService.cs
--- a/Engine/Source/Programs/Horde/Horde.Server/Jobs/Schedules/ScheduleService.cs
+++ b/Engine/Source/Programs/Horde/Horde.Server/Jobs/Schedules/ScheduleService.cs
@@ -469,7 +469,7 @@ namespace Horde.Server.Jobs.Schedules
 				CreateJobOptions options = new CreateJobOptions(templateRef.Config);
 				options.Priority = template.Priority;
 				template.GetDefaultParameters(options.Parameters, true);
-				template.GetArgumentsForParameters(null, options.Parameters, options.Arguments);
+				template.GetArgumentsForParameters(options.Parameters, options.Arguments);
 
 				IJob newJob = await _jobService.CreateJobAsync(null, stream.Config, templateId, template.Hash, graph, template.Name, change, codeChange, options, cancellationToken);
 				_logger.LogInformation("Started new job for {StreamId} template {TemplateId} at CL {Change} (Code CL {CodeChange}): {JobId}", stream.Id, templateId, change, codeChange, newJob.Id);
diff --git a/Engine/Source/Programs/Horde/Horde.Server/Jobs/Templates/ITemplate.cs b/Engine/Source/Programs/Horde/Horde.Server/Jobs/Templates/ITemplate.cs
--- a/Engine/Source/Programs/Horde/Horde.Server/Jobs/Templates/ITemplate.cs
+++ b/Engine/Source/Programs/Horde/Horde.Server/Jobs/Templates/ITemplate.cs
@@ -661,9 +661,9 @@ namespace Horde.Server.Jobs.Templates
 		/// <summary>
 		/// Gets the full argument list for a template
 		/// </summary>
-		public static void GetArgumentsForParameters(this ITemplate template, IReadOnlyList<string>? baseArguments, IReadOnlyDictionary<ParameterId, string> parameters, List<string> arguments)
+		public static void GetArgumentsForParameters(this ITemplate template, IReadOnlyDictionary<ParameterId, string> parameters, List<string> arguments)
 		{
-			arguments.AddRange(baseArguments ?? template.Arguments);
+			arguments.AddRange(template.Arguments);
 
 			foreach (Parameter parameter in template.Parameters)
 			{
