Horde: Add an AdditionalArguments property to jobs, which can be used to append arbitrary arguments to those derived from a job's parameters. This field is preserved - but not appended - to the arguments list if arguments are specified explicitly. Needs hooking up to the dashboard.

diff --git a/Engine/Source/Programs/Horde/Horde.Server/Jobs/IJob.cs b/Engine/Source/Programs/Horde/Horde.Server/Jobs/IJob.cs
--- a/Engine/Source/Programs/Horde/Horde.Server/Jobs/IJob.cs
+++ b/Engine/Source/Programs/Horde/Horde.Server/Jobs/IJob.cs
@@ -181,6 +181,11 @@ namespace Horde.Server.Jobs
 		/// </summary>
 		public IReadOnlyList<string> Arguments { get; }
 
+		/// <summary>
+		/// Additional arguments for the job, when a set of parameters are applied.
+		/// </summary>
+		public IReadOnlyList<string> AdditionalArguments { get; }
+
 		/// <summary>
 		/// Environment variables for the job
 		/// </summary>
diff --git a/Engine/Source/Programs/Horde/Horde.Server/Jobs/IJobCollection.cs b/Engine/Source/Programs/Horde/Horde.Server/Jobs/IJobCollection.cs
--- a/Engine/Source/Programs/Horde/Horde.Server/Jobs/IJobCollection.cs
+++ b/Engine/Source/Programs/Horde/Horde.Server/Jobs/IJobCollection.cs
@@ -77,6 +77,9 @@ namespace Horde.Server.Jobs
 		/// <inheritdoc cref="IJob.Arguments"/>
 		public List<string> Arguments { get; } = new List<string>();
 
+		/// <inheritdoc cref="IJob.AdditionalArguments"/>
+		public List<string> AdditionalArguments { get; } = new List<string>();
+
 		/// <inheritdoc cref="IJob.Environment"/>
 		public Dictionary<string, string> Environment { get; } = new Dictionary<string, string>(StringComparer.OrdinalIgnoreCase);
 
diff --git a/Engine/Source/Programs/Horde/Horde.Server/Jobs/JobCollection.cs b/Engine/Source/Programs/Horde/Horde.Server/Jobs/JobCollection.cs
--- a/Engine/Source/Programs/Horde/Horde.Server/Jobs/JobCollection.cs
+++ b/Engine/Source/Programs/Horde/Horde.Server/Jobs/JobCollection.cs
@@ -82,6 +82,7 @@ namespace Horde.Server.Jobs
 			IReadOnlyList<IJobStepBatch> IJob.Batches => Batches;
 			IReadOnlyDictionary<ParameterId, string> IJob.Parameters => _document.Parameters;
 			IReadOnlyList<string> IJob.Arguments => _document.Arguments;
+			IReadOnlyList<string> IJob.AdditionalArguments => _document.AdditionalArguments;
 			IReadOnlyDictionary<string, string> IJob.Environment => _document.Environment;
 			IReadOnlyList<int> IJob.Issues => _document.ReferencedByIssues;
 			ObjectId? IJob.NotificationTriggerId => _document.NotificationTriggerId;
@@ -452,6 +453,7 @@ namespace Horde.Server.Jobs
 			public Dictionary<ParameterId, string> Parameters { get; set; } = new Dictionary<ParameterId, string>();
 
 			public List<string> Arguments { get; set; } = new List<string>();
+			public List<string> AdditionalArguments { get; set; } = new List<string>();
 
 			[BsonDictionaryOptions(DictionaryRepresentation.ArrayOfDocuments)]
 			public Dictionary<string, string> Environment { get; set; } = new Dictionary<string, string>();
@@ -514,6 +516,7 @@ namespace Horde.Server.Jobs
 				NotificationChannelFilter = options.NotificationChannelFilter;
 				Parameters = new Dictionary<ParameterId, string>(options.Parameters);
 				Arguments.AddRange(options.Arguments);
+				AdditionalArguments.AddRange(options.AdditionalArguments);
 
 				foreach (KeyValuePair<string, string> pair in options.Environment)
 				{
diff --git a/Engine/Source/Programs/Horde/Horde.Server/Jobs/JobsController.cs b/Engine/Source/Programs/Horde/Horde.Server/Jobs/JobsController.cs
--- a/Engine/Source/Programs/Horde/Horde.Server/Jobs/JobsController.cs
+++ b/Engine/Source/Programs/Horde/Horde.Server/Jobs/JobsController.cs
@@ -200,6 +200,11 @@ namespace Horde.Server.Jobs
 			options.Claims.AddRange(User.Claims.Select(x => new AclClaimConfig(x)));
 			options.JobOptions ??= create.JobOptions;
 
+			if (create.AdditionalArguments != null)
+			{
+				options.AdditionalArguments.AddRange(create.AdditionalArguments);
+			}
+
 			if (create.Arguments != null && create.Arguments.Count > 0)
 			{
 				// Use the specific argument list specified in the request
@@ -226,6 +231,9 @@ namespace Horde.Server.Jobs
 
 				// Build the final arguments list from the combined parameter set
 				template.GetArgumentsForParameters(options.Parameters, options.Arguments);
+
+				// Add the additional arguments
+				options.Arguments.AddRange(options.AdditionalArguments);
 			}
 
 			// Merge the environment variables
@@ -599,6 +607,7 @@ namespace Horde.Server.Jobs
 			response.Reports = job.Reports?.ConvertAll(x => CreateGetReportResponse(x));
 			response.Parameters = job.Parameters.ToDictionary();
 			response.Arguments = job.Arguments.ToList();
+			response.AdditionalArguments = job.AdditionalArguments.ToList();
 			response.UpdateTime = new DateTimeOffset(job.UpdateTimeUtc);
 			response.UseArtifactsV2 = true;
 			response.UpdateIssues = job.UpdateIssues;
diff --git a/Engine/Source/Programs/Shared/EpicGames.Horde/Jobs/JobMessages.cs b/Engine/Source/Programs/Shared/EpicGames.Horde/Jobs/JobMessages.cs
--- a/Engine/Source/Programs/Shared/EpicGames.Horde/Jobs/JobMessages.cs
+++ b/Engine/Source/Programs/Shared/EpicGames.Horde/Jobs/JobMessages.cs
@@ -237,6 +237,11 @@ namespace EpicGames.Horde.Jobs
 		/// </summary>
 		public List<string>? Arguments { get; set; }
 
+		/// <summary>
+		/// Additional arguments for the job
+		/// </summary>
+		public List<string>? AdditionalArguments { get; set; }
+
 		/// <summary>
 		/// Private constructor for serialization
 		/// </summary>
@@ -504,6 +509,11 @@ namespace EpicGames.Horde.Jobs
 		/// </summary>
 		public List<string> Arguments { get; set; } = new List<string>();
 
+		/// <summary>
+		/// Additional command line arguments for the job for when using the parameters block
+		/// </summary>
+		public List<string> AdditionalArguments { get; set; } = new List<string>();
+
 		/// <summary>
 		/// The last update time for this job
 		/// </summary>
