Horde: Add a separate settings object for the driver app, and remove processes to terminate from the session object.

#rnx

diff --git a/Engine/Source/Programs/Horde/Horde.Agent.Driver/DriverSettings.cs b/Engine/Source/Programs/Horde/Horde.Agent.Driver/DriverSettings.cs
new file mode 100644
--- /dev/null
+++ b/Engine/Source/Programs/Horde/Horde.Agent.Driver/DriverSettings.cs
@@ -0,0 +1,55 @@
+﻿// Copyright Epic Games, Inc. All Rights Reserved.
+
+using Horde.Agent.Execution;
+using Horde.Agent.Utility;
+
+namespace Horde.Agent.Driver
+{
+	/// <summary>
+	/// Settings for the driver
+	/// </summary>
+	class DriverSettings
+	{
+		/// <summary>
+		/// The executor to use for jobs
+		/// </summary>
+		public string Executor { get; set; } = Execution.WorkspaceExecutor.Name;
+
+		/// <summary>
+		/// Settings for the local executor
+		/// </summary>
+		public LocalExecutorSettings LocalExecutor { get; set; } = new LocalExecutorSettings();
+
+		/// <summary>
+		/// Settings for the perforce executor
+		/// </summary>
+		public PerforceExecutorSettings PerforceExecutor { get; set; } = new PerforceExecutorSettings();
+
+		/// <summary>
+		/// List of process names to terminate after a job
+		/// </summary>
+		[Obsolete("Prefer using the ProcessesToTerminate list instead")]
+		public List<string> ProcessNamesToTerminate { get; } = new List<string>();
+
+		/// <summary>
+		/// List of process names to terminate after a lease completes, but not after a job step
+		/// </summary>
+		public List<ProcessToTerminate> ProcessesToTerminate { get; } = new List<ProcessToTerminate>();
+	}
+
+	/// <summary>
+	/// Specifies a process to terminate
+	/// </summary>
+	public class ProcessToTerminate
+	{
+		/// <summary>
+		/// Name of the process
+		/// </summary>
+		public string Name { get; set; } = String.Empty;
+
+		/// <summary>
+		/// When to terminate this process
+		/// </summary>
+		public List<TerminateCondition>? When { get; init; }
+	}
+}
diff --git a/Engine/Source/Programs/Horde/Horde.Agent.Driver/Execution/JobExecutor.cs b/Engine/Source/Programs/Horde/Horde.Agent.Driver/Execution/JobExecutor.cs
--- a/Engine/Source/Programs/Horde/Horde.Agent.Driver/Execution/JobExecutor.cs
+++ b/Engine/Source/Programs/Horde/Horde.Agent.Driver/Execution/JobExecutor.cs
@@ -15,6 +15,7 @@ using EpicGames.Horde.Storage;
 using EpicGames.Horde.Storage.Clients;
 using EpicGames.Horde.Storage.Nodes;
 using Grpc.Core;
+using Horde.Agent.Driver;
 using Horde.Agent.Parser;
 //using Horde.Agent.Services;
 using Horde.Agent.Utility;
@@ -63,7 +64,7 @@ namespace Horde.Agent.Execution
 		public Uri ServerUrl { get; }
 		public DirectoryReference WorkingDir { get; }
 		public IRpcConnection RpcConnection { get; }
-		public IReadOnlyDictionary<string, TerminateCondition>? ProcessNamesToTerminate { get; }
+		public IReadOnlyList<ProcessToTerminate>? ProcessesToTerminate { get; }
 		public HttpStorageClientFactory StorageFactory { get; }
 		public JobId JobId { get; }
 		public JobStepBatchId BatchId { get; }
@@ -71,12 +72,12 @@ namespace Horde.Agent.Execution
 		public string Token { get; }
 		public RpcJobOptions JobOptions { get; }
 
-		public JobExecutorOptions(Uri serverUrl, DirectoryReference workingDir, IRpcConnection rpcConnection, IReadOnlyDictionary<string, TerminateCondition>? processNamesToTerminate, HttpStorageClientFactory storageFactory, JobId jobId, JobStepBatchId batchId, RpcBeginBatchResponse batch, string token, RpcJobOptions jobOptions)
+		public JobExecutorOptions(Uri serverUrl, DirectoryReference workingDir, IRpcConnection rpcConnection, IReadOnlyList<ProcessToTerminate>? processesToTerminate, HttpStorageClientFactory storageFactory, JobId jobId, JobStepBatchId batchId, RpcBeginBatchResponse batch, string token, RpcJobOptions jobOptions)
 		{
 			ServerUrl = serverUrl;
 			WorkingDir = workingDir;
 			RpcConnection = rpcConnection;
-			ProcessNamesToTerminate = processNamesToTerminate;
+			ProcessesToTerminate = processesToTerminate;
 			StorageFactory = storageFactory;
 			JobId = jobId;
 			BatchId = batchId;
@@ -84,11 +85,6 @@ namespace Horde.Agent.Execution
 			Token = token;
 			JobOptions = jobOptions;
 		}
-
-//		public JobExecutorOptions(ISession session, HttpStorageClientFactory storageFactory, JobId jobId, JobStepBatchId batchId, RpcBeginBatchResponse batch, string token, RpcJobOptions jobOptions)
-//			: this(session.ServerUrl, session.WorkingDir, session.RpcConnection, session.ProcessNamesToTerminate, storageFactory, jobId, batchId, batch, token, jobOptions)
-//		{
-//		}
 	}
 
 	abstract class JobExecutor : IJobExecutor
@@ -1285,7 +1281,7 @@ namespace Horde.Agent.Execution
 		{
 			try
 			{
-				await TerminateProcessHelper.TerminateProcessesAsync(condition, Options.WorkingDir, Options.ProcessNamesToTerminate, logger, CancellationToken.None);
+				await TerminateProcessHelper.TerminateProcessesAsync(condition, Options.WorkingDir, Options.ProcessesToTerminate, logger, CancellationToken.None);
 			}
 			catch (Exception ex)
 			{
diff --git a/Engine/Source/Programs/Horde/Horde.Agent.Driver/Utility/TerminateProcessHelper.cs b/Engine/Source/Programs/Horde/Horde.Agent.Driver/Utility/TerminateProcessHelper.cs
--- a/Engine/Source/Programs/Horde/Horde.Agent.Driver/Utility/TerminateProcessHelper.cs
+++ b/Engine/Source/Programs/Horde/Horde.Agent.Driver/Utility/TerminateProcessHelper.cs
@@ -1,6 +1,7 @@
 // Copyright Epic Games, Inc. All Rights Reserved.
 
 using EpicGames.Core;
+using Horde.Agent.Driver;
 using Microsoft.Extensions.Logging;
 
 namespace Horde.Agent.Utility
@@ -16,11 +17,6 @@ namespace Horde.Agent.Utility
 		/// </summary>
 		None = 0,
 
-		/// <summary>
-		/// When a session starts
-		/// </summary>
-		BeforeSession = 1,
-
 		/// <summary>
 		/// Before running a conform
 		/// </summary>
@@ -50,28 +46,33 @@ namespace Horde.Agent.Utility
 		/// <summary>
 		/// Terminate processes matching certain criteria
 		/// </summary>
-		public static Task TerminateProcessesAsync(TerminateCondition condition, DirectoryReference workingDir, IReadOnlyDictionary<string, TerminateCondition>? processNamesToTerminate, ILogger logger, CancellationToken cancellationToken)
+		public static Task TerminateProcessesAsync(TerminateCondition condition, DirectoryReference workingDir, IReadOnlyList<ProcessToTerminate>? processesToTerminate, ILogger logger, CancellationToken cancellationToken)
 		{
 			// Terminate child processes from any previous runs
-			ProcessUtils.TerminateProcesses(x => ShouldTerminateProcess(x, condition, workingDir, processNamesToTerminate), logger, cancellationToken);
+			ProcessUtils.TerminateProcesses(x => ShouldTerminateProcess(x, condition, workingDir, processesToTerminate), logger, cancellationToken);
 			return Task.CompletedTask;
 		}
 
 		/// <summary>
 		/// Callback for determining whether a process should be terminated
 		/// </summary>
-		static bool ShouldTerminateProcess(FileReference imageFile, TerminateCondition condition, DirectoryReference workingDir, IReadOnlyDictionary<string, TerminateCondition>? processNamesToTerminate)
+		static bool ShouldTerminateProcess(FileReference imageFile, TerminateCondition condition, DirectoryReference workingDir, IReadOnlyList<ProcessToTerminate>? processesToTerminate)
 		{
 			if (imageFile.IsUnderDirectory(workingDir))
 			{
 				return true;
 			}
 
-			if (processNamesToTerminate != null)
+			if (processesToTerminate != null)
 			{
 				string fileName = imageFile.GetFileName();
-				if (processNamesToTerminate.TryGetValue(fileName, out TerminateCondition terminateFlags))
+				foreach (ProcessToTerminate processToTerminate in processesToTerminate)
 				{
+					TerminateCondition terminateFlags = TerminateCondition.None;
+					foreach (TerminateCondition when in processToTerminate.When ?? Enumerable.Empty<TerminateCondition>())
+					{
+						terminateFlags |= when;
+					}
 					if (terminateFlags == TerminateCondition.None || (terminateFlags & condition) != 0)
 					{
 						return true;
diff --git a/Engine/Source/Programs/Horde/Horde.Agent.Tests/Execution/WorkspaceExecutorTest.cs b/Engine/Source/Programs/Horde/Horde.Agent.Tests/Execution/WorkspaceExecutorTest.cs
--- a/Engine/Source/Programs/Horde/Horde.Agent.Tests/Execution/WorkspaceExecutorTest.cs
+++ b/Engine/Source/Programs/Horde/Horde.Agent.Tests/Execution/WorkspaceExecutorTest.cs
@@ -8,6 +8,7 @@ using System.Threading.Tasks;
 using EpicGames.Core;
 using EpicGames.Horde.Jobs;
 using EpicGames.Horde.Streams;
+using Horde.Agent.Driver;
 using Horde.Agent.Execution;
 using Horde.Agent.Services;
 using HordeCommon.Rpc;
@@ -79,7 +80,7 @@ public sealed class WorkspaceExecutorTest : IAsyncDisposable
 	public async Task RegularAndAutoSdkWorkspaceAsync()
 	{
 		RpcBeginBatchResponse batch = new RpcBeginBatchResponse { Change = 1 };
-		JobExecutorOptions executorOptions = new JobExecutorOptions(_session.ServerUrl, _session.WorkingDir, _session.RpcConnection, _session.ProcessNamesToTerminate, null!, _jobId, _batchId, batch, null!, new RpcJobOptions());
+		JobExecutorOptions executorOptions = new JobExecutorOptions(_session.ServerUrl, _session.WorkingDir, _session.RpcConnection, Array.Empty<ProcessToTerminate>(), null!, _jobId, _batchId, batch, null!, new RpcJobOptions());
 		using WorkspaceExecutor executor = new(executorOptions, _workspace, _autoSdkWorkspace, NullLogger.Instance);
 
 		await executor.InitializeAsync(_logger, CancellationToken.None);
@@ -93,7 +94,7 @@ public sealed class WorkspaceExecutorTest : IAsyncDisposable
 	public async Task EnvVarsAsync()
 	{
 		RpcBeginBatchResponse batch = new RpcBeginBatchResponse { Change = 1, StreamName = "//UE5/Main" };
-		JobExecutorOptions executorOptions = new JobExecutorOptions(_session.ServerUrl, _session.WorkingDir, _session.RpcConnection, _session.ProcessNamesToTerminate, null!, _jobId, _batchId, batch, null!, new RpcJobOptions());
+		JobExecutorOptions executorOptions = new JobExecutorOptions(_session.ServerUrl, _session.WorkingDir, _session.RpcConnection, Array.Empty<ProcessToTerminate>(), null!, _jobId, _batchId, batch, null!, new RpcJobOptions());
 		using WorkspaceExecutor executor = new(executorOptions, _workspace, _autoSdkWorkspace, NullLogger.Instance);
 		await executor.InitializeAsync(_logger, CancellationToken.None);
 
@@ -120,7 +121,7 @@ public sealed class WorkspaceExecutorTest : IAsyncDisposable
 		_workspace.SetFile(1000, "New/Feature/Foo.cs", "foo");
 
 		RpcBeginBatchResponse batch = new RpcBeginBatchResponse { Change = 1, PreflightChange = 1000 };
-		JobExecutorOptions executorOptions = new JobExecutorOptions(_session.ServerUrl, _session.WorkingDir, _session.RpcConnection, _session.ProcessNamesToTerminate, null!, preflightJobId, _batchId, batch, null!, new RpcJobOptions());
+		JobExecutorOptions executorOptions = new JobExecutorOptions(_session.ServerUrl, _session.WorkingDir, _session.RpcConnection, Array.Empty<ProcessToTerminate>(), null!, preflightJobId, _batchId, batch, null!, new RpcJobOptions());
 		using WorkspaceExecutor executor = new(executorOptions, _workspace, null, NullLogger.Instance);
 
 		await executor.InitializeAsync(_logger, CancellationToken.None);
@@ -137,7 +138,7 @@ public sealed class WorkspaceExecutorTest : IAsyncDisposable
 
 		_server.AddJob(noChangeJobId, _streamId, 0, 0);
 		RpcBeginBatchResponse batch = new RpcBeginBatchResponse { };
-		JobExecutorOptions executorOptions = new JobExecutorOptions(_session.ServerUrl, _session.WorkingDir, _session.RpcConnection, _session.ProcessNamesToTerminate, null!, noChangeJobId, _batchId, batch, null!, new RpcJobOptions());
+		JobExecutorOptions executorOptions = new JobExecutorOptions(_session.ServerUrl, _session.WorkingDir, _session.RpcConnection, Array.Empty<ProcessToTerminate>(), null!, noChangeJobId, _batchId, batch, null!, new RpcJobOptions());
 		using WorkspaceExecutor executor = new(executorOptions, _workspace, null, NullLogger.Instance);
 		await Assert.ThrowsExceptionAsync<WorkspaceMaterializationException>(() => executor.InitializeAsync(_logger, CancellationToken.None));
 	}
diff --git a/Engine/Source/Programs/Horde/Horde.Agent.Tests/Services/TelemetryServiceTest.cs b/Engine/Source/Programs/Horde/Horde.Agent.Tests/Services/TelemetryServiceTest.cs
--- a/Engine/Source/Programs/Horde/Horde.Agent.Tests/Services/TelemetryServiceTest.cs
+++ b/Engine/Source/Programs/Horde/Horde.Agent.Tests/Services/TelemetryServiceTest.cs
@@ -4,6 +4,7 @@ using System;
 using System.Collections.Generic;
 using System.Threading;
 using System.Threading.Tasks;
+using Horde.Agent.Driver;
 using Horde.Agent.Execution;
 using Horde.Agent.Leases.Handlers;
 using Horde.Agent.Services;
@@ -25,7 +26,7 @@ public sealed class TelemetryServiceTest : IDisposable
 		OptionsWrapper<AgentSettings> settingsOpt = new(settings);
 
 		using WorkerService workerService = new WorkerService(null!, null!, null!, null!);
-		JobHandler jobHandler = new JobHandler(new List<IJobExecutorFactory>(), settingsOpt, null!, null!);
+		JobHandler jobHandler = new JobHandler(new List<IJobExecutorFactory>(), settingsOpt, new OptionsWrapper<DriverSettings>(new DriverSettings()), null!, null!);
 		GrpcService grpcService = new(settingsOpt, NullLogger<GrpcService>.Instance, _loggerFactory);
 
 		_telemetryService = new TelemetryService(workerService, jobHandler, grpcService, null!, settingsOpt, NullLogger<TelemetryService>.Instance);
diff --git a/Engine/Source/Programs/Horde/Horde.Agent.Tests/WorkerServiceTest.cs b/Engine/Source/Programs/Horde/Horde.Agent.Tests/WorkerServiceTest.cs
--- a/Engine/Source/Programs/Horde/Horde.Agent.Tests/WorkerServiceTest.cs
+++ b/Engine/Source/Programs/Horde/Horde.Agent.Tests/WorkerServiceTest.cs
@@ -20,6 +20,7 @@ using EpicGames.Horde.Streams;
 using Google.Protobuf.WellKnownTypes;
 using Grpc.Core;
 using Grpc.Net.Client;
+using Horde.Agent.Driver;
 using Horde.Agent.Execution;
 using Horde.Agent.Leases;
 using Horde.Agent.Leases.Handlers;
@@ -98,6 +99,10 @@ namespace Horde.Agent.Tests
 				settings.ServerProfiles.Add(profile.Name, profile);
 				settings.Server = "test";
 				settings.WorkingDir = new DirectoryReference(Path.GetTempPath());
+			});
+
+			_serviceCollection.Configure<DriverSettings>(settings =>
+			{
 				settings.Executor = TestExecutor.Name; // Not really used since the executor is overridden in the tests
 			});
 
@@ -281,7 +286,6 @@ namespace Horde.Agent.Tests
 			fakeSession.Setup(x => x.AgentId).Returns(new EpicGames.Horde.Agents.AgentId("LocalAgent"));
 			fakeSession.Setup(x => x.SessionId).Returns(new EpicGames.Horde.Agents.Sessions.SessionId(default));
 			fakeSession.Setup(x => x.RpcConnection).Returns(rpcConnection);
-			fakeSession.Setup(x => x.ProcessNamesToTerminate).Returns((IReadOnlyDictionary<string, TerminateCondition>)new Dictionary<string, TerminateCondition>());
 			fakeSession.Setup(x => x.DisposeAsync()).Returns(new ValueTask());
 			fakeSession.Setup(x => x.WorkingDir).Returns(DirectoryReference.Combine(DirectoryReference.GetCurrentDirectory(), Guid.NewGuid().ToString()));
 			return fakeSession.Object;
diff --git a/Engine/Source/Programs/Horde/Horde.Agent/AgentApp.cs b/Engine/Source/Programs/Horde/Horde.Agent/AgentApp.cs
--- a/Engine/Source/Programs/Horde/Horde.Agent/AgentApp.cs
+++ b/Engine/Source/Programs/Horde/Horde.Agent/AgentApp.cs
@@ -15,6 +15,7 @@ using EpicGames.Horde.Storage;
 using EpicGames.Horde.Storage.Backends;
 using EpicGames.Horde.Storage.Bundles;
 using EpicGames.Horde.Storage.Clients;
+using Horde.Agent.Driver;
 using Horde.Agent.Execution;
 using Horde.Agent.Leases;
 using Horde.Agent.Leases.Handlers;
@@ -162,6 +163,7 @@ namespace Horde.Agent
 			// Add all the default 
 			IConfigurationSection configSection = configuration.GetSection(AgentSettings.SectionName);
 			services.AddOptions<AgentSettings>().Configure(options => configSection.Bind(options)).ValidateDataAnnotations();
+			services.AddOptions<DriverSettings>().Configure(options => configSection.Bind(options)).ValidateDataAnnotations();
 
 			ServerProfile serverProfile = settings.GetCurrentServerProfile();
 			if (String.IsNullOrEmpty(serverProfile.Url.Scheme) || String.IsNullOrEmpty(serverProfile.Url.Host))
diff --git a/Engine/Source/Programs/Horde/Horde.Agent/AgentSettings.cs b/Engine/Source/Programs/Horde/Horde.Agent/AgentSettings.cs
--- a/Engine/Source/Programs/Horde/Horde.Agent/AgentSettings.cs
+++ b/Engine/Source/Programs/Horde/Horde.Agent/AgentSettings.cs
@@ -81,22 +81,6 @@ namespace Horde.Agent
 		}
 	}
 
-	/// <summary>
-	/// Specifies a process to terminate
-	/// </summary>
-	public class ProcessToTerminate
-	{
-		/// <summary>
-		/// Name of the process
-		/// </summary>
-		public string Name { get; set; } = String.Empty;
-
-		/// <summary>
-		/// When to terminate this process
-		/// </summary>
-		public List<TerminateCondition>? When { get; init; }
-	}
-
 	/// <summary>
 	/// Global settings for the agent
 	/// </summary>
@@ -138,21 +122,6 @@ namespace Horde.Agent
 		/// </summary>
 		public bool Ephemeral { get; set; } = false;
 
-		/// <summary>
-		/// The executor to use for jobs
-		/// </summary>
-		public string Executor { get; set; } = Execution.WorkspaceExecutor.Name;
-
-		/// <summary>
-		/// Settings for the local executor
-		/// </summary>
-		public LocalExecutorSettings LocalExecutor { get; set; } = new LocalExecutorSettings();
-
-		/// <summary>
-		/// Settings for the perforce executor
-		/// </summary>
-		public PerforceExecutorSettings PerforceExecutor { get; set; } = new PerforceExecutorSettings();
-
 		/// <summary>
 		/// Working directory for leases and jobs (i.e where files from Perforce will be checked out) 
 		/// </summary>
@@ -173,17 +142,6 @@ namespace Horde.Agent
 		/// </summary>
 		public List<MountNetworkShare> Shares { get; } = new List<MountNetworkShare>();
 
-		/// <summary>
-		/// List of process names to terminate after a job
-		/// </summary>
-		[Obsolete("Prefer using the ProcessesToTerminate list instead")]
-		public List<string> ProcessNamesToTerminate { get; } = new List<string>();
-
-		/// <summary>
-		/// List of process names to terminate after a lease completes, but not after a job step
-		/// </summary>
-		public List<ProcessToTerminate> ProcessesToTerminate { get; } = new List<ProcessToTerminate>();
-
 		/// <summary>
 		/// Path to Wine executable. If null, execution under Wine is disabled
 		/// </summary>
@@ -292,37 +250,10 @@ namespace Horde.Agent
 				defaultServerProfile.Url = defaultServerUrl ?? new Uri("http://localhost:5000");
 				return defaultServerProfile;
 			}
-
+			
 			return GetServerProfile(Server);
 		}
 
-		/// <summary>
-		/// Gets a lookup of process name to the circumstances in which it should be terminate
-		/// </summary>
-		public Dictionary<string, TerminateCondition> GetProcessesToTerminateMap()
-		{
-			Dictionary<string, TerminateCondition> processesToTerminate = new Dictionary<string, TerminateCondition>(StringComparer.OrdinalIgnoreCase);
-#pragma warning disable CS0618 // Type or member is obsolete
-			foreach (string processName in ProcessNamesToTerminate)
-			{
-				processesToTerminate[processName] = TerminateCondition.None;
-			}
-#pragma warning restore CS0618 // Type or member is obsolete
-			foreach (ProcessToTerminate processToTerminate in ProcessesToTerminate)
-			{
-				TerminateCondition condition = default;
-				if (processToTerminate.When != null)
-				{
-					foreach (TerminateCondition when in processToTerminate.When)
-					{
-						condition |= when;
-					}
-				}
-				processesToTerminate[processToTerminate.Name] = condition;
-			}
-			return processesToTerminate;
-		}
-
 		/// <summary>
 		/// Path to file used for signaling impending termination and shutdown of the agent
 		/// </summary>
diff --git a/Engine/Source/Programs/Horde/Horde.Agent/Commands/Execution/ExecuteJobCommand.cs b/Engine/Source/Programs/Horde/Horde.Agent/Commands/Execution/ExecuteJobCommand.cs
--- a/Engine/Source/Programs/Horde/Horde.Agent/Commands/Execution/ExecuteJobCommand.cs
+++ b/Engine/Source/Programs/Horde/Horde.Agent/Commands/Execution/ExecuteJobCommand.cs
@@ -4,6 +4,7 @@ using EpicGames.Core;
 using EpicGames.Horde.Agents;
 using EpicGames.Horde.Agents.Leases;
 using EpicGames.Horde.Agents.Sessions;
+using Horde.Agent.Driver;
 using Horde.Agent.Leases.Handlers;
 using Horde.Agent.Services;
 using Horde.Agent.Utility;
@@ -33,24 +34,24 @@ namespace Horde.Agent.Commands.Execution
 
 		readonly GrpcService _grpcService;
 		readonly JobHandler _jobHandler;
-		readonly AgentSettings _settings;
+		readonly AgentSettings _agentSettings;
+		readonly DriverSettings _driverSettings;
 
-		public ExecuteJobCommand(GrpcService grpcService, JobHandler jobHandler, IOptions<AgentSettings> settings)
+		public ExecuteJobCommand(GrpcService grpcService, JobHandler jobHandler, IOptions<AgentSettings> agentSettings, IOptions<DriverSettings> driverSettings)
 		{
 			_grpcService = grpcService;
 			_jobHandler = jobHandler;
-			_settings = settings.Value;
+			_agentSettings = agentSettings.Value;
+			_driverSettings = driverSettings.Value;
 		}
 
 		public override async Task<int> ExecuteAsync(ILogger logger)
 		{
-			ServerProfile serverProfile = _settings.GetCurrentServerProfile();
+			ServerProfile serverProfile = _agentSettings.GetCurrentServerProfile();
 			ExecuteJobTask executeTask = ExecuteJobTask.Parser.ParseFrom(Convert.FromBase64String(Task));
 
-			Dictionary<string, TerminateCondition> processNamesToTerminate = _settings.GetProcessesToTerminateMap();
-
 			await using RpcConnection rpcConnection = new RpcConnection(ctx => _grpcService.CreateGrpcChannelAsync(executeTask.Token, ctx), logger);
-			await using Session session = new Session(serverProfile.Url, AgentId, SessionId, executeTask.Token, rpcConnection, WorkingDir, processNamesToTerminate);
+			await using Session session = new Session(serverProfile.Url, AgentId, SessionId, executeTask.Token, rpcConnection, WorkingDir);
 
 			await _jobHandler.ExecuteInternalAsync(session, LeaseId, executeTask, logger, CancellationToken.None);
 			return 0;
diff --git a/Engine/Source/Programs/Horde/Horde.Agent/Leases/Handlers/ConformHandler.cs b/Engine/Source/Programs/Horde/Horde.Agent/Leases/Handlers/ConformHandler.cs
--- a/Engine/Source/Programs/Horde/Horde.Agent/Leases/Handlers/ConformHandler.cs
+++ b/Engine/Source/Programs/Horde/Horde.Agent/Leases/Handlers/ConformHandler.cs
@@ -2,6 +2,7 @@
 
 using EpicGames.Horde.Agents.Leases;
 using EpicGames.Horde.Logs;
+using Horde.Agent.Driver;
 using Horde.Agent.Execution;
 using Horde.Agent.Services;
 using Horde.Agent.Utility;
@@ -15,15 +16,17 @@ namespace Horde.Agent.Leases.Handlers
 {
 	class ConformHandler : LeaseHandler<ConformTask>
 	{
-		readonly AgentSettings _settings;
+		readonly AgentSettings _agentSettings;
+		readonly DriverSettings _driverSettings;
 		readonly IServerLoggerFactory _serverLoggerFactory;
 
 		/// <summary>
 		/// Constructor
 		/// </summary>
-		public ConformHandler(IOptions<AgentSettings> settings, IServerLoggerFactory serverLoggerFactory)
+		public ConformHandler(IOptions<AgentSettings> agentSettings, IOptions<DriverSettings> driverSettings, IServerLoggerFactory serverLoggerFactory)
 		{
-			_settings = settings.Value;
+			_agentSettings = agentSettings.Value;
+			_driverSettings = driverSettings.Value;
 			_serverLoggerFactory = serverLoggerFactory;
 		}
 
@@ -46,14 +49,14 @@ namespace Horde.Agent.Leases.Handlers
 		async Task<LeaseResult> ExecuteInternalAsync(ISession session, LeaseId leaseId, ConformTask conformTask, ILogger logger, CancellationToken cancellationToken)
 		{
 			logger.LogInformation("Conforming, lease {LeaseId}", leaseId);
-			await TerminateProcessHelper.TerminateProcessesAsync(TerminateCondition.BeforeConform, session.WorkingDir, session.ProcessNamesToTerminate, logger, cancellationToken);
+			await TerminateProcessHelper.TerminateProcessesAsync(TerminateCondition.BeforeConform, session.WorkingDir, _driverSettings.ProcessesToTerminate, logger, cancellationToken);
 
 			bool removeUntrackedFiles = conformTask.RemoveUntrackedFiles;
 			IList<RpcAgentWorkspace> pendingWorkspaces = conformTask.Workspaces;
 			for (; ; )
 			{
-				bool isPerforceExecutor = _settings.Executor.Equals(PerforceExecutor.Name, StringComparison.OrdinalIgnoreCase);
-				bool isWorkspaceExecutor = _settings.Executor.Equals(WorkspaceExecutor.Name, StringComparison.OrdinalIgnoreCase);
+				bool isPerforceExecutor = _driverSettings.Executor.Equals(PerforceExecutor.Name, StringComparison.OrdinalIgnoreCase);
+				bool isWorkspaceExecutor = _driverSettings.Executor.Equals(WorkspaceExecutor.Name, StringComparison.OrdinalIgnoreCase);
 
 				// When using WorkspaceExecutor, only job options can override exact materializer to use
 				// It will default to ManagedWorkspaceMaterializer, which is compatible with the conform call below
@@ -62,13 +65,13 @@ namespace Horde.Agent.Leases.Handlers
 				bool isExecutorConformCompatible = isPerforceExecutor || isWorkspaceExecutor;
 
 				// Run the conform task
-				if (isExecutorConformCompatible && _settings.PerforceExecutor.RunConform)
+				if (isExecutorConformCompatible && _driverSettings.PerforceExecutor.RunConform)
 				{
 					await PerforceExecutor.ConformAsync(session.WorkingDir, pendingWorkspaces, removeUntrackedFiles, logger, cancellationToken);
 				}
 				else
 				{
-					logger.LogInformation("Skipping conform. Executor={Executor} RunConform={RunConform}", _settings.Executor, _settings.PerforceExecutor.RunConform);
+					logger.LogInformation("Skipping conform. Executor={Executor} RunConform={RunConform}", _driverSettings.Executor, _driverSettings.PerforceExecutor.RunConform);
 				}
 
 				// Update the new set of workspaces
diff --git a/Engine/Source/Programs/Horde/Horde.Agent/Leases/Handlers/JobHandler.cs b/Engine/Source/Programs/Horde/Horde.Agent/Leases/Handlers/JobHandler.cs
--- a/Engine/Source/Programs/Horde/Horde.Agent/Leases/Handlers/JobHandler.cs
+++ b/Engine/Source/Programs/Horde/Horde.Agent/Leases/Handlers/JobHandler.cs
@@ -10,6 +10,7 @@ using EpicGames.Horde.Logs;
 using EpicGames.Horde.Storage.Clients;
 using Google.Protobuf;
 using Grpc.Core;
+using Horde.Agent.Driver;
 using Horde.Agent.Execution;
 using Horde.Agent.Services;
 using Horde.Agent.Utility;
@@ -55,17 +56,19 @@ namespace Horde.Agent.Leases.Handlers
 		public string? CurrentBatchId { get; private set; } = null;
 
 		readonly IEnumerable<IJobExecutorFactory> _executorFactories;
-		readonly AgentSettings _settings;
+		readonly AgentSettings _agentSettings;
+		readonly DriverSettings _driverSettings;
 		readonly HttpStorageClientFactory _serverStorageFactory;
 		readonly IServerLoggerFactory _serverLoggerFactory;
 
 		/// <summary>
 		/// Constructor
 		/// </summary>
-		public JobHandler(IEnumerable<IJobExecutorFactory> executorFactories, IOptions<AgentSettings> settings, HttpStorageClientFactory storageClientFactory, IServerLoggerFactory serverLoggerFactory)
+		public JobHandler(IEnumerable<IJobExecutorFactory> executorFactories, IOptions<AgentSettings> agentSettings, IOptions<DriverSettings> driverSettings, HttpStorageClientFactory storageClientFactory, IServerLoggerFactory serverLoggerFactory)
 		{
 			_executorFactories = executorFactories;
-			_settings = settings.Value;
+			_agentSettings = agentSettings.Value;
+			_driverSettings = driverSettings.Value;
 			_serverStorageFactory = storageClientFactory;
 			_serverLoggerFactory = serverLoggerFactory;
 		}
@@ -97,7 +100,7 @@ namespace Horde.Agent.Leases.Handlers
 #pragma warning restore IL3000 // Avoid accessing Assembly file path when publishing as a single file						
 						arguments.Add("execute");
 						arguments.Add("job");
-						arguments.Add($"-Server={_settings.GetCurrentServerProfile().Name}");
+						arguments.Add($"-Server={_agentSettings.GetCurrentServerProfile().Name}");
 						arguments.Add($"-AgentId={session.AgentId}");
 						arguments.Add($"-SessionId={session.SessionId}");
 						arguments.Add($"-LeaseId={leaseId}");
@@ -170,7 +173,7 @@ namespace Horde.Agent.Leases.Handlers
 			RpcBeginBatchResponse batch = await session.RpcConnection.InvokeAsync<JobRpc.JobRpcClient, RpcBeginBatchResponse>(x => x.BeginBatchAsync(new RpcBeginBatchRequest(executeTask.JobId, executeTask.BatchId, leaseId), null, null, cancellationToken), cancellationToken);
 			try
 			{
-				JobExecutorOptions options = new JobExecutorOptions(session.ServerUrl, session.WorkingDir, session.RpcConnection, session.ProcessNamesToTerminate, _serverStorageFactory, executeTask.JobId, executeTask.BatchId, batch, executeTask.Token, jobOptions);
+				JobExecutorOptions options = new JobExecutorOptions(session.ServerUrl, session.WorkingDir, session.RpcConnection, _driverSettings.ProcessesToTerminate, _serverStorageFactory, executeTask.JobId, executeTask.BatchId, batch, executeTask.Token, jobOptions);
 				await ExecuteBatchAsync(session, leaseId, executeTask.Workspace, executeTask.AutoSdkWorkspace, options, logger, localLogger, cancellationToken);
 			}
 			catch (Exception ex)
@@ -215,10 +218,10 @@ namespace Horde.Agent.Leases.Handlers
 			IRpcConnection rpcClient = session.RpcConnection;
 
 			// Create an executor for this job
-			string executorName = String.IsNullOrEmpty(options.JobOptions.Executor) ? _settings.Executor : options.JobOptions.Executor;
+			string executorName = String.IsNullOrEmpty(options.JobOptions.Executor) ? _driverSettings.Executor : options.JobOptions.Executor;
 
 			logger.LogInformation("Executing batch {BatchId} using {Executor} executor", options.BatchId, executorName);
-			await session.TerminateProcessesAsync(TerminateCondition.BeforeBatch, logger, cancellationToken);
+			await TerminateProcessHelper.TerminateProcessesAsync(TerminateCondition.BeforeBatch, session.WorkingDir, _driverSettings.ProcessesToTerminate, logger, cancellationToken);
 
 			IJobExecutorFactory? executorFactory = _executorFactories.FirstOrDefault(x => x.Name.Equals(executorName, StringComparison.OrdinalIgnoreCase));
 			if (executorFactory == null)
@@ -329,7 +332,7 @@ namespace Horde.Agent.Leases.Handlers
 							}
 
 							// Kill any processes spawned by the step
-							await session.TerminateProcessesAsync(TerminateCondition.AfterStep, logger, cancellationToken);
+							await TerminateProcessHelper.TerminateProcessesAsync(TerminateCondition.AfterStep, session.WorkingDir, _driverSettings.ProcessesToTerminate, logger, cancellationToken);
 
 							// Wait for the logger to finish
 							await stepLogger.StopAsync();
@@ -374,7 +377,7 @@ namespace Horde.Agent.Leases.Handlers
 			// Terminate any processes which are still running
 			try
 			{
-				await session.TerminateProcessesAsync(TerminateCondition.AfterBatch, logger, CancellationToken.None);
+				await TerminateProcessHelper.TerminateProcessesAsync(TerminateCondition.AfterBatch, session.WorkingDir, _driverSettings.ProcessesToTerminate, logger, CancellationToken.None);
 			}
 			catch (Exception ex)
 			{
diff --git a/Engine/Source/Programs/Horde/Horde.Agent/Leases/LeaseManager.cs b/Engine/Source/Programs/Horde/Horde.Agent/Leases/LeaseManager.cs
--- a/Engine/Source/Programs/Horde/Horde.Agent/Leases/LeaseManager.cs
+++ b/Engine/Source/Programs/Horde/Horde.Agent/Leases/LeaseManager.cs
@@ -232,7 +232,7 @@ namespace Horde.Agent.Leases
 			IRpcConnection rpcCon = _session.RpcConnection;
 
 			// Terminate any remaining child processes from other instances
-			await _session.TerminateProcessesAsync(TerminateCondition.BeforeSession, _logger, stoppingToken);
+			ProcessUtils.TerminateProcesses(x => x.IsUnderDirectory(_session.WorkingDir), _logger, stoppingToken);
 
 			// Track how many updates we get in 10 seconds. We'll start rate limiting this if it looks like we've got a problem that's causing us to spam the server.
 			Stopwatch updateTimer = Stopwatch.StartNew();
diff --git a/Engine/Source/Programs/Horde/Horde.Agent/Services/SessionFactory.cs b/Engine/Source/Programs/Horde/Horde.Agent/Services/SessionFactory.cs
--- a/Engine/Source/Programs/Horde/Horde.Agent/Services/SessionFactory.cs
+++ b/Engine/Source/Programs/Horde/Horde.Agent/Services/SessionFactory.cs
@@ -8,6 +8,7 @@ using EpicGames.Horde.Agents;
 using EpicGames.Horde.Agents.Sessions;
 using Grpc.Core;
 using Grpc.Net.Client;
+using Horde.Agent.Driver;
 using Horde.Agent.Utility;
 using Horde.Common.Rpc;
 using HordeCommon.Rpc;
@@ -51,11 +52,6 @@ namespace Horde.Agent.Services
 		/// Working directory for sandboxes etc..
 		/// </summary>
 		DirectoryReference WorkingDir { get; }
-
-		/// <summary>
-		/// List of process names to terminate
-		/// </summary>
-		IReadOnlyDictionary<string, TerminateCondition> ProcessNamesToTerminate { get; }
 	}
 
 	/// <summary>
@@ -94,13 +90,10 @@ namespace Horde.Agent.Services
 		/// <inheritdoc/>
 		public DirectoryReference WorkingDir { get; }
 
-		/// <inheritdoc/>
-		public IReadOnlyDictionary<string, TerminateCondition> ProcessNamesToTerminate { get; }
-
 		/// <summary>
 		/// Constructor
 		/// </summary>
-		public Session(Uri serverUrl, AgentId agentId, SessionId sessionId, string token, IRpcConnection rpcConnection, DirectoryReference workingDir, IReadOnlyDictionary<string, TerminateCondition> processNamesToTerminate)
+		public Session(Uri serverUrl, AgentId agentId, SessionId sessionId, string token, IRpcConnection rpcConnection, DirectoryReference workingDir)
 		{
 			ServerUrl = serverUrl;
 			AgentId = agentId;
@@ -108,7 +101,6 @@ namespace Horde.Agent.Services
 			Token = token;
 			RpcConnection = rpcConnection;
 			WorkingDir = workingDir;
-			ProcessNamesToTerminate = processNamesToTerminate;
 		}
 
 		public class AgentRegistrationList
@@ -121,7 +113,7 @@ namespace Horde.Agent.Services
 		/// <summary>
 		/// Creates a new agent session
 		/// </summary>
-		public static async Task<Session> CreateAsync(CapabilitiesService capabilitiesService, GrpcService grpcService, StatusService statusService, IOptions<AgentSettings> settings, ILogger logger, CancellationToken cancellationToken)
+		public static async Task<Session> CreateAsync(CapabilitiesService capabilitiesService, GrpcService grpcService, StatusService statusService, IOptions<AgentSettings> settings, IOptions<DriverSettings> driverSettings, ILogger logger, CancellationToken cancellationToken)
 		{
 			AgentSettings currentSettings = settings.Value;
 
@@ -227,7 +219,7 @@ namespace Horde.Agent.Services
 #pragma warning disable CA2000 // False positive; ownership is transferred to new Session object.
 			Func<CancellationToken, Task<GrpcChannel>> createGrpcChannelAsync = ctx => grpcService.CreateGrpcChannelAsync(createSessionResponse.Token, ctx);
 			IRpcConnection rpcConnection = new RpcConnection(createGrpcChannelAsync, logger);
-			return new Session(serverProfile.Url, new AgentId(createSessionResponse.AgentId), SessionId.Parse(createSessionResponse.SessionId), createSessionResponse.Token, rpcConnection, workingDir, currentSettings.GetProcessesToTerminateMap());
+			return new Session(serverProfile.Url, new AgentId(createSessionResponse.AgentId), SessionId.Parse(createSessionResponse.SessionId), createSessionResponse.Token, rpcConnection, workingDir);
 #pragma warning restore CA2000
 		}
 
@@ -356,32 +348,24 @@ namespace Horde.Agent.Services
 		readonly CapabilitiesService _capabilitiesService;
 		readonly GrpcService _grpcService;
 		readonly StatusService _statusService;
-		readonly IOptions<AgentSettings> _settings;
+		readonly IOptions<AgentSettings> _agentSettings;
+		readonly IOptions<DriverSettings> _driverSettings;
 		readonly ILogger _logger;
 
 		/// <summary>
 		/// Constructor
 		/// </summary>
-		public SessionFactory(CapabilitiesService capabilitiesService, GrpcService grpcService, StatusService statusService, IOptions<AgentSettings> settings, ILogger<SessionFactory> logger)
+		public SessionFactory(CapabilitiesService capabilitiesService, GrpcService grpcService, StatusService statusService, IOptions<AgentSettings> agentSettings, IOptions<DriverSettings> driverSettings, ILogger<SessionFactory> logger)
 		{
 			_capabilitiesService = capabilitiesService;
 			_grpcService = grpcService;
 			_statusService = statusService;
-			_settings = settings;
+			_agentSettings = agentSettings;
+			_driverSettings = driverSettings;
 			_logger = logger;
 		}
 
 		/// <inheritdoc/>
-		public async Task<ISession> CreateAsync(CancellationToken cancellationToken) => await Session.CreateAsync(_capabilitiesService, _grpcService, _statusService, _settings, _logger, cancellationToken);
-	}
-
-	/// <summary>
-	/// Extension methods
-	/// </summary>
-	static class SessionExtensions
-	{
-		/// <inheritdoc cref="TerminateProcessHelper.TerminateProcessesAsync(TerminateCondition, DirectoryReference, IReadOnlyDictionary{string, TerminateCondition}, ILogger, CancellationToken)"/>
-		public static Task TerminateProcessesAsync(this ISession session, TerminateCondition condition, ILogger logger, CancellationToken cancellationToken)
-			=> TerminateProcessHelper.TerminateProcessesAsync(condition, session.WorkingDir, session.ProcessNamesToTerminate, logger, cancellationToken);
+		public async Task<ISession> CreateAsync(CancellationToken cancellationToken) => await Session.CreateAsync(_capabilitiesService, _grpcService, _statusService, _agentSettings, _driverSettings, _logger, cancellationToken);
 	}
 }
