Horde: Build fix.

#rnx

diff --git a/Engine/Source/Programs/Horde/Horde.Agent.Driver/Execution/JobExecutor.cs b/Engine/Source/Programs/Horde/Horde.Agent.Driver/Execution/JobExecutor.cs
--- a/Engine/Source/Programs/Horde/Horde.Agent.Driver/Execution/JobExecutor.cs
+++ b/Engine/Source/Programs/Horde/Horde.Agent.Driver/Execution/JobExecutor.cs
@@ -550,7 +550,7 @@ namespace Horde.Agent.Execution
 			}
 
 			RpcUpdateGraphRequest updateGraph = await ParseGraphUpdateAsync(definitionFile, logger, cancellationToken);
-			await jobRpc.UpdateGraphAsync(updateGraph, null, null, cancellationToken);
+			await JobRpc.UpdateGraphAsync(updateGraph, null, null, cancellationToken);
 
 			HashSet<string> validTargets = new HashSet<string>(StringComparer.OrdinalIgnoreCase);
 			validTargets.Add("Setup Build");
