Horde: Build fixes.

#rnx

diff --git a/Engine/Source/Programs/Horde/Horde.Agent/Leases/Handlers/JobHandler.cs b/Engine/Source/Programs/Horde/Horde.Agent/Leases/Handlers/JobHandler.cs
--- a/Engine/Source/Programs/Horde/Horde.Agent/Leases/Handlers/JobHandler.cs
+++ b/Engine/Source/Programs/Horde/Horde.Agent/Leases/Handlers/JobHandler.cs
@@ -182,13 +182,13 @@ namespace Horde.Agent.Leases.Handlers
 			{
 				if (cancellationToken.IsCancellationRequested && ex.IsCancellationException())
 				{
-					if (session.GrpcChannel.State != State.Healthy)
+					if (session.GrpcChannel.State == ConnectivityState.TransientFailure)
 					{
-						logger.LogInformation(ex, "Step was aborted");
+						logger.LogError(ex, "Connection to the server was lost; step aborted.");
 					}
 					else
 					{
-						logger.LogError(ex, "Connection to the server was lost; step aborted.");
+						logger.LogInformation(ex, "Step was aborted");
 					}
 					throw;
 				}
@@ -361,13 +361,13 @@ namespace Horde.Agent.Leases.Handlers
 			{
 				if (cancellationToken.IsCancellationRequested && ex.IsCancellationException())
 				{
-					if (session.RpcConnection.Healthy)
+					if (session.GrpcChannel.State == ConnectivityState.TransientFailure)
 					{
-						logger.LogError("Lease was aborted");
+						logger.LogError(ex, "Exception while executing batch: {Ex}", ex);
 					}
 					else
 					{
-						logger.LogError(ex, "Exception while executing batch: {Ex}", ex);
+						logger.LogError("Lease was aborted");
 					}
 				}
 			}
