Horde: Static analysis fixes, and fix for connectivity state exception.

#rnx

diff --git a/Engine/Source/Programs/Horde/Horde.Agent.Tests/Helper.cs b/Engine/Source/Programs/Horde/Horde.Agent.Tests/Helper.cs
--- a/Engine/Source/Programs/Horde/Horde.Agent.Tests/Helper.cs
+++ b/Engine/Source/Programs/Horde/Horde.Agent.Tests/Helper.cs
@@ -7,14 +7,11 @@ using System.Threading.Tasks;
 using EpicGames.Horde.Jobs;
 using Google.Protobuf.WellKnownTypes;
 using Grpc.Core;
-using Grpc.Net.Client;
 using Horde.Agent.Execution;
-using Horde.Agent.Utility;
 using Horde.Common.Rpc;
 using HordeCommon.Rpc;
 using HordeCommon.Rpc.Messages;
 using Microsoft.Extensions.Logging;
-using Microsoft.Extensions.Logging.Abstractions;
 using Microsoft.Extensions.Options;
 
 namespace Horde.Agent.Tests
diff --git a/Engine/Source/Programs/Horde/Horde.Agent.Tests/WorkerServiceTest.cs b/Engine/Source/Programs/Horde/Horde.Agent.Tests/WorkerServiceTest.cs
--- a/Engine/Source/Programs/Horde/Horde.Agent.Tests/WorkerServiceTest.cs
+++ b/Engine/Source/Programs/Horde/Horde.Agent.Tests/WorkerServiceTest.cs
@@ -310,7 +310,7 @@ namespace Horde.Agent.Tests
 		public readonly TaskCompletionSource<bool> UpdateSessionReceived = new();
 
 		private readonly GrpcChannel _grpcChannel;
-		private readonly FakeJobRpcClient _client;
+//		private readonly FakeJobRpcClient _client;
 
 		private class FakeHordeRpcClient : HordeRpc.HordeRpcClient
 		{
@@ -361,8 +361,8 @@ namespace Horde.Agent.Tests
 		{
 			_serverName = "FakeServer";
 			_logger = NullLogger<FakeHordeRpcServer>.Instance;
-			FakeHordeRpcClient hordeClient = new FakeHordeRpcClient(this);
-			_client = new FakeJobRpcClient(this);
+//			FakeHordeRpcClient hordeClient = new FakeHordeRpcClient(this);
+//			_client = new FakeJobRpcClient(this);
 			_grpcChannel = GrpcChannel.ForAddress(new Uri("http://horde-agent-test"), new GrpcChannelOptions());
 		}
 
diff --git a/Engine/Source/Programs/Horde/Horde.Agent/AgentSettings.cs b/Engine/Source/Programs/Horde/Horde.Agent/AgentSettings.cs
--- a/Engine/Source/Programs/Horde/Horde.Agent/AgentSettings.cs
+++ b/Engine/Source/Programs/Horde/Horde.Agent/AgentSettings.cs
@@ -3,8 +3,6 @@
 using System.ComponentModel.DataAnnotations;
 using EpicGames.Core;
 using EpicGames.Horde;
-using Horde.Agent.Execution;
-using Horde.Agent.Utility;
 using Microsoft.Extensions.Configuration;
 
 namespace Horde.Agent
@@ -250,7 +248,7 @@ namespace Horde.Agent
 				defaultServerProfile.Url = defaultServerUrl ?? new Uri("http://localhost:5000");
 				return defaultServerProfile;
 			}
-			
+
 			return GetServerProfile(Server);
 		}
 
diff --git a/Engine/Source/Programs/Horde/Horde.Agent/Commands/Execution/ExecuteJobCommand.cs b/Engine/Source/Programs/Horde/Horde.Agent/Commands/Execution/ExecuteJobCommand.cs
--- a/Engine/Source/Programs/Horde/Horde.Agent/Commands/Execution/ExecuteJobCommand.cs
+++ b/Engine/Source/Programs/Horde/Horde.Agent/Commands/Execution/ExecuteJobCommand.cs
@@ -5,10 +5,8 @@ using EpicGames.Horde.Agents;
 using EpicGames.Horde.Agents.Leases;
 using EpicGames.Horde.Agents.Sessions;
 using Grpc.Net.Client;
-using Horde.Agent.Driver;
 using Horde.Agent.Leases.Handlers;
 using Horde.Agent.Services;
-using Horde.Agent.Utility;
 using HordeCommon.Rpc.Tasks;
 using Microsoft.Extensions.Logging;
 using Microsoft.Extensions.Options;
@@ -36,14 +34,12 @@ namespace Horde.Agent.Commands.Execution
 		readonly GrpcService _grpcService;
 		readonly JobHandler _jobHandler;
 		readonly AgentSettings _agentSettings;
-		readonly DriverSettings _driverSettings;
 
-		public ExecuteJobCommand(GrpcService grpcService, JobHandler jobHandler, IOptions<AgentSettings> agentSettings, IOptions<DriverSettings> driverSettings)
+		public ExecuteJobCommand(GrpcService grpcService, JobHandler jobHandler, IOptions<AgentSettings> agentSettings)
 		{
 			_grpcService = grpcService;
 			_jobHandler = jobHandler;
 			_agentSettings = agentSettings.Value;
-			_driverSettings = driverSettings.Value;
 		}
 
 		public override async Task<int> ExecuteAsync(ILogger logger)
diff --git a/Engine/Source/Programs/Horde/Horde.Agent/Leases/Handlers/ConformHandler.cs b/Engine/Source/Programs/Horde/Horde.Agent/Leases/Handlers/ConformHandler.cs
--- a/Engine/Source/Programs/Horde/Horde.Agent/Leases/Handlers/ConformHandler.cs
+++ b/Engine/Source/Programs/Horde/Horde.Agent/Leases/Handlers/ConformHandler.cs
@@ -16,16 +16,14 @@ namespace Horde.Agent.Leases.Handlers
 {
 	class ConformHandler : LeaseHandler<ConformTask>
 	{
-		readonly AgentSettings _agentSettings;
 		readonly DriverSettings _driverSettings;
 		readonly IServerLoggerFactory _serverLoggerFactory;
 
 		/// <summary>
 		/// Constructor
 		/// </summary>
-		public ConformHandler(IOptions<AgentSettings> agentSettings, IOptions<DriverSettings> driverSettings, IServerLoggerFactory serverLoggerFactory)
+		public ConformHandler(IOptions<DriverSettings> driverSettings, IServerLoggerFactory serverLoggerFactory)
 		{
-			_agentSettings = agentSettings.Value;
 			_driverSettings = driverSettings.Value;
 			_serverLoggerFactory = serverLoggerFactory;
 		}
diff --git a/Engine/Source/Programs/Horde/Horde.Agent/Leases/Handlers/JobHandler.cs b/Engine/Source/Programs/Horde/Horde.Agent/Leases/Handlers/JobHandler.cs
--- a/Engine/Source/Programs/Horde/Horde.Agent/Leases/Handlers/JobHandler.cs
+++ b/Engine/Source/Programs/Horde/Horde.Agent/Leases/Handlers/JobHandler.cs
@@ -182,7 +182,7 @@ namespace Horde.Agent.Leases.Handlers
 			{
 				if (cancellationToken.IsCancellationRequested && ex.IsCancellationException())
 				{
-					if (session.GrpcChannel.State == ConnectivityState.TransientFailure)
+					if (session.ConnectivityState == ConnectivityState.TransientFailure)
 					{
 						logger.LogError(ex, "Connection to the server was lost; step aborted.");
 					}
@@ -361,7 +361,7 @@ namespace Horde.Agent.Leases.Handlers
 			{
 				if (cancellationToken.IsCancellationRequested && ex.IsCancellationException())
 				{
-					if (session.GrpcChannel.State == ConnectivityState.TransientFailure)
+					if (session.ConnectivityState == ConnectivityState.TransientFailure)
 					{
 						logger.LogError(ex, "Exception while executing batch: {Ex}", ex);
 					}
diff --git a/Engine/Source/Programs/Horde/Horde.Agent/Leases/LeaseManager.cs b/Engine/Source/Programs/Horde/Horde.Agent/Leases/LeaseManager.cs
--- a/Engine/Source/Programs/Horde/Horde.Agent/Leases/LeaseManager.cs
+++ b/Engine/Source/Programs/Horde/Horde.Agent/Leases/LeaseManager.cs
@@ -8,7 +8,6 @@ using Google.Protobuf.WellKnownTypes;
 using Grpc.Core;
 using Grpc.Net.Client;
 using Horde.Agent.Services;
-using Horde.Agent.Utility;
 using HordeCommon.Rpc;
 using HordeCommon.Rpc.Messages;
 using Microsoft.Extensions.DependencyInjection;
@@ -370,7 +369,7 @@ namespace Horde.Agent.Leases
 					}
 
 					// Update the current status
-					if (_session.GrpcChannel.State == ConnectivityState.TransientFailure)
+					if (_session.ConnectivityState == ConnectivityState.TransientFailure)
 					{
 						_statusService.Set(false, _activeLeases.Count, "Attempting to connect to server...");
 					}
diff --git a/Engine/Source/Programs/Horde/Horde.Agent/Services/SessionFactory.cs b/Engine/Source/Programs/Horde/Horde.Agent/Services/SessionFactory.cs
--- a/Engine/Source/Programs/Horde/Horde.Agent/Services/SessionFactory.cs
+++ b/Engine/Source/Programs/Horde/Horde.Agent/Services/SessionFactory.cs
@@ -8,7 +8,6 @@ using EpicGames.Horde.Agents;
 using EpicGames.Horde.Agents.Sessions;
 using Grpc.Core;
 using Grpc.Net.Client;
-using Horde.Agent.Driver;
 using Horde.Agent.Utility;
 using Horde.Common.Rpc;
 using HordeCommon.Rpc;
@@ -48,6 +47,11 @@ namespace Horde.Agent.Services
 		/// </summary>
 		GrpcChannel GrpcChannel { get; }
 
+		/// <summary>
+		/// State of the connection 
+		/// </summary>
+		ConnectivityState ConnectivityState { get; }
+
 		/// <summary>
 		/// Working directory for sandboxes etc..
 		/// </summary>
@@ -90,6 +94,9 @@ namespace Horde.Agent.Services
 		/// <inheritdoc/>
 		public DirectoryReference WorkingDir { get; }
 
+		/// <inheritdoc/>
+		public ConnectivityState ConnectivityState => ConnectivityState.Ready;
+
 		/// <summary>
 		/// Constructor
 		/// </summary>
@@ -113,7 +120,7 @@ namespace Horde.Agent.Services
 		/// <summary>
 		/// Creates a new agent session
 		/// </summary>
-		public static async Task<Session> CreateAsync(CapabilitiesService capabilitiesService, GrpcService grpcService, StatusService statusService, IOptions<AgentSettings> settings, IOptions<DriverSettings> driverSettings, ILogger logger, CancellationToken cancellationToken)
+		public static async Task<Session> CreateAsync(CapabilitiesService capabilitiesService, GrpcService grpcService, StatusService statusService, IOptions<AgentSettings> settings, ILogger logger, CancellationToken cancellationToken)
 		{
 			AgentSettings currentSettings = settings.Value;
 
@@ -349,23 +356,21 @@ namespace Horde.Agent.Services
 		readonly GrpcService _grpcService;
 		readonly StatusService _statusService;
 		readonly IOptions<AgentSettings> _agentSettings;
-		readonly IOptions<DriverSettings> _driverSettings;
 		readonly ILogger _logger;
 
 		/// <summary>
 		/// Constructor
 		/// </summary>
-		public SessionFactory(CapabilitiesService capabilitiesService, GrpcService grpcService, StatusService statusService, IOptions<AgentSettings> agentSettings, IOptions<DriverSettings> driverSettings, ILogger<SessionFactory> logger)
+		public SessionFactory(CapabilitiesService capabilitiesService, GrpcService grpcService, StatusService statusService, IOptions<AgentSettings> agentSettings, ILogger<SessionFactory> logger)
 		{
 			_capabilitiesService = capabilitiesService;
 			_grpcService = grpcService;
 			_statusService = statusService;
 			_agentSettings = agentSettings;
-			_driverSettings = driverSettings;
 			_logger = logger;
 		}
 
 		/// <inheritdoc/>
-		public async Task<ISession> CreateAsync(CancellationToken cancellationToken) => await Session.CreateAsync(_capabilitiesService, _grpcService, _statusService, _agentSettings, _driverSettings, _logger, cancellationToken);
+		public async Task<ISession> CreateAsync(CancellationToken cancellationToken) => await Session.CreateAsync(_capabilitiesService, _grpcService, _statusService, _agentSettings, _logger, cancellationToken);
 	}
 }
