Horde: Roll back changes to use grpc channels until I can track down cause for agent failures.

diff --git a/Engine/Source/Programs/Horde/Horde.Agent.Driver/Execution/JobExecutor.cs b/Engine/Source/Programs/Horde/Horde.Agent.Driver/Execution/JobExecutor.cs
--- a/Engine/Source/Programs/Horde/Horde.Agent.Driver/Execution/JobExecutor.cs
+++ b/Engine/Source/Programs/Horde/Horde.Agent.Driver/Execution/JobExecutor.cs
@@ -15,7 +15,6 @@ using EpicGames.Horde.Storage;
 using EpicGames.Horde.Storage.Clients;
 using EpicGames.Horde.Storage.Nodes;
 using Grpc.Core;
-using Grpc.Net.Client;
 using Horde.Agent.Driver;
 using Horde.Agent.Parser;
 //using Horde.Agent.Services;
@@ -64,7 +63,7 @@ namespace Horde.Agent.Execution
 	{
 		public Uri ServerUrl { get; }
 		public DirectoryReference WorkingDir { get; }
-		public GrpcChannel GrpcChannel { get; }
+		public IRpcConnection RpcConnection { get; }
 		public IReadOnlyList<ProcessToTerminate>? ProcessesToTerminate { get; }
 		public HttpStorageClientFactory StorageFactory { get; }
 		public JobId JobId { get; }
@@ -73,11 +72,11 @@ namespace Horde.Agent.Execution
 		public string Token { get; }
 		public RpcJobOptions JobOptions { get; }
 
-		public JobExecutorOptions(Uri serverUrl, DirectoryReference workingDir, GrpcChannel grpcChannel, IReadOnlyList<ProcessToTerminate>? processesToTerminate, HttpStorageClientFactory storageFactory, JobId jobId, JobStepBatchId batchId, RpcBeginBatchResponse batch, string token, RpcJobOptions jobOptions)
+		public JobExecutorOptions(Uri serverUrl, DirectoryReference workingDir, IRpcConnection rpcConnection, IReadOnlyList<ProcessToTerminate>? processesToTerminate, HttpStorageClientFactory storageFactory, JobId jobId, JobStepBatchId batchId, RpcBeginBatchResponse batch, string token, RpcJobOptions jobOptions)
 		{
 			ServerUrl = serverUrl;
 			WorkingDir = workingDir;
-			GrpcChannel = grpcChannel;
+			RpcConnection = rpcConnection;
 			ProcessesToTerminate = processesToTerminate;
 			StorageFactory = storageFactory;
 			JobId = jobId;
@@ -241,8 +240,7 @@ namespace Horde.Agent.Execution
 		protected HttpStorageClientFactory StorageFactory { get; }
 		protected RpcJobOptions JobOptions { get; }
 
-		protected GrpcChannel GrpcChannel => Options.GrpcChannel;
-		protected JobRpc.JobRpcClient JobRpc { get; }
+		protected IRpcConnection RpcConnection => Options.RpcConnection;
 		protected Dictionary<string, string> _remapAgentTypes = new Dictionary<string, string>(StringComparer.OrdinalIgnoreCase);
 
 		protected Dictionary<string, string> _envVars = new Dictionary<string, string>(StringComparer.OrdinalIgnoreCase);
@@ -257,7 +255,6 @@ namespace Horde.Agent.Execution
 			Batch = options.Batch;
 
 			JobOptions = options.JobOptions;
-			JobRpc = new JobRpc.JobRpcClient(GrpcChannel);
 
 			_envVars[HordeHttpClient.HordeUrlEnvVarName] = options.ServerUrl.ToString();
 			_envVars[HordeHttpClient.HordeTokenEnvVarName] = options.Token;
@@ -514,6 +511,8 @@ namespace Horde.Agent.Execution
 			using (GlobalTracer.Instance.BuildSpan("TempStorage").WithTag("resource", "Write").StartActive())
 			{
 				// Create the artifact
+				using IRpcClientRef<JobRpc.JobRpcClient> jobRpc = await RpcConnection.GetClientRefAsync<JobRpc.JobRpcClient>(cancellationToken);
+
 				ArtifactName artifactName = TempStorage.GetArtifactNameForNode(SetupStepName);
 				ArtifactType artifactType = ArtifactType.StepOutput;
 
@@ -523,7 +522,7 @@ namespace Horde.Agent.Execution
 				artifactRequest.Name = artifactName.ToString();
 				artifactRequest.Type = artifactType.ToString();
 
-				RpcCreateJobArtifactResponseV2 artifact = await JobRpc.CreateArtifactV2Async(artifactRequest, cancellationToken: cancellationToken);
+				RpcCreateJobArtifactResponseV2 artifact = await jobRpc.Client.CreateArtifactV2Async(artifactRequest, cancellationToken: cancellationToken);
 				ArtifactId artifactId = ArtifactId.Parse(artifact.Id);
 				logger.LogInformation("Creating output artifact {ArtifactId} '{ArtifactName}' ({ArtifactType}) with ref {RefName} in namespace {NamespaceId}", artifactId, artifactName, artifactType, artifact.RefName, artifact.NamespaceId);
 
@@ -550,7 +549,7 @@ namespace Horde.Agent.Execution
 			}
 
 			RpcUpdateGraphRequest updateGraph = await ParseGraphUpdateAsync(definitionFile, logger, cancellationToken);
-			await JobRpc.UpdateGraphAsync(updateGraph, null, null, cancellationToken);
+			await RpcConnection.InvokeAsync((JobRpc.JobRpcClient x) => x.UpdateGraphAsync(updateGraph, null, null, cancellationToken), cancellationToken);
 
 			HashSet<string> validTargets = new HashSet<string>(StringComparer.OrdinalIgnoreCase);
 			validTargets.Add("Setup Build");
@@ -771,13 +770,15 @@ namespace Horde.Agent.Execution
 
 				ArtifactName artifactName = TempStorage.GetArtifactNameForNode(SetupStepName);
 
+				using IRpcClientRef<JobRpc.JobRpcClient> jobRpc = await RpcConnection.GetClientRefAsync<JobRpc.JobRpcClient>(cancellationToken);
+
 				RpcGetJobArtifactRequest artifactRequest = new RpcGetJobArtifactRequest();
 				artifactRequest.JobId = JobId.ToString();
 				artifactRequest.StepId = step.StepId.ToString();
 				artifactRequest.Name = artifactName.ToString();
 				artifactRequest.Type = ArtifactType.StepOutput.ToString();
 
-				RpcGetJobArtifactResponse artifact = await JobRpc.GetArtifactAsync(artifactRequest, cancellationToken: cancellationToken);
+				RpcGetJobArtifactResponse artifact = await jobRpc.Client.GetArtifactAsync(artifactRequest, cancellationToken: cancellationToken);
 
 				NamespaceId namespaceId = new NamespaceId(artifact.NamespaceId);
 				RefName refName = new RefName(artifact.RefName);
@@ -826,13 +827,15 @@ namespace Horde.Agent.Execution
 		{
 			try
 			{
+				using IRpcClientRef<JobRpc.JobRpcClient> jobRpc = await RpcConnection.GetClientRefAsync<JobRpc.JobRpcClient>(cancellationToken);
+
 				RpcCreateJobArtifactRequestV2 artifactRequest = new RpcCreateJobArtifactRequestV2();
 				artifactRequest.JobId = JobId.ToString();
 				artifactRequest.StepId = stepId.ToString();
 				artifactRequest.Name = name.ToString();
 				artifactRequest.Type = type.ToString();
 
-				RpcCreateJobArtifactResponseV2 artifact = await JobRpc.CreateArtifactV2Async(artifactRequest, cancellationToken: cancellationToken);
+				RpcCreateJobArtifactResponseV2 artifact = await jobRpc.Client.CreateArtifactV2Async(artifactRequest, cancellationToken: cancellationToken);
 				ArtifactId artifactId = ArtifactId.Parse(artifact.Id);
 				Logger.LogInformation("Created artifact {ArtifactId} '{ArtifactName}' ({ArtifactType}) with ref {RefName} ({Link})", artifactId, name, type, artifact.RefName, $"{Options.ServerUrl}/api/v1/storage/{artifact.NamespaceId}/refs/{artifact.RefName}");
 
@@ -867,6 +870,8 @@ namespace Horde.Agent.Execution
 		{
 			DirectoryReference manifestDir = DirectoryReference.Combine(workspaceDir, "Engine", "Saved", "BuildGraph");
 
+			using IRpcClientRef<JobRpc.JobRpcClient> jobRpc = await RpcConnection.GetClientRefAsync<JobRpc.JobRpcClient>(cancellationToken);
+
 			// Create the mapping of tag names to file sets
 			Dictionary<string, HashSet<FileReference>> tagNameToFileSet = new Dictionary<string, HashSet<FileReference>>();
 
@@ -884,7 +889,7 @@ namespace Horde.Agent.Execution
 				string nodeName = input.Substring(0, slashIdx);
 				string tagName = input.Substring(slashIdx + 1);
 
-				TempStorageTagManifest fileList = await TempStorage.RetrieveTagAsync(JobRpc, JobId, step.StepId, StorageFactory, nodeName, tagName, manifestDir, logger, cancellationToken);
+				TempStorageTagManifest fileList = await TempStorage.RetrieveTagAsync(jobRpc, JobId, step.StepId, StorageFactory, nodeName, tagName, manifestDir, logger, cancellationToken);
 				tagNameToFileSet[tagName] = fileList.ToFileSet(workspaceDir);
 				inputStorageBlocks.UnionWith(fileList.Blocks);
 			}
@@ -897,7 +902,7 @@ namespace Horde.Agent.Execution
 				scope.Span.SetTag("blocks", inputStorageBlocks.Count);
 				foreach (TempStorageBlockRef inputStorageBlock in inputStorageBlocks)
 				{
-					TempStorageBlockManifest manifest = await TempStorage.RetrieveBlockAsync(JobRpc, JobId, step.StepId, StorageFactory, inputStorageBlock.NodeName, inputStorageBlock.OutputName, workspaceDir, manifestDir, logger, cancellationToken);
+					TempStorageBlockManifest manifest = await TempStorage.RetrieveBlockAsync(jobRpc, JobId, step.StepId, StorageFactory, inputStorageBlock.NodeName, inputStorageBlock.OutputName, workspaceDir, manifestDir, logger, cancellationToken);
 					inputManifests[inputStorageBlock] = manifest;
 				}
 				scope.Span.SetTag("size", inputManifests.Sum(x => x.Value.GetTotalSize()));
@@ -1032,7 +1037,7 @@ namespace Horde.Agent.Execution
 				artifactRequest.Name = TempStorage.GetArtifactNameForNode(step.Name).ToString();
 				artifactRequest.Type = ArtifactType.StepOutput.ToString();
 
-				RpcCreateJobArtifactResponseV2 artifact = await JobRpc.CreateArtifactV2Async(artifactRequest, cancellationToken: cancellationToken);
+				RpcCreateJobArtifactResponseV2 artifact = await jobRpc.Client.CreateArtifactV2Async(artifactRequest, cancellationToken: cancellationToken);
 				ArtifactId artifactId = ArtifactId.Parse(artifact.Id);
 				logger.LogInformation("Created artifact {ArtifactId} '{ArtifactName}' ({ArtifactType}) with ref {RefName} ({RefUrl})", artifactId, artifactRequest.Name, ArtifactType.StepOutput, artifact.RefName, $"{Options.ServerUrl.ToString().TrimEnd('/')}/api/v1/storage/{artifact.NamespaceId}/refs/{artifact.RefName}");
 
@@ -1111,7 +1116,7 @@ namespace Horde.Agent.Execution
 					artifactRequest.Keys.AddRange(graphArtifact.Keys);
 					artifactRequest.Metadata.AddRange(graphArtifact.Metadata);
 
-					RpcCreateJobArtifactResponseV2 artifact = await JobRpc.CreateArtifactV2Async(artifactRequest, cancellationToken: cancellationToken);
+					RpcCreateJobArtifactResponseV2 artifact = await jobRpc.Client.CreateArtifactV2Async(artifactRequest, cancellationToken: cancellationToken);
 					ArtifactId artifactId = ArtifactId.Parse(artifact.Id);
 					logger.LogInformation("Created artifact {ArtifactId} '{ArtifactName}' ({ArtifactType}) with ref {RefName} ({RefUrl})", artifactId, artifactRequest.Name, ArtifactType.StepOutput, artifact.RefName, $"{Options.ServerUrl.ToString().TrimEnd('/')}/api/v1/storage/{artifact.NamespaceId}/refs/{artifact.RefName}");
 
@@ -1658,7 +1663,7 @@ namespace Horde.Agent.Execution
 					}
 				}
 
-				await JobRpc.UpdateGraphAsync(updateGraph, null, null, cancellationToken);
+				await RpcConnection.InvokeAsync((JobRpc.JobRpcClient x) => x.UpdateGraphAsync(updateGraph, null, null, cancellationToken), cancellationToken);
 
 				HashSet<string> publishOutputNames = new HashSet<string>(updateGraph.Groups.SelectMany(x => x.Nodes).SelectMany(x => x.Inputs), StringComparer.OrdinalIgnoreCase);
 				foreach (string publishOutputName in publishOutputNames)
@@ -1836,7 +1841,7 @@ namespace Horde.Agent.Execution
 			request.Placement = report.Placement;
 			request.Name = report.Name;
 			request.Content = report.Content;
-			await JobRpc.CreateReportAsync(request);
+			await RpcConnection.InvokeAsync((JobRpc.JobRpcClient x) => x.CreateReportAsync(request), CancellationToken.None);
 		}
 
 		private static ISpan CreateTracingData(ISpan parent, TraceSpan span)
@@ -1867,11 +1872,17 @@ namespace Horde.Agent.Execution
 			return newSpan;
 		}
 
-		public async Task UploadTestDataAsync(JobStepId jobStepId, IEnumerable<KeyValuePair<string, object>> pairs)
+		protected async Task UploadTestDataAsync(JobStepId jobStepId, IEnumerable<KeyValuePair<string, object>> testData)
 		{
-			if (pairs.Any())
+			if (testData.Any())
 			{
-				using (AsyncClientStreamingCall<RpcUploadTestDataRequest, RpcUploadTestDataResponse> call = JobRpc.UploadTestData())
+				await RpcConnection.InvokeAsync((JobRpc.JobRpcClient x) => UploadTestDataAsync(x, jobStepId, testData), CancellationToken.None);
+			}
+		}
+
+		async Task<bool> UploadTestDataAsync(JobRpc.JobRpcClient rpcClient, JobStepId jobStepId, IEnumerable<KeyValuePair<string, object>> pairs)
+		{
+			using (AsyncClientStreamingCall<RpcUploadTestDataRequest, RpcUploadTestDataResponse> call = rpcClient.UploadTestData())
 				{
 					foreach (KeyValuePair<string, object> pair in pairs)
 					{
@@ -1890,7 +1901,7 @@ namespace Horde.Agent.Execution
 					await call.RequestStream.CompleteAsync();
 					await call.ResponseAsync;
 				}
-			}
+			return true;
 		}
 	}
 
diff --git a/Engine/Source/Programs/Horde/Horde.Agent.Driver/Execution/PerforceExecutor.cs b/Engine/Source/Programs/Horde/Horde.Agent.Driver/Execution/PerforceExecutor.cs
--- a/Engine/Source/Programs/Horde/Horde.Agent.Driver/Execution/PerforceExecutor.cs
+++ b/Engine/Source/Programs/Horde/Horde.Agent.Driver/Execution/PerforceExecutor.cs
@@ -124,7 +124,7 @@ namespace Horde.Agent.Execution
 					RpcUpdateJobRequest updateJobRequest = new RpcUpdateJobRequest();
 					updateJobRequest.JobId = JobId.ToString();
 					updateJobRequest.Change = Batch.Change;
-					await JobRpc.UpdateJobAsync(updateJobRequest, null, null, cancellationToken);
+					await RpcConnection.InvokeAsync((JobRpc.JobRpcClient x) => x.UpdateJobAsync(updateJobRequest, null, null, cancellationToken), cancellationToken);
 				}
 
 				// Sync the workspace
diff --git a/Engine/Source/Programs/Horde/Horde.Agent.Driver/Execution/TestExecutor.cs b/Engine/Source/Programs/Horde/Horde.Agent.Driver/Execution/TestExecutor.cs
--- a/Engine/Source/Programs/Horde/Horde.Agent.Driver/Execution/TestExecutor.cs
+++ b/Engine/Source/Programs/Horde/Horde.Agent.Driver/Execution/TestExecutor.cs
@@ -114,7 +114,7 @@ namespace Horde.Agent.Execution
 			updateGraph.Labels.Add(CreateLabel("Editors", "Project", new string[] { $"Compile {projectName}Editor Win64" }, Array.Empty<string>(), dependencyMap));
 			updateGraph.Labels.Add(CreateLabel("Clients", "Project", new string[] { $"Cook {projectName}Client Win64" }, new string[] { $"Publish {projectName}Client Win64" }, dependencyMap));
 
-			await JobRpc.UpdateGraphAsync(updateGraph, null, null, cancellationToken);
+			await RpcConnection.InvokeAsync((JobRpc.JobRpcClient x) => x.UpdateGraphAsync(updateGraph, null, null, cancellationToken), cancellationToken);
 
 			logger.LogInformation("**** FINISH JOB SETUP ****");
 			return true;
diff --git a/Engine/Source/Programs/Horde/Horde.Agent.Driver/Utility/JsonRpcLogSink.cs b/Engine/Source/Programs/Horde/Horde.Agent.Driver/Utility/JsonRpcLogSink.cs
--- a/Engine/Source/Programs/Horde/Horde.Agent.Driver/Utility/JsonRpcLogSink.cs
+++ b/Engine/Source/Programs/Horde/Horde.Agent.Driver/Utility/JsonRpcLogSink.cs
@@ -7,7 +7,6 @@ using EpicGames.Horde.Logs;
 using EpicGames.Horde.Storage;
 using Google.Protobuf;
 using Grpc.Core;
-using Grpc.Net.Client;
 using Horde.Common.Rpc;
 using HordeCommon.Rpc;
 using Microsoft.Extensions.Logging;
@@ -26,7 +25,7 @@ namespace Horde.Agent.Utility
 	{
 		const int FlushLength = 1024 * 1024;
 
-		readonly GrpcChannel _connection;
+		readonly IRpcConnection _connection;
 		readonly JobId? _jobId;
 		readonly JobStepBatchId? _jobBatchId;
 		readonly JobStepId? _jobStepId;
@@ -43,7 +42,7 @@ namespace Horde.Agent.Utility
 		AsyncEvent _tailTaskStop;
 		readonly AsyncEvent _newTailDataEvent = new AsyncEvent();
 
-		public JsonRpcAndStorageLogSink(GrpcChannel connection, LogId logId, JobId? jobId, JobStepBatchId? jobBatchId, JobStepId? jobStepId, IStorageClient store, ILogger logger)
+		public JsonRpcAndStorageLogSink(IRpcConnection connection, LogId logId, JobId? jobId, JobStepBatchId? jobBatchId, JobStepId? jobStepId, IStorageClient store, ILogger logger)
 		{
 			_connection = connection;
 			_logId = logId;
@@ -174,10 +173,9 @@ namespace Horde.Agent.Utility
 			// Update the outcome of this jobstep
 			if (_jobId != null && _jobBatchId != null && _jobStepId != null)
 			{
-				JobRpc.JobRpcClient jobRpc = new JobRpc.JobRpcClient(_connection);
 				try
 				{
-					await jobRpc.UpdateStepAsync(new RpcUpdateStepRequest(_jobId.Value, _jobBatchId.Value, _jobStepId.Value, JobStepState.Unspecified, outcome), cancellationToken: cancellationToken);
+					await _connection.InvokeAsync((JobRpc.JobRpcClient x) => x.UpdateStepAsync(new RpcUpdateStepRequest(_jobId.Value, _jobBatchId.Value, _jobStepId.Value, JobStepState.Unspecified, outcome)), cancellationToken);
 				}
 				catch (Exception ex)
 				{
@@ -189,8 +187,7 @@ namespace Horde.Agent.Utility
 		/// <inheritdoc/>
 		public async Task WriteEventsAsync(List<RpcCreateEventRequest> events, CancellationToken cancellationToken)
 		{
-			JobRpc.JobRpcClient jobRpc = new JobRpc.JobRpcClient(_connection);
-			await jobRpc.CreateEventsAsync(new RpcCreateEventsRequest(events), cancellationToken: cancellationToken);
+			await _connection.InvokeAsync((JobRpc.JobRpcClient x) => x.CreateEventsAsync(new RpcCreateEventsRequest(events)), cancellationToken);
 		}
 
 		/// <inheritdoc/>
@@ -221,9 +218,7 @@ namespace Horde.Agent.Utility
 			request.TargetHash = target.Hash.ToString();
 			request.TargetLocator = target.GetLocator().ToString();
 			request.Complete = complete;
-
-			LogRpcClient clientRef = new LogRpcClient(_connection);
-			await clientRef.UpdateLogAsync(request, cancellationToken: cancellationToken);
+			await _connection.InvokeAsync((LogRpcClient client) => client.UpdateLogAsync(request, cancellationToken: cancellationToken), cancellationToken);
 		}
 
 		protected virtual async Task<int> UpdateLogTailAsync(int tailNext, ReadOnlyMemory<byte> tailData, CancellationToken cancellationToken)
@@ -231,8 +226,8 @@ namespace Horde.Agent.Utility
 			DateTime deadline = DateTime.UtcNow.AddMinutes(2.0);
 			try
 			{
-				LogRpcClient clientRef = new LogRpcClient(_connection);
-				using AsyncDuplexStreamingCall<UpdateLogTailRequest, UpdateLogTailResponse> call = clientRef.UpdateLogTail(deadline: deadline, cancellationToken: cancellationToken);
+				using IRpcClientRef<LogRpcClient> clientRef = await _connection.GetClientRefAsync<LogRpcClient>(cancellationToken);
+				using AsyncDuplexStreamingCall<UpdateLogTailRequest, UpdateLogTailResponse> call = clientRef.Client.UpdateLogTail(deadline: deadline, cancellationToken: cancellationToken);
 
 				// Write the request to the server
 				UpdateLogTailRequest request = new UpdateLogTailRequest();
@@ -245,8 +240,14 @@ namespace Horde.Agent.Utility
 				// Wait until the server responds or we need to trigger a new update
 				Task<bool> moveNextAsync = call.ResponseStream.MoveNext();
 
-				Task task = await Task.WhenAny(moveNextAsync, _tailTaskStop.Task, Task.Delay(TimeSpan.FromMinutes(1.0), CancellationToken.None));
-				if (task == _tailTaskStop.Task)
+				Task task = await Task.WhenAny(moveNextAsync, clientRef.DisposingTask, _tailTaskStop.Task, Task.Delay(TimeSpan.FromMinutes(1.0), CancellationToken.None));
+				if (task == clientRef.DisposingTask)
+				{
+					TimeSpan graceDelay = TimeSpan.FromSeconds(10);
+					_logger.LogInformation("Cancelling long poll from client side (server migration). Backing off for {Delay} ms...", graceDelay.TotalMilliseconds);
+					await Task.Delay(graceDelay, cancellationToken);
+				}
+				else if (task == _tailTaskStop.Task)
 				{
 					_logger.LogInformation("Cancelling long poll from client side (complete)");
 				}
diff --git a/Engine/Source/Programs/Horde/Horde.Agent.Driver/Utility/RpcConnection.cs b/Engine/Source/Programs/Horde/Horde.Agent.Driver/Utility/RpcConnection.cs
new file mode 100644
--- /dev/null
+++ b/Engine/Source/Programs/Horde/Horde.Agent.Driver/Utility/RpcConnection.cs
@@ -0,0 +1,665 @@
+// Copyright Epic Games, Inc. All Rights Reserved.
+
+using System.Diagnostics;
+using System.Net;
+using Grpc.Core;
+using Grpc.Net.Client;
+using HordeCommon.Rpc;
+using Microsoft.Extensions.Logging;
+using OpenTracing;
+using OpenTracing.Util;
+
+namespace Horde.Agent.Utility
+{
+	/// <summary>
+	/// Reference to a gRPC client
+	/// </summary>
+	public interface IRpcClientRef<T> : IDisposable where T : ClientBase<T>
+	{
+		/// <summary>
+		/// The client instance
+		/// </summary>
+		T Client { get; }
+
+		/// <summary>
+		/// Task which completes when the client needs to be disposed of
+		/// </summary>
+		Task DisposingTask { get; }
+	}
+
+	/// <summary>
+	/// gRPC connection to the server
+	/// </summary>
+	public interface IRpcConnection : IAsyncDisposable
+	{
+		/// <summary>
+		/// Reports whether the server connection is healthy
+		/// </summary>
+		bool Healthy { get; }
+
+		/// <summary>
+		/// Logger for this connection
+		/// </summary>
+		ILogger Logger { get; }
+
+		/// <summary>
+		/// Attempts to get a client reference, returning immediately if there's not one available
+		/// </summary>
+		/// <returns></returns>
+		IRpcClientRef<TClient>? TryGetClientRef<TClient>() where TClient : ClientBase<TClient>;
+
+		/// <summary>
+		/// Obtains a new client reference object
+		/// </summary>
+		/// <param name="cancellationToken">Cancellation token for this request</param>
+		/// <returns>New client reference</returns>
+		Task<IRpcClientRef<TClient>> GetClientRefAsync<TClient>(CancellationToken cancellationToken) where TClient : ClientBase<TClient>;
+	}
+
+	static class RpcConnectionExtensions
+	{
+		private static readonly TimeSpan[] s_retryTimes =
+		{
+			TimeSpan.FromSeconds(1.0),
+			TimeSpan.FromSeconds(10.0),
+			TimeSpan.FromSeconds(30.0),
+		};
+
+		/// <summary>
+		/// Invokes an asynchronous command with a HordeRpcClient instance
+		/// </summary>
+		/// <typeparam name="TClient">The type of client to call</typeparam>
+		/// <typeparam name="TResult">Return type from the call</typeparam>
+		/// <param name="connection">The connection object</param>
+		/// <param name="func">The function to execute</param>
+		/// <param name="cancellationToken">Cancellation token for the call</param>
+		/// <returns>Response from the call</returns>
+		public static async Task<TResult> InvokeOnceAsync<TClient, TResult>(this IRpcConnection connection, Func<TClient, Task<TResult>> func, CancellationToken cancellationToken) where TClient : ClientBase<TClient>
+		{
+			using (IRpcClientRef<TClient> clientRef = await connection.GetClientRefAsync<TClient>(cancellationToken))
+			{
+				return await func(clientRef.Client);
+			}
+		}
+
+		/// <summary>
+		/// Obtains a client reference and executes a unary RPC
+		/// </summary>
+		/// <typeparam name="TClient">The type of client to call</typeparam>
+		/// <typeparam name="TResult">Return type from the call</typeparam>
+		/// <param name="connection">The connection object</param>
+		/// <param name="func">Method to execute with the RPC instance</param>
+		/// <param name="cancellationToken">Cancellation token for the call</param>
+		/// <returns>Response from the call</returns>
+		public static Task<TResult> InvokeOnceAsync<TClient, TResult>(this IRpcConnection connection, Func<TClient, AsyncUnaryCall<TResult>> func, CancellationToken cancellationToken) where TClient : ClientBase<TClient>
+		{
+			return InvokeOnceAsync<TClient, TResult>(connection, async (x) => await func(x), cancellationToken);
+		}
+
+		/// <summary>
+		/// Invokes an asynchronous command with a HordeRpcClient instance
+		/// </summary>
+		/// <typeparam name="TClient">The type of client to call</typeparam>
+		/// <typeparam name="TResult">Return type from the call</typeparam>
+		/// <param name="connection">The connection object</param>
+		/// <param name="func">The function to execute</param>
+		/// <param name="cancellationToken">Cancellation token for the call</param>
+		/// <returns>Response from the call</returns>
+		public static async Task<TResult> InvokeAsync<TClient, TResult>(this IRpcConnection connection, Func<TClient, Task<TResult>> func, CancellationToken cancellationToken) where TClient : ClientBase<TClient>
+		{
+			for (int attempt = 0; ; attempt++)
+			{
+				// Attempt the RPC call
+				using (IRpcClientRef<TClient> clientRef = await connection.GetClientRefAsync<TClient>(cancellationToken))
+				{
+					try
+					{
+						return await func(clientRef.Client);
+					}
+					catch (Exception ex)
+					{
+						// Otherwise check if we can try again
+						if (attempt >= s_retryTimes.Length)
+						{
+							connection.Logger.LogError("Max number of retry attempts reached");
+							throw;
+						}
+
+						// Check if the user requested cancellation
+						if (cancellationToken.IsCancellationRequested)
+						{
+							connection.Logger.LogInformation("Cancellation of gRPC call requested");
+							throw;
+						}
+
+						// Check we can retry it
+						if (!CanRetry(ex))
+						{
+							connection.Logger.LogError(ex, "Exception during RPC: {Message}", ex.Message);
+							throw;
+						}
+
+						RpcException? rpcEx = ex as RpcException;
+						if (rpcEx != null && rpcEx.StatusCode == StatusCode.Unavailable)
+						{
+							connection.Logger.LogInformation(ex, "Failure #{Attempt} during gRPC call (service unavailable). Retrying...", attempt + 1);
+						}
+						else
+						{
+							connection.Logger.LogError(ex, "Failure #{Attempt} during gRPC call.", attempt + 1);
+						}
+					}
+				}
+
+				// Wait before retrying
+				await Task.Delay(s_retryTimes[attempt], cancellationToken);
+			}
+		}
+
+		/// <summary>
+		/// Obtains a client reference and executes a unary RPC
+		/// </summary>
+		/// <typeparam name="TClient">The type of client to call</typeparam>
+		/// <typeparam name="TResult">Return type from the call</typeparam>
+		/// <param name="connection">The connection object</param>
+		/// <param name="func">Method to execute with the RPC instance</param>
+		/// <param name="cancellationToken">Cancellation token for the call</param>
+		/// <returns>Response from the call</returns>
+		public static Task<TResult> InvokeAsync<TClient, TResult>(this IRpcConnection connection, Func<TClient, AsyncUnaryCall<TResult>> func, CancellationToken cancellationToken) where TClient : ClientBase<TClient>
+		{
+			return InvokeAsync<TClient, TResult>(connection, async (x) => await func(x), cancellationToken);
+		}
+
+		/// <summary>
+		/// Determines if the given exception should be retried
+		/// </summary>
+		/// <param name="ex"></param>
+		/// <returns></returns>
+		static bool CanRetry(Exception ex)
+		{
+			RpcException? rpcEx = ex as RpcException;
+			if (rpcEx != null)
+			{
+				return rpcEx.StatusCode == StatusCode.Cancelled || rpcEx.StatusCode == StatusCode.Unavailable || rpcEx.StatusCode == StatusCode.DataLoss /* Interrupted streaming call */;
+			}
+			else
+			{
+				return false;
+			}
+		}
+	}
+
+	/// <summary>
+	/// Represents a connection to a GRPC server, which responds to server shutting down events and attempts to reconnect to another instance.
+	/// </summary>
+	class RpcConnection : IRpcConnection
+	{
+		/// <summary>
+		/// An instance of a connection to a particular server
+		/// </summary>
+		class RpcSubConnection : IAsyncDisposable
+		{
+			/// <summary>
+			/// The connection id
+			/// </summary>
+			public int ConnectionId { get; }
+
+			/// <summary>
+			/// Name of the server
+			/// </summary>
+			public string Name
+			{
+				get;
+			}
+
+			/// <summary>
+			/// The Grpc channel
+			/// </summary>
+			public GrpcChannel Channel { get; }
+
+			/// <summary>
+			/// Task which is set to indicate the connection is being disposed
+			/// </summary>
+			public Task DisposingTask => _disposingTaskSource.Task;
+
+			/// <summary>
+			/// Logger for messages about refcount changes
+			/// </summary>
+			private readonly ILogger _logger;
+
+			/// <summary>
+			/// Cache of created clients
+			/// </summary>
+			private readonly Dictionary<Type, ClientBase> _clients = new Dictionary<Type, ClientBase>();
+
+			/// <summary>
+			/// The current refcount
+			/// </summary>
+			private int _refCount;
+
+			/// <summary>
+			/// Whether the connection is terminating
+			/// </summary>
+			private bool _disposing;
+
+			/// <summary>
+			/// Cancellation token set 
+			/// </summary>
+			private readonly TaskCompletionSource<bool> _disposingTaskSource = new TaskCompletionSource<bool>(TaskCreationOptions.RunContinuationsAsynchronously);
+
+			/// <summary>
+			/// Wraps a task allowing the disposer to wait for clients to finish using this connection
+			/// </summary>
+			private readonly TaskCompletionSource<bool> _disposedTaskSource = new TaskCompletionSource<bool>(TaskCreationOptions.RunContinuationsAsynchronously);
+
+			/// <summary>
+			/// Constructor
+			/// </summary>
+			/// <param name="connectionId">The connection id</param>
+			/// <param name="name">Name of the server</param>
+			/// <param name="channel">The channel instance</param>
+			/// <param name="logger">Logger for debug messages</param>
+			public RpcSubConnection(int connectionId, string name, GrpcChannel channel, ILogger logger)
+			{
+				ConnectionId = connectionId;
+				Name = name;
+				Channel = channel;
+				_logger = logger;
+				_refCount = 1;
+			}
+
+			/// <summary>
+			/// Creates a client of the given type using this connection
+			/// </summary>
+			/// <typeparam name="T">The client type</typeparam>
+			/// <returns>Client of the given type</returns>
+			public T CreateClient<T>() where T : ClientBase => (T)CreateClient(typeof(T));
+
+			/// <summary>
+			/// Creates a client of an arbitrary type using the current channel
+			/// </summary>
+			/// <param name="clientType">The client type</param>
+			/// <returns></returns>
+			public ClientBase CreateClient(Type clientType)
+			{
+				lock (_clients)
+				{
+					ClientBase? client;
+					if (!_clients.TryGetValue(clientType, out client))
+					{
+						client = (ClientBase)Activator.CreateInstance(clientType, Channel)!;
+					}
+					return client;
+				}
+			}
+
+			/// <summary>
+			/// Attempts to increment the reference count for this subconnection. Fails if it's already zero.
+			/// </summary>
+			/// <returns>True if a reference was added</returns>
+			public bool TryAddRef()
+			{
+				for (; ; )
+				{
+					int initialRefCount = _refCount;
+					if (initialRefCount == 0)
+					{
+						return false;
+					}
+					if (Interlocked.CompareExchange(ref _refCount, initialRefCount + 1, initialRefCount) == initialRefCount)
+					{
+						return true;
+					}
+				}
+			}
+
+			/// <summary>
+			/// Release a reference to this connection
+			/// </summary>
+			public void Release()
+			{
+				int newRefCount = Interlocked.Decrement(ref _refCount);
+				if (_disposing)
+				{
+					_logger.LogInformation("Connection {ConnectionId}: Decrementing refcount to rpc server {ServerName} ({RefCount} rpc(s) in flight)", ConnectionId, Name, newRefCount);
+				}
+				if (newRefCount == 0)
+				{
+					_disposedTaskSource.SetResult(true);
+				}
+			}
+
+			/// <summary>
+			/// Dispose of this connection, waiting for all references to be released first
+			/// </summary>
+			/// <returns>Async task</returns>
+			public ValueTask DisposeAsync()
+			{
+				_logger.LogInformation("Connection {ConnectionId}: Disposing sub-connection", ConnectionId);
+				_disposingTaskSource.TrySetResult(true);
+				_disposing = true;
+				Release();
+				return new ValueTask(_disposedTaskSource.Task);
+			}
+		}
+
+		/// <summary>
+		/// Reference to a connection instance
+		/// </summary>
+		class RpcClientRef<TClient> : IRpcClientRef<TClient> where TClient : ClientBase<TClient>
+		{
+			RpcSubConnection _subConnection;
+
+			public RpcClientRef(RpcSubConnection subConnection)
+			{
+				_subConnection = subConnection;
+				Client = _subConnection.CreateClient<TClient>();
+			}
+
+			public void Dispose()
+			{
+				if (_subConnection != null)
+				{
+					_subConnection.Release();
+					_subConnection = null!;
+				}
+			}
+
+			public GrpcChannel Channel => _subConnection.Channel;
+
+			public TClient Client { get; }
+
+			public Task DisposingTask => _subConnection.DisposingTask;
+		}
+
+		private readonly Func<CancellationToken, Task<GrpcChannel>> _createGrpcChannelAsync;
+		private readonly TaskCompletionSource<bool> _stoppingTaskSource = new TaskCompletionSource<bool>(TaskCreationOptions.RunContinuationsAsynchronously);
+		private TaskCompletionSource<RpcSubConnection> _subConnectionTaskSource = new TaskCompletionSource<RpcSubConnection>(TaskCreationOptions.RunContinuationsAsynchronously);
+		private Task? _backgroundTask;
+		private bool _healthy;
+		private readonly ILogger _logger;
+
+		public bool Healthy => _healthy;
+		public ILogger Logger => _logger;
+
+		/// <summary>
+		/// Constructor
+		/// </summary>
+		/// <param name="createGrpcChannelAsync">Factory method for creating new GRPC channels</param>
+		/// <param name="logger">Logger instance</param>
+		public RpcConnection(Func<CancellationToken, Task<GrpcChannel>> createGrpcChannelAsync, ILogger logger)
+		{
+			_createGrpcChannelAsync = createGrpcChannelAsync;
+			_logger = logger;
+
+			_backgroundTask = Task.Run(() => ExecuteAsync());
+		}
+
+#pragma warning disable VSTHRD002 // Synchronously waiting on tasks or awaiters may cause deadlocks.
+		/// <summary>
+		/// Attempts to get a client reference, returning immediately if there's not one available
+		/// </summary>
+		/// <returns></returns>
+		public IRpcClientRef<TClient>? TryGetClientRef<TClient>() where TClient : ClientBase<TClient>
+		{
+			Task<RpcSubConnection> subConnectionTask = _subConnectionTaskSource.Task;
+			if (subConnectionTask.IsCompleted)
+			{
+				RpcSubConnection defaultClient = subConnectionTask.Result;
+				if (defaultClient.TryAddRef())
+				{
+					return new RpcClientRef<TClient>(defaultClient);
+				}
+			}
+			return null;
+		}
+#pragma warning restore VSTHRD002
+
+		/// <summary>
+		/// Obtains a new client reference object
+		/// </summary>
+		/// <typeparam name="TClient">Type of the client</typeparam>
+		/// <param name="cancellationToken">Cancellation token for this request</param>
+		/// <returns>New client reference</returns>
+		public async Task<IRpcClientRef<TClient>> GetClientRefAsync<TClient>(CancellationToken cancellationToken) where TClient : ClientBase<TClient>
+		{
+			for (; ; )
+			{
+				// Print a message while we're waiting for a client to be available
+				if (!_subConnectionTaskSource.Task.IsCompleted)
+				{
+					Stopwatch timer = Stopwatch.StartNew();
+					while (!_subConnectionTaskSource.Task.IsCompleted)
+					{
+						Task delay = Task.Delay(TimeSpan.FromSeconds(30.0), cancellationToken);
+						if (await Task.WhenAny(delay, _subConnectionTaskSource.Task, _stoppingTaskSource.Task) == delay)
+						{
+							await delay; // Allow the cancellation token to throw
+							_logger.LogInformation("Thread stalled for {Time:0.0}s while waiting for IRpcClientRef.\nStack trace:\n{StackTrace}", timer.Elapsed.TotalSeconds, Environment.StackTrace);
+						}
+						if (_stoppingTaskSource.Task.IsCompleted)
+						{
+							throw new OperationCanceledException("RpcConnection is being terminated");
+						}
+					}
+				}
+
+				// Get the new task and increase the refcount
+				RpcSubConnection subConnection = await _subConnectionTaskSource.Task;
+				if (subConnection.TryAddRef())
+				{
+					return new RpcClientRef<TClient>(subConnection);
+				}
+			}
+		}
+
+		/// <summary>
+		/// Background task to monitor the server state
+		/// </summary>
+		/// <returns>Async task</returns>
+		async Task ExecuteAsync()
+		{
+			// Counter for the connection number. Used to track log messages through async threads.
+			int connectionId = 0;
+
+			// This task source is marked as completed once a RPC connection receives a notification from the server that it's shutting down.
+			// Another connection will immediately be allocated for other connection requests to use, and the load balancer should route us to
+			// the new server.
+			TaskCompletionSource<bool>? reconnectTaskSource = null;
+
+			// List of connection handling async tasks we're tracking
+			List<Task> tasks = new List<Task>();
+			while (tasks.Count > 0 || !_stoppingTaskSource.Task.IsCompleted)
+			{
+				// Get the task indicating we want a new connection (or create a new one)
+				Task? reconnectTask = null;
+				if (!_stoppingTaskSource.Task.IsCompleted)
+				{
+					if (reconnectTaskSource == null || reconnectTaskSource.Task.IsCompleted)
+					{
+						// Need to avoid lambda capture of the reconnect task source
+						int newConnectionId = ++connectionId;
+						TaskCompletionSource<bool> newReconnectTaskSource = new TaskCompletionSource<bool>(TaskCreationOptions.RunContinuationsAsynchronously);
+						tasks.Add(Task.Run(() => HandleConnectionAsync(newConnectionId, newReconnectTaskSource)));
+						reconnectTaskSource = newReconnectTaskSource;
+					}
+					reconnectTask = reconnectTaskSource.Task;
+				}
+
+				// Build a list of tasks to wait for
+				List<Task> waitTasks = new List<Task>(tasks);
+				if (reconnectTask != null)
+				{
+					waitTasks.Add(reconnectTask);
+				}
+				await Task.WhenAny(waitTasks);
+
+				// Remove any complete tasks, logging any exceptions
+				for (int idx = 0; idx < tasks.Count; idx++)
+				{
+					Task task = tasks[idx];
+					if (task.IsCompleted)
+					{
+						await task;
+						tasks.RemoveAt(idx--);
+					}
+				}
+			}
+		}
+
+		/// <summary>
+		/// Background task to monitor the server state
+		/// </summary>
+		/// <param name="connectionId">The connection id</param>
+		/// <param name="reconnectTaskSource">Task source for reconnecting</param>
+		/// <returns>Async task</returns>
+		async Task HandleConnectionAsync(int connectionId, TaskCompletionSource<bool> reconnectTaskSource)
+		{
+			using IScope scope = GlobalTracer.Instance.BuildSpan("RpcConnection {ConnectionId}").StartActive();
+
+			_logger.LogInformation("Connection {ConnectionId}: Creating connection", connectionId);
+			try
+			{
+				await HandleConnectionInternalAsync(connectionId, reconnectTaskSource);
+			}
+			catch (RpcException ex) when (ex.StatusCode == StatusCode.Unavailable)
+			{
+				_logger.LogInformation("Connection {ConnectionId}: Rpc service is unavailable. Reconnecting.", connectionId);
+			}
+			catch (Exception ex)
+			{
+				_logger.LogError(ex, "Connection {ConnectionId}: Exception handling connection: {Message}", connectionId, ex.Message);
+			}
+			finally
+			{
+				TriggerReconnect(reconnectTaskSource);
+			}
+		}
+
+		/// <summary>
+		/// Background task to monitor the server state
+		/// </summary>
+		/// <returns>Async task</returns>
+		async Task HandleConnectionInternalAsync(int connectionId, TaskCompletionSource<bool> reconnectTaskSource)
+		{
+			using (GrpcChannel channel = await _createGrpcChannelAsync(CancellationToken.None))
+			{
+				HordeRpc.HordeRpcClient client = new HordeRpc.HordeRpcClient(channel);
+
+				RpcSubConnection? subConnection = null;
+				try
+				{
+					while (!reconnectTaskSource.Task.IsCompleted && !_stoppingTaskSource.Task.IsCompleted)
+					{
+						// Client HAS TO BE the one to end the call, by closing the stream. The server must keep receiving requests from us until it does.
+						using (AsyncDuplexStreamingCall<RpcQueryServerStateRequest, RpcQueryServerStateResponse> call = client.QueryServerStateV2())
+						{
+							// Send the name of this agent
+							RpcQueryServerStateRequest request = new RpcQueryServerStateRequest();
+							request.Name = Dns.GetHostName();
+							await call.RequestStream.WriteAsync(request);
+
+							// Read the server info
+							Task<bool> moveNextTask = call.ResponseStream.MoveNext();
+							if (!await moveNextTask)
+							{
+								_logger.LogError("Connection {ConnectionId}: No response from server", connectionId);
+								break;
+							}
+
+							// Read the server response
+							RpcQueryServerStateResponse response = call.ResponseStream.Current;
+							if (!response.Stopping)
+							{
+								// The first time we connect, log the server name and create the subconnection
+								if (subConnection == null)
+								{
+									_logger.LogInformation("Connection {ConnectionId}: Connected to rpc server {ServerName}", connectionId, response.Name);
+									subConnection = new RpcSubConnection(connectionId, response.Name, channel, _logger);
+									_subConnectionTaskSource.SetResult(subConnection);
+								}
+
+								// Set the healthy flag for reporting status to the tray app
+								_healthy = true;
+
+								// Wait for the StoppingTask token to be set, or the server to inform that it's shutting down
+								moveNextTask = call.ResponseStream.MoveNext();
+								await Task.WhenAny(_stoppingTaskSource.Task, Task.Delay(TimeSpan.FromSeconds(45.0)), moveNextTask);
+
+								// Update the response
+								if (moveNextTask.IsCompleted && await moveNextTask)
+								{
+									response = call.ResponseStream.Current;
+								}
+							}
+
+							// If the server is stopping, start reconnecting
+							if (response.Stopping)
+							{
+								_logger.LogInformation("Connection {ConnectionId}: Server is stopping. Triggering reconnect.", connectionId);
+								TriggerReconnect(reconnectTaskSource);
+							}
+
+							// Close the request stream
+							await call.RequestStream.CompleteAsync();
+
+							// Wait for the server to finish posting responses
+							while (await moveNextTask)
+							{
+								response = call.ResponseStream.Current;
+								moveNextTask = call.ResponseStream.MoveNext();
+							}
+						}
+					}
+					_logger.LogInformation("Connection {ConnectionId}: Closing connection to rpc server", connectionId);
+				}
+				finally
+				{
+					_healthy = false;
+
+					if (subConnection != null)
+					{
+						await subConnection.DisposeAsync();
+					}
+				}
+
+				// Wait a few seconds before finishing
+				await Task.Delay(TimeSpan.FromSeconds(5.0));
+			}
+		}
+
+		/// <summary>
+		/// Clears the current subconnection and triggers a reconnect
+		/// </summary>
+		/// <param name="reconnectTaskSource"></param>
+		void TriggerReconnect(TaskCompletionSource<bool> reconnectTaskSource)
+		{
+			if (!reconnectTaskSource.Task.IsCompleted)
+			{
+				// First reset the task source, so that nothing new starts to use the current connection
+				if (_subConnectionTaskSource.Task.IsCompleted)
+				{
+					_subConnectionTaskSource = new TaskCompletionSource<RpcSubConnection>(TaskCreationOptions.RunContinuationsAsynchronously);
+				}
+
+				// Now trigger another connection
+				reconnectTaskSource.SetResult(true);
+			}
+		}
+
+		/// <summary>
+		/// Dispose of this connection
+		/// </summary>
+		/// <returns>Async task</returns>
+		public async ValueTask DisposeAsync()
+		{
+			if (_backgroundTask != null)
+			{
+				_stoppingTaskSource.TrySetResult(true);
+
+				await _backgroundTask;
+				_backgroundTask = null!;
+			}
+		}
+	}
+}
diff --git a/Engine/Source/Programs/Horde/Horde.Agent.Driver/Utility/TempStorage.cs b/Engine/Source/Programs/Horde/Horde.Agent.Driver/Utility/TempStorage.cs
--- a/Engine/Source/Programs/Horde/Horde.Agent.Driver/Utility/TempStorage.cs
+++ b/Engine/Source/Programs/Horde/Horde.Agent.Driver/Utility/TempStorage.cs
@@ -590,7 +590,7 @@ namespace Horde.Storage.Utility
 		/// <param name="logger">Logger for output</param>
 		/// <param name="cancellationToken"></param>
 		/// <returns>The set of files</returns>
-		public static async Task<TempStorageTagManifest> RetrieveTagAsync(JobRpc.JobRpcClient jobRpc, JobId jobId, JobStepId stepId, HttpStorageClientFactory storageClientFactory, string nodeName, string tagName, DirectoryReference manifestDir, ILogger logger, CancellationToken cancellationToken)
+		public static async Task<TempStorageTagManifest> RetrieveTagAsync(IRpcClientRef<JobRpc.JobRpcClient> jobRpc, JobId jobId, JobStepId stepId, HttpStorageClientFactory storageClientFactory, string nodeName, string tagName, DirectoryReference manifestDir, ILogger logger, CancellationToken cancellationToken)
 		{
 			// Try to read the tag set from the local directory
 			FileReference localFileListLocation = GetTagManifestLocation(manifestDir, nodeName, tagName);
@@ -609,7 +609,7 @@ namespace Horde.Storage.Utility
 				artifactRequest.Name = artifactName.ToString();
 				artifactRequest.Type = artifactType.ToString();
 
-				RpcGetJobArtifactResponse artifact = await jobRpc.GetArtifactAsync(artifactRequest, cancellationToken: cancellationToken);
+				RpcGetJobArtifactResponse artifact = await jobRpc.Client.GetArtifactAsync(artifactRequest, cancellationToken: cancellationToken);
 
 				NamespaceId namespaceId = new NamespaceId(artifact.NamespaceId);
 				RefName refName = new RefName(artifact.RefName);
@@ -720,7 +720,7 @@ namespace Horde.Storage.Utility
 		/// <param name="logger">Logger for output</param>
 		/// <param name="cancellationToken"></param>
 		/// <returns>Manifest of the files retrieved</returns>
-		public static async Task<TempStorageBlockManifest> RetrieveBlockAsync(JobRpc.JobRpcClient jobRpc, JobId jobId, JobStepId stepId, HttpStorageClientFactory storageClientFactory, string nodeName, string blockName, DirectoryReference rootDir, DirectoryReference manifestDir, ILogger logger, CancellationToken cancellationToken)
+		public static async Task<TempStorageBlockManifest> RetrieveBlockAsync(IRpcClientRef<JobRpc.JobRpcClient> jobRpc, JobId jobId, JobStepId stepId, HttpStorageClientFactory storageClientFactory, string nodeName, string blockName, DirectoryReference rootDir, DirectoryReference manifestDir, ILogger logger, CancellationToken cancellationToken)
 		{
 			// Get the path to the local manifest
 			FileReference localManifestFile = GetBlockManifestLocation(manifestDir, nodeName, blockName);
@@ -747,7 +747,7 @@ namespace Horde.Storage.Utility
 				artifactRequest.Name = artifactName.ToString();
 				artifactRequest.Type = artifactType.ToString();
 
-				RpcGetJobArtifactResponse artifact = await jobRpc.GetArtifactAsync(artifactRequest, cancellationToken: cancellationToken);
+				RpcGetJobArtifactResponse artifact = await jobRpc.Client.GetArtifactAsync(artifactRequest, cancellationToken: cancellationToken);
 				NamespaceId namespaceId = new NamespaceId(artifact.NamespaceId);
 				RefName refName = new RefName(artifact.RefName);
 
diff --git a/Engine/Source/Programs/Horde/Horde.Agent.Tests/Execution/WorkspaceExecutorTest.cs b/Engine/Source/Programs/Horde/Horde.Agent.Tests/Execution/WorkspaceExecutorTest.cs
--- a/Engine/Source/Programs/Horde/Horde.Agent.Tests/Execution/WorkspaceExecutorTest.cs
+++ b/Engine/Source/Programs/Horde/Horde.Agent.Tests/Execution/WorkspaceExecutorTest.cs
@@ -47,14 +47,14 @@ public sealed class WorkspaceExecutorTest : IAsyncDisposable
 		_server.AddStream(_streamId, "//Foo/Main");
 		_server.AddAgentType(_streamId, AgentType);
 		_server.AddJob(_jobId, _streamId, 1, 0);
-		_session = FakeServerSessionFactory.CreateSession(_server.GetGrpcChannel());
+		_session = FakeServerSessionFactory.CreateSession(_server.GetConnection());
 
 		_autoSdkWorkspace.SetFile(1, "HostWin64/Android/base.h", "base");
 		_workspace.SetFile(1, "main.cpp", "main");
 		_workspace.SetFile(1, "foo/bar/baz.h", "baz");
 
 		RpcBeginBatchResponse batch = new RpcBeginBatchResponse { Change = 1 };
-		JobExecutorOptions executorOptions = new JobExecutorOptions(_session.ServerUrl, _session.WorkingDir, _server.GetGrpcChannel(), null, null!, _jobId, _batchId, batch, null!, new RpcJobOptions());
+		JobExecutorOptions executorOptions = new JobExecutorOptions(_session.ServerUrl, _session.WorkingDir, _session.RpcConnection, null, null!, _jobId, _batchId, batch, null!, new RpcJobOptions());
 		_executor = new(executorOptions, _workspace, null, NullLogger.Instance);
 	}
 
@@ -80,7 +80,7 @@ public sealed class WorkspaceExecutorTest : IAsyncDisposable
 	public async Task RegularAndAutoSdkWorkspaceAsync()
 	{
 		RpcBeginBatchResponse batch = new RpcBeginBatchResponse { Change = 1 };
-		JobExecutorOptions executorOptions = new JobExecutorOptions(_session.ServerUrl, _session.WorkingDir, _session.GrpcChannel, Array.Empty<ProcessToTerminate>(), null!, _jobId, _batchId, batch, null!, new RpcJobOptions());
+		JobExecutorOptions executorOptions = new JobExecutorOptions(_session.ServerUrl, _session.WorkingDir, _session.RpcConnection, Array.Empty<ProcessToTerminate>(), null!, _jobId, _batchId, batch, null!, new RpcJobOptions());
 		using WorkspaceExecutor executor = new(executorOptions, _workspace, _autoSdkWorkspace, NullLogger.Instance);
 
 		await executor.InitializeAsync(_logger, CancellationToken.None);
@@ -94,7 +94,7 @@ public sealed class WorkspaceExecutorTest : IAsyncDisposable
 	public async Task EnvVarsAsync()
 	{
 		RpcBeginBatchResponse batch = new RpcBeginBatchResponse { Change = 1, StreamName = "//UE5/Main" };
-		JobExecutorOptions executorOptions = new JobExecutorOptions(_session.ServerUrl, _session.WorkingDir, _session.GrpcChannel, Array.Empty<ProcessToTerminate>(), null!, _jobId, _batchId, batch, null!, new RpcJobOptions());
+		JobExecutorOptions executorOptions = new JobExecutorOptions(_session.ServerUrl, _session.WorkingDir, _session.RpcConnection, Array.Empty<ProcessToTerminate>(), null!, _jobId, _batchId, batch, null!, new RpcJobOptions());
 		using WorkspaceExecutor executor = new(executorOptions, _workspace, _autoSdkWorkspace, NullLogger.Instance);
 		await executor.InitializeAsync(_logger, CancellationToken.None);
 
@@ -121,7 +121,7 @@ public sealed class WorkspaceExecutorTest : IAsyncDisposable
 		_workspace.SetFile(1000, "New/Feature/Foo.cs", "foo");
 
 		RpcBeginBatchResponse batch = new RpcBeginBatchResponse { Change = 1, PreflightChange = 1000 };
-		JobExecutorOptions executorOptions = new JobExecutorOptions(_session.ServerUrl, _session.WorkingDir, _session.GrpcChannel, Array.Empty<ProcessToTerminate>(), null!, preflightJobId, _batchId, batch, null!, new RpcJobOptions());
+		JobExecutorOptions executorOptions = new JobExecutorOptions(_session.ServerUrl, _session.WorkingDir, _session.RpcConnection, Array.Empty<ProcessToTerminate>(), null!, preflightJobId, _batchId, batch, null!, new RpcJobOptions());
 		using WorkspaceExecutor executor = new(executorOptions, _workspace, null, NullLogger.Instance);
 
 		await executor.InitializeAsync(_logger, CancellationToken.None);
@@ -138,7 +138,7 @@ public sealed class WorkspaceExecutorTest : IAsyncDisposable
 
 		_server.AddJob(noChangeJobId, _streamId, 0, 0);
 		RpcBeginBatchResponse batch = new RpcBeginBatchResponse { };
-		JobExecutorOptions executorOptions = new JobExecutorOptions(_session.ServerUrl, _session.WorkingDir, _session.GrpcChannel, Array.Empty<ProcessToTerminate>(), null!, noChangeJobId, _batchId, batch, null!, new RpcJobOptions());
+		JobExecutorOptions executorOptions = new JobExecutorOptions(_session.ServerUrl, _session.WorkingDir, _session.RpcConnection, Array.Empty<ProcessToTerminate>(), null!, noChangeJobId, _batchId, batch, null!, new RpcJobOptions());
 		using WorkspaceExecutor executor = new(executorOptions, _workspace, null, NullLogger.Instance);
 		await Assert.ThrowsExceptionAsync<WorkspaceMaterializationException>(() => executor.InitializeAsync(_logger, CancellationToken.None));
 	}
diff --git a/Engine/Source/Programs/Horde/Horde.Agent.Tests/Helper.cs b/Engine/Source/Programs/Horde/Horde.Agent.Tests/Helper.cs
--- a/Engine/Source/Programs/Horde/Horde.Agent.Tests/Helper.cs
+++ b/Engine/Source/Programs/Horde/Horde.Agent.Tests/Helper.cs
@@ -7,11 +7,14 @@ using System.Threading.Tasks;
 using EpicGames.Horde.Jobs;
 using Google.Protobuf.WellKnownTypes;
 using Grpc.Core;
+using Grpc.Net.Client;
 using Horde.Agent.Execution;
+using Horde.Agent.Utility;
 using Horde.Common.Rpc;
 using HordeCommon.Rpc;
 using HordeCommon.Rpc.Messages;
 using Microsoft.Extensions.Logging;
+using Microsoft.Extensions.Logging.Abstractions;
 using Microsoft.Extensions.Options;
 
 namespace Horde.Agent.Tests
@@ -35,6 +38,69 @@ namespace Horde.Agent.Tests
 			=> null;
 	}
 
+	class RpcClientRefStub<TClient> : IRpcClientRef<TClient> where TClient : ClientBase<TClient>
+	{
+		public GrpcChannel Channel { get; }
+		public TClient Client { get; }
+		public Task DisposingTask { get; }
+
+		public RpcClientRefStub(GrpcChannel channel, TClient client)
+		{
+			Channel = channel;
+			Client = client;
+			DisposingTask = new TaskCompletionSource<bool>().Task;
+		}
+
+		public void Dispose()
+		{
+		}
+	}
+
+	class RpcConnectionStub : IRpcConnection
+	{
+		private readonly GrpcChannel _grpcChannel;
+		private readonly HordeRpc.HordeRpcClient _hordeRpcClient;
+		private readonly JobRpc.JobRpcClient _jobRpcClient;
+
+		public bool Healthy => true;
+		public ILogger Logger => NullLogger.Instance;
+
+		public RpcConnectionStub(GrpcChannel grpcChannel, HordeRpc.HordeRpcClient hordeRpcClient, JobRpc.JobRpcClient jobRpcClient)
+		{
+			_grpcChannel = grpcChannel;
+			_hordeRpcClient = hordeRpcClient;
+			_jobRpcClient = jobRpcClient;
+		}
+
+		public IRpcClientRef<TClient>? TryGetClientRef<TClient>() where TClient : ClientBase<TClient>
+		{
+			return (IRpcClientRef<TClient>)(object)new RpcClientRefStub<HordeRpc.HordeRpcClient>(_grpcChannel, _hordeRpcClient);
+		}
+
+		public Task<IRpcClientRef<TClient>> GetClientRefAsync<TClient>(CancellationToken cancellationToken) where TClient : ClientBase<TClient>
+		{
+			IRpcClientRef<TClient> rpcClientRefStub;
+			if (typeof(TClient) == typeof(HordeRpc.HordeRpcClient))
+			{
+				rpcClientRefStub = (IRpcClientRef<TClient>)(object)new RpcClientRefStub<HordeRpc.HordeRpcClient>(_grpcChannel, _hordeRpcClient);
+			}
+			else if (typeof(TClient) == typeof(JobRpc.JobRpcClient))
+			{
+				rpcClientRefStub = (IRpcClientRef<TClient>)(object)new RpcClientRefStub<JobRpc.JobRpcClient>(_grpcChannel, _jobRpcClient);
+			}
+			else
+			{
+				throw new NotImplementedException();
+			}
+			return Task.FromResult(rpcClientRefStub);
+		}
+
+		public ValueTask DisposeAsync()
+		{
+			return new ValueTask();
+		}
+	}
+
 	class JobRpcClientStub : JobRpc.JobRpcClient
 	{
 		public readonly Queue<RpcBeginStepResponse> BeginStepResponses = new Queue<RpcBeginStepResponse>();
diff --git a/Engine/Source/Programs/Horde/Horde.Agent.Tests/JsonRpcLoggerTests.cs b/Engine/Source/Programs/Horde/Horde.Agent.Tests/JsonRpcLoggerTests.cs
--- a/Engine/Source/Programs/Horde/Horde.Agent.Tests/JsonRpcLoggerTests.cs
+++ b/Engine/Source/Programs/Horde/Horde.Agent.Tests/JsonRpcLoggerTests.cs
@@ -12,7 +12,6 @@ using EpicGames.Horde.Logs;
 using EpicGames.Horde.Storage;
 using EpicGames.Horde.Storage.Bundles;
 using EpicGames.Horde.Storage.Bundles.V1;
-using Grpc.Net.Client;
 using Horde.Agent.Utility;
 using Microsoft.Extensions.Logging;
 using Microsoft.Extensions.Logging.Abstractions;
@@ -127,7 +126,7 @@ namespace Horde.Agent.Tests
 		{
 			public IBlobRef? Target { get; private set; }
 
-			public FakeJsonRpcLoggerBackend(GrpcChannel connection, LogId logId, JobId? jobId, JobStepBatchId? batchId, JobStepId? stepId, IStorageClient store, ILogger logger)
+			public FakeJsonRpcLoggerBackend(IRpcConnection connection, LogId logId, JobId? jobId, JobStepBatchId? batchId, JobStepId? stepId, IStorageClient store, ILogger logger)
 				: base(connection, logId, jobId, batchId, stepId, store, logger)
 			{
 			}
diff --git a/Engine/Source/Programs/Horde/Horde.Agent.Tests/WorkerServiceTest.cs b/Engine/Source/Programs/Horde/Horde.Agent.Tests/WorkerServiceTest.cs
--- a/Engine/Source/Programs/Horde/Horde.Agent.Tests/WorkerServiceTest.cs
+++ b/Engine/Source/Programs/Horde/Horde.Agent.Tests/WorkerServiceTest.cs
@@ -148,7 +148,6 @@ namespace Horde.Agent.Tests
 			}
 		}
 
-		[Ignore]
 		[TestMethod]
 		public async Task AbortExecuteJobTestAsync()
 		{
@@ -165,9 +164,10 @@ namespace Horde.Agent.Tests
 			executeJobTask.Workspace = new RpcAgentWorkspace();
 
 			JobRpcClientStub client = new JobRpcClientStub(NullLogger.Instance);
+			await using RpcConnectionStub rpcConnection = new RpcConnectionStub(null!, null!, client);
 
 			await using FakeHordeRpcServer fakeServer = new();
-			await using ISession session = FakeServerSessionFactory.CreateSession(null!);
+			await using ISession session = FakeServerSessionFactory.CreateSession(rpcConnection);
 
 			client.BeginStepResponses.Enqueue(new RpcBeginStepResponse { Name = "stepName1", StepId = _stepId1.ToString() });
 			client.BeginStepResponses.Enqueue(new RpcBeginStepResponse { Name = "stepName2", StepId = _stepId2.ToString() });
@@ -202,7 +202,6 @@ namespace Horde.Agent.Tests
 			Assert.AreEqual(JobStepState.Completed, (JobStepState)client.UpdateStepRequests[2].State);
 		}
 
-		[Ignore]
 		[TestMethod]
 		public async Task PollForStepAbortFailureTestAsync()
 		{
@@ -219,6 +218,7 @@ namespace Horde.Agent.Tests
 			jobHandler._stepAbortPollInterval = TimeSpan.FromMilliseconds(5);
 
 			JobRpcClientStub client = new JobRpcClientStub(NullLogger.Instance);
+			await using RpcConnectionStub rpcConnection = new RpcConnectionStub(null!, null!, client);
 
 			int c = 0;
 			client._getStepFunc = (request) =>
@@ -236,7 +236,7 @@ namespace Horde.Agent.Tests
 			using CancellationTokenSource stepCancelSource = new CancellationTokenSource();
 			TaskCompletionSource<bool> stepFinishedSource = new TaskCompletionSource<bool>();
 
-			await jobHandler.PollForStepAbortAsync(null!, _jobId, _batchId, _stepId2, stepCancelSource, stepFinishedSource.Task, NullLogger.Instance, stepPollCancelSource.Token);
+			await jobHandler.PollForStepAbortAsync(rpcConnection, _jobId, _batchId, _stepId2, stepCancelSource, stepFinishedSource.Task, NullLogger.Instance, stepPollCancelSource.Token);
 			Assert.IsTrue(stepCancelSource.IsCancellationRequested);
 		}
 
@@ -257,7 +257,7 @@ namespace Horde.Agent.Tests
 			cts.CancelAfter(20000);
 
 			await using FakeHordeRpcServer fakeServer = new();
-			await using ISession session = FakeServerSessionFactory.CreateSession(fakeServer.GetGrpcChannel());
+			await using ISession session = FakeServerSessionFactory.CreateSession(fakeServer.GetConnection());
 
 			LeaseManager manager = new LeaseManager(session, serviceProvider);
 
@@ -276,16 +276,16 @@ namespace Horde.Agent.Tests
 
 		public Task<ISession> CreateAsync(CancellationToken cancellationToken)
 		{
-			return Task.FromResult(CreateSession(_fakeServer.GetGrpcChannel()));
+			return Task.FromResult(CreateSession(_fakeServer.GetConnection()));
 		}
 
-		public static ISession CreateSession(GrpcChannel grpcChannel)
+		public static ISession CreateSession(IRpcConnection rpcConnection)
 		{
 			Mock<ISession> fakeSession = new Mock<ISession>(MockBehavior.Strict);
 			fakeSession.Setup(x => x.ServerUrl).Returns(new Uri("https://localhost:9999"));
 			fakeSession.Setup(x => x.AgentId).Returns(new EpicGames.Horde.Agents.AgentId("LocalAgent"));
 			fakeSession.Setup(x => x.SessionId).Returns(new EpicGames.Horde.Agents.Sessions.SessionId(default));
-			fakeSession.Setup(x => x.GrpcChannel).Returns(grpcChannel);
+			fakeSession.Setup(x => x.RpcConnection).Returns(rpcConnection);
 			fakeSession.Setup(x => x.DisposeAsync()).Returns(new ValueTask());
 			fakeSession.Setup(x => x.WorkingDir).Returns(DirectoryReference.Combine(DirectoryReference.GetCurrentDirectory(), Guid.NewGuid().ToString()));
 			return fakeSession.Object;
@@ -305,12 +305,15 @@ namespace Horde.Agent.Tests
 
 		private readonly Dictionary<StreamId, RpcGetStreamResponse> _streamIdToStreamResponse = new();
 		private readonly Dictionary<JobId, RpcGetJobResponse> _jobIdToJobResponse = new();
+		private readonly Mock<IRpcClientRef<HordeRpc.HordeRpcClient>> _mockClientRef;
+		private readonly Mock<IRpcConnection> _mockConnection;
 		private readonly ILogger<FakeHordeRpcServer> _logger;
 		public readonly TaskCompletionSource<bool> CreateSessionReceived = new();
 		public readonly TaskCompletionSource<bool> UpdateSessionReceived = new();
 
+		private readonly RpcConnectionStub _connection;
 		private readonly GrpcChannel _grpcChannel;
-//		private readonly FakeJobRpcClient _client;
+		private readonly FakeJobRpcClient _client;
 
 		private class FakeHordeRpcClient : HordeRpc.HordeRpcClient
 		{
@@ -361,9 +364,26 @@ namespace Horde.Agent.Tests
 		{
 			_serverName = "FakeServer";
 			_logger = NullLogger<FakeHordeRpcServer>.Instance;
-//			FakeHordeRpcClient hordeClient = new FakeHordeRpcClient(this);
-//			_client = new FakeJobRpcClient(this);
+			FakeHordeRpcClient hordeClient = new FakeHordeRpcClient(this);
+			_client = new FakeJobRpcClient(this);
+			_connection = new RpcConnectionStub(null!, hordeClient, _client);
 			_grpcChannel = GrpcChannel.ForAddress(new Uri("http://horde-agent-test"), new GrpcChannelOptions());
+
+			_mockClientRef = new Mock<IRpcClientRef<HordeRpc.HordeRpcClient>>();
+			_mockClientRef
+				.Setup(m => m.Client)
+				.Returns(() => hordeClient);
+
+			_mockConnection = new(MockBehavior.Strict);
+			_mockConnection
+				.Setup(m => m.TryGetClientRef<HordeRpc.HordeRpcClient>())
+				.Returns(() => _mockClientRef.Object);
+			_mockConnection
+				.Setup(m => m.GetClientRefAsync<HordeRpc.HordeRpcClient>(It.IsAny<CancellationToken>()))
+				.Returns(() => Task.FromResult(_mockClientRef.Object));
+			_mockConnection
+				.Setup(m => m.DisposeAsync())
+				.Returns(() => new ValueTask());
 		}
 
 		public void AddTestLease(string leaseId)
@@ -431,6 +451,11 @@ namespace Horde.Agent.Tests
 			};
 		}
 
+		public IRpcConnection GetConnection()
+		{
+			return _connection;
+		}
+
 		public GrpcChannel GetGrpcChannel()
 		{
 			return _grpcChannel;
@@ -519,8 +544,9 @@ namespace Horde.Agent.Tests
 				() => { });
 		}
 
-		public ValueTask DisposeAsync()
+		public async ValueTask DisposeAsync()
 		{
+			await _connection.DisposeAsync();
 			_grpcChannel.Dispose();
 
 			foreach (RpcGetStreamResponse stream in _streamIdToStreamResponse.Values)
@@ -533,8 +559,6 @@ namespace Horde.Agent.Tests
 					}
 				}
 			}
-
-			return default;
 		}
 	}
 
diff --git a/Engine/Source/Programs/Horde/Horde.Agent/Commands/Execution/ExecuteJobCommand.cs b/Engine/Source/Programs/Horde/Horde.Agent/Commands/Execution/ExecuteJobCommand.cs
--- a/Engine/Source/Programs/Horde/Horde.Agent/Commands/Execution/ExecuteJobCommand.cs
+++ b/Engine/Source/Programs/Horde/Horde.Agent/Commands/Execution/ExecuteJobCommand.cs
@@ -4,9 +4,9 @@ using EpicGames.Core;
 using EpicGames.Horde.Agents;
 using EpicGames.Horde.Agents.Leases;
 using EpicGames.Horde.Agents.Sessions;
-using Grpc.Net.Client;
 using Horde.Agent.Leases.Handlers;
 using Horde.Agent.Services;
+using Horde.Agent.Utility;
 using HordeCommon.Rpc.Tasks;
 using Microsoft.Extensions.Logging;
 using Microsoft.Extensions.Options;
@@ -47,7 +47,7 @@ namespace Horde.Agent.Commands.Execution
 			ServerProfile serverProfile = _agentSettings.GetCurrentServerProfile();
 			ExecuteJobTask executeTask = ExecuteJobTask.Parser.ParseFrom(Convert.FromBase64String(Task));
 
-			using GrpcChannel rpcConnection = await _grpcService.CreateGrpcChannelAsync(executeTask.Token, CancellationToken.None);
+			await using RpcConnection rpcConnection = new RpcConnection(ctx => _grpcService.CreateGrpcChannelAsync(executeTask.Token, ctx), logger);
 			await using Session session = new Session(serverProfile.Url, AgentId, SessionId, executeTask.Token, rpcConnection, WorkingDir);
 
 			await _jobHandler.ExecuteInternalAsync(session, LeaseId, executeTask, logger, CancellationToken.None);
diff --git a/Engine/Source/Programs/Horde/Horde.Agent/Leases/Handlers/ConformHandler.cs b/Engine/Source/Programs/Horde/Horde.Agent/Leases/Handlers/ConformHandler.cs
--- a/Engine/Source/Programs/Horde/Horde.Agent/Leases/Handlers/ConformHandler.cs
+++ b/Engine/Source/Programs/Horde/Horde.Agent/Leases/Handlers/ConformHandler.cs
@@ -78,9 +78,7 @@ namespace Horde.Agent.Leases.Handlers
 				request.Workspaces.AddRange(pendingWorkspaces);
 				request.RemoveUntrackedFiles = removeUntrackedFiles;
 
-				HordeRpc.HordeRpcClient hordeRpc = new HordeRpc.HordeRpcClient(session.GrpcChannel);
-
-				RpcUpdateAgentWorkspacesResponse response = await hordeRpc.UpdateAgentWorkspacesAsync(request, cancellationToken: cancellationToken);
+				RpcUpdateAgentWorkspacesResponse response = await session.RpcConnection.InvokeAsync((HordeRpc.HordeRpcClient x) => x.UpdateAgentWorkspacesAsync(request, null, null, cancellationToken), cancellationToken);
 				if (!response.Retry)
 				{
 					logger.LogInformation("Conform finished");
diff --git a/Engine/Source/Programs/Horde/Horde.Agent/Leases/Handlers/JobHandler.cs b/Engine/Source/Programs/Horde/Horde.Agent/Leases/Handlers/JobHandler.cs
--- a/Engine/Source/Programs/Horde/Horde.Agent/Leases/Handlers/JobHandler.cs
+++ b/Engine/Source/Programs/Horde/Horde.Agent/Leases/Handlers/JobHandler.cs
@@ -157,8 +157,6 @@ namespace Horde.Agent.Leases.Handlers
 
 		internal async Task<LeaseResult> ExecuteInternalAsync(ISession session, LeaseId leaseId, JobTaskInfo executeTask, ILogger localLogger, CancellationToken cancellationToken)
 		{
-			JobRpc.JobRpcClient jobRpc = new JobRpc.JobRpcClient(session.GrpcChannel);
-
 			// Create a storage client for this session
 			RpcJobOptions jobOptions = executeTask.JobOptions;
 			await using IServerLogger logger = _serverLoggerFactory.CreateLogger(session, executeTask.LogId, localLogger, executeTask.JobId, executeTask.BatchId, null, null);
@@ -172,23 +170,23 @@ namespace Horde.Agent.Leases.Handlers
 			logger.LogInformation("Executor: {Name}", jobOptions.Executor);
 
 			// Start executing the current batch
-			RpcBeginBatchResponse batch = await jobRpc.BeginBatchAsync(new RpcBeginBatchRequest(executeTask.JobId, executeTask.BatchId, leaseId), cancellationToken: cancellationToken);
+			RpcBeginBatchResponse batch = await session.RpcConnection.InvokeAsync<JobRpc.JobRpcClient, RpcBeginBatchResponse>(x => x.BeginBatchAsync(new RpcBeginBatchRequest(executeTask.JobId, executeTask.BatchId, leaseId), null, null, cancellationToken), cancellationToken);
 			try
 			{
-				JobExecutorOptions options = new JobExecutorOptions(session.ServerUrl, session.WorkingDir, session.GrpcChannel, _driverSettings.ProcessesToTerminate, _serverStorageFactory, executeTask.JobId, executeTask.BatchId, batch, executeTask.Token, jobOptions);
+				JobExecutorOptions options = new JobExecutorOptions(session.ServerUrl, session.WorkingDir, session.RpcConnection, _driverSettings.ProcessesToTerminate, _serverStorageFactory, executeTask.JobId, executeTask.BatchId, batch, executeTask.Token, jobOptions);
 				await ExecuteBatchAsync(session, leaseId, executeTask.Workspace, executeTask.AutoSdkWorkspace, options, logger, localLogger, cancellationToken);
 			}
 			catch (Exception ex)
 			{
 				if (cancellationToken.IsCancellationRequested && ex.IsCancellationException())
 				{
-					if (session.ConnectivityState == ConnectivityState.TransientFailure)
+					if (session.RpcConnection.Healthy)
 					{
-						logger.LogError(ex, "Connection to the server was lost; step aborted.");
+						logger.LogInformation(ex, "Step was aborted");
 					}
 					else
 					{
-						logger.LogInformation(ex, "Step was aborted");
+						logger.LogError(ex, "Connection to the server was lost; step aborted.");
 					}
 					throw;
 				}
@@ -206,7 +204,7 @@ namespace Horde.Agent.Leases.Handlers
 			}
 
 			// Mark the batch as complete
-			await jobRpc.FinishBatchAsync(new RpcFinishBatchRequest(executeTask.JobId, executeTask.BatchId, leaseId), cancellationToken: cancellationToken);
+			await session.RpcConnection.InvokeAsync((JobRpc.JobRpcClient x) => x.FinishBatchAsync(new RpcFinishBatchRequest(executeTask.JobId, executeTask.BatchId, leaseId), null, null, cancellationToken), cancellationToken);
 			logger.LogInformation("Done.");
 
 			return LeaseResult.Success;
@@ -217,7 +215,7 @@ namespace Horde.Agent.Leases.Handlers
 		/// </summary>
 		async Task ExecuteBatchAsync(ISession session, LeaseId leaseId, RpcAgentWorkspace workspaceInfo, RpcAgentWorkspace? autoSdkWorkspaceInfo, JobExecutorOptions options, ILogger logger, ILogger localLogger, CancellationToken cancellationToken)
 		{
-			JobRpc.JobRpcClient rpcClient = new JobRpc.JobRpcClient(session.GrpcChannel);
+			IRpcConnection rpcClient = session.RpcConnection;
 
 			// Create an executor for this job
 			string executorName = String.IsNullOrEmpty(options.JobOptions.Executor) ? _driverSettings.Executor : options.JobOptions.Executor;
@@ -249,7 +247,7 @@ namespace Horde.Agent.Leases.Handlers
 				for (; ; )
 				{
 					// Get the next step to execute
-					RpcBeginStepResponse stepResponse = await rpcClient.BeginStepAsync(new RpcBeginStepRequest(options.JobId, options.BatchId, leaseId), cancellationToken: cancellationToken);
+					RpcBeginStepResponse stepResponse = await rpcClient.InvokeAsync((JobRpc.JobRpcClient x) => x.BeginStepAsync(new RpcBeginStepRequest(options.JobId, options.BatchId, leaseId), null, null, cancellationToken), cancellationToken);
 					if (stepResponse.State == RpcBeginStepResponse.Types.Result.Waiting)
 					{
 						logger.LogInformation("Waiting for dependency to be ready");
@@ -349,7 +347,7 @@ namespace Horde.Agent.Leases.Handlers
 
 						// Update the server with the outcome from the step
 						logger.LogInformation("Marking step as complete (Outcome={Outcome}, State={StepState})", stepOutcome, stepState);
-						await rpcClient.UpdateStepAsync(new RpcUpdateStepRequest(options.JobId, options.BatchId, step.StepId, stepState, stepOutcome), cancellationToken: cancellationToken);
+						await rpcClient.InvokeAsync((JobRpc.JobRpcClient x) => x.UpdateStepAsync(new RpcUpdateStepRequest(options.JobId, options.BatchId, step.StepId, stepState, stepOutcome), null, null, cancellationToken), cancellationToken);
 					}
 
 					// Print the finishing state
@@ -361,15 +359,19 @@ namespace Horde.Agent.Leases.Handlers
 			{
 				if (cancellationToken.IsCancellationRequested && ex.IsCancellationException())
 				{
-					if (session.ConnectivityState == ConnectivityState.TransientFailure)
+					if (session.RpcConnection.Healthy)
 					{
-						logger.LogError(ex, "Exception while executing batch: {Ex}", ex);
+						logger.LogError("Lease was aborted");
 					}
 					else
 					{
-						logger.LogError("Lease was aborted");
+						logger.LogError("Connection to the server was lost; lease aborted.");
 					}
 				}
+				else
+				{
+					logger.LogError(ex, "Exception while executing batch: {Ex}", ex);
+				}
 			}
 
 			// Terminate any processes which are still running
@@ -427,14 +429,14 @@ namespace Horde.Agent.Leases.Handlers
 			}
 		}
 
-		internal async Task PollForStepAbortAsync(JobRpc.JobRpcClient rpcClient, JobId jobId, JobStepBatchId batchId, JobStepId stepId, CancellationTokenSource stepCancelSource, Task finishedTask, ILogger leaseLogger, CancellationToken cancellationToken)
+		internal async Task PollForStepAbortAsync(IRpcConnection rpcClient, JobId jobId, JobStepBatchId batchId, JobStepId stepId, CancellationTokenSource stepCancelSource, Task finishedTask, ILogger leaseLogger, CancellationToken cancellationToken)
 		{
 			while (!finishedTask.IsCompleted)
 			{
 				TimeSpan waitTime = _stepAbortPollInterval;
 				try
 				{
-					RpcGetStepResponse res = await rpcClient.GetStepAsync(new RpcGetStepRequest(jobId, batchId, stepId), cancellationToken: cancellationToken);
+					RpcGetStepResponse res = await rpcClient.InvokeAsync((JobRpc.JobRpcClient x) => x.GetStepAsync(new RpcGetStepRequest(jobId, batchId, stepId), null, null, cancellationToken), cancellationToken);
 					if (res.AbortRequested)
 					{
 						leaseLogger.LogDebug("Step was aborted by server (JobId={JobId} BatchId={BatchId} StepId={StepId})", jobId, batchId, stepId);
diff --git a/Engine/Source/Programs/Horde/Horde.Agent/Leases/Handlers/UpgradeHandler.cs b/Engine/Source/Programs/Horde/Horde.Agent/Leases/Handlers/UpgradeHandler.cs
--- a/Engine/Source/Programs/Horde/Horde.Agent/Leases/Handlers/UpgradeHandler.cs
+++ b/Engine/Source/Programs/Horde/Horde.Agent/Leases/Handlers/UpgradeHandler.cs
@@ -44,9 +44,8 @@ namespace Horde.Agent.Leases.Handlers
 
 				// Download the new software
 				FileInfo outputFile = new FileInfo(Path.Combine(upgradeDir.FullName, "Agent.zip"));
-
-				HordeRpc.HordeRpcClient hordeRpc = new HordeRpc.HordeRpcClient(session.GrpcChannel);
-				using (AsyncServerStreamingCall<RpcDownloadSoftwareResponse> cursor = hordeRpc.DownloadSoftware(new RpcDownloadSoftwareRequest(requiredVersion), null, null, cancellationToken))
+				using (IRpcClientRef<HordeRpc.HordeRpcClient> rpcClientRef = await session.RpcConnection.GetClientRefAsync<HordeRpc.HordeRpcClient>(cancellationToken))
+				using (AsyncServerStreamingCall<RpcDownloadSoftwareResponse> cursor = rpcClientRef.Client.DownloadSoftware(new RpcDownloadSoftwareRequest(requiredVersion), null, null, cancellationToken))
 				{
 					await using Stream outputStream = outputFile.Open(FileMode.Create);
 					while (await cursor.ResponseStream.MoveNext(cancellationToken))
diff --git a/Engine/Source/Programs/Horde/Horde.Agent/Leases/LeaseManager.cs b/Engine/Source/Programs/Horde/Horde.Agent/Leases/LeaseManager.cs
--- a/Engine/Source/Programs/Horde/Horde.Agent/Leases/LeaseManager.cs
+++ b/Engine/Source/Programs/Horde/Horde.Agent/Leases/LeaseManager.cs
@@ -6,8 +6,8 @@ using EpicGames.Horde.Agents.Leases;
 using EpicGames.Perforce.Managed;
 using Google.Protobuf.WellKnownTypes;
 using Grpc.Core;
-using Grpc.Net.Client;
 using Horde.Agent.Services;
+using Horde.Agent.Utility;
 using HordeCommon.Rpc;
 using HordeCommon.Rpc.Messages;
 using Microsoft.Extensions.DependencyInjection;
@@ -229,7 +229,7 @@ namespace Horde.Agent.Leases
 
 		async Task<SessionResult> HandleSessionAsync(bool shutdownAfterFinishedLease, CancellationToken stoppingToken)
 		{
-			HordeRpc.HordeRpcClient hordeRpc = new HordeRpc.HordeRpcClient(_session.GrpcChannel);
+			IRpcConnection rpcCon = _session.RpcConnection;
 
 			// Terminate any remaining child processes from other instances
 			ProcessUtils.TerminateProcesses(x => x.IsUnderDirectory(_session.WorkingDir), _logger, stoppingToken);
@@ -242,7 +242,7 @@ namespace Horde.Agent.Leases
 			await using BackgroundTask updateCapsTask = BackgroundTask.StartNew(ctx => UpdateCapabilitiesBackgroundAsync(_session.WorkingDir, ctx));
 
 			// Run another background task to send telemetry data
-			await using BackgroundTask telemetryTask = BackgroundTask.StartNew(ctx => SendTelemetryAsync(hordeRpc, ctx));
+			await using BackgroundTask telemetryTask = BackgroundTask.StartNew(ctx => SendTelemetryAsync(_session.RpcConnection, ctx));
 
 			// Loop until we're ready to exit
 			Stopwatch updateCapabilitiesTimer = Stopwatch.StartNew();
@@ -308,7 +308,7 @@ namespace Horde.Agent.Leases
 				using (stopping ? (CancellationTokenRegistration?)null : stoppingToken.Register(() => _updateLeasesEvent.Set()))
 				{
 					// Update the state with the server
-					RpcUpdateSessionResponse? updateSessionResponse = await UpdateSessionAsync(_session.GrpcChannel, updateSessionRequest, waitTask);
+					RpcUpdateSessionResponse? updateSessionResponse = await UpdateSessionAsync(_session.RpcConnection, updateSessionRequest, waitTask);
 
 					lock (_lockObject)
 					{
@@ -369,7 +369,7 @@ namespace Horde.Agent.Leases
 					}
 
 					// Update the current status
-					if (_session.ConnectivityState == ConnectivityState.TransientFailure)
+					if (!rpcCon.Healthy)
 					{
 						_statusService.Set(false, _activeLeases.Count, "Attempting to connect to server...");
 					}
@@ -413,18 +413,18 @@ namespace Horde.Agent.Leases
 		/// <summary>
 		/// Wrapper for <see cref="UpdateSessionInternalAsync"/> which filters/logs exceptions
 		/// </summary>
-		/// <param name="grpcChannel">The RPC client connection</param>
+		/// <param name="rpcConnection">The RPC client connection</param>
 		/// <param name="updateSessionRequest">The session update request</param>
 		/// <param name="waitTask">Task which can be used to jump out of the update early</param>
 		/// <returns>Response from the call</returns>
-		async Task<RpcUpdateSessionResponse?> UpdateSessionAsync(GrpcChannel grpcChannel, RpcUpdateSessionRequest updateSessionRequest, Task waitTask)
+		async Task<RpcUpdateSessionResponse?> UpdateSessionAsync(IRpcConnection rpcConnection, RpcUpdateSessionRequest updateSessionRequest, Task waitTask)
 		{
-			HordeRpc.HordeRpcClient hordeRpc = new HordeRpc.HordeRpcClient(grpcChannel);
+			using IRpcClientRef<HordeRpc.HordeRpcClient> rpcClientRef = await rpcConnection.GetClientRefAsync<HordeRpc.HordeRpcClient>(CancellationToken.None);
 
 			RpcUpdateSessionResponse? updateSessionResponse = null;
 			try
 			{
-				updateSessionResponse = await UpdateSessionInternalAsync(hordeRpc, updateSessionRequest, waitTask);
+				updateSessionResponse = await UpdateSessionInternalAsync(rpcClientRef.Client, updateSessionRequest, waitTask);
 				_updateSessionFailures = 0;
 			}
 			catch (RpcException ex)
@@ -593,7 +593,7 @@ namespace Horde.Agent.Leases
 		/// <summary>
 		/// Periodically send agent telemetry to the server
 		/// </summary>
-		async Task SendTelemetryAsync(HordeRpc.HordeRpcClient hordeRpc, CancellationToken cancellationToken)
+		async Task SendTelemetryAsync(IRpcConnection rpcConnection, CancellationToken cancellationToken)
 		{
 			while (!cancellationToken.IsCancellationRequested)
 			{
@@ -619,7 +619,8 @@ namespace Horde.Agent.Leases
 							request.UsedRam = memMetrics.Used / 1024;
 						}
 
-						await hordeRpc.UploadTelemetryAsync(request, cancellationToken: cancellationToken);
+						CallOptions callOptions = new CallOptions(cancellationToken: cancellationToken);
+						await rpcConnection.InvokeAsync<HordeRpc.HordeRpcClient, Empty>(client => client.UploadTelemetryAsync(request, callOptions), cancellationToken);
 					}
 				}
 				catch (OperationCanceledException) when (cancellationToken.IsCancellationRequested)
diff --git a/Engine/Source/Programs/Horde/Horde.Agent/Services/GrpcService.cs b/Engine/Source/Programs/Horde/Horde.Agent/Services/GrpcService.cs
--- a/Engine/Source/Programs/Horde/Horde.Agent/Services/GrpcService.cs
+++ b/Engine/Source/Programs/Horde/Horde.Agent/Services/GrpcService.cs
@@ -8,7 +8,6 @@ using System.Net.Sockets;
 using Grpc.Core;
 using Grpc.Core.Interceptors;
 using Grpc.Net.Client;
-using Grpc.Net.Client.Configuration;
 using Horde.Agent.Utility;
 using Microsoft.Extensions.Logging;
 using Microsoft.Extensions.Options;
@@ -134,20 +133,6 @@ namespace Horde.Agent.Services
 				}
 			}
 
-			ServiceConfig serviceConfig = new ServiceConfig();
-			serviceConfig.MethodConfigs.Add(new MethodConfig
-			{
-				Names = { MethodName.Default },
-				RetryPolicy = new RetryPolicy
-				{
-					MaxAttempts = 3,
-					InitialBackoff = TimeSpan.FromSeconds(1),
-					MaxBackoff = TimeSpan.FromSeconds(10),
-					BackoffMultiplier = 2.0,
-					RetryableStatusCodes = { StatusCode.Unavailable },
-				}
-			});
-
 			_logger.LogInformation("Connecting to rpc server {BaseUrl}", serverUri);
 			return GrpcChannel.ForAddress(serverUri, new GrpcChannelOptions
 			{
@@ -156,8 +141,7 @@ namespace Horde.Agent.Services
 				MaxSendMessageSize = 1024 * 1024 * 1024, // 1 GB
 				LoggerFactory = _loggerFactory,
 				HttpClient = httpClient,
-				DisposeHttpClient = true,
-				ServiceConfig = serviceConfig
+				DisposeHttpClient = true
 			});
 		}
 
diff --git a/Engine/Source/Programs/Horde/Horde.Agent/Services/ServerLoggerFactory.cs b/Engine/Source/Programs/Horde/Horde.Agent/Services/ServerLoggerFactory.cs
--- a/Engine/Source/Programs/Horde/Horde.Agent/Services/ServerLoggerFactory.cs
+++ b/Engine/Source/Programs/Horde/Horde.Agent/Services/ServerLoggerFactory.cs
@@ -72,7 +72,7 @@ namespace Horde.Agent.Services
 		{
 #pragma warning disable CA2000 // Dispose objects before losing scope
 			IStorageClient storageClient = _storageClientFactory.CreateClientWithPath($"api/v1/logs/{logId}", session.Token);
-			IJsonRpcLogSink sink = new JsonRpcAndStorageLogSink(session.GrpcChannel, logId, jobId, batchId, stepId, storageClient, _logger);
+			IJsonRpcLogSink sink = new JsonRpcAndStorageLogSink(session.RpcConnection, logId, jobId, batchId, stepId, storageClient, _logger);
 
 			return new ServerLogger(sink, logId, warnings, outputLevel, localLogger, _logger);
 #pragma warning restore CA2000 // Dispose objects before losing scope
diff --git a/Engine/Source/Programs/Horde/Horde.Agent/Services/SessionFactory.cs b/Engine/Source/Programs/Horde/Horde.Agent/Services/SessionFactory.cs
--- a/Engine/Source/Programs/Horde/Horde.Agent/Services/SessionFactory.cs
+++ b/Engine/Source/Programs/Horde/Horde.Agent/Services/SessionFactory.cs
@@ -45,7 +45,7 @@ namespace Horde.Agent.Services
 		/// <summary>
 		/// Connection to the server
 		/// </summary>
-		GrpcChannel GrpcChannel { get; }
+		IRpcConnection RpcConnection { get; }
 
 		/// <summary>
 		/// State of the connection 
@@ -89,7 +89,7 @@ namespace Horde.Agent.Services
 		public string Token { get; }
 
 		/// <inheritdoc/>
-		public GrpcChannel GrpcChannel { get; }
+		public IRpcConnection RpcConnection { get; }
 
 		/// <inheritdoc/>
 		public DirectoryReference WorkingDir { get; }
@@ -100,13 +100,13 @@ namespace Horde.Agent.Services
 		/// <summary>
 		/// Constructor
 		/// </summary>
-		public Session(Uri serverUrl, AgentId agentId, SessionId sessionId, string token, GrpcChannel grpcChannel, DirectoryReference workingDir)
+		public Session(Uri serverUrl, AgentId agentId, SessionId sessionId, string token, IRpcConnection rpcConnection, DirectoryReference workingDir)
 		{
 			ServerUrl = serverUrl;
 			AgentId = agentId;
 			SessionId = sessionId;
 			Token = token;
-			GrpcChannel = grpcChannel;
+			RpcConnection = rpcConnection;
 			WorkingDir = workingDir;
 		}
 
@@ -224,8 +224,9 @@ namespace Horde.Agent.Services
 
 			// Open a connection to the server
 #pragma warning disable CA2000 // False positive; ownership is transferred to new Session object.
-			GrpcChannel grpcChannel = await grpcService.CreateGrpcChannelAsync(createSessionResponse.Token, cancellationToken);
-			return new Session(serverProfile.Url, new AgentId(createSessionResponse.AgentId), SessionId.Parse(createSessionResponse.SessionId), createSessionResponse.Token, grpcChannel, workingDir);
+			Func<CancellationToken, Task<GrpcChannel>> createGrpcChannelAsync = ctx => grpcService.CreateGrpcChannelAsync(createSessionResponse.Token, ctx);
+			IRpcConnection rpcConnection = new RpcConnection(createGrpcChannelAsync, logger);
+			return new Session(serverProfile.Url, new AgentId(createSessionResponse.AgentId), SessionId.Parse(createSessionResponse.SessionId), createSessionResponse.Token, rpcConnection, workingDir);
 #pragma warning restore CA2000
 		}
 
@@ -233,10 +234,9 @@ namespace Horde.Agent.Services
 		/// Dispose of the current session
 		/// </summary>
 		/// <returns></returns>
-		public ValueTask DisposeAsync()
+		public async ValueTask DisposeAsync()
 		{
-			GrpcChannel.Dispose();
-			return default;
+			await RpcConnection.DisposeAsync();
 		}
 
 		static FileReference GetRegistrationFile()
