Horde: Fix agent attempting to terminate all processes on startup.

#rnx

diff --git a/Engine/Source/Programs/Horde/Horde.Agent.Driver/Utility/TerminateProcessHelper.cs b/Engine/Source/Programs/Horde/Horde.Agent.Driver/Utility/TerminateProcessHelper.cs
--- a/Engine/Source/Programs/Horde/Horde.Agent.Driver/Utility/TerminateProcessHelper.cs
+++ b/Engine/Source/Programs/Horde/Horde.Agent.Driver/Utility/TerminateProcessHelper.cs
@@ -68,14 +68,17 @@ namespace Horde.Agent.Utility
 				string fileName = imageFile.GetFileName();
 				foreach (ProcessToTerminate processToTerminate in processesToTerminate)
 				{
-					TerminateCondition terminateFlags = TerminateCondition.None;
-					foreach (TerminateCondition when in processToTerminate.When ?? Enumerable.Empty<TerminateCondition>())
+					if (String.Equals(processToTerminate.Name, fileName, StringComparison.OrdinalIgnoreCase))
 					{
-						terminateFlags |= when;
-					}
-					if (terminateFlags == TerminateCondition.None || (terminateFlags & condition) != 0)
-					{
-						return true;
+						TerminateCondition terminateFlags = TerminateCondition.None;
+						foreach (TerminateCondition when in processToTerminate.When ?? Enumerable.Empty<TerminateCondition>())
+						{
+							terminateFlags |= when;
+						}
+						if (terminateFlags == TerminateCondition.None || (terminateFlags & condition) != 0)
+						{
+							return true;
+						}
 					}
 				}
 			}
