Horde: Updating schema documentation

#rnx

diff --git a/Engine/Source/Programs/Horde/Docs/Deployment/AgentSettings.md b/Engine/Source/Programs/Horde/Docs/Deployment/AgentSettings.md
--- a/Engine/Source/Programs/Horde/Docs/Deployment/AgentSettings.md
+++ b/Engine/Source/Programs/Horde/Docs/Deployment/AgentSettings.md
@@ -11,15 +11,10 @@ Name | Description
 `name` | `string`<br>Name of agent to report as when connecting to server. By default, the computer's hostname will be used.
 `installed` | `boolean`<br>Whether the server is running in 'installed' mode. In this mode, on Windows, the default data directory will use the common application data folder (C:\ProgramData\Epic\Horde), and configuration data will be read from here and the registry. This setting is overridden to false for local builds from appsettings.Local.json.
 `ephemeral` | `boolean`<br>Whether agent should register as being ephemeral. Doing so will not persist any long-lived data on the server and once disconnected it's assumed to have been deleted permanently. Ideal for short-lived agents, such as spot instances on AWS EC2.
-`executor` | `string`<br>The executor to use for jobs
-`localExecutor` | [`LocalExecutorSettings`](#localexecutorsettings)<br>Settings for the local executor
-`perforceExecutor` | [`PerforceExecutorSettings`](#perforceexecutorsettings)<br>Settings for the perforce executor
 `workingDir` | [`DirectoryReference`](#directoryreference)<br>Working directory for leases and jobs (i.e where files from Perforce will be checked out)
 `logsDir` | [`DirectoryReference`](#directoryreference)<br>Directory where agent and lease logs are written
 `shareMountingEnabled` | `boolean`<br>Whether to mount the specified list of network shares
 `shares` | [`MountNetworkShare`](#mountnetworkshare)`[]`<br>List of network shares to mount
-`processNamesToTerminate` | `string[]`<br>List of process names to terminate after a job
-`processesToTerminate` | [`ProcessToTerminate`](#processtoterminate)`[]`<br>List of process names to terminate after a lease completes, but not after a job step
 `wineExecutablePath` | `string`<br>Path to Wine executable. If null, execution under Wine is disabled
 `containerEngineExecutablePath` | `string`<br>Path to container engine executable, such as /usr/bin/podman. If null, execution of compute workloads inside a container is disabled
 `writeStepOutputToLogger` | `boolean`<br>Whether to write step output to the logging device
@@ -44,23 +39,6 @@ Name | Description
 `thumbprint` | `string`<br>Thumbprint of a certificate to trust. Allows using self-signed certs for the server.
 `thumbprints` | `string[]`<br>Thumbprints of certificates to trust. Allows using self-signed certs for the server.
 
-## LocalExecutorSettings
-
-Settings for the local executor
-
-Name | Description
----- | -----------
-`workspaceDir` | `string`<br>Path to the local workspace to use with the local executor
-`runSteps` | `boolean`<br>Whether to actually execute steps, or just do job setup
-
-## PerforceExecutorSettings
-
-Settings for the perforce executor
-
-Name | Description
----- | -----------
-`runConform` | `boolean`<br>Whether to run conform jobs
-
 ## DirectoryReference
 
 Name | Description
@@ -76,25 +54,3 @@ Name | Description
 ---- | -----------
 `mountPoint` | `string`<br>Where the share should be mounted on the local machine. Must be a drive letter for Windows.
 `remotePath` | `string`<br>Path to the remote resource
-
-## ProcessToTerminate
-
-Specifies a process to terminate
-
-Name | Description
----- | -----------
-`name` | `string`<br>Name of the process
-`when` | [`TerminateCondition`](#terminatecondition-enum)`[]`<br>When to terminate this process
-
-## TerminateCondition (Enum)
-
-Flags for processes to terminate
-
-Name | Description
----- | -----------
-`None` | Not specified; terminate in all circumstances
-`BeforeSession` | When a session starts
-`BeforeConform` | Before running a conform
-`BeforeBatch` | Before executing a batch
-`AfterBatch` | Terminate at the end of a batch
-`AfterStep` | After a step completes
