Horde: Respond to lease even when agent is missing in REST API

#rnx

diff --git a/Engine/Source/Programs/Horde/Horde.Server/Agents/Leases/LeasesController.cs b/Engine/Source/Programs/Horde/Horde.Server/Agents/Leases/LeasesController.cs
--- a/Engine/Source/Programs/Horde/Horde.Server/Agents/Leases/LeasesController.cs
+++ b/Engine/Source/Programs/Horde/Horde.Server/Agents/Leases/LeasesController.cs
@@ -134,13 +134,8 @@ namespace Horde.Server.Agents.Leases
 			}
 
 			IAgent? agent = await _agentService.GetAgentAsync(lease.AgentId, cancellationToken);
-			if (agent == null)
-			{
-				return NotFound(lease.AgentId);
-			}
-
 			double? agentRate = null;
-			if (_globalConfig.Value.Authorize(ServerAclAction.ViewCosts, User))
+			if (_globalConfig.Value.Authorize(ServerAclAction.ViewCosts, User) && agent != null)
 			{
 				agentRate = await _agentService.GetRateAsync(agent.Id, cancellationToken);
 			}
