Horde: Automatically enroll agents with the server for local builds.

#rnx

diff --git a/Engine/Source/Programs/Horde/Horde.Server/Agents/Enrollment/EnrollmentService.cs b/Engine/Source/Programs/Horde/Horde.Server/Agents/Enrollment/EnrollmentService.cs
--- a/Engine/Source/Programs/Horde/Horde.Server/Agents/Enrollment/EnrollmentService.cs
+++ b/Engine/Source/Programs/Horde/Horde.Server/Agents/Enrollment/EnrollmentService.cs
@@ -9,6 +9,7 @@ using EpicGames.Horde.Agents;
 using EpicGames.Redis;
 using Horde.Server.Server;
 using Microsoft.Extensions.Hosting;
+using Microsoft.Extensions.Options;
 using StackExchange.Redis;
 
 namespace Horde.Server.Agents.Enrollment
@@ -31,17 +32,19 @@ namespace Horde.Server.Agents.Enrollment
 		readonly RedisHashKey<string, EnrollmentRequest> _requests = new("agents:registration:requests");
 		readonly RedisHashKey<string, string> _approvals = new("agents:registration:approvals");
 		readonly AsyncEvent _approvalEvent = new AsyncEvent();
+		readonly IOptionsMonitor<ServerSettings> _settings;
 
 		RedisSubscription? _subscription;
 
 		/// <summary>
 		/// Constructor
 		/// </summary>
-		public EnrollmentService(RedisService redisService, IClock clock)
+		public EnrollmentService(RedisService redisService, IClock clock, IOptionsMonitor<ServerSettings> settings)
 		{
 			_redisService = redisService;
 			_clock = clock;
 			_updateChannel = RedisChannel.Literal("agents:registration");
+			_settings = settings;
 		}
 
 		/// <inheritdoc/>
@@ -90,7 +93,7 @@ namespace Horde.Server.Agents.Enrollment
 		}
 
 		/// <summary>
-		/// Waits until a particular agent is approved
+		/// Adds an agent as a candidate for the farm
 		/// </summary>
 		/// <param name="key">Unique id for the request</param>
 		/// <param name="hostName">Name of the host</param>
@@ -106,7 +109,17 @@ namespace Horde.Server.Agents.Enrollment
 			_ = transaction.SortedSetAddAsync(_keys, key, GetTimestamp(_clock.UtcNow + TimeSpan.FromMinutes(2.0)));
 			_ = transaction.HashSetAsync(_requests, key, new EnrollmentRequest(key, hostName, description));
 
-			return await transaction.ExecuteAsync().WaitAsync(cancellationToken);
+			if (!await transaction.ExecuteAsync().WaitAsync(cancellationToken))
+			{
+				return false;
+			}
+
+			if (_settings.CurrentValue.AutoEnrollAgents)
+			{
+				await ApproveAsync(key, null, cancellationToken);
+			}
+
+			return true;
 		}
 
 		/// <summary>
diff --git a/Engine/Source/Programs/Horde/Horde.Server/ServerSettings.cs b/Engine/Source/Programs/Horde/Horde.Server/ServerSettings.cs
--- a/Engine/Source/Programs/Horde/Horde.Server/ServerSettings.cs
+++ b/Engine/Source/Programs/Horde/Horde.Server/ServerSettings.cs
@@ -509,6 +509,11 @@ namespace Horde.Server
 		/// </summary>
 		public AuthMethod AuthMethod { get; set; } = AuthMethod.Anonymous;
 
+		/// <summary>
+		/// Whether to automatically enroll agents in the farm
+		/// </summary>
+		public bool AutoEnrollAgents { get; set; }
+
 		/// <summary>
 		/// Optional profile name to report through the /api/v1/server/auth endpoint. Allows sharing auth tokens between providers configured through
 		/// the same profile name in OidcToken.exe config files.
diff --git a/Engine/Source/Programs/Horde/Horde.Server/appsettings.Local.json b/Engine/Source/Programs/Horde/Horde.Server/appsettings.Local.json
--- a/Engine/Source/Programs/Horde/Horde.Server/appsettings.Local.json
+++ b/Engine/Source/Programs/Horde/Horde.Server/appsettings.Local.json
@@ -10,6 +10,9 @@
 		// Don't allow conforms to run
 		"EnableConformTasks": false,
 
+		// Allow any local agents to automatically enrol
+		"AutoEnrollAgents": true,
+
 		// Apply config changes immediately on startup
 		"ForceConfigUpdateOnStartup": true,
 
