Horde: Create grpc channels with an explicit message handler, allowing us to monitor the state of the connection through the GrpcChannel instance.

#rnx

diff --git a/Engine/Source/Programs/Horde/Horde.Agent/Services/GrpcService.cs b/Engine/Source/Programs/Horde/Horde.Agent/Services/GrpcService.cs
--- a/Engine/Source/Programs/Horde/Horde.Agent/Services/GrpcService.cs
+++ b/Engine/Source/Programs/Horde/Horde.Agent/Services/GrpcService.cs
@@ -5,9 +5,11 @@ using System.Net.Http.Headers;
 using System.Net.Http.Json;
 using System.Net.Security;
 using System.Net.Sockets;
+using EpicGames.Horde.Server;
 using Grpc.Core;
 using Grpc.Core.Interceptors;
 using Grpc.Net.Client;
+using Grpc.Net.Client.Configuration;
 using Horde.Agent.Utility;
 using Microsoft.Extensions.Logging;
 using Microsoft.Extensions.Options;
@@ -21,11 +23,6 @@ namespace Horde.Agent.Services
 	/// </summary>
 	class GrpcService
 	{
-		class GetPortsResponse
-		{
-			public int UnencryptedHttp2 { get; set; }
-		}
-
 		private readonly IOptions<AgentSettings> _settings;
 		private readonly ServerProfile _serverProfile;
 		private readonly ILogger _logger;
@@ -65,40 +62,35 @@ namespace Horde.Agent.Services
 		/// <returns>New grpc channel</returns>
 		public async Task<GrpcChannel> CreateGrpcChannelAsync(string? bearerToken, CancellationToken cancellationToken)
 		{
-#pragma warning disable CA2000 // Dispose objects before losing scope
-			// HTTP handler is disposed by GrpcChannel below
-			SocketsHttpHandler httpHandler = new()
+			Uri serverUri = _serverProfile.Url;
+			bool useInsecureConnection = serverUri.Scheme.Equals("http", StringComparison.Ordinal);
+
+			// Get the server URL for gRPC traffic. If we're using an unencrpyted connection we need to use a different port for http/2, so 
+			// send a http1 request to the server to query it.
+			if (useInsecureConnection)
 			{
-				ConnectCallback = async (context, connectCt) =>
+				_logger.LogInformation("Querying server {BaseUrl} for rpc port", serverUri);
+				using (HttpClient httpClient = new HttpClient())
 				{
-					using IScope scope = GlobalTracer.Instance
-						.BuildSpan($"{nameof(GrpcService)}.ConnectCallback")
-						.WithResourceName(context.DnsEndPoint.Host)
-						.StartActive();
-
-					IPAddress[] ips;
-					using (IScope _ = GlobalTracer.Instance.BuildSpan("DnsResolve").StartActive())
+					httpClient.DefaultRequestHeaders.Add("Accept", "application/json");
+					httpClient.Timeout = TimeSpan.FromSeconds(210); // Need to make sure this doesn't cancel any long running gRPC streaming calls (eg. session update)
+					using (HttpResponseMessage response = await httpClient.GetAsync(new Uri(serverUri, "api/v1/server/ports"), cancellationToken))
 					{
-						ips = await Dns.GetHostAddressesAsync(context.DnsEndPoint.Host, connectCt);
-					}
-
-					IPEndPoint ipEndpoint = new(ips[0], context.DnsEndPoint.Port);
-					using (IScope _ = GlobalTracer.Instance.BuildSpan("TcpConnect").StartActive())
-					{
-						Socket socket = new(SocketType.Stream, ProtocolType.Tcp);
-						try
+						GetPortsResponse? ports = await response.Content.ReadFromJsonAsync<GetPortsResponse>(AgentApp.DefaultJsonSerializerOptions, cancellationToken);
+						if (ports != null && ports.UnencryptedHttp2.HasValue && ports.UnencryptedHttp2.Value != 0)
 						{
-							await socket.ConnectAsync(ipEndpoint, connectCt);
-							return new NetworkStream(socket, ownsSocket: true);
-						}
-						catch
-						{
-							socket.Dispose();
-							throw;
+							UriBuilder builder = new UriBuilder(serverUri);
+							builder.Port = ports.UnencryptedHttp2.Value;
+							serverUri = builder.Uri;
 						}
 					}
-				},
+				}
+			}
 
+#pragma warning disable CA2000 // Dispose objects before losing scope
+			// HTTP handler is disposed by GrpcChannel below
+			SocketsHttpHandler httpHandler = new SocketsHttpHandler()
+			{
 				SslOptions = new SslClientAuthenticationOptions
 				{
 					RemoteCertificateValidationCallback = (sender, cert, chain, errors) => CertificateHelper.CertificateValidationCallBack(_logger, sender, cert, chain, errors, _serverProfile)
@@ -106,43 +98,38 @@ namespace Horde.Agent.Services
 			};
 #pragma warning restore CA2000 // Dispose objects before losing scope
 
-			HttpClient httpClient = new(httpHandler, true);
-			httpClient.DefaultRequestHeaders.Add("Accept", "application/json");
-			if (bearerToken != null)
-			{
-				httpClient.DefaultRequestHeaders.Authorization = new AuthenticationHeaderValue("Bearer", bearerToken);
-			}
-
-			httpClient.Timeout = TimeSpan.FromSeconds(210); // Need to make sure this doesn't cancel any long running gRPC streaming calls (eg. session update)
+			// Create options for the new channel
+			GrpcChannelOptions channelOptions = new GrpcChannelOptions();
+			channelOptions.MaxReceiveMessageSize = 1024 * 1024 * 1024; // 1 GB 		// Required payloads coming from CAS service can be large
+			channelOptions.MaxSendMessageSize = 1024 * 1024 * 1024; // 1 GB
+			channelOptions.LoggerFactory = _loggerFactory;
+			channelOptions.HttpHandler = httpHandler;
+			channelOptions.DisposeHttpClient = true;
 
-			// Get the server URL for gRPC traffic. If we're using an unencrpyted connection we need to use a different port for http/2, so 
-			// send a http1 request to the server to query it.
-			Uri serverUri = _serverProfile.Url;
-			if (serverUri.Scheme.Equals("http", StringComparison.Ordinal))
+			// Configure requests to send the bearer token
+			if (!String.IsNullOrEmpty(bearerToken))
 			{
-				_logger.LogInformation("Querying server {BaseUrl} for rpc port", serverUri);
-				using (HttpResponseMessage response = await httpClient.GetAsync(new Uri(serverUri, "api/v1/server/ports"), cancellationToken))
+				CallCredentials callCredentials = CallCredentials.FromInterceptor((context, metadata) =>
 				{
-					GetPortsResponse? ports = await response.Content.ReadFromJsonAsync<GetPortsResponse>(AgentApp.DefaultJsonSerializerOptions, cancellationToken);
-					if (ports != null && ports.UnencryptedHttp2 != 0)
-					{
-						UriBuilder builder = new UriBuilder(serverUri);
-						builder.Port = ports.UnencryptedHttp2;
-						serverUri = builder.Uri;
-					}
+					metadata.Add("Authorization", $"Bearer {bearerToken}");
+					return Task.CompletedTask;
+				});
+
+				if (useInsecureConnection)
+				{
+					channelOptions.Credentials = ChannelCredentials.Create(ChannelCredentials.Insecure, callCredentials);
 				}
+				else
+				{
+					channelOptions.Credentials = ChannelCredentials.Create(ChannelCredentials.SecureSsl, callCredentials);
+				}
+
+				channelOptions.UnsafeUseInsecureChannelCallCredentials = useInsecureConnection;
 			}
 
+			// Create the channel
 			_logger.LogInformation("Connecting to rpc server {BaseUrl}", serverUri);
-			return GrpcChannel.ForAddress(serverUri, new GrpcChannelOptions
-			{
-				// Required payloads coming from CAS service can be large
-				MaxReceiveMessageSize = 1024 * 1024 * 1024, // 1 GB
-				MaxSendMessageSize = 1024 * 1024 * 1024, // 1 GB
-				LoggerFactory = _loggerFactory,
-				HttpClient = httpClient,
-				DisposeHttpClient = true
-			});
+			return GrpcChannel.ForAddress(serverUri, channelOptions);
 		}
 
 		/// <summary>
