Horde: Change agent driver into a library for the time being, to fix agent publish.

#rnx

diff --git a/Engine/Source/Programs/Horde/Horde.Agent.Driver/DriverApp.cs b/Engine/Source/Programs/Horde/Horde.Agent.Driver/DriverApp.cs
--- a/Engine/Source/Programs/Horde/Horde.Agent.Driver/DriverApp.cs
+++ b/Engine/Source/Programs/Horde/Horde.Agent.Driver/DriverApp.cs
@@ -8,7 +8,7 @@ namespace Horde.Agent.Driver
 {
 	class DriverApp
 	{
-		public static async Task<int> Main(string[] args)
+		public static async Task<int> MainAsync(string[] args)
 		{
 			CommandLineArguments arguments = new CommandLineArguments(args);
 
diff --git a/Engine/Source/Programs/Horde/Horde.Agent.Driver/Horde.Agent.Driver.csproj b/Engine/Source/Programs/Horde/Horde.Agent.Driver/Horde.Agent.Driver.csproj
--- a/Engine/Source/Programs/Horde/Horde.Agent.Driver/Horde.Agent.Driver.csproj
+++ b/Engine/Source/Programs/Horde/Horde.Agent.Driver/Horde.Agent.Driver.csproj
@@ -1,7 +1,7 @@
 <Project Sdk="Microsoft.NET.Sdk">
 
   <PropertyGroup>
-    <OutputType>Exe</OutputType>
+    <OutputType>Library</OutputType>
     <TargetFramework>net6.0</TargetFramework>
     <Nullable>enable</Nullable>
     <ImplicitUsings>enable</ImplicitUsings>
