Horde Fix issue with step run again bombing back to home page
#rnx

diff --git a/Engine/Source/Programs/Horde/HordeDashboard/package.json b/Engine/Source/Programs/Horde/HordeDashboard/package.json
--- a/Engine/Source/Programs/Horde/HordeDashboard/package.json
+++ b/Engine/Source/Programs/Horde/HordeDashboard/package.json
@@ -1,7 +1,7 @@
 {
   "name": "horde-dashboard",
   "private": true,
-  "version": "0.0.8",
+  "version": "0.1.0",
   "type": "module",  
   "scripts": {
     "dev": "vite",
diff --git a/Engine/Source/Programs/Horde/HordeDashboard/src/components/JobViewAll.tsx b/Engine/Source/Programs/Horde/HordeDashboard/src/components/JobViewAll.tsx
--- a/Engine/Source/Programs/Horde/HordeDashboard/src/components/JobViewAll.tsx
+++ b/Engine/Source/Programs/Horde/HordeDashboard/src/components/JobViewAll.tsx
@@ -362,7 +362,6 @@ const JobViewAllInner: React.FC<{ filter: JobFilterSimple }> = observer(({ filte
    const JobLabel: React.FC<{ item: JobItem; label: GetLabelStateResponse }> = ({ item, label }) => {
 
       const defaultLabel = item.job.defaultLabel;
-
       const aggregates = item.job.labels;
 
       // note details may not be loaded here, as only initialized on callout for optimization (details.getLabelIndex(label.Name, label.Category);)
@@ -374,7 +373,7 @@ const JobViewAllInner: React.FC<{ filter: JobFilterSimple }> = observer(({ filte
 
       let state: LabelState | undefined;
       let outcome: LabelOutcome | undefined;
-      if (defaultLabel) {
+      if (label === defaultLabel) {
          state = defaultLabel.state;
          outcome = defaultLabel.outcome;
       } else {
diff --git a/Engine/Source/Programs/Horde/HordeDashboard/src/components/LogSource.tsx b/Engine/Source/Programs/Horde/HordeDashboard/src/components/LogSource.tsx
--- a/Engine/Source/Programs/Horde/HordeDashboard/src/components/LogSource.tsx
+++ b/Engine/Source/Programs/Horde/HordeDashboard/src/components/LogSource.tsx
@@ -3,7 +3,7 @@
 import { action, makeObservable, observable } from "mobx";
 import moment, { Moment } from 'moment-timezone';
 import backend from '../backend';
-import { AgentData, BatchData, EventSeverity, GetArtifactResponseV2, GetLogEventResponse, GetLogFileResponse, IssueData, LeaseData, LogData, StepData, StreamData } from "../backend/Api";
+import { AgentData, BatchData, EventSeverity, GetArtifactResponseV2, GetJobsTabResponse, GetLogEventResponse, GetLogFileResponse, IssueData, LeaseData, LogData, StepData, StreamData } from "../backend/Api";
 import { getBatchSummaryMarkdown, getStepSummaryMarkdown, JobDetails } from "../backend/JobDetails";
 import { getLeaseElapsed, getStepPercent } from '../base/utilities/timeUtils';
 import { BreadcrumbItem } from './Breadcrumbs';
@@ -531,7 +531,6 @@ export class JobLogSource extends LogSource {
 
       const data = this.jobDetails.jobdata!;
 
-
       if (!this.stream) {
          return [];
       }
@@ -541,6 +540,16 @@ export class JobLogSource extends LogSource {
          projectName = "UE4";
       }
 
+      const tab = this.stream.tabs?.find((tab) => {
+         const jtab = tab as GetJobsTabResponse;
+         return !!jtab.templates?.find(t => t === data?.templateId);
+      });
+
+      let streamLink = `/stream/${this.stream.id}`;
+      if (tab) {
+         streamLink += `?tab=${tab.title}`;
+      }
+
       const crumbItems: BreadcrumbItem[] = [
          {
             text: projectName ?? "Unknown Project",
@@ -548,7 +557,7 @@ export class JobLogSource extends LogSource {
          },
          {
             text: this.stream.name,
-            link: `/stream/${this.stream.id}`
+            link: streamLink
          },
          {
             text: `${data?.name ?? ""} - ${this.clText}`,
diff --git a/Engine/Source/Programs/Horde/HordeDashboard/src/components/jobDetailsV2/JobDetailsViewCommon.tsx b/Engine/Source/Programs/Horde/HordeDashboard/src/components/jobDetailsV2/JobDetailsViewCommon.tsx
--- a/Engine/Source/Programs/Horde/HordeDashboard/src/components/jobDetailsV2/JobDetailsViewCommon.tsx
+++ b/Engine/Source/Programs/Horde/HordeDashboard/src/components/jobDetailsV2/JobDetailsViewCommon.tsx
@@ -366,7 +366,7 @@ export class JobDetailsV2 extends PollBase {
 
    }
 
-   private processGraph() {
+   processGraph() {
 
       const jobData = this.jobData;
 
diff --git a/Engine/Source/Programs/Horde/HordeDashboard/src/components/jobDetailsV2/StepRetryModal.tsx b/Engine/Source/Programs/Horde/HordeDashboard/src/components/jobDetailsV2/StepRetryModal.tsx
--- a/Engine/Source/Programs/Horde/HordeDashboard/src/components/jobDetailsV2/StepRetryModal.tsx
+++ b/Engine/Source/Programs/Horde/HordeDashboard/src/components/jobDetailsV2/StepRetryModal.tsx
@@ -212,10 +212,16 @@ export const StepRetryModal: React.FC<{ stepId: string; jobDetails: JobDetailsV2
 
       await backend.updateJobStep(jobData.id, batch.id, stepId, {
          retry: true
-      }).then((response) => {
+      }).then(async (response) => {
 
-         if (response.stepId) {
-            navigate(`/job/${jobData.id}?step=${response.stepId!}`)
+         if (response.stepId) {            
+            if (jobDetails.jobData?.id) {
+               // get new job data so we have the new batch/step
+               jobDetails.jobData = await backend.getJob(jobDetails.jobData?.id);
+               jobDetails.processGraph();
+            }
+            
+            navigate(`/job/${jobData.id}?step=${response.stepId!}`)            
          }
 
       }).catch((reason) => {
