Horde: Advanced Additional Args
#rnx

diff --git a/Engine/Source/Programs/Horde/HordeDashboard/src/backend/Api.ts b/Engine/Source/Programs/Horde/HordeDashboard/src/backend/Api.ts
--- a/Engine/Source/Programs/Horde/HordeDashboard/src/backend/Api.ts
+++ b/Engine/Source/Programs/Horde/HordeDashboard/src/backend/Api.ts
@@ -1139,6 +1139,9 @@ export type CreateJobRequest = {
 
 	/** Parameters for the job */
 	parameters?: Record<string, string>
+
+	/** Additional arguments for the job, job dialog v2 */
+	additionalArguments?: string[];
 }
 
 /**Response from creating a new job */
@@ -1299,6 +1302,8 @@ export type GetJobResponse = {
 	/**  Custom permissions for this object */
 	acl?: GetAclResponse;
 
+	/** Additional arguments for the job, job dialog v2 */
+	additionalArguments?: string[];
 }
 
 /**Request used to update a jobstep */
diff --git a/Engine/Source/Programs/Horde/HordeDashboard/src/components/build/NewBuildV2.tsx b/Engine/Source/Programs/Horde/HordeDashboard/src/components/build/NewBuildV2.tsx
--- a/Engine/Source/Programs/Horde/HordeDashboard/src/components/build/NewBuildV2.tsx
+++ b/Engine/Source/Programs/Horde/HordeDashboard/src/components/build/NewBuildV2.tsx
@@ -78,6 +78,8 @@ class BuildOptions {
 
       this.stream = stream;
 
+      this.advAdditionalArgs = jobDetails?.jobData?.additionalArguments?.join(" ");
+
       this.load();
    }
 
@@ -101,6 +103,7 @@ class BuildOptions {
    advJobName?: string;
    advJobPriority?: Priority;
    advUpdateIssues?: boolean;
+   advAdditionalArgs?: string;
 
    parameters: Record<string, string> = {};
 
@@ -150,7 +153,7 @@ class BuildOptions {
 
    get advancedModified(): boolean {
 
-      return !!((this.advJobPriority && this.advJobPriority !== Priority.Normal) || this.advUpdateIssues || !!this.advJobName)
+      return !!((this.advJobPriority && this.advJobPriority !== Priority.Normal) || this.advUpdateIssues || !!this.advJobName || !!this.advAdditionalArgs)
    }
 
    change?: number;
@@ -1273,6 +1276,11 @@ const AdvancedPanel: React.FC = observer(() => {
       height += (fargumentsHeight + 8)
    }
 
+   let showAdditionalArgs = !options.readOnly || !!options.advAdditionalArgs;
+   if (showAdditionalArgs) {
+      height += 52;
+   }
+
    let fparameters = "";
    let fparametersHeight = 0;
 
@@ -1302,7 +1310,7 @@ const AdvancedPanel: React.FC = observer(() => {
          <Stack>
             <Dropdown disabled={options.readOnly} key={"key_adv_priority"} defaultValue={options.advJobPriority} label="Priority" options={priorityOptions} onChange={(ev, option) => {
                options.advJobPriority = option?.key as Priority;
-               options.setChanged();
+            options.setChanged();
             }
             } />
          </Stack>
@@ -1321,6 +1329,7 @@ const AdvancedPanel: React.FC = observer(() => {
                options.setChanged();
             }} />
          </Stack>
+         {showAdditionalArgs && <TextField key={"key_adv_add_args"} multiline style={{ height: 48 }} resizable={false} readOnly={options.readOnly} spellCheck={false} defaultValue={options.advAdditionalArgs} label="Additional Arguments" onChange={(ev, newValue) => options.advAdditionalArgs = newValue} />}
          {showArgumentClipboardButton && <DefaultButton text="Copy Job Arguments to Clipboard" style={{ width: 240 }} onClick={() => copyToClipboard(
             options.jobDetails.jobData.arguments.map(arg => {
                if (arg.indexOf("=") !== -1) {
@@ -1336,6 +1345,7 @@ const AdvancedPanel: React.FC = observer(() => {
          )}/>}
          {!!farguments && <TextField style={{ height: fargumentsHeight }} key={"key_adv_job_arguments"} defaultValue={farguments} readOnly={true} label="Job Arguments" multiline resizable={false} />}
          {!!fparameters && <TextField style={{ height: fparametersHeight }} key={"key_adv_job_parameters]"} defaultValue={fparameters} readOnly={true} label="Job Parameters" multiline resizable={false} />}
+
          <Stack>
             <Checkbox disabled={options.readOnly}
                label="Template Editor"
@@ -1410,6 +1420,13 @@ const BuildModal: React.FC = observer(() => {
          }
       }
 
+      let additionalArgs = [];
+      if (options!.advAdditionalArgs) {
+         const argRegex = /"(\\"|[^"])*?"|[^ ]+/g;
+         options.advAdditionalArgs.trim().match(argRegex)?.forEach(arg => additionalArgs.push(arg.replace(/"/g, "")));
+      }
+
+
       const data: CreateJobRequest = {
          streamId: options.streamId,
          templateId: templateId,
@@ -1417,7 +1434,8 @@ const BuildModal: React.FC = observer(() => {
          priority: options.advJobPriority,
          updateIssues: updateIssues,
          changeQueries: changeQueries,
-         parameters: options.parameters
+         parameters: options.parameters,
+         additionalArguments: additionalArgs?.length ? additionalArgs : undefined
       };
 
       if (typeof (options.change) === 'number') {
