Horde: Copy IsUserAdministator() method into CapabilitiesService, to eliminate agent -> driver dependency.

#rnx

diff --git a/Engine/Source/Programs/Horde/Horde.Agent.Driver/Execution/JobExecutor.cs b/Engine/Source/Programs/Horde/Horde.Agent.Driver/Execution/JobExecutor.cs
--- a/Engine/Source/Programs/Horde/Horde.Agent.Driver/Execution/JobExecutor.cs
+++ b/Engine/Source/Programs/Horde/Horde.Agent.Driver/Execution/JobExecutor.cs
@@ -340,7 +340,7 @@ namespace Horde.Agent.Execution
 			}
 		}
 
-		public static bool IsUserAdministrator()
+		static bool IsUserAdministrator()
 		{
 			if (!OperatingSystem.IsWindows())
 			{
diff --git a/Engine/Source/Programs/Horde/Horde.Agent/Services/CapabilitiesService.cs b/Engine/Source/Programs/Horde/Horde.Agent/Services/CapabilitiesService.cs
--- a/Engine/Source/Programs/Horde/Horde.Agent/Services/CapabilitiesService.cs
+++ b/Engine/Source/Programs/Horde/Horde.Agent/Services/CapabilitiesService.cs
@@ -4,6 +4,7 @@ using System.Diagnostics;
 using System.Net;
 using System.Net.Sockets;
 using System.Runtime.InteropServices;
+using System.Security.Principal;
 using System.Text.RegularExpressions;
 using System.Xml;
 using Amazon.EC2;
@@ -12,7 +13,6 @@ using Amazon.Util;
 using EpicGames.Core;
 using EpicGames.Horde.Agents;
 using EpicGames.Horde.Compute;
-using Horde.Agent.Execution;
 using HordeCommon.Rpc.Messages;
 using Microsoft.Extensions.Logging;
 using Microsoft.Extensions.Options;
@@ -198,7 +198,7 @@ namespace Horde.Agent.Services
 				primaryDevice.Properties.Add($"User={Environment.UserName}");
 				primaryDevice.Properties.Add($"Domain={Environment.UserDomainName}");
 				primaryDevice.Properties.Add($"Interactive={Environment.UserInteractive}");
-				primaryDevice.Properties.Add($"Elevated={JobExecutor.IsUserAdministrator()}");
+				primaryDevice.Properties.Add($"Elevated={IsUserAdministrator()}");
 			}
 			else if (RuntimeInformation.IsOSPlatform(OSPlatform.Linux))
 			{
@@ -269,7 +269,7 @@ namespace Horde.Agent.Services
 				primaryDevice.Properties.Add($"User={Environment.UserName}");
 				primaryDevice.Properties.Add($"Domain={Environment.UserDomainName}");
 				primaryDevice.Properties.Add($"Interactive={Environment.UserInteractive}");
-				primaryDevice.Properties.Add($"Elevated={JobExecutor.IsUserAdministrator()}");
+				primaryDevice.Properties.Add($"Elevated={IsUserAdministrator()}");
 			}
 			else if (RuntimeInformation.IsOSPlatform(OSPlatform.OSX))
 			{
@@ -443,6 +443,30 @@ namespace Horde.Agent.Services
 			return agent;
 		}
 
+		/// <summary>
+		/// Determine if the current user is a running with elevated permissions on Windows (UAC)
+		/// </summary>
+		public static bool IsUserAdministrator()
+		{
+			if (!OperatingSystem.IsWindows())
+			{
+				return false;
+			}
+
+			try
+			{
+				using (WindowsIdentity identity = WindowsIdentity.GetCurrent())
+				{
+					WindowsPrincipal principal = new WindowsPrincipal(identity);
+					return principal.IsInRole(WindowsBuiltInRole.Administrator);
+				}
+			}
+			catch
+			{
+				return false;
+			}
+		}
+
 		/// <summary>
 		/// Resolve local IP address of agent
 		///
