Horde: Remove some unused gRPC messages.

#rnx

diff --git a/Engine/Source/Programs/Horde/Horde.Agent.Driver/Utilities/HordeRpc.cs b/Engine/Source/Programs/Horde/Horde.Agent.Driver/Utilities/HordeRpc.cs
--- a/Engine/Source/Programs/Horde/Horde.Agent.Driver/Utilities/HordeRpc.cs
+++ b/Engine/Source/Programs/Horde/Horde.Agent.Driver/Utilities/HordeRpc.cs
@@ -15,70 +15,6 @@ using Google.Protobuf.Collections;
 #pragma warning disable CA1716
 namespace HordeCommon.Rpc
 {
-	partial class RpcProperty
-	{
-		public RpcProperty(string name, string value)
-		{
-			Name = name;
-			Value = value;
-		}
-
-		public RpcProperty(KeyValuePair<string, string> pair)
-		{
-			Name = pair.Key;
-			Value = pair.Value;
-		}
-	}
-
-	partial class RpcPropertyUpdate
-	{
-		public RpcPropertyUpdate(string name, string? value)
-		{
-			Name = name;
-			Value = value;
-		}
-	}
-
-	static class PropertyExtensions
-	{
-		public static string GetValue(this RepeatedField<RpcProperty> properties, string name)
-		{
-			return properties.First(x => x.Name == name).Value;
-		}
-
-		public static bool TryGetValue(this RepeatedField<RpcProperty> properties, string name, [MaybeNullWhen(false)] out string result)
-		{
-			RpcProperty? property = properties.FirstOrDefault(x => x.Name == name);
-			if (property == null)
-			{
-				result = null!;
-				return false;
-			}
-			else
-			{
-				result = property.Value;
-				return true;
-			}
-		}
-	}
-
-	partial class RpcGetStreamRequest
-	{
-		public RpcGetStreamRequest(StreamId streamId)
-		{
-			StreamId = streamId.ToString();
-		}
-	}
-
-	partial class RpcUpdateStreamRequest
-	{
-		public RpcUpdateStreamRequest(StreamId streamId, Dictionary<string, string?> properties)
-		{
-			StreamId = streamId.ToString();
-			Properties.AddRange(properties.Select(x => new RpcPropertyUpdate(x.Key, x.Value)));
-		}
-	}
-
 	partial class RpcGetJobRequest
 	{
 		public RpcGetJobRequest(JobId jobId)
@@ -189,47 +125,5 @@ namespace HordeCommon.Rpc
 	}
 }
 
-namespace HordeCommon.Rpc.Messages.Telemetry
-{
-	partial class RpcAgentMetadataEvent
-	{
-		/// <summary>
-		/// Calculate an agent ID
-		/// </summary>
-		/// <returns>A unique hash for all fields</returns>
-		public long CalculateAgentId()
-		{
-			using SHA256 sha256 = SHA256.Create();
-			using MemoryStream ms = new(200);
-			using BinaryWriter bw = new(ms);
-
-			bw.Write(Ip ?? "<empty ip>");
-			bw.Write(Hostname ?? "<empty hostname>");
-			bw.Write(Region ?? "<empty region>");
-			bw.Write(AvailabilityZone ?? "<empty az>");
-			bw.Write(Environment ?? "<empty env>");
-			bw.Write(AgentVersion ?? "<empty version>");
-			bw.Write(Os ?? "<empty os>");
-			bw.Write(OsVersion ?? "<empty os version>");
-			bw.Write(Architecture ?? "<empty os architecture>");
-
-			foreach (KeyValuePair<string, string> pair in Properties)
-			{
-				bw.Write(pair.Key ?? "<empty key>");
-				bw.Write(pair.Value ?? "<empty value>");
-			}
-
-			foreach (string poolId in PoolIds)
-			{
-				bw.Write(poolId);
-			}
-
-			ms.Position = 0;
-			byte[] hash = sha256.ComputeHash(ms);
-			return BitConverter.ToInt64(hash, 0);
-		}
-	}
-}
-
 #pragma warning restore CA1716
 #pragma warning restore CS1591
\ No newline at end of file
diff --git a/Engine/Source/Programs/Horde/Horde.Agent.Tests/Helper.cs b/Engine/Source/Programs/Horde/Horde.Agent.Tests/Helper.cs
--- a/Engine/Source/Programs/Horde/Horde.Agent.Tests/Helper.cs
+++ b/Engine/Source/Programs/Horde/Horde.Agent.Tests/Helper.cs
@@ -67,12 +67,6 @@ namespace Horde.Agent.Tests
 			return Wrap(res);
 		}
 
-		public override AsyncUnaryCall<RpcGetStreamResponse> GetStreamAsync(RpcGetStreamRequest request, CallOptions options)
-		{
-			RpcGetStreamResponse res = new RpcGetStreamResponse();
-			return Wrap(res);
-		}
-
 		public override AsyncUnaryCall<RpcGetJobResponse> GetJobAsync(RpcGetJobRequest request, CallOptions options)
 		{
 			RpcGetJobResponse res = new RpcGetJobResponse();
diff --git a/Engine/Source/Programs/Horde/Horde.Agent.Tests/WorkerServiceTest.cs b/Engine/Source/Programs/Horde/Horde.Agent.Tests/WorkerServiceTest.cs
--- a/Engine/Source/Programs/Horde/Horde.Agent.Tests/WorkerServiceTest.cs
+++ b/Engine/Source/Programs/Horde/Horde.Agent.Tests/WorkerServiceTest.cs
@@ -378,16 +378,6 @@ namespace Horde.Agent.Tests
 				_outer = outer;
 			}
 
-			public override AsyncUnaryCall<RpcGetStreamResponse> GetStreamAsync(RpcGetStreamRequest request, CallOptions options)
-			{
-				if (_outer._streamIdToStreamResponse.TryGetValue(new StreamId(request.StreamId), out RpcGetStreamResponse? streamResponse))
-				{
-					return JobRpcClientStub.Wrap(streamResponse);
-				}
-
-				throw new RpcException(new Status(StatusCode.NotFound, $"Stream ID {request.StreamId} not found"));
-			}
-
 			public override AsyncUnaryCall<RpcGetJobResponse> GetJobAsync(RpcGetJobRequest request, CallOptions options)
 			{
 				if (_outer._jobIdToJobResponse.TryGetValue(JobId.Parse(request.JobId), out RpcGetJobResponse? jobResponse))
diff --git a/Engine/Source/Programs/Horde/Horde.Agent/Services/TelemetryService.cs b/Engine/Source/Programs/Horde/Horde.Agent/Services/TelemetryService.cs
--- a/Engine/Source/Programs/Horde/Horde.Agent/Services/TelemetryService.cs
+++ b/Engine/Source/Programs/Horde/Horde.Agent/Services/TelemetryService.cs
@@ -3,6 +3,7 @@
 using System.Diagnostics;
 using System.Runtime.InteropServices;
 using System.Runtime.Versioning;
+using System.Security.Cryptography;
 using System.Text;
 using EpicGames.Core;
 using Google.Protobuf.WellKnownTypes;
@@ -561,10 +562,46 @@ class TelemetryService : BackgroundService
 			Architecture = RuntimeInformation.OSArchitecture.ToString(),
 		};
 		e.PoolIds.AddRange(_workerService.PoolIds);
-		e.AgentId = e.CalculateAgentId();
+		e.AgentId = CalculateAgentId(e);
 		return e;
 	}
 
+	/// <summary>
+	/// Calculate an agent ID
+	/// </summary>
+	/// <returns>A unique hash for all fields</returns>
+	private static long CalculateAgentId(RpcAgentMetadataEvent e)
+	{
+		using SHA256 sha256 = SHA256.Create();
+		using MemoryStream ms = new(200);
+		using BinaryWriter bw = new(ms);
+
+		bw.Write(e.Ip ?? "<empty ip>");
+		bw.Write(e.Hostname ?? "<empty hostname>");
+		bw.Write(e.Region ?? "<empty region>");
+		bw.Write(e.AvailabilityZone ?? "<empty az>");
+		bw.Write(e.Environment ?? "<empty env>");
+		bw.Write(e.AgentVersion ?? "<empty version>");
+		bw.Write(e.Os ?? "<empty os>");
+		bw.Write(e.OsVersion ?? "<empty os version>");
+		bw.Write(e.Architecture ?? "<empty os architecture>");
+
+		foreach (KeyValuePair<string, string> pair in e.Properties)
+		{
+			bw.Write(pair.Key ?? "<empty key>");
+			bw.Write(pair.Value ?? "<empty value>");
+		}
+
+		foreach (string poolId in e.PoolIds)
+		{
+			bw.Write(poolId);
+		}
+
+		ms.Position = 0;
+		byte[] hash = sha256.ComputeHash(ms);
+		return BitConverter.ToInt64(hash, 0);
+	}
+
 	private static string GetOs()
 	{
 		if (RuntimeInformation.IsOSPlatform(OSPlatform.Windows))
diff --git a/Engine/Source/Programs/Horde/Horde.Server/Jobs/JobRpcService.cs b/Engine/Source/Programs/Horde/Horde.Server/Jobs/JobRpcService.cs
--- a/Engine/Source/Programs/Horde/Horde.Server/Jobs/JobRpcService.cs
+++ b/Engine/Source/Programs/Horde/Horde.Server/Jobs/JobRpcService.cs
@@ -79,29 +79,6 @@ namespace Horde.Server.Jobs
 			_logger = logger;
 		}
 
-		/// <summary>
-		/// Gets information about a stream
-		/// </summary>
-		/// <param name="request">Request arguments</param>
-		/// <param name="context">Context for the RPC call</param>
-		/// <returns>Information about the new agent</returns>
-		public override Task<RpcGetStreamResponse> GetStream(RpcGetStreamRequest request, ServerCallContext context)
-		{
-			StreamId streamIdValue = new StreamId(request.StreamId);
-
-			StreamConfig? streamConfig;
-			if (!_globalConfig.Value.TryGetStream(streamIdValue, out streamConfig))
-			{
-				throw new StructuredRpcException(StatusCode.NotFound, "Stream {StreamId} does not exist", request.StreamId);
-			}
-			if (!streamConfig.Authorize(StreamAclAction.ViewStream, context.GetHttpContext().User))
-			{
-				throw new StructuredRpcException(StatusCode.PermissionDenied, "Not authenticated to access stream {StreamId}", request.StreamId);
-			}
-
-			return Task.FromResult(streamConfig.ToRpcResponse());
-		}
-
 		/// <summary>
 		/// Gets information about a job
 		/// </summary>
