Horde: Updating schema documentation

#rnx

diff --git a/Engine/Source/Programs/Horde/Docs/Config/Schema/Globals.md b/Engine/Source/Programs/Horde/Docs/Config/Schema/Globals.md
--- a/Engine/Source/Programs/Horde/Docs/Config/Schema/Globals.md
+++ b/Engine/Source/Programs/Horde/Docs/Config/Schema/Globals.md
@@ -387,6 +387,8 @@ Name | Description
 `name` | `string`<br>Name of the tool
 `description` | `string`<br>Description for the tool
 `category` | `string`<br>Category for the tool. Will cause the tool to be shown in a different tab in the dashboard.
+`group` | `string`<br>Grouping key for different variations of the same tool. The dashboard will show these together.
+`platforms` | `string[]`<br>Platforms for this tool. Takes the form of a NET RID (https://learn.microsoft.com/en-us/dotnet/core/rid-catalog).
 `public` | `boolean`<br>Whether this tool should be exposed for download on a public endpoint without authentication
 `showInUgs` | `boolean`<br>Whether to show this tool for download in the UGS tools menu
 `showInDashboard` | `boolean`<br>Whether to show this tool for download in the dashboard
diff --git a/Engine/Source/Programs/Horde/Docs/Deployment/ServerSettings.md b/Engine/Source/Programs/Horde/Docs/Deployment/ServerSettings.md
--- a/Engine/Source/Programs/Horde/Docs/Deployment/ServerSettings.md
+++ b/Engine/Source/Programs/Horde/Docs/Deployment/ServerSettings.md
@@ -352,6 +352,8 @@ Name | Description
 `name` | `string`<br>Name of the tool
 `description` | `string`<br>Description for the tool
 `category` | `string`<br>Category for the tool. Will cause the tool to be shown in a different tab in the dashboard.
+`group` | `string`<br>Grouping key for different variations of the same tool. The dashboard will show these together.
+`platforms` | `string[]`<br>Platforms for this tool. Takes the form of a NET RID (https://learn.microsoft.com/en-us/dotnet/core/rid-catalog).
 `public` | `boolean`<br>Whether this tool should be exposed for download on a public endpoint without authentication
 `showInUgs` | `boolean`<br>Whether to show this tool for download in the UGS tools menu
 `showInDashboard` | `boolean`<br>Whether to show this tool for download in the dashboard
