Horde: Move functionality for updating job state from severity log levels into a separate logger. Reimplement tests using gRPC mocks.

#rnx

diff --git a/Engine/Source/Programs/Horde/Horde.Agent.Driver/Commands/Utilities/BuildGraphCommand.cs b/Engine/Source/Programs/Horde/Horde.Agent.Driver/Commands/Utilities/BuildGraphCommand.cs
--- a/Engine/Source/Programs/Horde/Horde.Agent.Driver/Commands/Utilities/BuildGraphCommand.cs
+++ b/Engine/Source/Programs/Horde/Horde.Agent.Driver/Commands/Utilities/BuildGraphCommand.cs
@@ -19,71 +19,6 @@ namespace Horde.Agent.Commands.Utilities
 	[Command("buildgraph", "Executes a BuildGraph script with the given arguments using a build of UAT within the current branch, and runs the output through the log processor")]
 	class BuildGraphCommand : Command
 	{
-		sealed class LogSink : IJsonRpcLogSink, IAsyncDisposable, IDisposable
-		{
-			readonly ILogger _inner;
-			readonly FileStream _eventStream;
-			readonly FileStream _outputStream;
-
-			public LogSink(DirectoryReference baseDir, ILogger inner)
-			{
-				_inner = inner;
-
-				FileReference eventsFile = FileReference.Combine(baseDir, "horde-events.txt");
-				_eventStream = FileReference.Open(eventsFile, FileMode.Create);
-				inner.LogInformation("Writing events to {File}", eventsFile);
-
-				FileReference outputFile = FileReference.Combine(baseDir, "horde-output.txt");
-				_outputStream = FileReference.Open(outputFile, FileMode.Create);
-				inner.LogInformation("Writing output to {File}", outputFile);
-			}
-
-			public ValueTask DisposeAsync()
-			{
-				Dispose();
-				return new ValueTask();
-			}
-
-			public void Dispose()
-			{
-				_eventStream.Dispose();
-				_outputStream.Dispose();
-			}
-
-			public Task SetOutcomeAsync(JobStepOutcome outcome, CancellationToken cancellationToken) => Task.CompletedTask;
-
-			public async Task WriteEventsAsync(List<RpcCreateEventRequest> events, CancellationToken cancellationToken)
-			{
-				foreach (RpcCreateEventRequest request in events)
-				{
-					JsonSerializerOptions options = new JsonSerializerOptions();
-					options.Converters.Add(new JsonStringEnumConverter());
-					await JsonSerializer.SerializeAsync(_eventStream, request, options, cancellationToken);
-					_eventStream.Write(Encoding.UTF8.GetBytes(Environment.NewLine));
-				}
-			}
-
-			public async Task WriteOutputAsync(RpcWriteOutputRequest request, CancellationToken cancellationToken)
-			{
-				PrintJson(request.Data.Span);
-				await _outputStream.WriteAsync(request.Data.Memory, cancellationToken);
-			}
-
-			void PrintJson(ReadOnlySpan<byte> span)
-			{
-				while (span.Length > 0)
-				{
-					int idx = span.IndexOf((byte)'\n');
-					if (idx == -1)
-					{
-						break;
-					}
-					_inner.LogInformation("{Line}", Encoding.UTF8.GetString(span.Slice(0, idx)));
-					span = span.Slice(idx + 1);
-				}
-			}
-		}
-
 		readonly List<string> _arguments = new List<string>();
 
 		/// <inheritdoc/>
@@ -123,34 +58,29 @@ namespace Horde.Agent.Commands.Utilities
 				baseDir = nextDir;
 			}
 
-			using (LogSink sink = new LogSink(runUatBat.Directory, logger))
+			using (LogParser filter = new LogParser(logger, new List<string>()))
 			{
-				await using (ServerLogger jsonLogger = new ServerLogger(sink, default, null, LogLevel.Information, logger, logger))
+				using (ManagedProcessGroup processGroup = new ManagedProcessGroup())
 				{
-					using (LogParser filter = new LogParser(jsonLogger, new List<string>()))
-					{
-						using (ManagedProcessGroup processGroup = new ManagedProcessGroup())
-						{
-							Dictionary<string, string> newEnvironment = ManagedProcess.GetCurrentEnvVars();
-							newEnvironment["UE_STDOUT_JSON"] = "1";
+					Dictionary<string, string> newEnvironment = ManagedProcess.GetCurrentEnvVars();
+					newEnvironment["UE_STDOUT_JSON"] = "1";
 
-							List<string> allArguments = new List<string>();
-							allArguments.Add(runUatBat.FullName);
-							allArguments.Add("BuildGraph");
-							allArguments.AddRange(_arguments);
+					List<string> allArguments = new List<string>();
+					allArguments.Add(runUatBat.FullName);
+					allArguments.Add("BuildGraph");
+					allArguments.AddRange(_arguments);
 
-							string fileName = Environment.GetEnvironmentVariable("COMSPEC") ?? "cmd.exe";
-							string commandLine = $"/C \"{CommandLineArguments.Join(allArguments)}\"";
-							using (ManagedProcess process = new ManagedProcess(processGroup, fileName, commandLine, runUatBat.Directory.FullName, newEnvironment, null, ProcessPriorityClass.Normal))
-							{
-								await process.CopyToAsync((buffer, offset, length) => filter.WriteData(buffer.AsMemory(offset, length)), 4096, CancellationToken.None);
-								await process.WaitForExitAsync(CancellationToken.None);
-							}
-						}
-						filter.Flush();
+					string fileName = Environment.GetEnvironmentVariable("COMSPEC") ?? "cmd.exe";
+					string commandLine = $"/C \"{CommandLineArguments.Join(allArguments)}\"";
+					using (ManagedProcess process = new ManagedProcess(processGroup, fileName, commandLine, runUatBat.Directory.FullName, newEnvironment, null, ProcessPriorityClass.Normal))
+					{
+						await process.CopyToAsync((buffer, offset, length) => filter.WriteData(buffer.AsMemory(offset, length)), 4096, CancellationToken.None);
+						await process.WaitForExitAsync(CancellationToken.None);
 					}
 				}
+				filter.Flush();
 			}
+
 			return 0;
 		}
 	}
diff --git a/Engine/Source/Programs/Horde/Horde.Agent.Driver/Utility/JobStepLogger.cs b/Engine/Source/Programs/Horde/Horde.Agent.Driver/Utility/JobStepLogger.cs
new file mode 100644
--- /dev/null
+++ b/Engine/Source/Programs/Horde/Horde.Agent.Driver/Utility/JobStepLogger.cs
@@ -0,0 +1,125 @@
+// Copyright Epic Games, Inc. All Rights Reserved.
+
+using System;
+using System.Collections.Generic;
+using System.Linq;
+using System.Text;
+using System.Threading.Tasks;
+using EpicGames.Horde;
+using EpicGames.Horde.Jobs;
+using EpicGames.Horde.Logs;
+using Horde.Agent.Utility;
+using Horde.Common.Rpc;
+using HordeCommon.Rpc;
+using Microsoft.Extensions.Logging;
+
+namespace Horde.Agent.Driver.Utility
+{
+	/// <summary>
+	/// Logger which updates the job step outcome whenever a warning or error is detected.
+	/// </summary>
+	class JobStepLogger : IServerLogger
+	{
+		/// <summary>
+		/// The current outcome for this step. Updated to reflect any errors and warnings that occurred.
+		/// </summary>
+		public JobStepOutcome Outcome => _outcome;
+
+		readonly IHordeClient _hordeClient;
+		readonly JobId _jobId;
+		readonly JobStepBatchId _jobBatchId;
+		readonly JobStepId _jobStepId;
+
+		readonly ServerLogger _inner;
+		JobStepOutcome _outcome;
+		Task _updateOutcomeTask;
+		CancellationTokenSource _cancellationSource;
+		ILogger _internalLogger;
+
+		/// <summary>
+		/// Constructor
+		/// </summary>
+		public JobStepLogger(IHordeClient hordeClient, LogId logId, ILogger localLogger, JobId jobId, JobStepBatchId jobBatchId, JobStepId jobStepId, bool? warnings, LogLevel outputLevel, ILogger internalLogger)
+		{
+			_hordeClient = hordeClient;
+			_jobId = jobId;
+			_jobBatchId = jobBatchId;
+			_jobStepId = jobStepId;
+			_inner = new ServerLogger(hordeClient, logId, warnings, outputLevel, localLogger, internalLogger);
+			_outcome = JobStepOutcome.Success;
+			_updateOutcomeTask = Task.CompletedTask;
+			_cancellationSource = new CancellationTokenSource();
+			_internalLogger = internalLogger;
+		}
+
+		/// <inheritdoc/>
+		public IDisposable? BeginScope<TState>(TState state) where TState : notnull
+			=> _inner.BeginScope(state);
+
+		/// <inheritdoc/>
+		public async ValueTask DisposeAsync()
+		{
+			if (!_updateOutcomeTask.IsCompleted)
+			{
+				_cancellationSource.Cancel();
+				await _updateOutcomeTask;
+				_cancellationSource.Dispose();
+			}
+
+			await _inner.DisposeAsync();
+		}
+
+		/// <inheritdoc/>
+		public bool IsEnabled(LogLevel logLevel)
+			=> _inner.IsEnabled(logLevel);
+
+		/// <inheritdoc/>
+		public void Log<TState>(LogLevel logLevel, EventId eventId, TState state, Exception? exception, Func<TState, Exception?, string> formatter)
+		{
+			// Forward the log data to the inner writer
+			_inner.Log<TState>(logLevel, eventId, state, exception, formatter);
+
+			// Update the state of this job if this is an error status
+			JobStepOutcome newOutcome = _outcome;
+			if (logLevel == LogLevel.Error || logLevel == LogLevel.Critical)
+			{
+				newOutcome = JobStepOutcome.Failure;
+			}
+			else if (logLevel == LogLevel.Warning && Outcome != JobStepOutcome.Failure)
+			{
+				newOutcome = JobStepOutcome.Warnings;
+			}
+
+			// If it changed, create a new task to update the job state on the server
+			if (newOutcome != _outcome)
+			{
+				Task prevUpdateTask = _updateOutcomeTask;
+				_updateOutcomeTask = Task.Run(() => UpdateOutcomeAsync(prevUpdateTask, newOutcome, _cancellationSource.Token), _cancellationSource.Token);
+			}
+		}
+
+		/// <inheritdoc/>
+		public async Task StopAsync()
+		{
+			await _inner.StopAsync();
+			await _updateOutcomeTask;
+		}
+
+		async Task UpdateOutcomeAsync(Task prevTask, JobStepOutcome newOutcome, CancellationToken cancellationToken)
+		{
+			// Wait for the last task to complete
+			await prevTask;
+
+			// Update the outcome of this jobstep
+			try
+			{
+				JobRpc.JobRpcClient jobRpc = await _hordeClient.CreateGrpcClientAsync<JobRpc.JobRpcClient>(cancellationToken);
+				await jobRpc.UpdateStepAsync(new RpcUpdateStepRequest(_jobId, _jobBatchId, _jobStepId, JobStepState.Unspecified, newOutcome), cancellationToken: cancellationToken);
+			}
+			catch (Exception ex)
+			{
+				_internalLogger.LogWarning(ex, "Unable to update step outcome to {NewOutcome}", Outcome);
+			}
+		}
+	}
+}
diff --git a/Engine/Source/Programs/Horde/Horde.Agent.Driver/Utility/JsonRpcLogSink.cs b/Engine/Source/Programs/Horde/Horde.Agent.Driver/Utility/JsonRpcLogSink.cs
--- a/Engine/Source/Programs/Horde/Horde.Agent.Driver/Utility/JsonRpcLogSink.cs
+++ b/Engine/Source/Programs/Horde/Horde.Agent.Driver/Utility/JsonRpcLogSink.cs
@@ -15,21 +15,11 @@ using static Horde.Common.Rpc.LogRpc;
 
 namespace Horde.Agent.Utility
 {
-	interface IJsonRpcLogSink : IAsyncDisposable
-	{
-		Task WriteEventsAsync(List<RpcCreateEventRequest> events, CancellationToken cancellationToken);
-		Task WriteOutputAsync(RpcWriteOutputRequest request, CancellationToken cancellationToken);
-		Task SetOutcomeAsync(JobStepOutcome outcome, CancellationToken cancellationToken);
-	}
-
-	class JsonRpcAndStorageLogSink : IJsonRpcLogSink, IAsyncDisposable
+	class JsonRpcAndStorageLogSink : IAsyncDisposable
 	{
 		const int FlushLength = 1024 * 1024;
 
 		readonly IHordeClient _hordeClient;
-		readonly JobId? _jobId;
-		readonly JobStepBatchId? _jobBatchId;
-		readonly JobStepId? _jobStepId;
 		readonly LogId _logId;
 		readonly LogBuilder _builder;
 		readonly IStorageClient _store;
@@ -43,13 +33,10 @@ namespace Horde.Agent.Utility
 		AsyncEvent _tailTaskStop;
 		readonly AsyncEvent _newTailDataEvent = new AsyncEvent();
 
-		public JsonRpcAndStorageLogSink(IHordeClient hordeClient, LogId logId, JobId? jobId, JobStepBatchId? jobBatchId, JobStepId? jobStepId, ILogger logger)
+		public JsonRpcAndStorageLogSink(IHordeClient hordeClient, LogId logId, ILogger logger)
 		{
 			_hordeClient = hordeClient;
 			_logId = logId;
-			_jobId = jobId;
-			_jobBatchId = jobBatchId;
-			_jobStepId = jobStepId;
 			_builder = new LogBuilder(LogFormat.Json, logger);
 			_store = _hordeClient.CreateStorageClient(logId);
 			_writer = _store.CreateBlobWriter();
@@ -168,24 +155,6 @@ namespace Horde.Agent.Utility
 			return lines;
 		}
 
-		/// <inheritdoc/>
-		public async Task SetOutcomeAsync(JobStepOutcome outcome, CancellationToken cancellationToken)
-		{
-			// Update the outcome of this jobstep
-			if (_jobId != null && _jobBatchId != null && _jobStepId != null)
-			{
-				JobRpc.JobRpcClient jobRpc = await _hordeClient.CreateGrpcClientAsync<JobRpc.JobRpcClient>(cancellationToken);
-				try
-				{
-					await jobRpc.UpdateStepAsync(new RpcUpdateStepRequest(_jobId.Value, _jobBatchId.Value, _jobStepId.Value, JobStepState.Unspecified, outcome), cancellationToken: cancellationToken);
-				}
-				catch (Exception ex)
-				{
-					_logger.LogWarning(ex, "Unable to update step outcome to {NewOutcome}", outcome);
-				}
-			}
-		}
-
 		/// <inheritdoc/>
 		public async Task WriteEventsAsync(List<RpcCreateEventRequest> events, CancellationToken cancellationToken)
 		{
diff --git a/Engine/Source/Programs/Horde/Horde.Agent.Driver/Utility/ServerLogger.cs b/Engine/Source/Programs/Horde/Horde.Agent.Driver/Utility/ServerLogger.cs
--- a/Engine/Source/Programs/Horde/Horde.Agent.Driver/Utility/ServerLogger.cs
+++ b/Engine/Source/Programs/Horde/Horde.Agent.Driver/Utility/ServerLogger.cs
@@ -3,9 +3,11 @@
 using System.Text;
 using System.Threading.Channels;
 using EpicGames.Core;
+using EpicGames.Horde;
 using EpicGames.Horde.Jobs;
 using EpicGames.Horde.Logs;
 using Google.Protobuf;
+using Horde.Common.Rpc;
 using HordeCommon.Rpc;
 using Microsoft.Extensions.Logging;
 
@@ -18,11 +20,6 @@ namespace Horde.Agent.Utility
 	/// </summary>
 	public interface IServerLogger : ILogger, IAsyncDisposable
 	{
-		/// <summary>
-		/// Outcome of the job step, including any warnings/errors
-		/// </summary>
-		JobStepOutcome Outcome { get; }
-
 		/// <summary>
 		/// Flushes the logger with the server and stops the background work
 		/// </summary>
@@ -34,7 +31,7 @@ namespace Horde.Agent.Utility
 	/// </summary>
 	sealed class ServerLogger : IServerLogger
 	{
-		readonly IJsonRpcLogSink _sink;
+		readonly JsonRpcAndStorageLogSink _sink;
 		readonly LogId _logId;
 		readonly bool _warnings;
 		readonly LogLevel _outputLevel;
@@ -43,27 +40,18 @@ namespace Horde.Agent.Utility
 		readonly Channel<JsonLogEvent> _dataChannel;
 		Task? _dataWriter;
 
-		/// <summary>
-		/// The current outcome for this step. Updated to reflect any errors and warnings that occurred.
-		/// </summary>
-		public JobStepOutcome Outcome
-		{
-			get;
-			private set;
-		}
-
 		/// <summary>
 		/// Constructor
 		/// </summary>
-		/// <param name="sink">Sink for log events</param>
+		/// <param name="hordeClient">Horde instance to write to</param>
 		/// <param name="logId">The log id to write to</param>
 		/// <param name="warnings">Whether to include warnings in the output</param>
 		/// <param name="outputLevel">Minimum level for output</param>
 		/// <param name="localLogger">Logger to forward log messages to</param>
 		/// <param name="agentLogger">Logger for systemic messages</param>
-		public ServerLogger(IJsonRpcLogSink sink, LogId logId, bool? warnings, LogLevel outputLevel, ILogger localLogger, ILogger agentLogger)
+		public ServerLogger(IHordeClient hordeClient, LogId logId, bool? warnings, LogLevel outputLevel, ILogger localLogger, ILogger agentLogger)
 		{
-			_sink = sink;
+			_sink = new JsonRpcAndStorageLogSink(hordeClient, logId, agentLogger); 
 			_logId = logId;
 			_warnings = warnings ?? true;
 			_outputLevel = outputLevel;
@@ -71,8 +59,6 @@ namespace Horde.Agent.Utility
 			_agentLogger = agentLogger;
 			_dataChannel = Channel.CreateUnbounded<JsonLogEvent>();
 			_dataWriter = Task.Run(() => RunDataWriterAsync());
-
-			Outcome = JobStepOutcome.Success;
 		}
 
 		/// <inheritdoc/>
@@ -98,18 +84,6 @@ namespace Horde.Agent.Utility
 
 		private void WriteFormattedEvent(JsonLogEvent jsonLogEvent)
 		{
-			// Update the state of this job if this is an error status
-			LogLevel level = jsonLogEvent.Level;
-			if (level == LogLevel.Error || level == LogLevel.Critical)
-			{
-				Outcome = JobStepOutcome.Failure;
-			}
-			else if (level == LogLevel.Warning && Outcome != JobStepOutcome.Failure)
-			{
-				Outcome = JobStepOutcome.Warnings;
-			}
-
-			// Write the event
 			if (!_dataChannel.Writer.TryWrite(jsonLogEvent))
 			{
 				throw new InvalidOperationException("Expected unbounded writer to complete immediately");
@@ -162,9 +136,6 @@ namespace Horde.Agent.Utility
 			JsonRpcLogWriter writer = new JsonRpcLogWriter();
 			List<RpcCreateEventRequest> events = new List<RpcCreateEventRequest>();
 
-			// The current jobstep outcome
-			JobStepOutcome postedOutcome = JobStepOutcome.Success;
-
 			// Whether we've written the flush command
 			for (; ; )
 			{
@@ -233,20 +204,6 @@ namespace Horde.Agent.Utility
 					}
 				}
 
-				// Update the outcome of this jobstep
-				if (Outcome != postedOutcome)
-				{
-					try
-					{
-						await _sink.SetOutcomeAsync(Outcome, CancellationToken.None);
-					}
-					catch (Exception ex)
-					{
-						_agentLogger.LogWarning(ex, "Unable to update step outcome to {NewOutcome}", Outcome);
-					}
-					postedOutcome = Outcome;
-				}
-
 				// Wait for more data to be available
 				if (writer.PacketLength <= 0 && !await _dataChannel.Reader.WaitToReadAsync())
 				{
diff --git a/Engine/Source/Programs/Horde/Horde.Agent.Tests/JsonRpcLoggerTests.cs b/Engine/Source/Programs/Horde/Horde.Agent.Tests/JsonRpcLoggerTests.cs
--- a/Engine/Source/Programs/Horde/Horde.Agent.Tests/JsonRpcLoggerTests.cs
+++ b/Engine/Source/Programs/Horde/Horde.Agent.Tests/JsonRpcLoggerTests.cs
@@ -5,7 +5,9 @@ using System.Collections.Generic;
 using System.Linq;
 using System.Text;
 using System.Threading;
+using System.Threading.Channels;
 using System.Threading.Tasks;
+using Amazon.EC2.Model;
 using EpicGames.Core;
 using EpicGames.Horde;
 using EpicGames.Horde.Compute;
@@ -16,6 +18,8 @@ using EpicGames.Horde.Storage.Bundles;
 using Grpc.Core;
 using Grpc.Net.Client;
 using Horde.Agent.Utility;
+using Horde.Common.Rpc;
+using Microsoft.CodeAnalysis.CSharp.Syntax;
 using Microsoft.Extensions.Logging;
 using Microsoft.Extensions.Logging.Abstractions;
 using Microsoft.VisualStudio.TestTools.UnitTesting;
@@ -125,30 +129,92 @@ namespace Horde.Agent.Tests
 			Assert.AreEqual(result[0], Msg);
 		}
 
-		class FakeJsonRpcLoggerBackend : JsonRpcAndStorageLogSink
+		class GrpcHelpers
 		{
-			public IBlobRef? Target { get; private set; }
+			class StreamReader<T> : IAsyncStreamReader<T> where T : class
+			{
+				readonly ChannelReader<T> _reader;
+				T? _current;
+
+				public T Current => _current ?? throw new InvalidOperationException();
+
+				public StreamReader(ChannelReader<T> reader)
+					=> _reader = reader;
+
+				public async Task<bool> MoveNext(CancellationToken cancellationToken)
+				{
+					while (await _reader.WaitToReadAsync(cancellationToken))
+					{
+						if (_reader.TryRead(out _current))
+						{
+							return true;
+						}
+					}
+					return false;
+				}
+			}
 
-			public FakeJsonRpcLoggerBackend(IHordeClient hordeClient, LogId logId, JobId? jobId, JobStepBatchId? batchId, JobStepId? stepId, ILogger logger)
-				: base(hordeClient, logId, jobId, batchId, stepId, logger)
+			class StreamWriter<T> : IClientStreamWriter<T>
 			{
+				readonly ChannelWriter<T> _writer;
+
+				public StreamWriter(ChannelWriter<T> writer)
+					=> _writer = writer;
+
+				public WriteOptions? WriteOptions { get => throw new NotImplementedException(); set => throw new NotImplementedException(); }
+
+				public Task CompleteAsync()
+				{
+					_writer.Complete();
+					return Task.CompletedTask;
+				}
+
+				public async Task WriteAsync(T message)
+					=> await _writer.WriteAsync(message);
 			}
 
-			protected override Task UpdateLogAsync(IBlobRef target, int lineCount, bool complete, CancellationToken cancellationToken)
+			public static AsyncDuplexStreamingCall<TRequest, TResponse> CreateDuplexCall<TRequest, TResponse>(Func<ChannelReader<TRequest>, ChannelWriter<TResponse>, Task> func) 
+				where TRequest : class 
+				where TResponse : class
 			{
-				Target = target;
-				return Task.CompletedTask;
+				Channel<TRequest> requests = Channel.CreateUnbounded<TRequest>();
+				Channel<TResponse> responses = Channel.CreateUnbounded<TResponse>();
+
+				BackgroundTask runner = BackgroundTask.StartNew(async ctx =>
+				{
+					await func(requests.Reader, responses.Writer);
+					responses.Writer.Complete();
+				});
+
+				return new AsyncDuplexStreamingCall<TRequest, TResponse>(new StreamWriter<TRequest>(requests.Writer), new StreamReader<TResponse>(responses.Reader), Task.FromResult<Metadata>(null!), () => Status.DefaultSuccess, () => null!, () => runner.DisposeAsync().AsTask().Wait());
 			}
+		}
 
-			protected override async Task<int> UpdateLogTailAsync(int tailNext, ReadOnlyMemory<byte> tailData, CancellationToken cancellationToken)
+		class FakeLogRpcClient : LogRpc.LogRpcClient
+		{
+			public Dictionary<LogId, BlobLocator> Logs { get; } = new Dictionary<LogId, BlobLocator>();
+
+			public override AsyncUnaryCall<UpdateLogResponse> UpdateLogAsync(UpdateLogRequest request, CallOptions options)
 			{
-				await Task.Delay(TimeSpan.FromSeconds(2.0), cancellationToken);
-				return -1;
+				Logs[LogId.Parse(request.LogId)] = new BlobLocator(request.TargetLocator);
+				return new AsyncUnaryCall<UpdateLogResponse>(Task.FromResult(new UpdateLogResponse()), Task.FromResult(new Metadata()), () => Status.DefaultSuccess, () => new Metadata(), () => { });
+			}
+
+			public override AsyncDuplexStreamingCall<UpdateLogTailRequest, UpdateLogTailResponse> UpdateLogTail(CallOptions options)
+			{
+				return GrpcHelpers.CreateDuplexCall<UpdateLogTailRequest, UpdateLogTailResponse>(async (reader, writer) =>
+				{
+					while (await reader.WaitToReadAsync())
+					{
+						reader.TryRead(out UpdateLogTailRequest? request);
+					}
+				});
 			}
 		}
 
 		class FakeHordeClient : IHordeClient
 		{
+			public FakeLogRpcClient LogRpc { get; } = new FakeLogRpcClient();
 			public Dictionary<string, BundleStorageClient> StorageClients { get; } = new Dictionary<string, BundleStorageClient>();
 
 			public Uri ServerUrl => throw new NotImplementedException();
@@ -190,7 +256,13 @@ namespace Horde.Agent.Tests
 				=> throw new NotImplementedException();
 
 			public Task<TClient> CreateGrpcClientAsync<TClient>(CancellationToken cancellationToken = default) where TClient : ClientBase<TClient>
-				=> throw new NotImplementedException();
+			{
+				if (typeof(TClient) == typeof(LogRpc.LogRpcClient))
+				{
+					return Task.FromResult<TClient>((TClient)(object)LogRpc);
+				}
+				throw new NotImplementedException();
+			}
 
 			public bool HasValidAccessToken()
 				=> throw new NotImplementedException();
@@ -202,21 +274,21 @@ namespace Horde.Agent.Tests
 			await using BundleCache cache = new BundleCache();
 			await using FakeHordeClient hordeClient = new FakeHordeClient();
 
-			const int Count = 20000;
+			LogId logId = default;
 
-			LogNode file;
-			await using (FakeJsonRpcLoggerBackend sink = new FakeJsonRpcLoggerBackend(hordeClient, default, null, null, null, NullLogger.Instance))
+			const int Count = 20000;
+			await using (ServerLogger logger = new ServerLogger(hordeClient, logId, null, LogLevel.Information, NullLogger.Instance, NullLogger.Instance))
 			{
-				await using (ServerLogger logger = new ServerLogger(sink, default, null, LogLevel.Information, NullLogger.Instance, NullLogger.Instance))
+				for (int idx = 0; idx < Count; idx++)
 				{
-					for (int idx = 0; idx < Count; idx++)
-					{
-						logger.LogInformation("Testing {Number}", idx);
-					}
+					logger.LogInformation("Testing {Number}", idx);
 				}
-				file = await sink.Target!.ReadBlobAsync<LogNode>();
 			}
 
+			// Read the log
+			using IStorageClient storageClient = hordeClient.CreateStorageClient(logId);
+			LogNode file = await storageClient.CreateBlobHandle(hordeClient.LogRpc.Logs[logId]).ReadBlobAsync<LogNode>();
+
 			// Check the index text
 			List<Utf8String> extractedIndexText = new List<Utf8String>();
 
diff --git a/Engine/Source/Programs/Horde/Horde.Agent.Tests/WorkerServiceTest.cs b/Engine/Source/Programs/Horde/Horde.Agent.Tests/WorkerServiceTest.cs
--- a/Engine/Source/Programs/Horde/Horde.Agent.Tests/WorkerServiceTest.cs
+++ b/Engine/Source/Programs/Horde/Horde.Agent.Tests/WorkerServiceTest.cs
@@ -53,8 +53,6 @@ namespace Horde.Agent.Tests
 
 		class FakeServerLogger : IServerLogger
 		{
-			public JobStepOutcome Outcome => JobStepOutcome.Success;
-
 			public IDisposable? BeginScope<TState>(TState state) where TState : notnull => NullLogger.Instance.BeginScope<TState>(state);
 
 			public ValueTask DisposeAsync() => new ValueTask();
@@ -68,7 +66,7 @@ namespace Horde.Agent.Tests
 
 		class FakeServerLoggerFactory : IServerLoggerFactory
 		{
-			public IServerLogger CreateLogger(IHordeClient hordeClient, LogId logId, ILogger localLogger, JobId? jobId, JobStepBatchId? batchId, JobStepId? stepId, bool? warnings = null, LogLevel outputLevel = LogLevel.Information) => new FakeServerLogger();
+			public IServerLogger CreateLogger(IHordeClient hordeClient, LogId logId, ILogger localLogger, bool? warnings = null, LogLevel outputLevel = LogLevel.Information) => new FakeServerLogger();
 		}
 
 		internal static IJobExecutor NullExecutor = new SimpleTestExecutor(async (step, logger, cancellationToken) =>
diff --git a/Engine/Source/Programs/Horde/Horde.Agent/Leases/Handlers/JobHandler.cs b/Engine/Source/Programs/Horde/Horde.Agent/Leases/Handlers/JobHandler.cs
--- a/Engine/Source/Programs/Horde/Horde.Agent/Leases/Handlers/JobHandler.cs
+++ b/Engine/Source/Programs/Horde/Horde.Agent/Leases/Handlers/JobHandler.cs
@@ -10,6 +10,7 @@ using EpicGames.Horde.Logs;
 using Google.Protobuf;
 using Grpc.Core;
 using Horde.Agent.Driver;
+using Horde.Agent.Driver.Utility;
 using Horde.Agent.Execution;
 using Horde.Agent.Services;
 using Horde.Agent.Utility;
@@ -18,6 +19,7 @@ using HordeCommon.Rpc;
 using HordeCommon.Rpc.Messages;
 using HordeCommon.Rpc.Tasks;
 using Microsoft.Extensions.Logging;
+using Microsoft.Extensions.Logging.Abstractions;
 using Microsoft.Extensions.Options;
 using OpenTracing;
 using OpenTracing.Util;
@@ -158,7 +160,7 @@ namespace Horde.Agent.Leases.Handlers
 
 			// Create a storage client for this session
 			RpcJobOptions jobOptions = executeTask.JobOptions;
-			await using IServerLogger logger = _serverLoggerFactory.CreateLogger(session.HordeClient, executeTask.LogId, localLogger, executeTask.JobId, executeTask.BatchId, null, null);
+			await using IServerLogger logger = _serverLoggerFactory.CreateLogger(session.HordeClient, executeTask.LogId, localLogger, null);
 
 			logger.LogInformation("Executing job \"{JobName}\", jobId {JobId}, batchId {BatchId}, leaseId {LeaseId}, agentVersion {AgentVersion}", executeTask.JobName, executeTask.JobId, executeTask.BatchId, leaseId, AgentApp.Version);
 
@@ -310,7 +312,7 @@ namespace Horde.Agent.Leases.Handlers
 					{
 						// Start writing to the log file
 #pragma warning disable CA2000 // Dispose objects before losing scope
-						await using (IServerLogger stepLogger = _serverLoggerFactory.CreateLogger(session.HordeClient, step.LogId, localLogger, options.JobId, options.BatchId, step.StepId, step.Warnings))
+						await using (JobStepLogger stepLogger = new JobStepLogger(session.HordeClient, step.LogId, localLogger, options.JobId, options.BatchId, step.StepId, step.Warnings, LogLevel.Debug, logger))
 						{
 							// Execute the task
 							using CancellationTokenSource stepPollCancelSource = new CancellationTokenSource();
diff --git a/Engine/Source/Programs/Horde/Horde.Agent/Services/ServerLoggerFactory.cs b/Engine/Source/Programs/Horde/Horde.Agent/Services/ServerLoggerFactory.cs
--- a/Engine/Source/Programs/Horde/Horde.Agent/Services/ServerLoggerFactory.cs
+++ b/Engine/Source/Programs/Horde/Horde.Agent/Services/ServerLoggerFactory.cs
@@ -19,34 +19,10 @@ namespace Horde.Agent.Services
 		/// <param name="hordeClient">The Horde client</param>
 		/// <param name="logId">The log id</param>
 		/// <param name="localLogger">Local log output device</param>
-		/// <param name="jobId">Job containing the log</param>
-		/// <param name="batchId">Job batch id</param>
-		/// <param name="stepId">The job step id</param>
 		/// <param name="warnings">Whether to suppress warnings</param>
 		/// <param name="outputLevel">Minimum output level for messages</param>
 		/// <returns>New logger instance</returns>
-		IServerLogger CreateLogger(IHordeClient hordeClient, LogId logId, ILogger localLogger, JobId? jobId, JobStepBatchId? batchId, JobStepId? stepId, bool? warnings, LogLevel outputLevel = LogLevel.Information);
-	}
-
-	/// <summary>
-	/// Extension methods for <see cref="IServerLoggerFactory"/>
-	/// </summary>
-	static class ServerLoggerFactoryExtensions
-	{
-		/// <summary>
-		/// Creates a logger which uploads data to the server
-		/// </summary>
-		/// <param name="service">Service instance</param>
-		/// <param name="hordeClient">The current session</param>
-		/// <param name="logId">The log identifier</param>
-		/// <param name="localLogger">Local log output device</param>
-		/// <param name="warnings">Whether to suppress warnings</param>
-		/// <param name="outputLevel">Minimum output level for messages</param>
-		/// <returns>New logger instance</returns>
-		public static IServerLogger CreateLogger(this IServerLoggerFactory service, IHordeClient hordeClient, LogId logId, ILogger localLogger, bool? warnings, LogLevel outputLevel = LogLevel.Information)
-		{
-			return service.CreateLogger(hordeClient, logId, localLogger, null, null, null, warnings, outputLevel);
-		}
+		IServerLogger CreateLogger(IHordeClient hordeClient, LogId logId, ILogger localLogger, bool? warnings, LogLevel outputLevel = LogLevel.Information);
 	}
 
 	/// <summary>
@@ -65,11 +41,10 @@ namespace Horde.Agent.Services
 		}
 
 		/// <inheritdoc/>
-		public IServerLogger CreateLogger(IHordeClient hordeClient, LogId logId, ILogger localLogger, JobId? jobId, JobStepBatchId? batchId, JobStepId? stepId, bool? warnings, LogLevel outputLevel)
+		public IServerLogger CreateLogger(IHordeClient hordeClient, LogId logId, ILogger localLogger, bool? warnings, LogLevel outputLevel)
 		{
 #pragma warning disable CA2000 // Dispose objects before losing scope
-			IJsonRpcLogSink sink = new JsonRpcAndStorageLogSink(hordeClient, logId, jobId, batchId, stepId, _logger);
-			return new ServerLogger(sink, logId, warnings, outputLevel, localLogger, _logger);
+			return new ServerLogger(hordeClient, logId, warnings, outputLevel, localLogger, _logger);
 #pragma warning restore CA2000 // Dispose objects before losing scope
 		}
 	}
