Horde: Merge JsonRpcLogSink and ServerLogger classees.

#rnx

diff --git a/Engine/Source/Programs/Horde/Horde.Agent.Driver/Utility/JsonRpcLogSink.cs b/Engine/Source/Programs/Horde/Horde.Agent.Driver/Utility/JsonRpcLogSink.cs
deleted file mode 100644
diff --git a/Engine/Source/Programs/Horde/Horde.Agent.Driver/Utility/ServerLogger.cs b/Engine/Source/Programs/Horde/Horde.Agent.Driver/Utility/ServerLogger.cs
--- a/Engine/Source/Programs/Horde/Horde.Agent.Driver/Utility/ServerLogger.cs
+++ b/Engine/Source/Programs/Horde/Horde.Agent.Driver/Utility/ServerLogger.cs
@@ -6,13 +6,16 @@ using EpicGames.Core;
 using EpicGames.Horde;
 using EpicGames.Horde.Jobs;
 using EpicGames.Horde.Logs;
+using EpicGames.Horde.Storage;
 using Google.Protobuf;
+using Grpc.Core;
 using Horde.Common.Rpc;
 using HordeCommon.Rpc;
 using Microsoft.Extensions.Logging;
 
 namespace Horde.Agent.Utility
 {
+	using static Horde.Common.Rpc.LogRpc;
 	using ByteString = Google.Protobuf.ByteString;
 
 	/// <summary>
@@ -31,11 +34,24 @@ namespace Horde.Agent.Utility
 	/// </summary>
 	sealed class ServerLogger : IServerLogger
 	{
-		readonly JsonRpcAndStorageLogSink _sink;
+		const int FlushLength = 1024 * 1024;
+
+		readonly IHordeClient _hordeClient;
 		readonly LogId _logId;
+		readonly LogBuilder _builder;
+		readonly IStorageClient _store;
+		readonly IBlobWriter _writer;
+
+		int _bufferLength;
+
+		// Tailing task
+		readonly Task _tailTask;
+		AsyncEvent _tailTaskStop;
+		readonly AsyncEvent _newTailDataEvent = new AsyncEvent();
+
 		readonly bool _warnings;
 		readonly LogLevel _outputLevel;
-		readonly ILogger _forwardLogger;
+		readonly ILogger _localLogger;
 		readonly ILogger _agentLogger;
 		readonly Channel<JsonLogEvent> _dataChannel;
 		Task? _dataWriter;
@@ -51,11 +67,19 @@ namespace Horde.Agent.Utility
 		/// <param name="agentLogger">Logger for systemic messages</param>
 		public ServerLogger(IHordeClient hordeClient, LogId logId, bool? warnings, LogLevel outputLevel, ILogger localLogger, ILogger agentLogger)
 		{
-			_sink = new JsonRpcAndStorageLogSink(hordeClient, logId, agentLogger); 
+			_hordeClient = hordeClient;
+			_logId = logId;
+			_builder = new LogBuilder(LogFormat.Json, agentLogger);
+			_store = _hordeClient.CreateStorageClient(logId);
+			_writer = _store.CreateBlobWriter();
+
+			_tailTaskStop = new AsyncEvent();
+			_tailTask = Task.Run(() => TickTailAsync());
+
 			_logId = logId;
 			_warnings = warnings ?? true;
 			_outputLevel = outputLevel;
-			_forwardLogger = localLogger;
+			_localLogger = localLogger;
 			_agentLogger = agentLogger;
 			_dataChannel = Channel.CreateUnbounded<JsonLogEvent>();
 			_dataWriter = Task.Run(() => RunDataWriterAsync());
@@ -64,7 +88,7 @@ namespace Horde.Agent.Utility
 		/// <inheritdoc/>
 		public void Log<TState>(LogLevel logLevel, EventId eventId, TState state, Exception? exception, Func<TState, Exception?, string> formatter)
 		{
-			_forwardLogger.Log(logLevel, eventId, state, exception, formatter);
+			_localLogger.Log(logLevel, eventId, state, exception, formatter);
 
 			// Downgrade warnings to information if not required
 			if (logLevel == LogLevel.Warning && !_warnings)
@@ -77,10 +101,10 @@ namespace Horde.Agent.Utility
 		}
 
 		/// <inheritdoc/>
-		public bool IsEnabled(LogLevel logLevel) => logLevel >= _outputLevel || _forwardLogger.IsEnabled(logLevel);
+		public bool IsEnabled(LogLevel logLevel) => logLevel >= _outputLevel || _localLogger.IsEnabled(logLevel);
 
 		/// <inheritdoc/>
-		public IDisposable? BeginScope<TState>(TState state) where TState : notnull => _forwardLogger.BeginScope(state);
+		public IDisposable? BeginScope<TState>(TState state) where TState : notnull => _localLogger.BeginScope(state);
 
 		private void WriteFormattedEvent(JsonLogEvent jsonLogEvent)
 		{
@@ -110,7 +134,27 @@ namespace Horde.Agent.Utility
 		public async ValueTask DisposeAsync()
 		{
 			await StopAsync();
-			await _sink.DisposeAsync();
+
+			_agentLogger.LogInformation("Disposing json log task");
+
+			if (_tailTaskStop != null)
+			{
+				_tailTaskStop.Latch();
+				_newTailDataEvent.Latch();
+
+				await _tailTask;
+				_tailTaskStop = null!;
+			}
+
+			if (_writer != null)
+			{
+				await _writer.DisposeAsync();
+			}
+
+			if (_store != null)
+			{
+				_store.Dispose();
+			}
 		}
 
 		/// <summary>
@@ -181,7 +225,7 @@ namespace Horde.Agent.Utility
 					(ReadOnlyMemory<byte> packet, int packetLineCount) = writer.CreatePacket();
 					try
 					{
-						await _sink.WriteOutputAsync(new RpcWriteOutputRequest(_logId, packetOffset, packetLineIndex, UnsafeByteOperations.UnsafeWrap(packet), false), CancellationToken.None);
+						await WriteOutputAsync(new RpcWriteOutputRequest(_logId, packetOffset, packetLineIndex, UnsafeByteOperations.UnsafeWrap(packet), false), CancellationToken.None);
 						packetOffset += packet.Length;
 						packetLineIndex += packetLineCount;
 					}
@@ -196,7 +240,7 @@ namespace Horde.Agent.Utility
 				{
 					try
 					{
-						await _sink.WriteEventsAsync(events, CancellationToken.None);
+						await WriteEventsAsync(events, CancellationToken.None);
 					}
 					catch (Exception ex)
 					{
@@ -209,7 +253,7 @@ namespace Horde.Agent.Utility
 				{
 					try
 					{
-						await _sink.WriteOutputAsync(new RpcWriteOutputRequest(_logId, packetOffset, packetLineIndex, ByteString.Empty, true), CancellationToken.None);
+						await WriteOutputAsync(new RpcWriteOutputRequest(_logId, packetOffset, packetLineIndex, ByteString.Empty, true), CancellationToken.None);
 					}
 					catch (Exception ex)
 					{
@@ -231,5 +275,176 @@ namespace Horde.Agent.Utility
 				_agentLogger.LogError(ex, "Exception while trying to parse line count from data ({Message})", Encoding.UTF8.GetString(span));
 			}
 		}
+
+		async Task TickTailAsync()
+		{
+			for (; ; )
+			{
+				try
+				{
+					await TickTailInternalAsync();
+					break;
+				}
+				catch (OperationCanceledException ex)
+				{
+					_agentLogger.LogInformation(ex, "Cancelled log tailing task");
+					break;
+				}
+				catch (Exception ex)
+				{
+					_agentLogger.LogError(ex, "Exception on log tailing task ({LogId}): {Message}", _logId, ex.Message);
+					await Task.Delay(TimeSpan.FromSeconds(10.0));
+				}
+			}
+		}
+
+		async Task TickTailInternalAsync()
+		{
+			int tailNext = -1;
+			Task tickTask = Task.CompletedTask;
+			while (!_tailTaskStop.IsSet())
+			{
+				Task newTailDataTask = _newTailDataEvent.Task;
+				int initialTailNext = tailNext;
+
+				// Get the data to send to the server
+				ReadOnlyMemory<byte> tailData = ReadOnlyMemory<byte>.Empty;
+				if (tailNext != -1)
+				{
+					(tailNext, tailData) = _builder.ReadTailData(tailNext, 16 * 1024);
+				}
+
+				// If we don't have any updates for the server, wait until we do. We need to ensure
+				// we keep pumping the RPC with the server in case the requested tail next value changes,
+				// and to make sure that we don't expire the existing tail data.
+				if (tailNext != -1 && tailData.IsEmpty && tailNext == initialTailNext && !tickTask.IsCompleted)
+				{
+					_agentLogger.LogInformation("No tail data available for log {LogId} after line {TailNext}; waiting for more...", _logId, tailNext);
+					await Task.WhenAny(newTailDataTask, tickTask);
+					continue;
+				}
+
+				string start = "";
+				if (tailData.Length > 0)
+				{
+					start = Encoding.UTF8.GetString(tailData.Slice(0, Math.Min(tailData.Length, 256)).Span);
+				}
+
+				// Update the next tailing position
+				int numLines = CountLines(tailData.Span);
+				_agentLogger.LogInformation("Setting log {LogId} tail = {TailNext}, data = {TailDataSize} bytes, {NumLines} lines ('{Start}')", _logId, tailNext, tailData.Length, numLines, start);
+
+				int newTailNext = await UpdateLogTailAsync(tailNext, tailData, CancellationToken.None);
+				_agentLogger.LogInformation("Log {LogId} tail next = {TailNext}", _logId, newTailNext);
+
+				if (newTailNext != tailNext)
+				{
+					tailNext = newTailNext;
+					_agentLogger.LogInformation("Modified tail position for log {LogId} to {TailNext}", _logId, tailNext);
+				}
+
+				tickTask = Task.Delay(TimeSpan.FromSeconds(10.0));
+			}
+			_agentLogger.LogInformation("Finishing log tail task");
+		}
+
+		static int CountLines(ReadOnlySpan<byte> data)
+		{
+			int lines = 0;
+			for (int idx = 0; idx < data.Length; idx++)
+			{
+				if (data[idx] == '\n')
+				{
+					lines++;
+				}
+			}
+			return lines;
+		}
+
+		/// <inheritdoc/>
+		public async Task WriteEventsAsync(List<RpcCreateEventRequest> events, CancellationToken cancellationToken)
+		{
+			JobRpc.JobRpcClient jobRpc = await _hordeClient.CreateGrpcClientAsync<JobRpc.JobRpcClient>(cancellationToken);
+			await jobRpc.CreateEventsAsync(new RpcCreateEventsRequest(events), cancellationToken: cancellationToken);
+		}
+
+		/// <inheritdoc/>
+		public async Task WriteOutputAsync(RpcWriteOutputRequest request, CancellationToken cancellationToken)
+		{
+			_builder.WriteData(request.Data.Memory);
+			_bufferLength += request.Data.Length;
+
+			if (request.Flush || _bufferLength > FlushLength)
+			{
+				IBlobRef<LogNode> target = await _builder.FlushAsync(_writer, request.Flush, cancellationToken);
+				await UpdateLogAsync(target, _builder.LineCount, request.Flush, cancellationToken);
+				_bufferLength = 0;
+			}
+
+			_newTailDataEvent.Set();
+		}
+
+		#region RPC calls
+
+		async Task UpdateLogAsync(IBlobRef target, int lineCount, bool complete, CancellationToken cancellationToken)
+		{
+			_agentLogger.LogInformation("Updating log {LogId} to line {LineCount}, target {Locator}", _logId, lineCount, target.GetLocator());
+
+			UpdateLogRequest request = new UpdateLogRequest();
+			request.LogId = _logId.ToString();
+			request.LineCount = lineCount;
+			request.TargetHash = target.Hash.ToString();
+			request.TargetLocator = target.GetLocator().ToString();
+			request.Complete = complete;
+
+			LogRpcClient clientRef = await _hordeClient.CreateGrpcClientAsync<LogRpcClient>(cancellationToken);
+			await clientRef.UpdateLogAsync(request, cancellationToken: cancellationToken);
+		}
+
+		async Task<int> UpdateLogTailAsync(int tailNext, ReadOnlyMemory<byte> tailData, CancellationToken cancellationToken)
+		{
+			DateTime deadline = DateTime.UtcNow.AddMinutes(2.0);
+			try
+			{
+				LogRpcClient clientRef = await _hordeClient.CreateGrpcClientAsync<LogRpcClient>(cancellationToken);
+				using AsyncDuplexStreamingCall<UpdateLogTailRequest, UpdateLogTailResponse> call = clientRef.UpdateLogTail(deadline: deadline, cancellationToken: cancellationToken);
+
+				// Write the request to the server
+				UpdateLogTailRequest request = new UpdateLogTailRequest();
+				request.LogId = _logId.ToString();
+				request.TailNext = tailNext;
+				request.TailData = UnsafeByteOperations.UnsafeWrap(tailData);
+				await call.RequestStream.WriteAsync(request, cancellationToken);
+				_agentLogger.LogInformation("Writing log data: {LogId}, {TailNext}, {TailData} bytes", _logId, tailNext, tailData.Length);
+
+				// Wait until the server responds or we need to trigger a new update
+				Task<bool> moveNextAsync = call.ResponseStream.MoveNext();
+
+				Task task = await Task.WhenAny(moveNextAsync, _tailTaskStop.Task, Task.Delay(TimeSpan.FromMinutes(1.0), CancellationToken.None));
+				if (task == _tailTaskStop.Task)
+				{
+					_agentLogger.LogInformation("Cancelling long poll from client side (complete)");
+				}
+
+				// Close the request stream to indicate that we're finished
+				await call.RequestStream.CompleteAsync();
+
+				// Wait for a response or a new update to come in, then close the request stream
+				UpdateLogTailResponse? response = null;
+				while (await moveNextAsync)
+				{
+					response = call.ResponseStream.Current;
+					moveNextAsync = call.ResponseStream.MoveNext();
+				}
+				return response?.TailNext ?? -1;
+			}
+			catch (RpcException ex) when (ex.StatusCode == StatusCode.DeadlineExceeded)
+			{
+				_agentLogger.LogDebug(ex, "Log tail deadline exceeded, ignoring.");
+				return -1;
+			}
+		}
+
+		#endregion
 	}
 }
