Horde: Agent telemetry in log views
#rnx

diff --git a/Engine/Source/Programs/Horde/HordeDashboard/src/components/LogSource.tsx b/Engine/Source/Programs/Horde/HordeDashboard/src/components/LogSource.tsx
--- a/Engine/Source/Programs/Horde/HordeDashboard/src/components/LogSource.tsx
+++ b/Engine/Source/Programs/Horde/HordeDashboard/src/components/LogSource.tsx
@@ -8,6 +8,7 @@ import { getBatchSummaryMarkdown, getStepSummaryMarkdown, JobDetails } from "../
 import { getLeaseElapsed, getStepPercent } from '../base/utilities/timeUtils';
 import { BreadcrumbItem } from './Breadcrumbs';
 import { LogItem } from './LogRender';
+import { AgentTelemetryHandler } from "./agents/AgentTelemetrySparkline";
 
 const stripNewlines = /(\r\n|\n|\r)/gm;
 
@@ -57,6 +58,10 @@ export abstract class LogSource {
       return undefined;
    }
 
+   get agentTelemetry(): AgentTelemetryHandler | undefined {
+      return undefined;
+   }
+
    getAgentId(): string | undefined {
       return "";
    }
@@ -243,6 +248,7 @@ export abstract class LogSource {
    @action
    setActive(active: boolean) {
       this.active = active;
+      this.agentTelemetry?.setActive(active);
    }
 
    @action
@@ -373,6 +379,7 @@ export class JobLogSource extends LogSource {
    clear() {
       super.clear();
       this.jobDetails.clear();
+      this.agentTelemetry?.stop()
    }
 
    get errors(): GetLogEventResponse[] {
@@ -387,7 +394,6 @@ export class JobLogSource extends LogSource {
       return this.jobDetails.events;
    }
 
-
    get issues(): IssueData[] {
       return this.jobDetails.issues;
    }
@@ -418,6 +424,11 @@ export class JobLogSource extends LogSource {
       return undefined;
    }
 
+   get agentTelemetry(): AgentTelemetryHandler | undefined 
+   {
+      return this._agentTelemety;
+   }
+
    getAgentId(): string {
       return this.agentId ?? "";
    }
@@ -505,6 +516,8 @@ export class JobLogSource extends LogSource {
       if (this.batch) {
          this.jobName = `Batch-${this.batch.id}`;
       }
+      
+      this.agentTelemetry?.set(this.agentId ?? "", new Date(this.startTime as any), this.step?.finishTime ? new Date(this.step.finishTime) : undefined);      
 
    }
 
@@ -597,6 +610,8 @@ export class JobLogSource extends LogSource {
    
    artifactsV2?: GetArtifactResponseV2[];
 
+   _agentTelemety: AgentTelemetryHandler = new AgentTelemetryHandler();
+
    jobDetails: JobDetails = new JobDetails(undefined, undefined, undefined, true);
 }
 
diff --git a/Engine/Source/Programs/Horde/HordeDashboard/src/components/LogStyle.tsx b/Engine/Source/Programs/Horde/HordeDashboard/src/components/LogStyle.tsx
--- a/Engine/Source/Programs/Horde/HordeDashboard/src/components/LogStyle.tsx
+++ b/Engine/Source/Programs/Horde/HordeDashboard/src/components/LogStyle.tsx
@@ -57,8 +57,7 @@ export const getLogStyles = () => {
 
    const logStyleBase = _logStyleBase ?? mergeStyleSets({
       container: {
-         overflow: 'auto',
-         height: 'calc(100vh - 292px)',
+         overflow: 'auto',         
          marginTop: 8,
       },
       logLine: [
diff --git a/Engine/Source/Programs/Horde/HordeDashboard/src/components/LogView.tsx b/Engine/Source/Programs/Horde/HordeDashboard/src/components/LogView.tsx
--- a/Engine/Source/Programs/Horde/HordeDashboard/src/components/LogView.tsx
+++ b/Engine/Source/Programs/Horde/HordeDashboard/src/components/LogView.tsx
@@ -16,6 +16,7 @@ import { useWindowSize } from '../base/utilities/hooks';
 import { displayTimeZone, getElapsedString } from '../base/utilities/timeUtils';
 import { getHordeStyling } from '../styles/Styles';
 import { getHordeTheme } from '../styles/theme';
+import { AgentTelemetrySparkline } from "./agents/AgentTelemetrySparkline";
 import { JobArtifactsModal } from './artifacts/ArtifactsModal';
 import { Breadcrumbs } from './Breadcrumbs';
 import { ChangeContextMenu, ChangeContextMenuTarget } from './ChangeButton';
@@ -472,6 +473,7 @@ export const LogList: React.FC<{ logId: string }> = observer(({ logId }) => {
    const [issueHistory, setIssueHistory] = useState(false);
    const [logHistory, setLogHistory] = useState(false);
    const [logArtifacts, setLogArtifacts] = useState("");
+   const [logTelemetry, setLogTelemetry] = useState(false);
    const [logError, setLogError] = useState("")
 
    let [historyAgentId, setHistoryAgentId] = useState<string | undefined>(undefined);
@@ -514,7 +516,7 @@ export const LogList: React.FC<{ logId: string }> = observer(({ logId }) => {
 
       return () => {
          globalHandler = globalSearchState = undefined;
-         window.removeEventListener("keydown", handler);
+         window.removeEventListener("keydown", handler);         
       };
 
    }, []);
@@ -727,13 +729,24 @@ export const LogList: React.FC<{ logId: string }> = observer(({ logId }) => {
          const eyeColor = modeColors.text + "44";
 
          return (
-            <Stack key={`key_log_line_${item.lineNumber}`} style={{ width: "max-content", height: handler.lineHeight }} onClick={() => {
-               const search = new URLSearchParams(window.location.search);
-               search.set("lineindex", (item.lineNumber - 1).toString());
-               const url = `${window.location.pathname}?` + search.toString();
-
-               navigate(url, { replace: true })
-            }}>
+            <Stack key={`key_log_line_${item.lineNumber}`} style={{ width: "max-content", height: handler.lineHeight }}
+               onMouseEnter={(ev) => {
+
+                  if (item.line?.time) {
+                     let time = item.line?.time;
+                     if (!time.endsWith("Z")) {
+                        time += "Z";
+                     }
+                     logSource?.agentTelemetry?.setCurrentTime(new Date(time));
+                  }
+               }}
+               onClick={() => {
+                  const search = new URLSearchParams(window.location.search);
+                  search.set("lineindex", (item.lineNumber - 1).toString());
+                  const url = `${window.location.pathname}?` + search.toString();
+
+                  navigate(url, { replace: true })
+               }}>
                <div style={{ position: "relative" }}>
                   <Stack className={styles.logLine} style={{ position: "relative" }} tokens={{ childrenGap: 8 }} horizontal disableShrink={true}>
                      <Stack styles={{ root: { color: "#c0c0c0", width: 80, textAlign: "right", userSelect: "none", fontSize: handler.fontSize } }}>{prefix + item.lineNumber}</Stack>
@@ -742,7 +755,7 @@ export const LogList: React.FC<{ logId: string }> = observer(({ logId }) => {
                         {(!item.issueId || !ev) && <Stack styles={{ root: { color: "#8a8a8a", width: tsWidth, whiteSpace: "nowrap", fontSize: handler.fontSize, userSelect: "none" } }}> {timestamp}</Stack>}
                         {!!item.issueId && !!ev && <IssueButton item={item} event={ev!} />}
                         <div className={styles.logLineOuter}> <Stack styles={{ root: { paddingLeft: 8, paddingRight: 8, position: "relative", verticalAlign: "center" } }}> {renderLine(navigate, item.line, item.lineNumber, handler.lineRenderStyle, searchState.search)}
-                           <Stack id={`callout_target_${item?.lineNumber}`} style={{ position: "absolute", cursor: "pointer", userSelect: "none", left: "-12px", top: "0px" }} onClick={() => {
+                           <Stack id={`callout_target_${item?.lineNumber}`} style={{ position: "absolute", cursor: "pointer", userSelect: "none", left: "-12px", top: "0px", zIndex: 100 }} onClick={() => {
                               handler.infoLine = item.lineNumber;
                               handler.externalUpdate();
                            }}><FontIcon id="infoview" style={{ fontSize: 14, color: eyeColor }} iconName="Eye" /></Stack>
@@ -1005,7 +1018,16 @@ export const LogList: React.FC<{ logId: string }> = observer(({ logId }) => {
          onClick: () => setLogHistory(true)
       })
 
-
+      if (logSource.agentTelemetry) {
+         menuProps.items.push({
+            key: 'jobstep_agent_telemetry',
+            text: logTelemetry ? 'Hide Telemetry' : 'Show Telemetry',
+            onClick: () => {
+               logSource.agentTelemetry?.show(!logTelemetry);
+               setLogTelemetry(!logTelemetry)               
+            }
+         })
+      }
    }
 
    function updateError() {
@@ -1015,7 +1037,7 @@ export const LogList: React.FC<{ logId: string }> = observer(({ logId }) => {
       let lineIdx = errors[handler.currentError].lineIndex;
       const search = new URLSearchParams(window.location.search);
       search.set("lineindex", (lineIdx).toString());
-      const url = `${window.location.pathname}?` + search.toString();               
+      const url = `${window.location.pathname}?` + search.toString();
       lineIdx -= 10;
       if (lineIdx < 0) {
          lineIdx = 0;
@@ -1059,10 +1081,10 @@ export const LogList: React.FC<{ logId: string }> = observer(({ logId }) => {
       handler.currentError = undefined;
 
       handler.stopTrailing();
-      let lineIdx = warnings[handler.currentWarning].lineIndex ;
+      let lineIdx = warnings[handler.currentWarning].lineIndex;
       const search = new URLSearchParams(window.location.search);
       search.set("lineindex", (lineIdx).toString());
-      const url = `${window.location.pathname}?` + search.toString();               
+      const url = `${window.location.pathname}?` + search.toString();
 
       lineIdx -= 10
       if (lineIdx < 0) {
@@ -1125,6 +1147,12 @@ export const LogList: React.FC<{ logId: string }> = observer(({ logId }) => {
 
    const buttonNoneText = dashboard.darktheme ? "#949898" : "#616e85";
 
+   let heightAdjust = 292;
+   if (logTelemetry) {
+      // sparkline height
+      heightAdjust += 120;
+   }
+
    return <Stack>
       {!!fixme && <IssueModalV2 issueId={query.get("issue")} popHistoryOnClose={issueHistory} />}
       {!!fixme && logHistory && <StepHistoryModal jobDetails={fixme!} stepId={fixme!.stepByLogId(logId)?.id} onClose={() => setLogHistory(false)} />}
@@ -1276,6 +1304,7 @@ export const LogList: React.FC<{ logId: string }> = observer(({ logId }) => {
                         <Markdown>{summaryText}</Markdown>
                      </Stack>
                   </Stack>
+                  {logTelemetry && !! logSource.agentTelemetry && <AgentTelemetrySparkline handler={logSource.agentTelemetry} />}
                   <Stack horizontalAlign="center" style={{ paddingBottom: 12 }}>
                      <Separator styles={{ root: { fontSize: 0, width: "100%", padding: 0, selectors: { '::before': { background: dashboard.darktheme ? '#313638' : '#D3D2D1' } } } }} />
                   </Stack>
@@ -1284,10 +1313,11 @@ export const LogList: React.FC<{ logId: string }> = observer(({ logId }) => {
             </Stack>
          </Stack>
 
+
          <Stack style={{ backgroundColor: hordeTheme.horde.contentBackground, paddingLeft: "24px", paddingRight: "24px" }}>
             <Stack tokens={{ childrenGap: 0 }}>
                <FocusZone direction={FocusZoneDirection.vertical} isInnerZoneKeystroke={() => { return true; }} defaultActiveElement="#LogList" style={{ padding: 0, margin: 0 }} >
-                  <div className={handler.style.container} data-is-scrollable={true}
+                  <div className={handler.style.container} data-is-scrollable={true} style={{ height: `calc(100vh - ${heightAdjust}px)` }}
                      onScroll={(ev) => {
 
                         const element: any = ev.target;
diff --git a/Engine/Source/Programs/Horde/HordeDashboard/src/components/agents/AgentTelemetrySparkline.tsx b/Engine/Source/Programs/Horde/HordeDashboard/src/components/agents/AgentTelemetrySparkline.tsx
new file mode 100644
--- /dev/null
+++ b/Engine/Source/Programs/Horde/HordeDashboard/src/components/agents/AgentTelemetrySparkline.tsx
@@ -0,0 +1,242 @@
+import { Label, Spinner, SpinnerSize, Stack, Text } from "@fluentui/react";
+import * as d3 from "d3";
+import { observer } from "mobx-react-lite";
+import { Sparklines, SparklinesLine, SparklinesReferenceLine } from "react-sparklines";
+import backend from "../../backend";
+import { GetAgentResponse, GetAgentTelemetrySampleResponse } from "../../backend/Api";
+import dashboard from "../../backend/Dashboard";
+import { PollBase } from "../../backend/PollBase";
+
+const sparkWidth = 1380 / 2;
+const sparkHeight = 64;
+
+export class AgentTelemetryHandler extends PollBase {
+
+   constructor() {
+      super(15000);
+
+   }
+
+   show(visible: boolean) {
+
+      if (this.visible === visible) {
+         return;
+      }
+
+      visible ? this.start() : this.stop();
+
+      this.visible = visible;
+      this.setUpdated();
+
+   }
+
+   setActive(value: boolean) {
+      if (this.active === value) {
+         return;
+      }
+
+      this.active = value;
+   }
+
+   set(agentId: string, startTime: Date, endTime?: Date) {
+
+      this.agentId = agentId;
+      this.startTime = startTime;
+      this.endTime = endTime;
+   }
+
+   setCurrentTime(currentTime?: Date) {
+
+      if (!this.visible) {
+         this.currentTime = undefined;
+         return;
+      }
+
+      this.currentTime = currentTime;
+
+      // chart selectors
+      const charts = ["#telemetry_cpu > svg", "#telemetry_ram > svg"]
+
+      const startTime = this.lastStartTime;
+      const endTime = this.lastEndTime;
+
+      if (this.currentTime && startTime && endTime) {
+
+         let x = (this.currentTime.getTime() - startTime.getTime()) / (endTime.getTime() - startTime.getTime());
+
+         if (x < 0 || x > 1) {
+            return;
+         }
+
+         x *= sparkWidth;
+         x = Math.floor(x);
+         
+         if (this.lastX === x) {
+            return;
+         }
+
+         this.lastX = x;
+
+         charts.forEach(chartSelector => {
+
+            const chart = d3.select(chartSelector);
+
+            if (!chart) {
+               return;
+            }
+
+            chart.selectAll(".cline").remove();
+
+            chart.append("line")
+               .attr("class", "cline")
+               .attr("x1", () => x)
+               .attr("x2", () => x)
+               .attr("y1", () => 0)
+               .attr("y2", () => sparkHeight)
+               .attr("stroke-width", () => 2)
+               .attr("stroke", dashboard.darktheme ? "#6D6C6B" : "#4D4C4BAA")            
+
+         })
+
+      }
+
+   }
+
+   async poll(): Promise<void> {
+
+
+      try {
+
+         if (!this.agentId || !this.startTime) {
+            return;
+         }
+
+         if (!this.active && this.data?.length) {
+            return;
+         }
+
+         let endTime = this.endTime;
+         if (!endTime) {
+            endTime = new Date();
+         }         
+
+         // clamp to 24 hours
+         if ((endTime.getTime() - this.startTime.getTime()) > 86400000) {
+            endTime = new Date(this.startTime.getTime() + 86400000);
+         }
+
+         if (!this.agent) {
+            this.agent = await backend.getAgent(this.agentId);
+         }
+
+         if (this.startTime === this.lastStartTime && this.endTime === this.lastEndTime) {
+            return;
+         }
+
+         this.data = await backend.getAgentTelemetry(this.agentId, this.startTime, endTime);
+
+         this.lastStartTime = this.startTime;
+         this.lastEndTime = endTime;
+
+         this.setUpdated();
+
+      } catch (err) {
+         console.error(err);
+      }
+   }
+
+   active = true;
+   lastX: number = -1;
+   lastStartTime?: Date;
+   lastEndTime?: Date;
+
+   data?: GetAgentTelemetrySampleResponse[];
+
+   agent?: GetAgentResponse;
+
+   visible = false;
+   agentId?: string;
+   startTime?: Date;
+   endTime?: Date;
+   currentTime?: Date;
+
+}
+
+export const AgentTelemetrySparkline: React.FC<{ handler: AgentTelemetryHandler }> = observer(({ handler }) => {
+
+
+   handler.subscribe();
+
+   const data = handler.data;
+
+   if (data === undefined) {
+      return <Stack horizontal horizontalAlign="center" style={{padding: 32}}>
+         <Spinner size={SpinnerSize.large} />
+      </Stack>
+   }
+
+   if (!data.length) {
+      return <Stack horizontal horizontalAlign="center" style={{padding: 32}}>
+         <Text variant="mediumPlus">No Telemetry Data</Text>
+      </Stack>
+   }
+
+   if (!handler.visible || !data?.length) {
+      return;
+   }
+
+   const ramData = data.map(d => {
+      return d.usedRam / d.totalRam;
+   })
+
+   const cpuData = data.map(d => {
+      return (d.userCpu + d.systemCpu) / 100;
+   })
+
+   let cpuText = "";
+   let ramText = "";
+
+   let device: any;
+   const devices = handler.agent?.capabilities?.devices;
+   if (devices?.length) {
+      device = devices[0];
+   }
+   device?.properties?.forEach(v => {
+      if (v.startsWith("CPU=")) {
+         cpuText = `${v.replace("CPU=", "")}`;
+      }
+
+      if (v.startsWith("RAM=")) {
+         ramText = `${v.replace("RAM=", "")} GB`;
+      }
+   })
+
+   return <Stack horizontal horizontalAlign="center" tokens={{ childrenGap: 24 }} style={{paddingBottom: 8}}>
+      <Stack style={{ width: sparkWidth }}>
+         <Stack horizontal>
+            <Label>CPU</Label>
+            <Stack grow></Stack>
+            {!!cpuText && <Label>{cpuText}</Label>}
+         </Stack>
+         <div id="telemetry_cpu">
+            <Sparklines width={sparkWidth} height={sparkHeight} data={cpuData} max={1} style={{ backgroundColor: dashboard.darktheme ? "#060709" : "#F3F2F1", padding: 8, border: "solid 1px #181A1B" }}>
+               <SparklinesLine color={dashboard.darktheme ? "lightblue" : "#1E90FF"} />
+               <SparklinesReferenceLine type="avg" />
+            </Sparklines>
+         </div>
+      </Stack>
+      <Stack style={{ width: sparkWidth }}>
+         <Stack horizontal>
+            <Label>RAM</Label>
+            <Stack grow></Stack>
+            {!!ramText && <Label>{ramText}</Label>}
+         </Stack>
+         <div id="telemetry_ram">
+            <Sparklines width={sparkWidth} height={sparkHeight} data={ramData} max={1} style={{ backgroundColor: dashboard.darktheme ? "#060709" : "#F3F2F1", padding: 8, border: "solid 1px #181A1B" }}>
+               <SparklinesLine color={dashboard.darktheme ? "lightblue" : "#1E90FF"} />
+               <SparklinesReferenceLine type="avg" />
+            </Sparklines>
+         </div>
+      </Stack>
+   </Stack>
+})
\ No newline at end of file
