Horde: Fix static analysis warnings.

#rnx

diff --git a/Engine/Source/Programs/Horde/Horde.Agent.Driver/Commands/Utilities/BuildGraphCommand.cs b/Engine/Source/Programs/Horde/Horde.Agent.Driver/Commands/Utilities/BuildGraphCommand.cs
--- a/Engine/Source/Programs/Horde/Horde.Agent.Driver/Commands/Utilities/BuildGraphCommand.cs
+++ b/Engine/Source/Programs/Horde/Horde.Agent.Driver/Commands/Utilities/BuildGraphCommand.cs
@@ -1,14 +1,8 @@
 // Copyright Epic Games, Inc. All Rights Reserved.
 
 using System.Diagnostics;
-using System.Text;
-using System.Text.Json;
-using System.Text.Json.Serialization;
 using EpicGames.Core;
-using EpicGames.Horde.Jobs;
 using Horde.Agent.Parser;
-using Horde.Agent.Utility;
-using HordeCommon.Rpc;
 using Microsoft.Extensions.Logging;
 
 namespace Horde.Agent.Commands.Utilities
diff --git a/Engine/Source/Programs/Horde/Horde.Agent.Driver/Utilities/HordeRpc.cs b/Engine/Source/Programs/Horde/Horde.Agent.Driver/Utilities/HordeRpc.cs
--- a/Engine/Source/Programs/Horde/Horde.Agent.Driver/Utilities/HordeRpc.cs
+++ b/Engine/Source/Programs/Horde/Horde.Agent.Driver/Utilities/HordeRpc.cs
@@ -1,13 +1,9 @@
 // Copyright Epic Games, Inc. All Rights Reserved.
 
-using System.Diagnostics.CodeAnalysis;
-using System.Security.Cryptography;
 using EpicGames.Horde.Agents.Leases;
 using EpicGames.Horde.Jobs;
 using EpicGames.Horde.Logs;
-using EpicGames.Horde.Streams;
 using Google.Protobuf;
-using Google.Protobuf.Collections;
 
 // These partial classes/extensions operate on generated gRPC and Protobuf code.
 // Warnings below are disabled to avoid documenting every class touched.
diff --git a/Engine/Source/Programs/Horde/Horde.Agent.Driver/Utility/JobStepLogger.cs b/Engine/Source/Programs/Horde/Horde.Agent.Driver/Utility/JobStepLogger.cs
--- a/Engine/Source/Programs/Horde/Horde.Agent.Driver/Utility/JobStepLogger.cs
+++ b/Engine/Source/Programs/Horde/Horde.Agent.Driver/Utility/JobStepLogger.cs
@@ -1,10 +1,5 @@
 // Copyright Epic Games, Inc. All Rights Reserved.
 
-using System;
-using System.Collections.Generic;
-using System.Linq;
-using System.Text;
-using System.Threading.Tasks;
 using EpicGames.Horde;
 using EpicGames.Horde.Jobs;
 using EpicGames.Horde.Logs;
diff --git a/Engine/Source/Programs/Horde/Horde.Agent.Driver/Utility/ServerLogger.cs b/Engine/Source/Programs/Horde/Horde.Agent.Driver/Utility/ServerLogger.cs
--- a/Engine/Source/Programs/Horde/Horde.Agent.Driver/Utility/ServerLogger.cs
+++ b/Engine/Source/Programs/Horde/Horde.Agent.Driver/Utility/ServerLogger.cs
@@ -4,7 +4,6 @@ using System.Text;
 using System.Threading.Channels;
 using EpicGames.Core;
 using EpicGames.Horde;
-using EpicGames.Horde.Jobs;
 using EpicGames.Horde.Logs;
 using EpicGames.Horde.Storage;
 using Google.Protobuf;
diff --git a/Engine/Source/Programs/Horde/Horde.Agent/Leases/Handlers/JobHandler.cs b/Engine/Source/Programs/Horde/Horde.Agent/Leases/Handlers/JobHandler.cs
--- a/Engine/Source/Programs/Horde/Horde.Agent/Leases/Handlers/JobHandler.cs
+++ b/Engine/Source/Programs/Horde/Horde.Agent/Leases/Handlers/JobHandler.cs
@@ -19,7 +19,6 @@ using HordeCommon.Rpc;
 using HordeCommon.Rpc.Messages;
 using HordeCommon.Rpc.Tasks;
 using Microsoft.Extensions.Logging;
-using Microsoft.Extensions.Logging.Abstractions;
 using Microsoft.Extensions.Options;
 using OpenTracing;
 using OpenTracing.Util;
@@ -35,7 +34,7 @@ namespace Horde.Agent.Leases.Handlers
 		/// Exposed as internal to ease testing.
 		/// </summary>
 		internal TimeSpan _stepAbortPollInterval = TimeSpan.FromSeconds(5);
-		
+
 		/// <summary>
 		/// How long to wait before retrying a failed step abort check request
 		/// </summary>
diff --git a/Engine/Source/Programs/Horde/Horde.Agent/Services/ServerLoggerFactory.cs b/Engine/Source/Programs/Horde/Horde.Agent/Services/ServerLoggerFactory.cs
--- a/Engine/Source/Programs/Horde/Horde.Agent/Services/ServerLoggerFactory.cs
+++ b/Engine/Source/Programs/Horde/Horde.Agent/Services/ServerLoggerFactory.cs
@@ -1,7 +1,6 @@
 // Copyright Epic Games, Inc. All Rights Reserved.
 
 using EpicGames.Horde;
-using EpicGames.Horde.Jobs;
 using EpicGames.Horde.Logs;
 using Horde.Agent.Utility;
 using Microsoft.Extensions.Logging;
diff --git a/Engine/Source/Programs/Horde/Horde.Server/Jobs/JobRpcService.cs b/Engine/Source/Programs/Horde/Horde.Server/Jobs/JobRpcService.cs
--- a/Engine/Source/Programs/Horde/Horde.Server/Jobs/JobRpcService.cs
+++ b/Engine/Source/Programs/Horde/Horde.Server/Jobs/JobRpcService.cs
@@ -17,7 +17,6 @@ using EpicGames.Horde.Agents.Sessions;
 using EpicGames.Horde.Artifacts;
 using EpicGames.Horde.Jobs;
 using EpicGames.Horde.Logs;
-using EpicGames.Horde.Streams;
 using Google.Protobuf.WellKnownTypes;
 using Grpc.Core;
 using Horde.Common.Rpc;
