Horde: Fix another static analysis warning.

#rnx

diff --git a/Engine/Source/Programs/Horde/Horde.Agent.Driver/Utility/JobStepLogger.cs b/Engine/Source/Programs/Horde/Horde.Agent.Driver/Utility/JobStepLogger.cs
--- a/Engine/Source/Programs/Horde/Horde.Agent.Driver/Utility/JobStepLogger.cs
+++ b/Engine/Source/Programs/Horde/Horde.Agent.Driver/Utility/JobStepLogger.cs
@@ -28,8 +28,8 @@ namespace Horde.Agent.Driver.Utility
 		readonly ServerLogger _inner;
 		JobStepOutcome _outcome;
 		Task _updateOutcomeTask;
-		CancellationTokenSource _cancellationSource;
-		ILogger _internalLogger;
+		readonly CancellationTokenSource _cancellationSource;
+		readonly ILogger _internalLogger;
 
 		/// <summary>
 		/// Constructor
@@ -90,6 +90,7 @@ namespace Horde.Agent.Driver.Utility
 			{
 				Task prevUpdateTask = _updateOutcomeTask;
 				_updateOutcomeTask = Task.Run(() => UpdateOutcomeAsync(prevUpdateTask, newOutcome, _cancellationSource.Token), _cancellationSource.Token);
+				_outcome = newOutcome;
 			}
 		}
 
