Horde: Remove some agent-specific functionality from ServerLogger (masking out warnings, wrapping a local logger device).

#rnx

diff --git a/Engine/Source/Programs/Horde/Horde.Agent.Driver/Utility/JobStepLogger.cs b/Engine/Source/Programs/Horde/Horde.Agent.Driver/Utility/JobStepLogger.cs
--- a/Engine/Source/Programs/Horde/Horde.Agent.Driver/Utility/JobStepLogger.cs
+++ b/Engine/Source/Programs/Horde/Horde.Agent.Driver/Utility/JobStepLogger.cs
@@ -24,8 +24,10 @@ namespace Horde.Agent.Driver.Utility
 		readonly JobId _jobId;
 		readonly JobStepBatchId _jobBatchId;
 		readonly JobStepId _jobStepId;
+		readonly bool _warnings;
 
-		readonly ServerLogger _inner;
+		readonly ServerLogger _serverLogger;
+		readonly ILogger _localLogger;
 		JobStepOutcome _outcome;
 		Task _updateOutcomeTask;
 		readonly CancellationTokenSource _cancellationSource;
@@ -40,7 +42,9 @@ namespace Horde.Agent.Driver.Utility
 			_jobId = jobId;
 			_jobBatchId = jobBatchId;
 			_jobStepId = jobStepId;
-			_inner = new ServerLogger(hordeClient, logId, warnings, outputLevel, localLogger, internalLogger);
+			_warnings = warnings ?? true;
+			_serverLogger = new ServerLogger(hordeClient, logId, outputLevel, internalLogger);
+			_localLogger = localLogger;
 			_outcome = JobStepOutcome.Success;
 			_updateOutcomeTask = Task.CompletedTask;
 			_cancellationSource = new CancellationTokenSource();
@@ -49,7 +53,7 @@ namespace Horde.Agent.Driver.Utility
 
 		/// <inheritdoc/>
 		public IDisposable? BeginScope<TState>(TState state) where TState : notnull
-			=> _inner.BeginScope(state);
+			=> _localLogger.BeginScope(state);
 
 		/// <inheritdoc/>
 		public async ValueTask DisposeAsync()
@@ -61,18 +65,27 @@ namespace Horde.Agent.Driver.Utility
 				_cancellationSource.Dispose();
 			}
 
-			await _inner.DisposeAsync();
+			await _serverLogger.DisposeAsync();
 		}
 
 		/// <inheritdoc/>
 		public bool IsEnabled(LogLevel logLevel)
-			=> _inner.IsEnabled(logLevel);
+			=> _serverLogger.IsEnabled(logLevel) || _localLogger.IsEnabled(logLevel);
 
 		/// <inheritdoc/>
 		public void Log<TState>(LogLevel logLevel, EventId eventId, TState state, Exception? exception, Func<TState, Exception?, string> formatter)
 		{
+			// Downgrade warnings to information if not required
+			if (logLevel == LogLevel.Warning && !_warnings)
+			{
+				logLevel = LogLevel.Information;
+			}
+
+			// Write to the local logger
+			_localLogger.Log<TState>(logLevel, eventId, state, exception, formatter);
+
 			// Forward the log data to the inner writer
-			_inner.Log<TState>(logLevel, eventId, state, exception, formatter);
+			_serverLogger.Log<TState>(logLevel, eventId, state, exception, formatter);
 
 			// Update the state of this job if this is an error status
 			JobStepOutcome newOutcome = _outcome;
@@ -97,7 +110,7 @@ namespace Horde.Agent.Driver.Utility
 		/// <inheritdoc/>
 		public async Task StopAsync()
 		{
-			await _inner.StopAsync();
+			await _serverLogger.StopAsync();
 			await _updateOutcomeTask;
 		}
 
diff --git a/Engine/Source/Programs/Horde/Horde.Agent.Driver/Utility/ServerLogger.cs b/Engine/Source/Programs/Horde/Horde.Agent.Driver/Utility/ServerLogger.cs
--- a/Engine/Source/Programs/Horde/Horde.Agent.Driver/Utility/ServerLogger.cs
+++ b/Engine/Source/Programs/Horde/Horde.Agent.Driver/Utility/ServerLogger.cs
@@ -48,10 +48,8 @@ namespace Horde.Agent.Utility
 		AsyncEvent _tailTaskStop;
 		readonly AsyncEvent _newTailDataEvent = new AsyncEvent();
 
-		readonly bool _warnings;
 		readonly LogLevel _outputLevel;
-		readonly ILogger _localLogger;
-		readonly ILogger _agentLogger;
+		readonly ILogger _internalLogger;
 		readonly Channel<JsonLogEvent> _dataChannel;
 		Task? _dataWriter;
 
@@ -60,15 +58,13 @@ namespace Horde.Agent.Utility
 		/// </summary>
 		/// <param name="hordeClient">Horde instance to write to</param>
 		/// <param name="logId">The log id to write to</param>
-		/// <param name="warnings">Whether to include warnings in the output</param>
-		/// <param name="outputLevel">Minimum level for output</param>
-		/// <param name="localLogger">Logger to forward log messages to</param>
-		/// <param name="agentLogger">Logger for systemic messages</param>
-		public ServerLogger(IHordeClient hordeClient, LogId logId, bool? warnings, LogLevel outputLevel, ILogger localLogger, ILogger agentLogger)
+		/// <param name="minimumLevel">Minimum level for output</param>
+		/// <param name="internalLogger">Logger for systemic messages</param>
+		public ServerLogger(IHordeClient hordeClient, LogId logId, LogLevel minimumLevel, ILogger internalLogger)
 		{
 			_hordeClient = hordeClient;
 			_logId = logId;
-			_builder = new LogBuilder(LogFormat.Json, agentLogger);
+			_builder = new LogBuilder(LogFormat.Json, internalLogger);
 			_store = _hordeClient.CreateStorageClient(logId);
 			_writer = _store.CreateBlobWriter();
 
@@ -76,10 +72,8 @@ namespace Horde.Agent.Utility
 			_tailTask = Task.Run(() => TickTailAsync());
 
 			_logId = logId;
-			_warnings = warnings ?? true;
-			_outputLevel = outputLevel;
-			_localLogger = localLogger;
-			_agentLogger = agentLogger;
+			_outputLevel = minimumLevel;
+			_internalLogger = internalLogger;
 			_dataChannel = Channel.CreateUnbounded<JsonLogEvent>();
 			_dataWriter = Task.Run(() => RunDataWriterAsync());
 		}
@@ -87,23 +81,16 @@ namespace Horde.Agent.Utility
 		/// <inheritdoc/>
 		public void Log<TState>(LogLevel logLevel, EventId eventId, TState state, Exception? exception, Func<TState, Exception?, string> formatter)
 		{
-			_localLogger.Log(logLevel, eventId, state, exception, formatter);
-
-			// Downgrade warnings to information if not required
-			if (logLevel == LogLevel.Warning && !_warnings)
-			{
-				logLevel = LogLevel.Information;
-			}
-
 			JsonLogEvent jsonLogEvent = JsonLogEvent.FromLoggerState(logLevel, eventId, state, exception, formatter);
 			WriteFormattedEvent(jsonLogEvent);
 		}
 
 		/// <inheritdoc/>
-		public bool IsEnabled(LogLevel logLevel) => logLevel >= _outputLevel || _localLogger.IsEnabled(logLevel);
+		public bool IsEnabled(LogLevel logLevel) => logLevel >= _outputLevel;
 
 		/// <inheritdoc/>
-		public IDisposable? BeginScope<TState>(TState state) where TState : notnull => _localLogger.BeginScope(state);
+		public IDisposable? BeginScope<TState>(TState state) where TState : notnull
+			=> null;
 
 		private void WriteFormattedEvent(JsonLogEvent jsonLogEvent)
 		{
@@ -134,7 +121,7 @@ namespace Horde.Agent.Utility
 		{
 			await StopAsync();
 
-			_agentLogger.LogInformation("Disposing json log task");
+			_internalLogger.LogInformation("Disposing json log task");
 
 			if (_tailTaskStop != null)
 			{
@@ -230,7 +217,7 @@ namespace Horde.Agent.Utility
 					}
 					catch (Exception ex)
 					{
-						_agentLogger.LogWarning(ex, "Unable to write data to server (log {LogId}, offset {Offset}, length {Length}, lines {StartLine}-{EndLine})", _logId, packetOffset, packet.Length, packetLineIndex, packetLineIndex + packetLineCount);
+						_internalLogger.LogWarning(ex, "Unable to write data to server (log {LogId}, offset {Offset}, length {Length}, lines {StartLine}-{EndLine})", _logId, packetOffset, packet.Length, packetLineIndex, packetLineIndex + packetLineCount);
 					}
 				}
 
@@ -243,7 +230,7 @@ namespace Horde.Agent.Utility
 					}
 					catch (Exception ex)
 					{
-						_agentLogger.LogWarning(ex, "Unable to create events");
+						_internalLogger.LogWarning(ex, "Unable to create events");
 					}
 				}
 
@@ -256,7 +243,7 @@ namespace Horde.Agent.Utility
 					}
 					catch (Exception ex)
 					{
-						_agentLogger.LogWarning(ex, "Unable to flush data to server (log {LogId}, offset {Offset})", _logId, packetOffset);
+						_internalLogger.LogWarning(ex, "Unable to flush data to server (log {LogId}, offset {Offset})", _logId, packetOffset);
 					}
 					break;
 				}
@@ -271,7 +258,7 @@ namespace Horde.Agent.Utility
 			}
 			catch (Exception ex)
 			{
-				_agentLogger.LogError(ex, "Exception while trying to parse line count from data ({Message})", Encoding.UTF8.GetString(span));
+				_internalLogger.LogError(ex, "Exception while trying to parse line count from data ({Message})", Encoding.UTF8.GetString(span));
 			}
 		}
 
@@ -286,12 +273,12 @@ namespace Horde.Agent.Utility
 				}
 				catch (OperationCanceledException ex)
 				{
-					_agentLogger.LogInformation(ex, "Cancelled log tailing task");
+					_internalLogger.LogInformation(ex, "Cancelled log tailing task");
 					break;
 				}
 				catch (Exception ex)
 				{
-					_agentLogger.LogError(ex, "Exception on log tailing task ({LogId}): {Message}", _logId, ex.Message);
+					_internalLogger.LogError(ex, "Exception on log tailing task ({LogId}): {Message}", _logId, ex.Message);
 					await Task.Delay(TimeSpan.FromSeconds(10.0));
 				}
 			}
@@ -318,7 +305,7 @@ namespace Horde.Agent.Utility
 				// and to make sure that we don't expire the existing tail data.
 				if (tailNext != -1 && tailData.IsEmpty && tailNext == initialTailNext && !tickTask.IsCompleted)
 				{
-					_agentLogger.LogInformation("No tail data available for log {LogId} after line {TailNext}; waiting for more...", _logId, tailNext);
+					_internalLogger.LogInformation("No tail data available for log {LogId} after line {TailNext}; waiting for more...", _logId, tailNext);
 					await Task.WhenAny(newTailDataTask, tickTask);
 					continue;
 				}
@@ -331,20 +318,20 @@ namespace Horde.Agent.Utility
 
 				// Update the next tailing position
 				int numLines = CountLines(tailData.Span);
-				_agentLogger.LogInformation("Setting log {LogId} tail = {TailNext}, data = {TailDataSize} bytes, {NumLines} lines ('{Start}')", _logId, tailNext, tailData.Length, numLines, start);
+				_internalLogger.LogInformation("Setting log {LogId} tail = {TailNext}, data = {TailDataSize} bytes, {NumLines} lines ('{Start}')", _logId, tailNext, tailData.Length, numLines, start);
 
 				int newTailNext = await UpdateLogTailAsync(tailNext, tailData, CancellationToken.None);
-				_agentLogger.LogInformation("Log {LogId} tail next = {TailNext}", _logId, newTailNext);
+				_internalLogger.LogInformation("Log {LogId} tail next = {TailNext}", _logId, newTailNext);
 
 				if (newTailNext != tailNext)
 				{
 					tailNext = newTailNext;
-					_agentLogger.LogInformation("Modified tail position for log {LogId} to {TailNext}", _logId, tailNext);
+					_internalLogger.LogInformation("Modified tail position for log {LogId} to {TailNext}", _logId, tailNext);
 				}
 
 				tickTask = Task.Delay(TimeSpan.FromSeconds(10.0));
 			}
-			_agentLogger.LogInformation("Finishing log tail task");
+			_internalLogger.LogInformation("Finishing log tail task");
 		}
 
 		static int CountLines(ReadOnlySpan<byte> data)
@@ -387,7 +374,7 @@ namespace Horde.Agent.Utility
 
 		async Task UpdateLogAsync(IBlobRef target, int lineCount, bool complete, CancellationToken cancellationToken)
 		{
-			_agentLogger.LogInformation("Updating log {LogId} to line {LineCount}, target {Locator}", _logId, lineCount, target.GetLocator());
+			_internalLogger.LogInformation("Updating log {LogId} to line {LineCount}, target {Locator}", _logId, lineCount, target.GetLocator());
 
 			UpdateLogRequest request = new UpdateLogRequest();
 			request.LogId = _logId.ToString();
@@ -414,7 +401,7 @@ namespace Horde.Agent.Utility
 				request.TailNext = tailNext;
 				request.TailData = UnsafeByteOperations.UnsafeWrap(tailData);
 				await call.RequestStream.WriteAsync(request, cancellationToken);
-				_agentLogger.LogInformation("Writing log data: {LogId}, {TailNext}, {TailData} bytes", _logId, tailNext, tailData.Length);
+				_internalLogger.LogInformation("Writing log data: {LogId}, {TailNext}, {TailData} bytes", _logId, tailNext, tailData.Length);
 
 				// Wait until the server responds or we need to trigger a new update
 				Task<bool> moveNextAsync = call.ResponseStream.MoveNext();
@@ -422,7 +409,7 @@ namespace Horde.Agent.Utility
 				Task task = await Task.WhenAny(moveNextAsync, _tailTaskStop.Task, Task.Delay(TimeSpan.FromMinutes(1.0), CancellationToken.None));
 				if (task == _tailTaskStop.Task)
 				{
-					_agentLogger.LogInformation("Cancelling long poll from client side (complete)");
+					_internalLogger.LogInformation("Cancelling long poll from client side (complete)");
 				}
 
 				// Close the request stream to indicate that we're finished
@@ -439,7 +426,7 @@ namespace Horde.Agent.Utility
 			}
 			catch (RpcException ex) when (ex.StatusCode == StatusCode.DeadlineExceeded)
 			{
-				_agentLogger.LogDebug(ex, "Log tail deadline exceeded, ignoring.");
+				_internalLogger.LogDebug(ex, "Log tail deadline exceeded, ignoring.");
 				return -1;
 			}
 		}
diff --git a/Engine/Source/Programs/Horde/Horde.Agent.Tests/JsonRpcLoggerTests.cs b/Engine/Source/Programs/Horde/Horde.Agent.Tests/JsonRpcLoggerTests.cs
--- a/Engine/Source/Programs/Horde/Horde.Agent.Tests/JsonRpcLoggerTests.cs
+++ b/Engine/Source/Programs/Horde/Horde.Agent.Tests/JsonRpcLoggerTests.cs
@@ -7,11 +7,9 @@ using System.Text;
 using System.Threading;
 using System.Threading.Channels;
 using System.Threading.Tasks;
-using Amazon.EC2.Model;
 using EpicGames.Core;
 using EpicGames.Horde;
 using EpicGames.Horde.Compute;
-using EpicGames.Horde.Jobs;
 using EpicGames.Horde.Logs;
 using EpicGames.Horde.Storage;
 using EpicGames.Horde.Storage.Bundles;
@@ -19,7 +17,6 @@ using Grpc.Core;
 using Grpc.Net.Client;
 using Horde.Agent.Utility;
 using Horde.Common.Rpc;
-using Microsoft.CodeAnalysis.CSharp.Syntax;
 using Microsoft.Extensions.Logging;
 using Microsoft.Extensions.Logging.Abstractions;
 using Microsoft.VisualStudio.TestTools.UnitTesting;
@@ -173,8 +170,8 @@ namespace Horde.Agent.Tests
 					=> await _writer.WriteAsync(message);
 			}
 
-			public static AsyncDuplexStreamingCall<TRequest, TResponse> CreateDuplexCall<TRequest, TResponse>(Func<ChannelReader<TRequest>, ChannelWriter<TResponse>, Task> func) 
-				where TRequest : class 
+			public static AsyncDuplexStreamingCall<TRequest, TResponse> CreateDuplexCall<TRequest, TResponse>(Func<ChannelReader<TRequest>, ChannelWriter<TResponse>, Task> func)
+				where TRequest : class
 				where TResponse : class
 			{
 				Channel<TRequest> requests = Channel.CreateUnbounded<TRequest>();
@@ -277,7 +274,7 @@ namespace Horde.Agent.Tests
 			LogId logId = default;
 
 			const int Count = 20000;
-			await using (ServerLogger logger = new ServerLogger(hordeClient, logId, null, LogLevel.Information, NullLogger.Instance, NullLogger.Instance))
+			await using (ServerLogger logger = new ServerLogger(hordeClient, logId, LogLevel.Information, NullLogger.Instance))
 			{
 				for (int idx = 0; idx < Count; idx++)
 				{
diff --git a/Engine/Source/Programs/Horde/Horde.Agent.Tests/WorkerServiceTest.cs b/Engine/Source/Programs/Horde/Horde.Agent.Tests/WorkerServiceTest.cs
--- a/Engine/Source/Programs/Horde/Horde.Agent.Tests/WorkerServiceTest.cs
+++ b/Engine/Source/Programs/Horde/Horde.Agent.Tests/WorkerServiceTest.cs
@@ -66,7 +66,7 @@ namespace Horde.Agent.Tests
 
 		class FakeServerLoggerFactory : IServerLoggerFactory
 		{
-			public IServerLogger CreateLogger(IHordeClient hordeClient, LogId logId, ILogger localLogger, bool? warnings = null, LogLevel outputLevel = LogLevel.Information) => new FakeServerLogger();
+			public IServerLogger CreateLogger(IHordeClient hordeClient, LogId logId, LogLevel outputLevel = LogLevel.Information) => new FakeServerLogger();
 		}
 
 		internal static IJobExecutor NullExecutor = new SimpleTestExecutor(async (step, logger, cancellationToken) =>
diff --git a/Engine/Source/Programs/Horde/Horde.Agent/Leases/Handlers/ComputeHandler.cs b/Engine/Source/Programs/Horde/Horde.Agent/Leases/Handlers/ComputeHandler.cs
--- a/Engine/Source/Programs/Horde/Horde.Agent/Leases/Handlers/ComputeHandler.cs
+++ b/Engine/Source/Programs/Horde/Horde.Agent/Leases/Handlers/ComputeHandler.cs
@@ -36,7 +36,7 @@ namespace Horde.Agent.Leases.Handlers
 		/// <inheritdoc/>
 		public override async Task<LeaseResult> ExecuteAsync(ISession session, LeaseId leaseId, ComputeTask computeTask, ILogger localLogger, CancellationToken cancellationToken)
 		{
-			await using IServerLogger? serverLogger = (computeTask.LogId != null) ? _serverLoggerFactory.CreateLogger(session.HordeClient, LogId.Parse(computeTask.LogId), localLogger, null, LogLevel.Trace) : null;
+			await using IServerLogger? serverLogger = (computeTask.LogId != null) ? _serverLoggerFactory.CreateLogger(session.HordeClient, LogId.Parse(computeTask.LogId), LogLevel.Trace).WithLocalLogger(localLogger) : null;
 			ILogger logger = serverLogger ?? localLogger;
 
 			if (!String.IsNullOrEmpty(computeTask.ParentLeaseId))
diff --git a/Engine/Source/Programs/Horde/Horde.Agent/Leases/Handlers/ConformHandler.cs b/Engine/Source/Programs/Horde/Horde.Agent/Leases/Handlers/ConformHandler.cs
--- a/Engine/Source/Programs/Horde/Horde.Agent/Leases/Handlers/ConformHandler.cs
+++ b/Engine/Source/Programs/Horde/Horde.Agent/Leases/Handlers/ConformHandler.cs
@@ -31,7 +31,7 @@ namespace Horde.Agent.Leases.Handlers
 		/// <inheritdoc/>
 		public override async Task<LeaseResult> ExecuteAsync(ISession session, LeaseId leaseId, ConformTask conformTask, ILogger localLogger, CancellationToken cancellationToken)
 		{
-			await using IServerLogger serverLogger = _serverLoggerFactory.CreateLogger(session.HordeClient, LogId.Parse(conformTask.LogId), localLogger, null);
+			await using IServerLogger serverLogger = _serverLoggerFactory.CreateLogger(session.HordeClient, LogId.Parse(conformTask.LogId)).WithLocalLogger(localLogger);
 			try
 			{
 				LeaseResult result = await ExecuteInternalAsync(session, leaseId, conformTask, serverLogger, cancellationToken);
diff --git a/Engine/Source/Programs/Horde/Horde.Agent/Leases/Handlers/JobHandler.cs b/Engine/Source/Programs/Horde/Horde.Agent/Leases/Handlers/JobHandler.cs
--- a/Engine/Source/Programs/Horde/Horde.Agent/Leases/Handlers/JobHandler.cs
+++ b/Engine/Source/Programs/Horde/Horde.Agent/Leases/Handlers/JobHandler.cs
@@ -159,7 +159,7 @@ namespace Horde.Agent.Leases.Handlers
 
 			// Create a storage client for this session
 			RpcJobOptions jobOptions = executeTask.JobOptions;
-			await using IServerLogger logger = _serverLoggerFactory.CreateLogger(session.HordeClient, executeTask.LogId, localLogger, null);
+			await using IServerLogger logger = _serverLoggerFactory.CreateLogger(session.HordeClient, executeTask.LogId).WithLocalLogger(localLogger);
 
 			logger.LogInformation("Executing job \"{JobName}\", jobId {JobId}, batchId {BatchId}, leaseId {LeaseId}, agentVersion {AgentVersion}", executeTask.JobName, executeTask.JobId, executeTask.BatchId, leaseId, AgentApp.Version);
 
diff --git a/Engine/Source/Programs/Horde/Horde.Agent/Leases/Handlers/UpgradeHandler.cs b/Engine/Source/Programs/Horde/Horde.Agent/Leases/Handlers/UpgradeHandler.cs
--- a/Engine/Source/Programs/Horde/Horde.Agent/Leases/Handlers/UpgradeHandler.cs
+++ b/Engine/Source/Programs/Horde/Horde.Agent/Leases/Handlers/UpgradeHandler.cs
@@ -28,7 +28,7 @@ namespace Horde.Agent.Leases.Handlers
 		/// <inheritdoc/>
 		public override async Task<LeaseResult> ExecuteAsync(ISession session, LeaseId leaseId, UpgradeTask task, ILogger localLogger, CancellationToken cancellationToken)
 		{
-			await using IServerLogger logger = _serverLoggerFactory.CreateLogger(session.HordeClient, LogId.Parse(task.LogId), localLogger, null);
+			await using IServerLogger logger = _serverLoggerFactory.CreateLogger(session.HordeClient, LogId.Parse(task.LogId)).WithLocalLogger(localLogger);
 
 			string requiredVersion = task.SoftwareId;
 
diff --git a/Engine/Source/Programs/Horde/Horde.Agent/Services/ServerLoggerFactory.cs b/Engine/Source/Programs/Horde/Horde.Agent/Services/ServerLoggerFactory.cs
--- a/Engine/Source/Programs/Horde/Horde.Agent/Services/ServerLoggerFactory.cs
+++ b/Engine/Source/Programs/Horde/Horde.Agent/Services/ServerLoggerFactory.cs
@@ -17,11 +17,9 @@ namespace Horde.Agent.Services
 		/// </summary>
 		/// <param name="hordeClient">The Horde client</param>
 		/// <param name="logId">The log id</param>
-		/// <param name="localLogger">Local log output device</param>
-		/// <param name="warnings">Whether to suppress warnings</param>
-		/// <param name="outputLevel">Minimum output level for messages</param>
+		/// <param name="minimumLevel">Minimum output level for messages</param>
 		/// <returns>New logger instance</returns>
-		IServerLogger CreateLogger(IHordeClient hordeClient, LogId logId, ILogger localLogger, bool? warnings, LogLevel outputLevel = LogLevel.Information);
+		IServerLogger CreateLogger(IHordeClient hordeClient, LogId logId, LogLevel minimumLevel = LogLevel.Information);
 	}
 
 	/// <summary>
@@ -40,11 +38,60 @@ namespace Horde.Agent.Services
 		}
 
 		/// <inheritdoc/>
-		public IServerLogger CreateLogger(IHordeClient hordeClient, LogId logId, ILogger localLogger, bool? warnings, LogLevel outputLevel)
+		public IServerLogger CreateLogger(IHordeClient hordeClient, LogId logId, LogLevel minimumLevel)
 		{
 #pragma warning disable CA2000 // Dispose objects before losing scope
-			return new ServerLogger(hordeClient, logId, warnings, outputLevel, localLogger, _logger);
+			return new ServerLogger(hordeClient, logId, minimumLevel, _logger);
 #pragma warning restore CA2000 // Dispose objects before losing scope
 		}
 	}
+
+	static class ServerLoggerExtensions
+	{
+		sealed class ServerLoggerWithLocalLogger : IServerLogger
+		{
+			readonly IServerLogger _serverLogger;
+			readonly ILogger _localLogger;
+
+			public ServerLoggerWithLocalLogger(IServerLogger serverLogger, ILogger localLogger)
+			{
+				_serverLogger = serverLogger;
+				_localLogger = localLogger;
+			}
+
+			/// <inheritdoc/>
+			public IDisposable? BeginScope<TState>(TState state) where TState : notnull
+				=> _localLogger.BeginScope(state);
+
+			/// <inheritdoc/>
+			public ValueTask DisposeAsync()
+				=> _serverLogger.DisposeAsync();
+
+			/// <inheritdoc/>
+			public bool IsEnabled(LogLevel logLevel)
+				=> _serverLogger.IsEnabled(logLevel) || _localLogger.IsEnabled(logLevel);
+
+			/// <inheritdoc/>
+			public void Log<TState>(LogLevel logLevel, EventId eventId, TState state, Exception? exception, Func<TState, Exception?, string> formatter)
+			{
+				_serverLogger.Log(logLevel, eventId, state, exception, formatter);
+				_localLogger.Log(logLevel, eventId, state, exception, formatter);
+			}
+
+			/// <inheritdoc/>
+			public Task StopAsync()
+				=> _serverLogger.StopAsync();
+		}
+
+		/// <summary>
+		/// Creates a logger which uploads data to the server
+		/// </summary>
+		/// <param name="serverLogger">Logger for the server</param>
+		/// <param name="localLogger">Local log output device</param>
+		/// <returns>New logger instance</returns>
+		public static IServerLogger WithLocalLogger(this IServerLogger serverLogger, ILogger localLogger)
+		{
+			return new ServerLoggerWithLocalLogger(serverLogger, localLogger);
+		}
+	}
 }
