Horde: Move methods for creating log events into LogRpc.

#rnx

diff --git a/Engine/Source/Programs/Horde/Horde.Agent.Driver/Utilities/HordeRpc.cs b/Engine/Source/Programs/Horde/Horde.Agent.Driver/Utilities/HordeRpc.cs
--- a/Engine/Source/Programs/Horde/Horde.Agent.Driver/Utilities/HordeRpc.cs
+++ b/Engine/Source/Programs/Horde/Horde.Agent.Driver/Utilities/HordeRpc.cs
@@ -81,25 +81,6 @@ namespace HordeCommon.Rpc
 		}
 	}
 
-	partial class RpcCreateEventRequest
-	{
-		public RpcCreateEventRequest(LogEventSeverity severity, LogId logId, int lineIndex, int lineCount)
-		{
-			Severity = (int)severity;
-			LogId = logId.ToString();
-			LineIndex = lineIndex;
-			LineCount = lineCount;
-		}
-	}
-
-	partial class RpcCreateEventsRequest
-	{
-		public RpcCreateEventsRequest(IEnumerable<RpcCreateEventRequest> events)
-		{
-			Events.AddRange(events);
-		}
-	}
-
 	partial class RpcWriteOutputRequest
 	{
 		public RpcWriteOutputRequest(LogId logId, long offset, int lineIndex, ByteString data, bool flush)
diff --git a/Engine/Source/Programs/Horde/Horde.Agent.Driver/Utility/ServerLogger.cs b/Engine/Source/Programs/Horde/Horde.Agent.Driver/Utility/ServerLogger.cs
--- a/Engine/Source/Programs/Horde/Horde.Agent.Driver/Utility/ServerLogger.cs
+++ b/Engine/Source/Programs/Horde/Horde.Agent.Driver/Utility/ServerLogger.cs
@@ -164,7 +164,7 @@ namespace Horde.Agent.Utility
 
 			// Buffer for events read in a single iteration
 			JsonRpcLogWriter writer = new JsonRpcLogWriter();
-			List<RpcCreateEventRequest> events = new List<RpcCreateEventRequest>();
+			List<RpcCreateLogEventRequest> events = new List<RpcCreateLogEventRequest>();
 
 			// Whether we've written the flush command
 			for (; ; )
@@ -250,11 +250,11 @@ namespace Horde.Agent.Utility
 			}
 		}
 
-		void AddEvent(ReadOnlySpan<byte> span, int lineIndex, int lineCount, LogEventSeverity severity, List<RpcCreateEventRequest> events)
+		void AddEvent(ReadOnlySpan<byte> span, int lineIndex, int lineCount, LogEventSeverity severity, List<RpcCreateLogEventRequest> events)
 		{
 			try
 			{
-				events.Add(new RpcCreateEventRequest(severity, _logId, lineIndex, lineCount));
+				events.Add(new RpcCreateLogEventRequest{ Severity = (int) severity, LogId = _logId.ToString(), LineIndex = lineIndex, LineCount = lineCount });
 			}
 			catch (Exception ex)
 			{
@@ -348,10 +348,14 @@ namespace Horde.Agent.Utility
 		}
 
 		/// <inheritdoc/>
-		public async Task WriteEventsAsync(List<RpcCreateEventRequest> events, CancellationToken cancellationToken)
+		public async Task WriteEventsAsync(List<RpcCreateLogEventRequest> events, CancellationToken cancellationToken)
 		{
-			JobRpc.JobRpcClient jobRpc = await _hordeClient.CreateGrpcClientAsync<JobRpc.JobRpcClient>(cancellationToken);
-			await jobRpc.CreateEventsAsync(new RpcCreateEventsRequest(events), cancellationToken: cancellationToken);
+			LogRpc.LogRpcClient logRpc = await _hordeClient.CreateGrpcClientAsync<LogRpc.LogRpcClient>(cancellationToken);
+
+			RpcCreateLogEventsRequest request = new RpcCreateLogEventsRequest();
+			request.Events.AddRange(events);
+
+			await logRpc.CreateLogEventsAsync(request, cancellationToken: cancellationToken);
 		}
 
 		/// <inheritdoc/>
@@ -376,7 +380,7 @@ namespace Horde.Agent.Utility
 		{
 			_internalLogger.LogInformation("Updating log {LogId} to line {LineCount}, target {Locator}", _logId, lineCount, target.GetLocator());
 
-			UpdateLogRequest request = new UpdateLogRequest();
+			RpcUpdateLogRequest request = new RpcUpdateLogRequest();
 			request.LogId = _logId.ToString();
 			request.LineCount = lineCount;
 			request.TargetHash = target.Hash.ToString();
@@ -393,10 +397,10 @@ namespace Horde.Agent.Utility
 			try
 			{
 				LogRpcClient clientRef = await _hordeClient.CreateGrpcClientAsync<LogRpcClient>(cancellationToken);
-				using AsyncDuplexStreamingCall<UpdateLogTailRequest, UpdateLogTailResponse> call = clientRef.UpdateLogTail(deadline: deadline, cancellationToken: cancellationToken);
+				using AsyncDuplexStreamingCall<RpcUpdateLogTailRequest, RpcUpdateLogTailResponse> call = clientRef.UpdateLogTail(deadline: deadline, cancellationToken: cancellationToken);
 
 				// Write the request to the server
-				UpdateLogTailRequest request = new UpdateLogTailRequest();
+				RpcUpdateLogTailRequest request = new RpcUpdateLogTailRequest();
 				request.LogId = _logId.ToString();
 				request.TailNext = tailNext;
 				request.TailData = UnsafeByteOperations.UnsafeWrap(tailData);
@@ -416,7 +420,7 @@ namespace Horde.Agent.Utility
 				await call.RequestStream.CompleteAsync();
 
 				// Wait for a response or a new update to come in, then close the request stream
-				UpdateLogTailResponse? response = null;
+				RpcUpdateLogTailResponse? response = null;
 				while (await moveNextAsync)
 				{
 					response = call.ResponseStream.Current;
diff --git a/Engine/Source/Programs/Horde/Horde.Agent.Tests/Helper.cs b/Engine/Source/Programs/Horde/Horde.Agent.Tests/Helper.cs
--- a/Engine/Source/Programs/Horde/Horde.Agent.Tests/Helper.cs
+++ b/Engine/Source/Programs/Horde/Horde.Agent.Tests/Helper.cs
@@ -95,7 +95,8 @@ namespace Horde.Agent.Tests
 			return Wrap(res);
 		}
 
-		public override AsyncUnaryCall<Empty> CreateEventsAsync(RpcCreateEventsRequest request, CallOptions options)
+		[Obsolete("Use LogRpc.CreateLogEvents instead")]
+		public override AsyncUnaryCall<Empty> CreateEventsAsync(RpcCreateLogEventsRequest request, CallOptions options)
 		{
 			_logger.LogDebug("CreateEventsAsync: {Request}", request);
 			Empty res = new Empty();
diff --git a/Engine/Source/Programs/Horde/Horde.Agent.Tests/JsonRpcLoggerTests.cs b/Engine/Source/Programs/Horde/Horde.Agent.Tests/JsonRpcLoggerTests.cs
--- a/Engine/Source/Programs/Horde/Horde.Agent.Tests/JsonRpcLoggerTests.cs
+++ b/Engine/Source/Programs/Horde/Horde.Agent.Tests/JsonRpcLoggerTests.cs
@@ -191,19 +191,19 @@ namespace Horde.Agent.Tests
 		{
 			public Dictionary<LogId, BlobLocator> Logs { get; } = new Dictionary<LogId, BlobLocator>();
 
-			public override AsyncUnaryCall<UpdateLogResponse> UpdateLogAsync(UpdateLogRequest request, CallOptions options)
+			public override AsyncUnaryCall<RpcUpdateLogResponse> UpdateLogAsync(RpcUpdateLogRequest request, CallOptions options)
 			{
 				Logs[LogId.Parse(request.LogId)] = new BlobLocator(request.TargetLocator);
-				return new AsyncUnaryCall<UpdateLogResponse>(Task.FromResult(new UpdateLogResponse()), Task.FromResult(new Metadata()), () => Status.DefaultSuccess, () => new Metadata(), () => { });
+				return new AsyncUnaryCall<RpcUpdateLogResponse>(Task.FromResult(new RpcUpdateLogResponse()), Task.FromResult(new Metadata()), () => Status.DefaultSuccess, () => new Metadata(), () => { });
 			}
 
-			public override AsyncDuplexStreamingCall<UpdateLogTailRequest, UpdateLogTailResponse> UpdateLogTail(CallOptions options)
+			public override AsyncDuplexStreamingCall<RpcUpdateLogTailRequest, RpcUpdateLogTailResponse> UpdateLogTail(CallOptions options)
 			{
-				return GrpcHelpers.CreateDuplexCall<UpdateLogTailRequest, UpdateLogTailResponse>(async (reader, writer) =>
+				return GrpcHelpers.CreateDuplexCall<RpcUpdateLogTailRequest, RpcUpdateLogTailResponse>(async (reader, writer) =>
 				{
 					while (await reader.WaitToReadAsync())
 					{
-						reader.TryRead(out UpdateLogTailRequest? request);
+						reader.TryRead(out RpcUpdateLogTailRequest? request);
 					}
 				});
 			}
diff --git a/Engine/Source/Programs/Horde/Horde.Server/Jobs/JobRpcService.cs b/Engine/Source/Programs/Horde/Horde.Server/Jobs/JobRpcService.cs
--- a/Engine/Source/Programs/Horde/Horde.Server/Jobs/JobRpcService.cs
+++ b/Engine/Source/Programs/Horde/Horde.Server/Jobs/JobRpcService.cs
@@ -663,14 +663,15 @@ namespace Horde.Server.Jobs
 		/// <param name="request">Request arguments</param>
 		/// <param name="context">Context for the RPC call</param>
 		/// <returns>Information about the new agent</returns>
-		public override async Task<Empty> CreateEvents(RpcCreateEventsRequest request, ServerCallContext context)
+		[Obsolete("Use LogRpc.CreateLogEvents instead")]
+		public override async Task<Empty> CreateEvents(RpcCreateLogEventsRequest request, ServerCallContext context)
 		{
 			if (!_globalConfig.Value.Authorize(LogAclAction.CreateEvent, context.GetHttpContext().User))
 			{
 				throw new StructuredRpcException(StatusCode.PermissionDenied, "Access denied");
 			}
 
-			foreach (IGrouping<string, RpcCreateEventRequest> createEventGroup in request.Events.GroupBy(x => x.LogId))
+			foreach (IGrouping<string, RpcCreateLogEventRequest> createEventGroup in request.Events.GroupBy(x => x.LogId))
 			{
 				ILog? log = await _logCollection.GetAsync(LogId.Parse(createEventGroup.Key), context.CancellationToken);
 				if (log == null)
@@ -679,7 +680,7 @@ namespace Horde.Server.Jobs
 				}
 
 				List<NewLogEventData> newEvents = new List<NewLogEventData>();
-				foreach (RpcCreateEventRequest createEvent in createEventGroup)
+				foreach (RpcCreateLogEventRequest createEvent in createEventGroup)
 				{
 					NewLogEventData newEvent = new NewLogEventData();
 					newEvent.Severity = (LogEventSeverity)createEvent.Severity;
diff --git a/Engine/Source/Programs/Horde/Horde.Server/Logs/LogRpcService.cs b/Engine/Source/Programs/Horde/Horde.Server/Logs/LogRpcService.cs
--- a/Engine/Source/Programs/Horde/Horde.Server/Logs/LogRpcService.cs
+++ b/Engine/Source/Programs/Horde/Horde.Server/Logs/LogRpcService.cs
@@ -1,17 +1,22 @@
 // Copyright Epic Games, Inc. All Rights Reserved.
 
 using System;
+using System.Collections.Generic;
+using System.Linq;
 using System.Threading;
 using System.Threading.Tasks;
 using EpicGames.Core;
 using EpicGames.Horde.Logs;
 using EpicGames.Horde.Storage;
+using Google.Protobuf.WellKnownTypes;
 using Grpc.Core;
 using Horde.Common.Rpc;
+using Horde.Server.Server;
 using Horde.Server.Storage;
 using Horde.Server.Utilities;
 using Microsoft.AspNetCore.Authorization;
 using Microsoft.Extensions.Logging;
+using Microsoft.Extensions.Options;
 
 namespace Horde.Server.Logs
 {
@@ -24,21 +29,23 @@ namespace Horde.Server.Logs
 		readonly ILogCollection _logCollection;
 		readonly LogTailService _logTailService;
 		readonly StorageService _storageService;
+		readonly IOptionsSnapshot<GlobalConfig> _globalConfig;
 		readonly ILogger<LogRpcService> _logger;
 
 		/// <summary>
 		/// Constructor
 		/// </summary>
-		public LogRpcService(ILogCollection logCollection, LogTailService logTailService, StorageService storageService, ILogger<LogRpcService> logger)
+		public LogRpcService(ILogCollection logCollection, LogTailService logTailService, StorageService storageService, IOptionsSnapshot<GlobalConfig> globalConfig, ILogger<LogRpcService> logger)
 		{
 			_logCollection = logCollection;
 			_logTailService = logTailService;
 			_storageService = storageService;
+			_globalConfig = globalConfig;
 			_logger = logger;
 		}
 
 		/// <inheritdoc/>
-		public override async Task<UpdateLogResponse> UpdateLog(UpdateLogRequest request, ServerCallContext context)
+		public override async Task<RpcUpdateLogResponse> UpdateLog(RpcUpdateLogRequest request, ServerCallContext context)
 		{
 			ILog? log = await _logCollection.GetAsync(LogId.Parse(request.LogId), context.CancellationToken);
 			if (log == null)
@@ -67,19 +74,19 @@ namespace Horde.Server.Logs
 
 			await _logTailService.FlushAsync(log.Id, request.LineCount);
 
-			return new UpdateLogResponse();
+			return new RpcUpdateLogResponse();
 		}
 
 		/// <inheritdoc/>
-		public override async Task UpdateLogTail(IAsyncStreamReader<UpdateLogTailRequest> requestStream, IServerStreamWriter<UpdateLogTailResponse> responseStream, ServerCallContext context)
+		public override async Task UpdateLogTail(IAsyncStreamReader<RpcUpdateLogTailRequest> requestStream, IServerStreamWriter<RpcUpdateLogTailResponse> responseStream, ServerCallContext context)
 		{
-			UpdateLogTailResponse response = new UpdateLogTailResponse();
+			RpcUpdateLogTailResponse response = new RpcUpdateLogTailResponse();
 			response.TailNext = -1;
 
 			Task<bool> moveNextTask = requestStream.MoveNext();
 			while (await moveNextTask)
 			{
-				UpdateLogTailRequest request = requestStream.Current;
+				RpcUpdateLogTailRequest request = requestStream.Current;
 				LogId logId = LogId.Parse(request.LogId);
 				_logger.LogDebug("Updating log tail for {LogId}: line {LineIdx}, size {Size}", logId, request.TailNext, request.TailData.Length);
 
@@ -118,5 +125,37 @@ namespace Horde.Server.Logs
 
 			await responseStream.WriteAsync(response);
 		}
+
+		/// <inheritdoc/>
+		public override async Task<Empty> CreateLogEvents(RpcCreateLogEventsRequest request, ServerCallContext context)
+		{
+			if (!_globalConfig.Value.Authorize(LogAclAction.CreateEvent, context.GetHttpContext().User))
+			{
+				throw new StructuredRpcException(StatusCode.PermissionDenied, "Access denied");
+			}
+
+			foreach (IGrouping<string, RpcCreateLogEventRequest> createEventGroup in request.Events.GroupBy(x => x.LogId))
+			{
+				ILog? log = await _logCollection.GetAsync(LogId.Parse(createEventGroup.Key), context.CancellationToken);
+				if (log == null)
+				{
+					throw new StructuredRpcException(StatusCode.NotFound, "Log not found");
+				}
+
+				List<NewLogEventData> newEvents = new List<NewLogEventData>();
+				foreach (RpcCreateLogEventRequest createEvent in createEventGroup)
+				{
+					NewLogEventData newEvent = new NewLogEventData();
+					newEvent.Severity = (LogEventSeverity)createEvent.Severity;
+					newEvent.LineIndex = createEvent.LineIndex;
+					newEvent.LineCount = createEvent.LineCount;
+					newEvents.Add(newEvent);
+				}
+
+				await log.AddEventsAsync(newEvents, context.CancellationToken);
+			}
+
+			return new Empty();
+		}
 	}
 }
