Horde: Strict null check and TS upgrade
#rnx
#skipci

diff --git a/Engine/Source/Programs/Horde/HordeDashboard/package.json b/Engine/Source/Programs/Horde/HordeDashboard/package.json
--- a/Engine/Source/Programs/Horde/HordeDashboard/package.json
+++ b/Engine/Source/Programs/Horde/HordeDashboard/package.json
@@ -18,7 +18,6 @@
     "@fluentui/react": "^8.106.2",
     "@fluentui/react-charting": "5.2.0",
     "d3": "^7.6.1",
-    "idb": "^7.1.1",
     "markdown-to-jsx": "^7.1.8",
     "mobx": "^6.12.2",
     "mobx-react-lite": "^4.0.7",
@@ -40,7 +39,7 @@
     "eslint": "^8.57.0",
     "eslint-plugin-react-hooks": "^4.6.0",
     "eslint-plugin-react-refresh": "^0.4.6",
-    "typescript": "^5.2.2",
+    "typescript": "^5.4.5",
     "vite": "^5.2.0",
     "@types/d3": "^7.4.0",
     "@types/markdown-to-jsx": "7.0.1",
diff --git a/Engine/Source/Programs/Horde/HordeDashboard/src/backend/Backend.ts b/Engine/Source/Programs/Horde/HordeDashboard/src/backend/Backend.ts
--- a/Engine/Source/Programs/Horde/HordeDashboard/src/backend/Backend.ts
+++ b/Engine/Source/Programs/Horde/HordeDashboard/src/backend/Backend.ts
@@ -271,7 +271,7 @@ export class Backend {
         return new Promise<GetAgentLeaseResponse[]>((resolve, reject) => {
             this.backend.get(`/api/v1/agents/${agentId}/leases?${paramString}`).then((response) => {
                 (response.data as GetAgentLeaseResponse[]).forEach(s => {
-                    s.startTime = s.startTime ? new Date(s.startTime) : undefined;
+                    s.startTime = s.startTime ? new Date(s.startTime) : new Date();
                     s.finishTime = s.finishTime ? new Date(s.finishTime) : undefined;
                 })
                 resolve(response.data as GetAgentLeaseResponse[]);
diff --git a/Engine/Source/Programs/Horde/HordeDashboard/src/backend/JobDetails.ts b/Engine/Source/Programs/Horde/HordeDashboard/src/backend/JobDetails.ts
--- a/Engine/Source/Programs/Horde/HordeDashboard/src/backend/JobDetails.ts
+++ b/Engine/Source/Programs/Horde/HordeDashboard/src/backend/JobDetails.ts
@@ -92,7 +92,7 @@ export class JobDetails {
                 return;
             }
 
-            const requests = [];
+            const requests:any = [];
 
             if (this.stepId !== stepId || this.labelIdx !== labelIdx) {
 
@@ -310,6 +310,10 @@ export class JobDetails {
 
     getStepName(stepId: string | undefined, includeRetry: boolean = true): string {
 
+        if (!stepId?.length) {
+            return "";
+        }
+
         const step = this.stepById(stepId);
 
         if (!step?.name) {
@@ -507,7 +511,7 @@ export class JobDetails {
         this.stream = projectStore.streamById(jobdata.streamId);
         this.batches = jobdata.batches ?? [];
 
-        this.labels = jobdata.labels.map((label, index) => {
+        this.labels = jobdata.labels?.map((label, index) => {
 
             return {
                 default: false,
@@ -515,7 +519,7 @@ export class JobDetails {
                 internal: label,
                 timing: undefined
             }
-        });
+        }) ?? [];
 
         const defaultLabel = jobdata.defaultLabel;
         if (defaultLabel) {
@@ -659,7 +663,7 @@ export class JobDetails {
             this.updating = true;
             const cancelID = this.cancelID++;
 
-            let requests = [];
+            let requests:any = [];
 
             requests.push(backend.getJob(this.id));
 
diff --git a/Engine/Source/Programs/Horde/HordeDashboard/src/backend/TestDataHandler.ts b/Engine/Source/Programs/Horde/HordeDashboard/src/backend/TestDataHandler.ts
--- a/Engine/Source/Programs/Horde/HordeDashboard/src/backend/TestDataHandler.ts
+++ b/Engine/Source/Programs/Horde/HordeDashboard/src/backend/TestDataHandler.ts
@@ -196,7 +196,7 @@ export class TestDataWrapper implements TestData {
         let stepName = "";
         const batch = this.jobdata.batches?.find(b => b.steps.find(s => s.id === this.stepId));
         const stepNode = batch?.steps.find(s => s.id === this.stepId);
-        stepName = stepNode.name;
+        stepName = stepNode?.name ?? "";
 
         return this.stepName;
     }
diff --git a/Engine/Source/Programs/Horde/HordeDashboard/src/components/AgentView.tsx b/Engine/Source/Programs/Horde/HordeDashboard/src/components/AgentView.tsx
--- a/Engine/Source/Programs/Horde/HordeDashboard/src/components/AgentView.tsx
+++ b/Engine/Source/Programs/Horde/HordeDashboard/src/components/AgentView.tsx
@@ -2514,8 +2514,8 @@ export const AgentViewInner: React.FC<{ agentId?: string, poolId?: string, searc
                </Stack>
             );
          case 'pools':
-            const poolItems = [];
-            const poolSearchNames = [];
+            const poolItems:any = [];
+            const poolSearchNames:string[] = [];
             if (agent.pools) {
                const poolObjs: PoolData[] = [];
                // get actual pool objects
@@ -2589,7 +2589,7 @@ export const AgentViewInner: React.FC<{ agentId?: string, poolId?: string, searc
             return <Stack horizontal horizontalAlign={"start"} styles={{ root: { overflow: "auto", height: '100%' } }} tokens={{ childrenGap: 4 }}>{poolItems}</Stack>;
          case 'status':
             const leases: any = [];
-            const leaseSearchItems = [];
+            const leaseSearchItems:string[] = [];
             if (agent.leases) {
                agent.leases.forEach(lease => {
 
@@ -2637,7 +2637,9 @@ export const AgentViewInner: React.FC<{ agentId?: string, poolId?: string, searc
                         <Text key={"statusText_" + agent.id + "_" + lease.id}>{statusText}</Text>
                      </Stack.Item>);
                   }
-                  leaseSearchItems.push(lease.name);
+                  if (lease.name) {
+                     leaseSearchItems.push(lease.name);
+                  }                  
                });
             }
             // if there are no leases, we'll push some other state.
diff --git a/Engine/Source/Programs/Horde/HordeDashboard/src/components/AuditLog.tsx b/Engine/Source/Programs/Horde/HordeDashboard/src/components/AuditLog.tsx
--- a/Engine/Source/Programs/Horde/HordeDashboard/src/components/AuditLog.tsx
+++ b/Engine/Source/Programs/Horde/HordeDashboard/src/components/AuditLog.tsx
@@ -855,7 +855,7 @@ class AuditLine {
             return t;
          }
 
-         const rtags = [];
+         const rtags:any = [];
 
          const key = `log_line_${idx}_${index}_fragment`;
 
diff --git a/Engine/Source/Programs/Horde/HordeDashboard/src/components/IssueViewV2.tsx b/Engine/Source/Programs/Horde/HordeDashboard/src/components/IssueViewV2.tsx
--- a/Engine/Source/Programs/Horde/HordeDashboard/src/components/IssueViewV2.tsx
+++ b/Engine/Source/Programs/Horde/HordeDashboard/src/components/IssueViewV2.tsx
@@ -2524,7 +2524,7 @@ const CreateExternalIssueModal: React.FC<{ onClose: () => void }> = ({ onClose }
 
       const project = projects.find(p => p.name === projectName)!;
 
-      let componentName = undefined;
+      let componentName: string | undefined;
       for (let id in project.components) {
          if (project.components[id] === state.componentName) {
             componentName = state.componentName;
diff --git a/Engine/Source/Programs/Horde/HordeDashboard/src/components/JobDetailCommon.tsx b/Engine/Source/Programs/Horde/HordeDashboard/src/components/JobDetailCommon.tsx
--- a/Engine/Source/Programs/Horde/HordeDashboard/src/components/JobDetailCommon.tsx
+++ b/Engine/Source/Programs/Horde/HordeDashboard/src/components/JobDetailCommon.tsx
@@ -88,7 +88,7 @@ export const getBatchText = (item: StepItem): string | undefined => {
 
    const batch = item.batch;
 
-   let statusText = undefined;
+   let statusText: string | undefined;
 
    if (batch) {
 
@@ -131,7 +131,7 @@ export const getBatchText = (item: StepItem): string | undefined => {
 
       if (batch.error === JobStepBatchError.SyncingFailed) {
          statusText = `${item.agentId} : Syncing failed`;
-      }      
+      }
 
       if (batch.error === JobStepBatchError.Cancelled) {
          if (item.agentId) {
diff --git a/Engine/Source/Programs/Horde/HordeDashboard/src/components/JobViewAll.tsx b/Engine/Source/Programs/Horde/HordeDashboard/src/components/JobViewAll.tsx
--- a/Engine/Source/Programs/Horde/HordeDashboard/src/components/JobViewAll.tsx
+++ b/Engine/Source/Programs/Horde/HordeDashboard/src/components/JobViewAll.tsx
@@ -433,7 +433,7 @@ const JobViewAllInner: React.FC<{ filter: JobFilterSimple }> = observer(({ filte
       });
 
       const catergories = new Set<string>();
-      view.forEach(label => catergories.add(label.dashboardCategory));
+      view.forEach(label => { if (label.dashboardCategory?.length) catergories.add(label.dashboardCategory) });
 
       const sorted = Array.from(catergories).sort((a, b) => { return a < b ? -1 : 1 });
 
diff --git a/Engine/Source/Programs/Horde/HordeDashboard/src/components/LogRender.tsx b/Engine/Source/Programs/Horde/HordeDashboard/src/components/LogRender.tsx
--- a/Engine/Source/Programs/Horde/HordeDashboard/src/components/LogRender.tsx
+++ b/Engine/Source/Programs/Horde/HordeDashboard/src/components/LogRender.tsx
@@ -259,7 +259,7 @@ const renderTags = (navigate: NavigateFunction, line: LogLine, lineNumber: numbe
          return t;
       }
 
-      const rtags = [];
+      const rtags:any = [];
 
       const key = `log_line_${lineNumber}_${idx}_${index}_fragment`;
 
diff --git a/Engine/Source/Programs/Horde/HordeDashboard/src/components/NotificationDropdown.tsx b/Engine/Source/Programs/Horde/HordeDashboard/src/components/NotificationDropdown.tsx
--- a/Engine/Source/Programs/Horde/HordeDashboard/src/components/NotificationDropdown.tsx
+++ b/Engine/Source/Programs/Horde/HordeDashboard/src/components/NotificationDropdown.tsx
@@ -480,7 +480,7 @@ export const NotificationDropdown: React.FC<{ jobDetails: JobDetailsV2 }> = obse
    }
    if (label) {
       type = "Label";
-      const aggregate = jobDetails.findLabel(label.stateResponse.dashboardName, label.stateResponse.dashboardCategory)!;
+      const aggregate = jobDetails.findLabel(label.stateResponse.dashboardName ?? "", label.stateResponse.dashboardCategory)!;
       complete = aggregate?.stateResponse?.state === LabelState.Complete;
       subscribed = subHandler.nlabel?.slack ?? false;
    }
diff --git a/Engine/Source/Programs/Horde/HordeDashboard/src/components/TestReportView.tsx b/Engine/Source/Programs/Horde/HordeDashboard/src/components/TestReportView.tsx
--- a/Engine/Source/Programs/Horde/HordeDashboard/src/components/TestReportView.tsx
+++ b/Engine/Source/Programs/Horde/HordeDashboard/src/components/TestReportView.tsx
@@ -60,7 +60,7 @@ const getTestdataBreadcrumbProps = async (testdata: TestDataWrapper) => {
             link: testdata.getJobLink(),
         },
         {
-            text: stepName,
+            text: stepName ?? "",
             link: testdata.getJobStepLink(),
         },
         {
diff --git a/Engine/Source/Programs/Horde/HordeDashboard/src/components/UserHome.tsx b/Engine/Source/Programs/Horde/HordeDashboard/src/components/UserHome.tsx
--- a/Engine/Source/Programs/Horde/HordeDashboard/src/components/UserHome.tsx
+++ b/Engine/Source/Programs/Horde/HordeDashboard/src/components/UserHome.tsx
@@ -911,7 +911,7 @@ const JobsPanel: React.FC<{ includeOtherPreflights: boolean }> = observer(({ inc
       });
 
       const catergories = new Set<string>();
-      view.forEach(label => catergories.add(label.dashboardCategory));
+      view.forEach(label => { if (label.dashboardCategory?.length) { catergories.add(label.dashboardCategory) } });
 
       const sorted = Array.from(catergories).sort((a, b) => { return a < b ? -1 : 1 });
 
diff --git a/Engine/Source/Programs/Horde/HordeDashboard/src/components/build/NewBuildV2.tsx b/Engine/Source/Programs/Horde/HordeDashboard/src/components/build/NewBuildV2.tsx
--- a/Engine/Source/Programs/Horde/HordeDashboard/src/components/build/NewBuildV2.tsx
+++ b/Engine/Source/Programs/Horde/HordeDashboard/src/components/build/NewBuildV2.tsx
@@ -162,7 +162,12 @@ class BuildOptions {
    queryShelvedChange?: number;
    autoSubmit?: boolean;
 
-   static get() { return BuildOptions.instance };
+   static get() {
+      if (!BuildOptions.instance) {
+         throw `Build Options have not been instantiated`;
+      }
+      return BuildOptions.instance!;
+   }
 
    template?: GetTemplateRefResponse;
    templates: GetTemplateRefResponse[] = [];
@@ -354,6 +359,9 @@ class BuildOptions {
 
 
    validate(errors: ValidationError[]): boolean {
+      if (!this.template) {
+         return true;
+      }
 
       if (this.preflightChange && !this.template.allowPreflights) {
          errors.push({ error: `Template "${this.template.name}" does not allow preflights`, paramString: "Shelved Change", value: this.preflightChange.toString() })
@@ -393,6 +401,9 @@ class BuildOptions {
    estimateHeight(previewPanel?: boolean): number {
 
       const template = this.template;
+      if (!template) {
+         return 0;
+      }
       let estimatedHeight = parameterGap * template.parameters.length;
 
       const estimateHeight = (param: ParameterData) => {
@@ -497,7 +508,7 @@ class BuildOptions {
    private setTemplate(templateId: string, notifyChanged = true) {
 
       const template = this.allTemplates.find(t => t.id === templateId);
-      if (template.id === this.template?.id) {
+      if (!template || template.id === this.template?.id) {
          return;
       }
 
@@ -506,7 +517,7 @@ class BuildOptions {
       }
 
       const source = { ...template };
-      source.hash = undefined;
+      source.hash = "";
 
       this.previewSource = JSON.stringify(source, null, 4);
 
@@ -567,7 +578,7 @@ class BuildOptions {
                if (any) {
 
                   list.items.forEach(i => {
-                     this.parameters[i.id] = undefined;
+                     (this.parameters as any)[i.id] = undefined;
                      qid = `id-${i.id}`;
                      if (query.get(qid)) {
                         this.parameters[i.id] = query.get(qid) as string;
@@ -700,7 +711,7 @@ class BuildOptions {
          if (allowtemplatechange) {
             templates = [...this.allTemplates];
          } else {
-            const template = this.allTemplates.find(t => t.name === this.jobDetails.template?.name);
+            const template = this.allTemplates.find(t => t.name === this.jobDetails?.template?.name);
             if (template) {
                templates = [template];
             }
@@ -745,8 +756,8 @@ const BoolParameter: React.FC<{ param: BoolParameterData }> = observer(({ param
          disabled={options.readOnly}
          checked={options.parameters[param.id] == "true"}
          onChange={(ev, value) => {
-            ev.preventDefault();
-            options.onBooleanChanged(param.id, value);
+            ev?.preventDefault();
+            options.onBooleanChanged(param.id, value ?? false);
          }}
       />
    </Stack>
@@ -769,7 +780,7 @@ const TextParameter: React.FC<{ param: TextParameterData }> = observer(({ param
          value={options.parameters[param.id] ?? ""}
          disabled={options.readOnly}
          onChange={(ev, value) => {
-            options.onTextChanged(param.id, value);
+            options.onTextChanged(param.id, value ?? "");
          }}
       />
    </Stack>
@@ -1331,7 +1342,7 @@ const AdvancedPanel: React.FC = observer(() => {
          </Stack>
          {showAdditionalArgs && <TextField key={"key_adv_add_args"} multiline style={{ height: 48 }} resizable={false} readOnly={options.readOnly} spellCheck={false} defaultValue={options.advAdditionalArgs} label="Additional Arguments" onChange={(ev, newValue) => options.advAdditionalArgs = newValue} />}
          {showArgumentClipboardButton && <DefaultButton text="Copy Job Arguments to Clipboard" style={{ width: 240 }} onClick={() => copyToClipboard(
-            options.jobDetails.jobData.arguments.map(arg => {
+            options.jobDetails?.jobData?.arguments.map(arg => {
                if (arg.indexOf("=") !== -1) {
                   const components = arg.split("=");
                   if (components[1].indexOf(" ") === -1) {
@@ -1354,7 +1365,7 @@ const AdvancedPanel: React.FC = observer(() => {
                   if (checked) {
                      options.onModeChanged("Basic");
                   }
-                  options.onPreviewChanged(checked);
+                  options.onPreviewChanged(checked ?? false);
                }} />
          </Stack>
 
@@ -1420,7 +1431,7 @@ const BuildModal: React.FC = observer(() => {
          }
       }
 
-      let additionalArgs = [];
+      let additionalArgs:string[] = [];
       if (options!.advAdditionalArgs) {
          const argRegex = /"(\\"|[^"])*?"|[^ ]+/g;
          options.advAdditionalArgs.trim().match(argRegex)?.forEach(arg => additionalArgs.push(arg.replace(/"/g, "")));
diff --git a/Engine/Source/Programs/Horde/HordeDashboard/src/components/jobDetailsV2/JobDetailLabels.tsx b/Engine/Source/Programs/Horde/HordeDashboard/src/components/jobDetailsV2/JobDetailLabels.tsx
--- a/Engine/Source/Programs/Horde/HordeDashboard/src/components/jobDetailsV2/JobDetailLabels.tsx
+++ b/Engine/Source/Programs/Horde/HordeDashboard/src/components/jobDetailsV2/JobDetailLabels.tsx
@@ -22,19 +22,18 @@ export const LabelsPanelV2: React.FC<{ jobDetails: JobDetailsV2, dataView: JobDa
 
    const labelIdx = query.get("label") ? parseInt(query.get("label")!) : undefined;
    const qlabel = jobDetails.labelByIndex(labelIdx);
-   if (typeof (labelIdx) === `number` && !qlabel) {            
+   if (typeof (labelIdx) === `number` && !qlabel) {
       return null;
    }
-   
+
    const jobFilter = jobDetails.filter;
 
    // subscribe   
    if (jobFilter.inputChanged) { }
 
    const batchFilter = query.get("batch");
-   if (batchFilter)
-   {
-      return null;   
+   if (batchFilter) {
+      return null;
    }
 
    let labels = jobDetails.labels.filter(label => label.stateResponse.state !== LabelState.Unspecified);
@@ -78,7 +77,7 @@ export const LabelsPanelV2: React.FC<{ jobDetails: JobDetailsV2, dataView: JobDa
 
       return include;
 
-   });   
+   });
 
    if (!labels) {
       return null;
@@ -118,7 +117,7 @@ export const LabelsPanelV2: React.FC<{ jobDetails: JobDetailsV2, dataView: JobDa
          i.labels = i.labels.filter(label => (label.stateResponse.dashboardName ?? "").toLowerCase().indexOf(filter.currentInput!.toLowerCase()) !== -1);
          return i.labels.length !== 0;
       });
-   }   
+   }
 
    if (!items.length) {
       return null;
@@ -154,7 +153,7 @@ export const LabelsPanelV2: React.FC<{ jobDetails: JobDetailsV2, dataView: JobDa
 
                   let filtered = false;
                   if (qlabel) {
-                     if (qlabel.stateResponse.dashboardCategory!== label.stateResponse.dashboardCategory || qlabel.stateResponse.dashboardName !== label.stateResponse.dashboardName) {
+                     if (qlabel.stateResponse.dashboardCategory !== label.stateResponse.dashboardCategory || qlabel.stateResponse.dashboardName !== label.stateResponse.dashboardName) {
                         filtered = true;
                      }
                   }
@@ -165,17 +164,20 @@ export const LabelsPanelV2: React.FC<{ jobDetails: JobDetailsV2, dataView: JobDa
                               navigate(location.pathname)
                            } else {
 
-                              const idx = jobDetails.labelIndex(label.stateResponse.dashboardName, label.stateResponse.dashboardCategory);
+                              let idx = -1;
+                              if (label.stateResponse.dashboardName?.length) {
+                                 idx = jobDetails.labelIndex(label.stateResponse.dashboardName, label.stateResponse.dashboardCategory);
+                              }
                               if (idx >= 0) {
                                  navigate(location.pathname + `?label=${idx}`)
                               } else {
-                                 navigate(location.pathname, {replace: true})
-                              }                              
+                                 navigate(location.pathname, { replace: true })
+                              }
                            }
 
                         }}
-                        key={label.stateResponse.dashboardName  ?? ""} style={{ backgroundColor: color.primaryColor, color: textColor, filter: filtered ? "brightness(0.70)" : undefined }}
-                        text={label.stateResponse.dashboardName  ?? ""}>
+                        key={label.stateResponse.dashboardName ?? ""} style={{ backgroundColor: color.primaryColor, color: textColor, filter: filtered ? "brightness(0.70)" : undefined }}
+                        text={label.stateResponse.dashboardName ?? ""}>
                         {!!color.secondaryColor && <div style={{
                            borderLeft: "10px solid transparent",
                            borderRight: `10px solid ${color.secondaryColor}`,
@@ -200,7 +202,7 @@ export const LabelsPanelV2: React.FC<{ jobDetails: JobDetailsV2, dataView: JobDa
 
    return <Stack styles={{ root: { paddingTop: 18, paddingRight: 12 } }}>
       <Stack style={{ paddingLeft: 12 }}>
-         <Stack styles={{root:{selectors: {'.ms-DetailsRow': {backgroundColor: dashboard.darktheme ? modeColors.content : undefined}}}}}>
+         <Stack styles={{ root: { selectors: { '.ms-DetailsRow': { backgroundColor: dashboard.darktheme ? modeColors.content : undefined } } } }}>
             <DetailsList
                isHeaderVisible={false}
                indentWidth={0}
diff --git a/Engine/Source/Programs/Horde/HordeDashboard/src/components/jobDetailsV2/JobDetailStepHistoryV2.tsx b/Engine/Source/Programs/Horde/HordeDashboard/src/components/jobDetailsV2/JobDetailStepHistoryV2.tsx
--- a/Engine/Source/Programs/Horde/HordeDashboard/src/components/jobDetailsV2/JobDetailStepHistoryV2.tsx
+++ b/Engine/Source/Programs/Horde/HordeDashboard/src/components/jobDetailsV2/JobDetailStepHistoryV2.tsx
@@ -406,7 +406,7 @@ class StepHistoryRenderer {
          .attr("markerHeight", 5)
          .attr("orient", 'auto-start-reverse')
          .append("path")
-         .attr('d', d3.line()(arrowPoints))
+         .attr('d', (d3.line() as any)(arrowPoints))
          .style("fill", dashboard.darktheme ? "#6D6C6B" : "#4D4C4B");
 
       //svg.append("g")
diff --git a/Engine/Source/Programs/Horde/HordeDashboard/src/components/jobDetailsV2/JobDetailStepTrendsV2.tsx b/Engine/Source/Programs/Horde/HordeDashboard/src/components/jobDetailsV2/JobDetailStepTrendsV2.tsx
--- a/Engine/Source/Programs/Horde/HordeDashboard/src/components/jobDetailsV2/JobDetailStepTrendsV2.tsx
+++ b/Engine/Source/Programs/Horde/HordeDashboard/src/components/jobDetailsV2/JobDetailStepTrendsV2.tsx
@@ -337,7 +337,7 @@ class StepTrendsRenderer {
          .attr("markerHeight", 5)
          .attr("orient", 'auto-start-reverse')
          .append("path")
-         .attr('d', d3.line()(arrowPoints))
+         .attr('d', (d3.line() as any)(arrowPoints))
          .style("fill", dashboard.darktheme ? "#6D6C6B" : "#4D4C4B");
 
       //svg.append("g")
@@ -366,7 +366,7 @@ class StepTrendsRenderer {
       const lineI = d3.range(data.length);
 
       const curve = d3.curveLinear;
-      const line = d3.line()
+      const line:any = d3.line()
          .defined(i => true)
          .curve(curve)
          .x(i => { return x(new Date(data[i as any].startTime!).getTime() / 1000) })
@@ -445,7 +445,7 @@ class StepTrendsRenderer {
 
          const strokeWidth = zoomLevel > 4 ? 2 : 1;
 
-         const scaledLine = d3.line()
+         const scaledLine:any = d3.line()
             .defined(i => true)
             .curve(curve)
             .x(i => { return x(new Date(data[i as any].startTime!).getTime() / 1000) })
@@ -472,7 +472,7 @@ class StepTrendsRenderer {
             .attr("height", (d: any) => scaleY(0) - scaleY(dataView.durations.get(d.jobId)!))
 
 
-         const scaledLine = d3.line()
+         const scaledLine:any = d3.line()
             .defined(i => true)
             .curve(curve)
             .x(i => { return x(new Date(data[i as any].startTime!).getTime() / 1000) })
diff --git a/Engine/Source/Programs/Horde/HordeDashboard/src/components/jobDetailsV2/JobDetailViewSteps.tsx b/Engine/Source/Programs/Horde/HordeDashboard/src/components/jobDetailsV2/JobDetailViewSteps.tsx
--- a/Engine/Source/Programs/Horde/HordeDashboard/src/components/jobDetailsV2/JobDetailViewSteps.tsx
+++ b/Engine/Source/Programs/Horde/HordeDashboard/src/components/jobDetailsV2/JobDetailViewSteps.tsx
@@ -395,6 +395,9 @@ export const StepsPanelInner: React.FC<{ jobDetails: JobDetailsV2, depStepId?: s
 
       const getStepsRecursive = (stepId: string) => {
          const rstep = jobDetails.stepById(stepId);
+         if (!rstep) {
+            return;
+         }
          [rstep.inputDependencies, rstep.orderDependencies].flat().forEach(id => {
             const s = jobDetails.stepById(id);
             if (s && !stepFilter.find(s => s.id === id)) {
@@ -515,7 +518,7 @@ export const StepsPanelInner: React.FC<{ jobDetails: JobDetailsV2, depStepId?: s
          return;
       }
 
-      const b = jobDetails.jobData.batches.find(b => b.steps.findIndex(s => s.id === step.id) !== -1);
+      const b = jobDetails.jobData?.batches?.find(b => b.steps.findIndex(s => s.id === step.id) !== -1);
       if (b)
          stepBatches.add(b.id);
    });
diff --git a/Engine/Source/Programs/Horde/HordeDashboard/src/components/jobDetailsV2/JobDetailsViewCommon.tsx b/Engine/Source/Programs/Horde/HordeDashboard/src/components/jobDetailsV2/JobDetailsViewCommon.tsx
--- a/Engine/Source/Programs/Horde/HordeDashboard/src/components/jobDetailsV2/JobDetailsViewCommon.tsx
+++ b/Engine/Source/Programs/Horde/HordeDashboard/src/components/jobDetailsV2/JobDetailsViewCommon.tsx
@@ -547,7 +547,7 @@ export class JobDetailsV2 extends PollBase {
          return;
       }
 
-      let requests = [];
+      let requests:any[] = [];
 
       const initialRequest = !this.jobData;
 
@@ -678,6 +678,9 @@ export class JobDetailsV2 extends PollBase {
       const getStepsRecursive = (stepId: string) => {
 
          const step = this.stepById(stepId);
+         if (!step) {
+            return;
+         }
 
          [step.inputDependencies, step.orderDependencies].flat().forEach(id => {
             const s = this.stepById(id);
diff --git a/Engine/Source/Programs/Horde/HordeDashboard/src/components/jobDetailsV2/JobOperationsBar.tsx b/Engine/Source/Programs/Horde/HordeDashboard/src/components/jobDetailsV2/JobOperationsBar.tsx
--- a/Engine/Source/Programs/Horde/HordeDashboard/src/components/jobDetailsV2/JobOperationsBar.tsx
+++ b/Engine/Source/Programs/Horde/HordeDashboard/src/components/jobDetailsV2/JobOperationsBar.tsx
@@ -40,7 +40,7 @@ export const JobOperations: React.FC<{ jobDetails: JobDetailsV2 }> = observer(({
 
    const stepId = query.get("step") ? query.get("step")! : undefined;
    const batchFilter = query.get("batch");
-   let newBuildVersion: string = query.get("newbuildversion") ? query.get("newbuildversion") : "1";
+   let newBuildVersion: string = query.get("newbuildversion") ? query.get("newbuildversion")! : "1";
 
    // show new build version 2 if we have new parameters store
    const hasJobParameters = !!(Object.keys(jobDetails?.jobData?.parameters ?? {}).length)
diff --git a/Engine/Source/Programs/Horde/HordeDashboard/src/components/jobDetailsV2/StepRetryModal.tsx b/Engine/Source/Programs/Horde/HordeDashboard/src/components/jobDetailsV2/StepRetryModal.tsx
--- a/Engine/Source/Programs/Horde/HordeDashboard/src/components/jobDetailsV2/StepRetryModal.tsx
+++ b/Engine/Source/Programs/Horde/HordeDashboard/src/components/jobDetailsV2/StepRetryModal.tsx
@@ -242,7 +242,7 @@ export const StepRetryModal: React.FC<{ stepId: string; jobDetails: JobDetailsV2
 
       const job = jobData;
 
-      const args = [];
+      const args:string[] = [];
       args.push(`-Target=Setup Build`);
       args.push(`-Target=${step!.name}`);
 
diff --git a/Engine/Source/Programs/Horde/HordeDashboard/src/components/telemetry/TelemetryLineGraph.ts b/Engine/Source/Programs/Horde/HordeDashboard/src/components/telemetry/TelemetryLineGraph.ts
--- a/Engine/Source/Programs/Horde/HordeDashboard/src/components/telemetry/TelemetryLineGraph.ts
+++ b/Engine/Source/Programs/Horde/HordeDashboard/src/components/telemetry/TelemetryLineGraph.ts
@@ -91,7 +91,7 @@ export class TelemetryLineRenderer {
          .selectAll("path")
          .data(gvalues)
          .join("path")
-         .attr("stroke", d => { return d[0][3] !== undefined ? graphColors[d[0][3] as number] : "#8ab8ff" })
+         .attr("stroke", (d:any) => { return d[0][3] !== undefined ? graphColors[d[0][3] as number] : "#8ab8ff" })
          .attr("d", line as any);
 
       svg.append("g")
diff --git a/Engine/Source/Programs/Horde/HordeDashboard/src/components/test/IssueStepReport.tsx b/Engine/Source/Programs/Horde/HordeDashboard/src/components/test/IssueStepReport.tsx
--- a/Engine/Source/Programs/Horde/HordeDashboard/src/components/test/IssueStepReport.tsx
+++ b/Engine/Source/Programs/Horde/HordeDashboard/src/components/test/IssueStepReport.tsx
@@ -73,7 +73,8 @@ class ToolHandler {
       }
 
 
-      const stepName = this.stepName = batch.steps.find(s => s.id === stepId)?.name;
+      const stepName = batch.steps.find(s => s.id === stepId)?.name;
+      this.stepName = stepName ?? "";
 
       if (!stepName) {
          console.error("Bad step name");
diff --git a/Engine/Source/Programs/Horde/HordeDashboard/tsconfig.json b/Engine/Source/Programs/Horde/HordeDashboard/tsconfig.json
--- a/Engine/Source/Programs/Horde/HordeDashboard/tsconfig.json
+++ b/Engine/Source/Programs/Horde/HordeDashboard/tsconfig.json
@@ -21,7 +21,8 @@
 
     /* Linting */
     "strict": false,
-    "noFallthroughCasesInSwitch": true
+    "noFallthroughCasesInSwitch": true,
+    "strictNullChecks": true
     /*
     "noUnusedLocals": true,
     "noUnusedParameters": true,    
