Horde: Remove legacy RPC messages for uploading log data; all data is written directly to the storage system now.

#rnx

diff --git a/Engine/Source/Programs/Horde/Horde.Agent.Driver/Utilities/HordeRpc.cs b/Engine/Source/Programs/Horde/Horde.Agent.Driver/Utilities/HordeRpc.cs
--- a/Engine/Source/Programs/Horde/Horde.Agent.Driver/Utilities/HordeRpc.cs
+++ b/Engine/Source/Programs/Horde/Horde.Agent.Driver/Utilities/HordeRpc.cs
@@ -81,18 +81,6 @@ namespace HordeCommon.Rpc
 		}
 	}
 
-	partial class RpcWriteOutputRequest
-	{
-		public RpcWriteOutputRequest(LogId logId, long offset, int lineIndex, ByteString data, bool flush)
-		{
-			LogId = logId.ToString();
-			Offset = offset;
-			LineIndex = lineIndex;
-			Data = data;
-			Flush = flush;
-		}
-	}
-
 	partial class RpcDownloadSoftwareRequest
 	{
 		public RpcDownloadSoftwareRequest(string version)
diff --git a/Engine/Source/Programs/Horde/Horde.Agent.Driver/Utility/ServerLogger.cs b/Engine/Source/Programs/Horde/Horde.Agent.Driver/Utility/ServerLogger.cs
--- a/Engine/Source/Programs/Horde/Horde.Agent.Driver/Utility/ServerLogger.cs
+++ b/Engine/Source/Programs/Horde/Horde.Agent.Driver/Utility/ServerLogger.cs
@@ -211,7 +211,7 @@ namespace Horde.Agent.Utility
 					(ReadOnlyMemory<byte> packet, int packetLineCount) = writer.CreatePacket();
 					try
 					{
-						await WriteOutputAsync(new RpcWriteOutputRequest(_logId, packetOffset, packetLineIndex, UnsafeByteOperations.UnsafeWrap(packet), false), CancellationToken.None);
+						await WriteOutputAsync(_logId, packetOffset, packetLineIndex, packet, false, CancellationToken.None);
 						packetOffset += packet.Length;
 						packetLineIndex += packetLineCount;
 					}
@@ -239,7 +239,7 @@ namespace Horde.Agent.Utility
 				{
 					try
 					{
-						await WriteOutputAsync(new RpcWriteOutputRequest(_logId, packetOffset, packetLineIndex, ByteString.Empty, true), CancellationToken.None);
+						await WriteOutputAsync(_logId, packetOffset, packetLineIndex, ReadOnlyMemory<byte>.Empty, true, CancellationToken.None);
 					}
 					catch (Exception ex)
 					{
@@ -359,15 +359,15 @@ namespace Horde.Agent.Utility
 		}
 
 		/// <inheritdoc/>
-		public async Task WriteOutputAsync(RpcWriteOutputRequest request, CancellationToken cancellationToken)
+		public async Task WriteOutputAsync(LogId logId, long offset, int lineIndex, ReadOnlyMemory<byte> data, bool flush, CancellationToken cancellationToken)
 		{
-			_builder.WriteData(request.Data.Memory);
-			_bufferLength += request.Data.Length;
+			_builder.WriteData(data);
+			_bufferLength += data.Length;
 
-			if (request.Flush || _bufferLength > FlushLength)
+			if (flush || _bufferLength > FlushLength)
 			{
-				IBlobRef<LogNode> target = await _builder.FlushAsync(_writer, request.Flush, cancellationToken);
-				await UpdateLogAsync(target, _builder.LineCount, request.Flush, cancellationToken);
+				IBlobRef<LogNode> target = await _builder.FlushAsync(_writer, flush, cancellationToken);
+				await UpdateLogAsync(target, _builder.LineCount, flush, cancellationToken);
 				_bufferLength = 0;
 			}
 
