Horde: Fix dashboard build
#rnx

diff --git a/Engine/Source/Programs/Horde/HordeDashboard/src/components/HistoryModal.tsx b/Engine/Source/Programs/Horde/HordeDashboard/src/components/HistoryModal.tsx
--- a/Engine/Source/Programs/Horde/HordeDashboard/src/components/HistoryModal.tsx
+++ b/Engine/Source/Programs/Horde/HordeDashboard/src/components/HistoryModal.tsx
@@ -391,7 +391,7 @@ const TelemetryTooltip: React.FC<{ id: string }> = observer(({ id }) => {
    let offsetX = 32;
    let translateX = "0%";
 
-   if (tipX > 600) {
+   if ((tipX ?? 0) > 600) {
       offsetX = -32;
       translateX = "-100%";
    }
@@ -469,7 +469,7 @@ const TelemetryTooltip: React.FC<{ id: string }> = observer(({ id }) => {
       position: "absolute",
       display: "block",
       top: `${tooltip.y}px`,
-      left: `${tooltip.x + offsetX}px`,
+      left: `${(tooltip.x ?? 0) + offsetX}px`,
       backgroundColor: modeColors.background,
       zIndex: 1,
       border: "solid",
@@ -581,6 +581,7 @@ const TelemetryPanel: React.FC<{}> = observer(({ }) => {
       if (cpuKey == "user") return d.userCpu;
       if (cpuKey == "system") return d.systemCpu;
       if (cpuKey == "combined") return d.userCpu + d.systemCpu;
+      return 1;
    })
 
    const min = data[0].time;
@@ -604,16 +605,19 @@ const TelemetryPanel: React.FC<{}> = observer(({ }) => {
    if (cpuKey == "combined") cpuLabel = "CPU - Combined";
 
 
-   const device = state.selectedAgent.capabilities?.devices[0];
-   device.properties?.forEach(v => {
-      if (v.startsWith("CPU=")) {
-         cpuText = `${v.replace("CPU=", "")}`;
-      }
-
-      if (v.startsWith("RAM=")) {
-         ramText = `${v.replace("RAM=", "")} GB`;
-      }
-   })
+   const devices = state.selectedAgent?.capabilities?.devices;
+   if (devices?.length) {
+      const device = devices[0];
+      device.properties?.forEach(v => {
+         if (v.startsWith("CPU=")) {
+            cpuText = `${v.replace("CPU=", "")}`;
+         }
+   
+         if (v.startsWith("RAM=")) {
+            ramText = `${v.replace("RAM=", "")} GB`;
+         }
+      })   
+   }
 
    function relativeCoords(event: any, id: string) {
 
@@ -666,6 +670,9 @@ const TelemetryPanel: React.FC<{}> = observer(({ }) => {
                options={cpuSelections}
                selectedKey={cpuKey}
                onChange={(ev, option, index, value) => {
+                  if (!option) {
+                     return;
+                  }
                   setCPUKey(option.key as string);
                }}
             />
