Horde: Add free/total disk space to agent telemetry.

#rnx

diff --git a/Engine/Source/Programs/Horde/Horde.Agent/AgentApp.cs b/Engine/Source/Programs/Horde/Horde.Agent/AgentApp.cs
--- a/Engine/Source/Programs/Horde/Horde.Agent/AgentApp.cs
+++ b/Engine/Source/Programs/Horde/Horde.Agent/AgentApp.cs
@@ -225,7 +225,7 @@ namespace Horde.Agent
 
 			services.AddMemoryCache();
 
-			services.AddSingleton<ISystemMetrics>(sp => CreateSystemMetrics(sp.GetRequiredService<ILogger<ISystemMetrics>>()));
+			services.AddSingleton<ISystemMetrics>(sp => CreateSystemMetrics(settings.WorkingDir, sp.GetRequiredService<ILogger<ISystemMetrics>>()));
 
 			// Allow commands to augment the service collection for their own DI service providers
 			services.AddSingleton<DefaultServices>(x => new DefaultServices(configuration, services));
@@ -235,13 +235,13 @@ namespace Horde.Agent
 			return await CommandHost.RunAsync(arguments, serviceProvider, typeof(Commands.Service.RunCommand));
 		}
 
-		static ISystemMetrics CreateSystemMetrics(ILogger logger)
+		static ISystemMetrics CreateSystemMetrics(DirectoryReference workingDir, ILogger logger)
 		{
 			if (RuntimeInformation.IsOSPlatform(OSPlatform.Windows))
 			{
 				try
 				{
-					return new WindowsSystemMetrics();
+					return new WindowsSystemMetrics(workingDir);
 				}
 				catch (Exception e)
 				{
diff --git a/Engine/Source/Programs/Horde/Horde.Agent/Leases/LeaseManager.cs b/Engine/Source/Programs/Horde/Horde.Agent/Leases/LeaseManager.cs
--- a/Engine/Source/Programs/Horde/Horde.Agent/Leases/LeaseManager.cs
+++ b/Engine/Source/Programs/Horde/Horde.Agent/Leases/LeaseManager.cs
@@ -598,8 +598,9 @@ namespace Horde.Agent.Leases
 				{
 					CpuMetrics? cpuMetrics = _systemMetrics.GetCpu();
 					MemoryMetrics? memMetrics = _systemMetrics.GetMemory();
+					DiskMetrics? diskMetrics = _systemMetrics.GetDisk();
 
-					if (cpuMetrics != null || memMetrics != null)
+					if (cpuMetrics != null || memMetrics != null || diskMetrics != null)
 					{
 						RpcUploadTelemetryRequest request = new RpcUploadTelemetryRequest();
 						request.AgentId = _session.AgentId.ToString();
@@ -615,6 +616,11 @@ namespace Horde.Agent.Leases
 							request.FreeRam = memMetrics.Available / 1024;
 							request.UsedRam = memMetrics.Used / 1024;
 						}
+						if (diskMetrics != null)
+						{
+							request.FreeDisk = (ulong)(diskMetrics.FreeSpace / (1024 * 1024));
+							request.TotalDisk = (ulong)(diskMetrics.TotalSize / (1024 * 1024));
+						}
 
 						await hordeRpc.UploadTelemetryAsync(request, cancellationToken: cancellationToken);
 					}
diff --git a/Engine/Source/Programs/Horde/Horde.Agent/Services/TelemetryService.cs b/Engine/Source/Programs/Horde/Horde.Agent/Services/TelemetryService.cs
--- a/Engine/Source/Programs/Horde/Horde.Agent/Services/TelemetryService.cs
+++ b/Engine/Source/Programs/Horde/Horde.Agent/Services/TelemetryService.cs
@@ -96,6 +96,22 @@ public class MemoryMetrics
 	}
 }
 
+/// <summary>
+/// Metrics for free disk space
+/// </summary>
+public class DiskMetrics
+{
+	/// <summary>
+	/// Amount of free space on the drive
+	/// </summary>
+	public long FreeSpace { get; set; }
+
+	/// <summary>
+	/// Total size of the drive
+	/// </summary>
+	public long TotalSize { get; set; }
+}
+
 /// <summary>
 /// OS agnostic interface for retrieving system metrics (CPU, memory etc)
 /// </summary>
@@ -112,6 +128,12 @@ public interface ISystemMetrics
 	/// </summary>
 	/// <returns>An object with memory usage metrics</returns>
 	MemoryMetrics? GetMemory();
+
+	/// <summary>
+	/// Gets HDD metrics
+	/// </summary>
+	/// <returns>An object with disk usage metrics</returns>
+	DiskMetrics? GetDisk();
 }
 
 /// <summary>
@@ -126,6 +148,10 @@ public sealed class DefaultSystemMetrics : ISystemMetrics
 	/// <inheritdoc/>
 	public MemoryMetrics? GetMemory()
 		=> null;
+
+	/// <inheritdoc/>
+	public DiskMetrics? GetDisk()
+		=> null;
 }
 
 /// <summary>
@@ -145,11 +171,18 @@ public sealed class WindowsSystemMetrics : ISystemMetrics, IDisposable
 	private readonly uint _totalPhysicalMemory = GetPhysicalMemory();
 	private readonly PerformanceCounter _memAvailableBytes = new(Memory, "Available Bytes");
 
+	private readonly DirectoryReference _workingDir;
+	private readonly Stopwatch _driveInfoUpdateTimer = Stopwatch.StartNew();
+	private DriveInfo? _driveInfo;
+
 	/// <summary>
 	/// Constructor
 	/// </summary>
-	public WindowsSystemMetrics()
+	public WindowsSystemMetrics(DirectoryReference workingDir)
 	{
+		_workingDir = workingDir;
+		_driveInfo = GetDriveInfo();
+
 		GetCpu(); // Trigger this to ensure performance counter has a fetched value. Avoids an initial zero result when called later.
 	}
 
@@ -187,6 +220,34 @@ public sealed class WindowsSystemMetrics : ISystemMetrics, IDisposable
 		};
 	}
 
+	/// <inheritdoc/>
+	public DiskMetrics? GetDisk()
+	{
+		// Update the metrics if we don't have a recent sample
+		if (_driveInfoUpdateTimer.Elapsed > TimeSpan.FromSeconds(20))
+		{
+			_driveInfo = GetDriveInfo();
+			_driveInfoUpdateTimer.Restart();
+		}
+
+		if (_driveInfo == null)
+		{
+			return null;
+		}
+
+		return new DiskMetrics { FreeSpace = _driveInfo.AvailableFreeSpace, TotalSize = _driveInfo.TotalSize };
+	}
+
+	DriveInfo? GetDriveInfo()
+	{
+		string? rootDir = Path.GetPathRoot(_workingDir.FullName);
+		if (rootDir == null)
+		{
+			return null;
+		}
+		return new DriveInfo(rootDir);
+	}
+
 	private static uint GetPhysicalMemory()
 	{
 		using CimSession session = CimSession.Create(null);
diff --git a/Engine/Source/Programs/Horde/Horde.Server/Agents/AgentsController.cs b/Engine/Source/Programs/Horde/Horde.Server/Agents/AgentsController.cs
--- a/Engine/Source/Programs/Horde/Horde.Server/Agents/AgentsController.cs
+++ b/Engine/Source/Programs/Horde/Horde.Server/Agents/AgentsController.cs
@@ -382,7 +382,7 @@ namespace Horde.Server.Agents
 			IReadOnlyList<IAgentTelemetry> telemetry = await _agentTelemetryCollection.FindAsync(agentId, minTime, maxTime, HttpContext.RequestAborted);
 
 			List<GetAgentTelemetrySampleResponse> samples = new List<GetAgentTelemetrySampleResponse>();
-			samples.AddRange(telemetry.Select(x => new GetAgentTelemetrySampleResponse(x.TimeUtc, x.UserCpu, x.IdleCpu, x.SystemCpu, x.FreeRam, x.UsedRam, x.TotalRam)));
+			samples.AddRange(telemetry.Select(x => new GetAgentTelemetrySampleResponse(x.TimeUtc, x.UserCpu, x.IdleCpu, x.SystemCpu, x.FreeRam, x.UsedRam, x.TotalRam, x.FreeDisk, x.TotalDisk)));
 
 			return new GetAgentTelemetryResponse(samples);
 		}
diff --git a/Engine/Source/Programs/Horde/Horde.Server/Agents/Telemetry/AgentTelemetryCollection.cs b/Engine/Source/Programs/Horde/Horde.Server/Agents/Telemetry/AgentTelemetryCollection.cs
--- a/Engine/Source/Programs/Horde/Horde.Server/Agents/Telemetry/AgentTelemetryCollection.cs
+++ b/Engine/Source/Programs/Horde/Horde.Server/Agents/Telemetry/AgentTelemetryCollection.cs
@@ -49,6 +49,12 @@ namespace Horde.Server.Agents.Telemetry
 
 			[BsonElement("ram_used")]
 			public int UsedRam { get; set; }
+
+			[BsonElement("hdd_tot")]
+			public long TotalDisk { get; set; }
+
+			[BsonElement("hdd_free")]
+			public long FreeDisk { get; set; }
 		}
 
 		static TimeSpan RetainDays { get; } = TimeSpan.FromDays(30.0);
@@ -123,6 +129,8 @@ namespace Horde.Server.Agents.Telemetry
 			document.FreeRam = telemetry.FreeRamMb;
 			document.UsedRam = telemetry.UsedRamMb;
 			document.TotalRam = telemetry.TotalRamMb;
+			document.FreeDisk = telemetry.FreeDiskMb;
+			document.TotalDisk = telemetry.TotalDiskMb;
 			_writer.Write(document);
 		}
 
diff --git a/Engine/Source/Programs/Horde/Horde.Server/Agents/Telemetry/IAgentTelemetry.cs b/Engine/Source/Programs/Horde/Horde.Server/Agents/Telemetry/IAgentTelemetry.cs
--- a/Engine/Source/Programs/Horde/Horde.Server/Agents/Telemetry/IAgentTelemetry.cs
+++ b/Engine/Source/Programs/Horde/Horde.Server/Agents/Telemetry/IAgentTelemetry.cs
@@ -43,5 +43,15 @@ namespace Horde.Server.Agents.Telemetry
 		/// Used memory (megabytes)
 		/// </summary>
 		int UsedRam { get; set; }
+
+		/// <summary>
+		/// Total disk space (megabytes)
+		/// </summary>
+		long TotalDisk { get; set; }
+
+		/// <summary>
+		/// Free disk space (megabytes)
+		/// </summary>
+		long FreeDisk { get; set; }
 	}
 }
diff --git a/Engine/Source/Programs/Horde/Horde.Server/Agents/Telemetry/IAgentTelemetryCollection.cs b/Engine/Source/Programs/Horde/Horde.Server/Agents/Telemetry/IAgentTelemetryCollection.cs
--- a/Engine/Source/Programs/Horde/Horde.Server/Agents/Telemetry/IAgentTelemetryCollection.cs
+++ b/Engine/Source/Programs/Horde/Horde.Server/Agents/Telemetry/IAgentTelemetryCollection.cs
@@ -31,5 +31,5 @@ namespace Horde.Server.Agents.Telemetry
 	/// <summary>
 	/// New telemetry data
 	/// </summary>
-	public record class NewAgentTelemetry(float UserCpuPct, float IdleCpuPct, float SystemCpuPct, int FreeRamMb, int UsedRamMb, int TotalRamMb);
+	public record class NewAgentTelemetry(float UserCpuPct, float IdleCpuPct, float SystemCpuPct, int FreeRamMb, int UsedRamMb, int TotalRamMb, long FreeDiskMb, long TotalDiskMb);
 }
diff --git a/Engine/Source/Programs/Horde/Horde.Server/Server/RpcService.cs b/Engine/Source/Programs/Horde/Horde.Server/Server/RpcService.cs
--- a/Engine/Source/Programs/Horde/Horde.Server/Server/RpcService.cs
+++ b/Engine/Source/Programs/Horde/Horde.Server/Server/RpcService.cs
@@ -493,7 +493,7 @@ namespace Horde.Server.Server
 		{
 			_logger.LogDebug("Posting telemetry data for {AgentId}", new AgentId(request.AgentId));
 
-			NewAgentTelemetry telemetry = new NewAgentTelemetry(request.UserCpu, request.IdleCpu, request.SystemCpu, (int)request.FreeRam, (int)request.UsedRam, (int)request.TotalRam);
+			NewAgentTelemetry telemetry = new NewAgentTelemetry(request.UserCpu, request.IdleCpu, request.SystemCpu, (int)request.FreeRam, (int)request.UsedRam, (int)request.TotalRam, (long)request.FreeDisk, (long)request.TotalDisk);
 			_agentTelemetryCollection.Add(new AgentId(request.AgentId), telemetry);
 
 			return Task.FromResult(new Empty());
diff --git a/Engine/Source/Programs/Shared/EpicGames.Horde/Agents/AgentMessages.cs b/Engine/Source/Programs/Shared/EpicGames.Horde/Agents/AgentMessages.cs
--- a/Engine/Source/Programs/Shared/EpicGames.Horde/Agents/AgentMessages.cs
+++ b/Engine/Source/Programs/Shared/EpicGames.Horde/Agents/AgentMessages.cs
@@ -127,5 +127,5 @@ namespace EpicGames.Horde.Agents
 	/// <summary>
 	/// Telemetry data sample
 	/// </summary>
-	public record GetAgentTelemetrySampleResponse(DateTime Time, float UserCpu, float IdleCpu, float SystemCpu, int FreeRam, int UsedRam, int TotalRam);
+	public record GetAgentTelemetrySampleResponse(DateTime Time, float UserCpu, float IdleCpu, float SystemCpu, int FreeRam, int UsedRam, int TotalRam, long FreeDisk, long TotalDisk);
 }
