Horde: Remove session interface from core job handler.

#rnx

diff --git a/Engine/Source/Programs/Horde/Horde.Agent/Commands/Execution/ExecuteJobCommand.cs b/Engine/Source/Programs/Horde/Horde.Agent/Commands/Execution/ExecuteJobCommand.cs
--- a/Engine/Source/Programs/Horde/Horde.Agent/Commands/Execution/ExecuteJobCommand.cs
+++ b/Engine/Source/Programs/Horde/Horde.Agent/Commands/Execution/ExecuteJobCommand.cs
@@ -46,7 +46,7 @@ namespace Horde.Agent.Commands.Execution
 			IHordeClient hordeClient = _hordeClientFactory.Create(executeTask.Token);
 			await using Session session = new Session(AgentId, SessionId, WorkingDir, hordeClient);
 
-			await _jobHandler.ExecuteInternalAsync(session, LeaseId, executeTask, logger, CancellationToken.None);
+			await _jobHandler.ExecuteInternalAsync(session.HordeClient, session.WorkingDir, LeaseId, executeTask, logger, CancellationToken.None);
 			return 0;
 		}
 	}
diff --git a/Engine/Source/Programs/Horde/Horde.Agent/Leases/Handlers/JobHandler.cs b/Engine/Source/Programs/Horde/Horde.Agent/Leases/Handlers/JobHandler.cs
--- a/Engine/Source/Programs/Horde/Horde.Agent/Leases/Handlers/JobHandler.cs
+++ b/Engine/Source/Programs/Horde/Horde.Agent/Leases/Handlers/JobHandler.cs
@@ -4,6 +4,7 @@ using System.Diagnostics;
 using System.Reflection;
 using System.Runtime.InteropServices;
 using EpicGames.Core;
+using EpicGames.Horde;
 using EpicGames.Horde.Agents.Leases;
 using EpicGames.Horde.Jobs;
 using EpicGames.Horde.Logs;
@@ -135,7 +136,7 @@ namespace Horde.Agent.Leases.Handlers
 					return LeaseResult.Success;
 				}
 
-				return await ExecuteInternalAsync(session, leaseId, executeTask, localLogger, cancellationToken);
+				return await ExecuteInternalAsync(session.HordeClient, session.WorkingDir, leaseId, executeTask, localLogger, cancellationToken);
 			}
 			finally
 			{
@@ -145,19 +146,19 @@ namespace Horde.Agent.Leases.Handlers
 			}
 		}
 
-		internal async Task<LeaseResult> ExecuteInternalAsync(ISession session, LeaseId leaseId, ExecuteJobTask executeTask, ILogger localLogger, CancellationToken cancellationToken)
+		internal async Task<LeaseResult> ExecuteInternalAsync(IHordeClient hordeClient, DirectoryReference workingDir, LeaseId leaseId, ExecuteJobTask executeTask, ILogger localLogger, CancellationToken cancellationToken)
 		{
 			JobTaskInfo jobTaskInfo = new JobTaskInfo(executeTask.JobName, JobId.Parse(executeTask.JobId), JobStepBatchId.Parse(executeTask.BatchId), executeTask.JobOptions, executeTask.Token, LogId.Parse(executeTask.LogId), executeTask.Workspace, executeTask.AutoSdkWorkspace);
-			return await ExecuteInternalAsync(session, leaseId, jobTaskInfo, localLogger, cancellationToken);
+			return await ExecuteInternalAsync(hordeClient, workingDir, leaseId, jobTaskInfo, localLogger, cancellationToken);
 		}
 
-		internal async Task<LeaseResult> ExecuteInternalAsync(ISession session, LeaseId leaseId, JobTaskInfo executeTask, ILogger localLogger, CancellationToken cancellationToken)
+		internal async Task<LeaseResult> ExecuteInternalAsync(IHordeClient hordeClient, DirectoryReference workingDir, LeaseId leaseId, JobTaskInfo executeTask, ILogger localLogger, CancellationToken cancellationToken)
 		{
-			JobRpc.JobRpcClient jobRpc = await session.HordeClient.CreateGrpcClientAsync<JobRpc.JobRpcClient>(cancellationToken);
+			JobRpc.JobRpcClient jobRpc = await hordeClient.CreateGrpcClientAsync<JobRpc.JobRpcClient>(cancellationToken);
 
 			// Create a storage client for this session
 			RpcJobOptions jobOptions = executeTask.JobOptions;
-			await using IServerLogger logger = session.HordeClient.CreateServerLogger(executeTask.LogId).WithLocalLogger(localLogger);
+			await using IServerLogger logger = hordeClient.CreateServerLogger(executeTask.LogId).WithLocalLogger(localLogger);
 
 			logger.LogInformation("Executing job \"{JobName}\", jobId {JobId}, batchId {BatchId}, leaseId {LeaseId}, agentVersion {AgentVersion}", executeTask.JobName, executeTask.JobId, executeTask.BatchId, leaseId, AgentApp.Version);
 
@@ -171,14 +172,14 @@ namespace Horde.Agent.Leases.Handlers
 			RpcBeginBatchResponse batch = await jobRpc.BeginBatchAsync(new RpcBeginBatchRequest(executeTask.JobId, executeTask.BatchId, leaseId), cancellationToken: cancellationToken);
 			try
 			{
-				JobExecutorOptions options = new JobExecutorOptions(session.HordeClient, session.WorkingDir, _driverSettings.ProcessesToTerminate, executeTask.JobId, executeTask.BatchId, batch, jobOptions);
-				await ExecuteBatchAsync(session, leaseId, executeTask.Workspace, executeTask.AutoSdkWorkspace, options, logger, localLogger, cancellationToken);
+				JobExecutorOptions options = new JobExecutorOptions(hordeClient, workingDir, _driverSettings.ProcessesToTerminate, executeTask.JobId, executeTask.BatchId, batch, jobOptions);
+				await ExecuteBatchAsync(hordeClient, workingDir, leaseId, executeTask.Workspace, executeTask.AutoSdkWorkspace, options, logger, localLogger, cancellationToken);
 			}
 			catch (Exception ex)
 			{
 				if (cancellationToken.IsCancellationRequested && ex.IsCancellationException())
 				{
-					if (!session.HordeClient.HasValidAccessToken())
+					if (!hordeClient.HasValidAccessToken())
 					{
 						logger.LogError(ex, "Connection to the server was lost; step aborted.");
 					}
@@ -211,15 +212,15 @@ namespace Horde.Agent.Leases.Handlers
 		/// <summary>
 		/// Executes a batch
 		/// </summary>
-		async Task ExecuteBatchAsync(ISession session, LeaseId leaseId, RpcAgentWorkspace workspaceInfo, RpcAgentWorkspace? autoSdkWorkspaceInfo, JobExecutorOptions options, ILogger logger, ILogger localLogger, CancellationToken cancellationToken)
+		async Task ExecuteBatchAsync(IHordeClient hordeClient, DirectoryReference workingDir, LeaseId leaseId, RpcAgentWorkspace workspaceInfo, RpcAgentWorkspace? autoSdkWorkspaceInfo, JobExecutorOptions options, ILogger logger, ILogger localLogger, CancellationToken cancellationToken)
 		{
-			JobRpc.JobRpcClient rpcClient = await session.HordeClient.CreateGrpcClientAsync<JobRpc.JobRpcClient>(cancellationToken);
+			JobRpc.JobRpcClient rpcClient = await hordeClient.CreateGrpcClientAsync<JobRpc.JobRpcClient>(cancellationToken);
 
 			// Create an executor for this job
 			string executorName = String.IsNullOrEmpty(options.JobOptions.Executor) ? _driverSettings.Executor : options.JobOptions.Executor;
 
 			logger.LogInformation("Executing batch {BatchId} using {Executor} executor", options.BatchId, executorName);
-			await TerminateProcessHelper.TerminateProcessesAsync(TerminateCondition.BeforeBatch, session.WorkingDir, _driverSettings.ProcessesToTerminate, logger, cancellationToken);
+			await TerminateProcessHelper.TerminateProcessesAsync(TerminateCondition.BeforeBatch, workingDir, _driverSettings.ProcessesToTerminate, logger, cancellationToken);
 
 			IJobExecutorFactory? executorFactory = _executorFactories.FirstOrDefault(x => x.Name.Equals(executorName, StringComparison.OrdinalIgnoreCase));
 			if (executorFactory == null)
@@ -269,11 +270,11 @@ namespace Horde.Agent.Leases.Handlers
 					string? driveName;
 					if (RuntimeInformation.IsOSPlatform(OSPlatform.Windows))
 					{
-						driveName = Path.GetPathRoot(session.WorkingDir.FullName);
+						driveName = Path.GetPathRoot(workingDir.FullName);
 					}
 					else
 					{
-						driveName = session.WorkingDir.FullName;
+						driveName = workingDir.FullName;
 					}
 
 					float availableFreeSpace = 0;
@@ -309,7 +310,7 @@ namespace Horde.Agent.Leases.Handlers
 					{
 						// Start writing to the log file
 #pragma warning disable CA2000 // Dispose objects before losing scope
-						await using (JobStepLogger stepLogger = new JobStepLogger(session.HordeClient, step.LogId, localLogger, options.JobId, options.BatchId, step.StepId, step.Warnings, LogLevel.Debug, logger))
+						await using (JobStepLogger stepLogger = new JobStepLogger(hordeClient, step.LogId, localLogger, options.JobId, options.BatchId, step.StepId, step.Warnings, LogLevel.Debug, logger))
 						{
 							// Execute the task
 							using CancellationTokenSource stepPollCancelSource = new CancellationTokenSource();
@@ -330,7 +331,7 @@ namespace Horde.Agent.Leases.Handlers
 							}
 
 							// Kill any processes spawned by the step
-							await TerminateProcessHelper.TerminateProcessesAsync(TerminateCondition.AfterStep, session.WorkingDir, _driverSettings.ProcessesToTerminate, logger, cancellationToken);
+							await TerminateProcessHelper.TerminateProcessesAsync(TerminateCondition.AfterStep, workingDir, _driverSettings.ProcessesToTerminate, logger, cancellationToken);
 
 							// Wait for the logger to finish
 							await stepLogger.StopAsync();
@@ -357,7 +358,7 @@ namespace Horde.Agent.Leases.Handlers
 			{
 				if (cancellationToken.IsCancellationRequested && ex.IsCancellationException())
 				{
-					if (!session.HordeClient.HasValidAccessToken())
+					if (!hordeClient.HasValidAccessToken())
 					{
 						logger.LogError(ex, "Exception while executing batch: {Ex}", ex);
 					}
@@ -371,7 +372,7 @@ namespace Horde.Agent.Leases.Handlers
 			// Terminate any processes which are still running
 			try
 			{
-				await TerminateProcessHelper.TerminateProcessesAsync(TerminateCondition.AfterBatch, session.WorkingDir, _driverSettings.ProcessesToTerminate, logger, CancellationToken.None);
+				await TerminateProcessHelper.TerminateProcessesAsync(TerminateCondition.AfterBatch, workingDir, _driverSettings.ProcessesToTerminate, logger, CancellationToken.None);
 			}
 			catch (Exception ex)
 			{
