Horde: Static analysis fixes.

#rnx

diff --git a/Engine/Source/Programs/Horde/Horde.Agent.Tests/WorkerServiceTest.cs b/Engine/Source/Programs/Horde/Horde.Agent.Tests/WorkerServiceTest.cs
--- a/Engine/Source/Programs/Horde/Horde.Agent.Tests/WorkerServiceTest.cs
+++ b/Engine/Source/Programs/Horde/Horde.Agent.Tests/WorkerServiceTest.cs
@@ -9,7 +9,6 @@ using System.Threading.Channels;
 using System.Threading.Tasks;
 using EpicGames.Core;
 using EpicGames.Horde;
-using EpicGames.Horde.Agents.Leases;
 using EpicGames.Horde.Compute;
 using EpicGames.Horde.Jobs;
 using EpicGames.Horde.Logs;
diff --git a/Engine/Source/Programs/Horde/Horde.Agent/Leases/Handlers/JobHandler.cs b/Engine/Source/Programs/Horde/Horde.Agent/Leases/Handlers/JobHandler.cs
--- a/Engine/Source/Programs/Horde/Horde.Agent/Leases/Handlers/JobHandler.cs
+++ b/Engine/Source/Programs/Horde/Horde.Agent/Leases/Handlers/JobHandler.cs
@@ -2,27 +2,18 @@
 
 using System.Diagnostics;
 using System.Reflection;
-using System.Runtime.InteropServices;
-using Amazon.SimpleSystemsManagement.Model.Internal.MarshallTransformations;
 using EpicGames.Core;
 using EpicGames.Horde;
 using EpicGames.Horde.Agents.Leases;
 using EpicGames.Horde.Jobs;
 using EpicGames.Horde.Logs;
 using Google.Protobuf;
-using Grpc.Core;
 using Horde.Agent.Driver;
-using Horde.Agent.Driver.Utility;
 using Horde.Agent.Execution;
 using Horde.Agent.Services;
-using Horde.Agent.Utility;
-using Horde.Common.Rpc;
-using HordeCommon.Rpc;
-using HordeCommon.Rpc.Messages;
 using HordeCommon.Rpc.Tasks;
 using Microsoft.Extensions.Logging;
 using Microsoft.Extensions.Options;
-using OpenTracing;
 using OpenTracing.Util;
 
 namespace Horde.Agent.Leases.Handlers
@@ -154,9 +145,10 @@ namespace Horde.Agent.Leases.Handlers
 			LogId logId = LogId.Parse(executeTask.LogId);
 
 			JobExecutorOptions options = new JobExecutorOptions(hordeClient, workingDir, _driverSettings.ProcessesToTerminate, jobId, batchId, leaseId, logId, executeTask.JobOptions);
-			using JobExecutor executor = executorFactory.CreateExecutor(executeTask.Workspace, executeTask.AutoSdkWorkspace, options);
 
+			using JobExecutor executor = executorFactory.CreateExecutor(executeTask.Workspace, executeTask.AutoSdkWorkspace, options);
 			await executor.ExecuteAsync(localLogger, cancellationToken);
+
 			return LeaseResult.Success;
 		}
 	}
diff --git a/Engine/Source/Programs/Horde/Horde.Agent/Leases/LeaseLoggerFactory.cs b/Engine/Source/Programs/Horde/Horde.Agent/Leases/LeaseLoggerFactory.cs
--- a/Engine/Source/Programs/Horde/Horde.Agent/Leases/LeaseLoggerFactory.cs
+++ b/Engine/Source/Programs/Horde/Horde.Agent/Leases/LeaseLoggerFactory.cs
@@ -2,7 +2,6 @@
 
 using EpicGames.Core;
 using EpicGames.Horde.Agents.Leases;
-using EpicGames.Horde.Logs;
 using Microsoft.Extensions.Logging;
 using Microsoft.Extensions.Options;
 
