Horde: Add commands to driver for executing jobs and conforms. Rename namespace to Horde.Agent.Driver.

#rnx

diff --git a/Engine/Source/Programs/Horde/Horde.Agent.Driver/Commands/Execution/ExecuteConformCommand.cs b/Engine/Source/Programs/Horde/Horde.Agent.Driver/Commands/Execution/ExecuteConformCommand.cs
new file mode 100644
--- /dev/null
+++ b/Engine/Source/Programs/Horde/Horde.Agent.Driver/Commands/Execution/ExecuteConformCommand.cs
@@ -0,0 +1,55 @@
+// Copyright Epic Games, Inc. All Rights Reserved.
+
+using EpicGames.Core;
+using EpicGames.Horde;
+using EpicGames.Horde.Agents;
+using EpicGames.Horde.Agents.Leases;
+using EpicGames.Horde.Agents.Sessions;
+using Horde.Agent.Driver.Execution;
+using HordeCommon.Rpc.Tasks;
+using Microsoft.Extensions.Logging;
+using Microsoft.Extensions.Options;
+
+namespace Horde.Agent.Driver.Commands.Execution
+{
+	[Command("Execute", "Conform", "Executes a conform")]
+	class ExecuteConformCommand : Command
+	{
+		[CommandLine("-AgentId=", Required = true)]
+		public AgentId AgentId { get; set; }
+
+		[CommandLine("-SessionId=", Required = true)]
+		public SessionId SessionId { get; set; }
+
+		[CommandLine("-LeaseId", Required = true)]
+		public LeaseId LeaseId { get; set; }
+
+		[CommandLine("-Task=", Required = true)]
+		public string Task { get; set; } = null!;
+
+		[CommandLine("-WorkingDir=", Required = true)]
+		public DirectoryReference WorkingDir { get; set; } = null!;
+
+		readonly IHordeClientFactory _hordeClientFactory;
+		readonly IOptions<DriverSettings> _driverSettings;
+
+		public ExecuteConformCommand(IHordeClientFactory hordeClientFactory, IOptions<DriverSettings> driverSettings)
+		{
+			_hordeClientFactory = hordeClientFactory;
+			_driverSettings = driverSettings;
+		}
+
+		public override async Task<int> ExecuteAsync(ILogger logger)
+		{
+			ConformTask conformTask = ConformTask.Parser.ParseFrom(Convert.FromBase64String(Task));
+
+			string accessToken = null!;
+			IHordeClient hordeClient = _hordeClientFactory.Create(accessToken);
+
+			ConformExecutor conformExecutor = new ConformExecutor(hordeClient, WorkingDir, AgentId, LeaseId, conformTask, _driverSettings.Value, logger);
+			await conformExecutor.ExecuteAsync(CancellationToken.None);
+
+			return 0;
+		}
+	}
+}
diff --git a/Engine/Source/Programs/Horde/Horde.Agent.Driver/Commands/Execution/ExecuteJobCommand.cs b/Engine/Source/Programs/Horde/Horde.Agent.Driver/Commands/Execution/ExecuteJobCommand.cs
new file mode 100644
--- /dev/null
+++ b/Engine/Source/Programs/Horde/Horde.Agent.Driver/Commands/Execution/ExecuteJobCommand.cs
@@ -0,0 +1,54 @@
+// Copyright Epic Games, Inc. All Rights Reserved.
+
+using EpicGames.Core;
+using EpicGames.Horde;
+using EpicGames.Horde.Agents;
+using EpicGames.Horde.Agents.Leases;
+using EpicGames.Horde.Agents.Sessions;
+using Horde.Agent.Driver.Execution;
+using HordeCommon.Rpc.Tasks;
+using Microsoft.Extensions.Logging;
+using Microsoft.Extensions.Options;
+
+namespace Horde.Agent.Driver.Commands.Execution
+{
+	[Command("Execute", "Job", "Executes a job")]
+	class ExecuteJobCommand : Command
+	{
+		[CommandLine("-AgentId=", Required = true)]
+		public AgentId AgentId { get; set; }
+
+		[CommandLine("-SessionId=", Required = true)]
+		public SessionId SessionId { get; set; }
+
+		[CommandLine("-LeaseId", Required = true)]
+		public LeaseId LeaseId { get; set; }
+
+		[CommandLine("-Task=", Required = true)]
+		public string Task { get; set; } = null!;
+
+		[CommandLine("-WorkingDir=", Required = true)]
+		public DirectoryReference WorkingDir { get; set; } = null!;
+
+		readonly IHordeClientFactory _hordeClientFactory;
+		readonly IEnumerable<IJobExecutorFactory> _jobExecutorFactories;
+		readonly IOptions<DriverSettings> _driverSettings;
+
+		public ExecuteJobCommand(IHordeClientFactory hordeClientFactory, IEnumerable<IJobExecutorFactory> jobExecutorFactories, IOptions<DriverSettings> driverSettings)
+		{
+			_hordeClientFactory = hordeClientFactory;
+			_jobExecutorFactories = jobExecutorFactories;
+			_driverSettings = driverSettings;
+		}
+
+		public override async Task<int> ExecuteAsync(ILogger logger)
+		{
+			ExecuteJobTask executeTask = ExecuteJobTask.Parser.ParseFrom(Convert.FromBase64String(Task));
+
+			IHordeClient hordeClient = _hordeClientFactory.Create(executeTask.Token);
+
+			await JobExecutorHelpers.ExecuteAsync(hordeClient, WorkingDir, LeaseId, executeTask, _jobExecutorFactories, _driverSettings.Value, logger, CancellationToken.None);
+			return 0;
+		}
+	}
+}
diff --git a/Engine/Source/Programs/Horde/Horde.Agent.Driver/Commands/Parse/ParseCommand.cs b/Engine/Source/Programs/Horde/Horde.Agent.Driver/Commands/Parse/ParseCommand.cs
--- a/Engine/Source/Programs/Horde/Horde.Agent.Driver/Commands/Parse/ParseCommand.cs
+++ b/Engine/Source/Programs/Horde/Horde.Agent.Driver/Commands/Parse/ParseCommand.cs
@@ -3,10 +3,10 @@
 using System.ComponentModel;
 using EpicGames.Core;
 using EpicGames.Perforce;
-using Horde.Agent.Parser;
+using Horde.Agent.Driver.Parser;
 using Microsoft.Extensions.Logging;
 
-namespace Horde.Agent.Commands.Parse
+namespace Horde.Agent.Driver.Commands.Parse
 {
 	/// <summary>
 	/// Installs the agent as a service
diff --git a/Engine/Source/Programs/Horde/Horde.Agent.Driver/Commands/Utilities/BuildGraphCommand.cs b/Engine/Source/Programs/Horde/Horde.Agent.Driver/Commands/Utilities/BuildGraphCommand.cs
--- a/Engine/Source/Programs/Horde/Horde.Agent.Driver/Commands/Utilities/BuildGraphCommand.cs
+++ b/Engine/Source/Programs/Horde/Horde.Agent.Driver/Commands/Utilities/BuildGraphCommand.cs
@@ -2,10 +2,10 @@
 
 using System.Diagnostics;
 using EpicGames.Core;
-using Horde.Agent.Parser;
+using Horde.Agent.Driver.Parser;
 using Microsoft.Extensions.Logging;
 
-namespace Horde.Agent.Commands.Utilities
+namespace Horde.Agent.Driver.Commands.Utilities
 {
 	/// <summary>
 	/// Runs a BuildGraph script and captures the processed log output
diff --git a/Engine/Source/Programs/Horde/Horde.Agent.Driver/DriverSettings.cs b/Engine/Source/Programs/Horde/Horde.Agent.Driver/DriverSettings.cs
--- a/Engine/Source/Programs/Horde/Horde.Agent.Driver/DriverSettings.cs
+++ b/Engine/Source/Programs/Horde/Horde.Agent.Driver/DriverSettings.cs
@@ -1,7 +1,7 @@
 ﻿// Copyright Epic Games, Inc. All Rights Reserved.
 
-using Horde.Agent.Execution;
-using Horde.Agent.Utility;
+using Horde.Agent.Driver.Execution;
+using Horde.Agent.Driver.Utility;
 
 namespace Horde.Agent.Driver
 {
@@ -13,7 +13,7 @@ namespace Horde.Agent.Driver
 		/// <summary>
 		/// The executor to use for jobs
 		/// </summary>
-		public string Executor { get; set; } = Execution.WorkspaceExecutor.Name;
+		public string Executor { get; set; } = WorkspaceExecutor.Name;
 
 		/// <summary>
 		/// Settings for the local executor
diff --git a/Engine/Source/Programs/Horde/Horde.Agent.Driver/Execution/ConformExecutor.cs b/Engine/Source/Programs/Horde/Horde.Agent.Driver/Execution/ConformExecutor.cs
new file mode 100644
--- /dev/null
+++ b/Engine/Source/Programs/Horde/Horde.Agent.Driver/Execution/ConformExecutor.cs
@@ -0,0 +1,100 @@
+// Copyright Epic Games, Inc. All Rights Reserved.
+
+using EpicGames.Core;
+using EpicGames.Horde;
+using EpicGames.Horde.Agents;
+using EpicGames.Horde.Agents.Leases;
+using EpicGames.Horde.Logs;
+using Horde.Agent.Driver.Utility;
+using HordeCommon.Rpc;
+using HordeCommon.Rpc.Messages;
+using HordeCommon.Rpc.Tasks;
+using Microsoft.Extensions.Logging;
+
+namespace Horde.Agent.Driver.Execution
+{
+	class ConformExecutor
+	{
+		readonly IHordeClient _hordeClient;
+		readonly DirectoryReference _workingDir;
+		readonly AgentId _agentId;
+		readonly LeaseId _leaseId;
+		readonly ConformTask _conformTask;
+		readonly DriverSettings _driverSettings;
+		readonly ILogger _localLogger;
+
+		public ConformExecutor(IHordeClient hordeClient, DirectoryReference workingDir, AgentId agentId, LeaseId leaseId, ConformTask conformTask, DriverSettings driverSettings, ILogger localLogger)
+		{
+			_hordeClient = hordeClient;
+			_workingDir = workingDir;
+			_agentId = agentId;
+			_leaseId = leaseId;
+			_conformTask = conformTask;
+			_driverSettings = driverSettings;
+			_localLogger = localLogger;
+		}
+
+		public async Task ExecuteAsync(CancellationToken cancellationToken)
+		{
+			await using IServerLogger serverLogger = _hordeClient.CreateServerLogger(LogId.Parse(_conformTask.LogId)).WithLocalLogger(_localLogger);
+			try
+			{
+				await ExecuteInternalAsync(serverLogger, cancellationToken);
+			}
+			catch (Exception ex)
+			{
+				serverLogger.LogError(ex, "Unhandled exception while running conform: {Message}", ex.Message);
+				throw;
+			}
+		}
+
+		async Task ExecuteInternalAsync(ILogger logger, CancellationToken cancellationToken)
+		{
+			logger.LogInformation("Conforming, lease {LeaseId}", _leaseId);
+			await TerminateProcessHelper.TerminateProcessesAsync(TerminateCondition.BeforeConform, _workingDir, _driverSettings.ProcessesToTerminate, logger, cancellationToken);
+
+			bool removeUntrackedFiles = _conformTask.RemoveUntrackedFiles;
+			IList<RpcAgentWorkspace> pendingWorkspaces = _conformTask.Workspaces;
+			for (; ; )
+			{
+				bool isPerforceExecutor = _driverSettings.Executor.Equals(PerforceExecutor.Name, StringComparison.OrdinalIgnoreCase);
+				bool isWorkspaceExecutor = _driverSettings.Executor.Equals(WorkspaceExecutor.Name, StringComparison.OrdinalIgnoreCase);
+
+				// When using WorkspaceExecutor, only job options can override exact materializer to use
+				// It will default to ManagedWorkspaceMaterializer, which is compatible with the conform call below
+				// Therefore, compatibility is assumed for now. Exact materializer to use should be changed to a per workspace setting.
+				// See WorkspaceExecutorFactory.CreateExecutor
+				bool isExecutorConformCompatible = isPerforceExecutor || isWorkspaceExecutor;
+
+				// Run the conform task
+				if (isExecutorConformCompatible && _driverSettings.PerforceExecutor.RunConform)
+				{
+					await PerforceExecutor.ConformAsync(_workingDir, pendingWorkspaces, removeUntrackedFiles, logger, cancellationToken);
+				}
+				else
+				{
+					logger.LogInformation("Skipping conform. Executor={Executor} RunConform={RunConform}", _driverSettings.Executor, _driverSettings.PerforceExecutor.RunConform);
+				}
+
+				// Update the new set of workspaces
+				RpcUpdateAgentWorkspacesRequest request = new RpcUpdateAgentWorkspacesRequest();
+				request.AgentId = _agentId.ToString();
+				request.Workspaces.AddRange(pendingWorkspaces);
+				request.RemoveUntrackedFiles = removeUntrackedFiles;
+
+				HordeRpc.HordeRpcClient hordeRpc = await _hordeClient.CreateGrpcClientAsync<HordeRpc.HordeRpcClient>(cancellationToken);
+
+				RpcUpdateAgentWorkspacesResponse response = await hordeRpc.UpdateAgentWorkspacesAsync(request, cancellationToken: cancellationToken);
+				if (!response.Retry)
+				{
+					logger.LogInformation("Conform finished");
+					break;
+				}
+
+				logger.LogInformation("Pending workspaces have changed - running conform again...");
+				pendingWorkspaces = response.PendingWorkspaces;
+				removeUntrackedFiles = response.RemoveUntrackedFiles;
+			}
+		}
+	}
+}
diff --git a/Engine/Source/Programs/Horde/Horde.Agent.Driver/Execution/JobExecutor.cs b/Engine/Source/Programs/Horde/Horde.Agent.Driver/Execution/JobExecutor.cs
--- a/Engine/Source/Programs/Horde/Horde.Agent.Driver/Execution/JobExecutor.cs
+++ b/Engine/Source/Programs/Horde/Horde.Agent.Driver/Execution/JobExecutor.cs
@@ -15,13 +15,9 @@ using EpicGames.Horde.Logs;
 using EpicGames.Horde.Storage;
 using EpicGames.Horde.Storage.Nodes;
 using Grpc.Core;
-using Horde.Agent.Driver;
+using Horde.Agent.Driver.Parser;
 using Horde.Agent.Driver.Utility;
-using Horde.Agent.Parser;
-//using Horde.Agent.Services;
-using Horde.Agent.Utility;
 using Horde.Common.Rpc;
-using Horde.Storage.Utility;
 using HordeCommon;
 using HordeCommon.Rpc;
 using HordeCommon.Rpc.Messages;
@@ -30,7 +26,7 @@ using Microsoft.Extensions.Logging;
 using OpenTracing;
 using OpenTracing.Util;
 
-namespace Horde.Agent.Execution
+namespace Horde.Agent.Driver.Execution
 {
 	record class JobStepInfo
 	(
@@ -225,7 +221,7 @@ namespace Horde.Agent.Execution
 		protected IReadOnlyList<ProcessToTerminate>? ProcessesToTerminate { get; }
 
 		protected RpcBeginBatchResponse Batch { get; private set; }
-		private LogId _logId;
+		private readonly LogId _logId;
 
 		protected List<string> _additionalArguments = new List<string>();
 		protected bool _allowTargetChanges;
@@ -234,7 +230,6 @@ namespace Horde.Agent.Execution
 
 		protected RpcJobOptions JobOptions { get; }
 
-
 		protected JobRpc.JobRpcClient JobRpc { get; private set; }
 		protected Dictionary<string, string> _remapAgentTypes = new Dictionary<string, string>(StringComparer.OrdinalIgnoreCase);
 
@@ -2204,4 +2199,28 @@ namespace Horde.Agent.Execution
 
 		JobExecutor CreateExecutor(RpcAgentWorkspace workspaceInfo, RpcAgentWorkspace? autoSdkWorkspaceInfo, JobExecutorOptions options);
 	}
+
+	static class JobExecutorHelpers
+	{
+		public static async Task ExecuteAsync(IHordeClient hordeClient, DirectoryReference workingDir, LeaseId leaseId, ExecuteJobTask executeTask, IEnumerable<IJobExecutorFactory> executorFactories, DriverSettings driverSettings, ILogger logger, CancellationToken cancellationToken)
+		{
+			// Create an executor for this job
+			string executorName = String.IsNullOrEmpty(executeTask.JobOptions.Executor) ? driverSettings.Executor : executeTask.JobOptions.Executor;
+
+			IJobExecutorFactory? executorFactory = executorFactories.FirstOrDefault(x => x.Name.Equals(executorName, StringComparison.OrdinalIgnoreCase));
+			if (executorFactory == null)
+			{
+				throw new InvalidOperationException($"Unable to find executor '{executorName}'");
+			}
+
+			JobId jobId = JobId.Parse(executeTask.JobId);
+			JobStepBatchId batchId = JobStepBatchId.Parse(executeTask.BatchId);
+			LogId logId = LogId.Parse(executeTask.LogId);
+
+			JobExecutorOptions options = new JobExecutorOptions(hordeClient, workingDir, driverSettings.ProcessesToTerminate, jobId, batchId, leaseId, logId, executeTask.JobOptions);
+
+			using JobExecutor executor = executorFactory.CreateExecutor(executeTask.Workspace, executeTask.AutoSdkWorkspace, options);
+			await executor.ExecuteAsync(logger, cancellationToken);
+		}
+	}
 }
diff --git a/Engine/Source/Programs/Horde/Horde.Agent.Driver/Execution/LocalExecutor.cs b/Engine/Source/Programs/Horde/Horde.Agent.Driver/Execution/LocalExecutor.cs
--- a/Engine/Source/Programs/Horde/Horde.Agent.Driver/Execution/LocalExecutor.cs
+++ b/Engine/Source/Programs/Horde/Horde.Agent.Driver/Execution/LocalExecutor.cs
@@ -5,7 +5,7 @@ using HordeCommon.Rpc.Messages;
 using Microsoft.Extensions.Logging;
 using Microsoft.Extensions.Options;
 
-namespace Horde.Agent.Execution
+namespace Horde.Agent.Driver.Execution
 {
 	/// <summary>
 	/// Settings for the local executor
diff --git a/Engine/Source/Programs/Horde/Horde.Agent.Driver/Execution/ManagedWorkspaceMaterializer.cs b/Engine/Source/Programs/Horde/Horde.Agent.Driver/Execution/ManagedWorkspaceMaterializer.cs
--- a/Engine/Source/Programs/Horde/Horde.Agent.Driver/Execution/ManagedWorkspaceMaterializer.cs
+++ b/Engine/Source/Programs/Horde/Horde.Agent.Driver/Execution/ManagedWorkspaceMaterializer.cs
@@ -2,13 +2,13 @@
 
 using EpicGames.Core;
 using EpicGames.Perforce.Managed;
-using Horde.Agent.Utility;
+using Horde.Agent.Driver.Utility;
 using HordeCommon.Rpc.Messages;
 using Microsoft.Extensions.Logging;
 using OpenTracing;
 using OpenTracing.Util;
 
-namespace Horde.Agent.Execution;
+namespace Horde.Agent.Driver.Execution;
 
 /// <summary>
 /// Workspace materializer wrapping ManagedWorkspace and WorkspaceInfo
diff --git a/Engine/Source/Programs/Horde/Horde.Agent.Driver/Execution/PerforceExecutor.cs b/Engine/Source/Programs/Horde/Horde.Agent.Driver/Execution/PerforceExecutor.cs
--- a/Engine/Source/Programs/Horde/Horde.Agent.Driver/Execution/PerforceExecutor.cs
+++ b/Engine/Source/Programs/Horde/Horde.Agent.Driver/Execution/PerforceExecutor.cs
@@ -5,14 +5,14 @@ using System.Text;
 using EpicGames.Core;
 using EpicGames.Perforce;
 using EpicGames.Perforce.Managed;
-using Horde.Agent.Utility;
+using Horde.Agent.Driver.Utility;
 using HordeCommon.Rpc;
 using HordeCommon.Rpc.Messages;
 using Microsoft.Extensions.Logging;
 using OpenTracing;
 using OpenTracing.Util;
 
-namespace Horde.Agent.Execution
+namespace Horde.Agent.Driver.Execution
 {
 	/// <summary>
 	/// Settings for the perforce executor
diff --git a/Engine/Source/Programs/Horde/Horde.Agent.Driver/Execution/TestExecutor.cs b/Engine/Source/Programs/Horde/Horde.Agent.Driver/Execution/TestExecutor.cs
--- a/Engine/Source/Programs/Horde/Horde.Agent.Driver/Execution/TestExecutor.cs
+++ b/Engine/Source/Programs/Horde/Horde.Agent.Driver/Execution/TestExecutor.cs
@@ -4,18 +4,18 @@ using System.Diagnostics.CodeAnalysis;
 using EpicGames.Core;
 using EpicGames.Horde.Jobs;
 using EpicGames.Perforce;
-using Horde.Agent.Parser;
+using Horde.Agent.Driver.Parser;
 using HordeCommon.Rpc;
 using HordeCommon.Rpc.Messages;
 using Microsoft.Extensions.Logging;
 
-namespace Horde.Agent.Execution
+namespace Horde.Agent.Driver.Execution
 {
 	class TestExecutor : JobExecutor
 	{
 		public const string Name = "Test";
 
-		Dictionary<string, string> _arguments = new Dictionary<string, string>(StringComparer.OrdinalIgnoreCase);
+		readonly Dictionary<string, string> _arguments = new Dictionary<string, string>(StringComparer.OrdinalIgnoreCase);
 
 		public TestExecutor(JobExecutorOptions options, ILogger logger)
 			: base(options, logger)
diff --git a/Engine/Source/Programs/Horde/Horde.Agent.Driver/Execution/WorkspaceExecutor.cs b/Engine/Source/Programs/Horde/Horde.Agent.Driver/Execution/WorkspaceExecutor.cs
--- a/Engine/Source/Programs/Horde/Horde.Agent.Driver/Execution/WorkspaceExecutor.cs
+++ b/Engine/Source/Programs/Horde/Horde.Agent.Driver/Execution/WorkspaceExecutor.cs
@@ -2,14 +2,14 @@
 
 using System.Text.RegularExpressions;
 using EpicGames.Core;
-using Horde.Agent.Utility;
+using Horde.Agent.Driver.Utility;
 using HordeCommon.Rpc;
 using HordeCommon.Rpc.Messages;
 using Microsoft.Extensions.Logging;
 using OpenTracing;
 using OpenTracing.Util;
 
-namespace Horde.Agent.Execution
+namespace Horde.Agent.Driver.Execution
 {
 	class WorkspaceExecutor : JobExecutor
 	{
diff --git a/Engine/Source/Programs/Horde/Horde.Agent.Driver/Execution/WorkspaceMaterializer.cs b/Engine/Source/Programs/Horde/Horde.Agent.Driver/Execution/WorkspaceMaterializer.cs
--- a/Engine/Source/Programs/Horde/Horde.Agent.Driver/Execution/WorkspaceMaterializer.cs
+++ b/Engine/Source/Programs/Horde/Horde.Agent.Driver/Execution/WorkspaceMaterializer.cs
@@ -4,7 +4,7 @@ using EpicGames.Core;
 using HordeCommon.Rpc.Messages;
 using Microsoft.Extensions.Logging;
 
-namespace Horde.Agent.Execution;
+namespace Horde.Agent.Driver.Execution;
 
 /// <summary>
 /// Exception for workspace materializer
@@ -138,7 +138,7 @@ public interface IWorkspaceMaterializer : IDisposable
 	/// <param name="preflightChangeNum">Preflight change number to add</param>
 	/// <param name="options">Additional options</param>
 	/// <param name="cancellationToken">Cancellation token for the call</param>
-	/// <exception cref="Horde.Agent.Execution.WorkspaceMaterializationException">Thrown if syncing fails</exception>
+	/// <exception cref="Horde.Agent.Driver.Execution.WorkspaceMaterializationException">Thrown if syncing fails</exception>
 	/// <returns>Async task</returns>
 	public Task SyncAsync(int changeNum, int preflightChangeNum, SyncOptions options, CancellationToken cancellationToken);
 }
diff --git a/Engine/Source/Programs/Horde/Horde.Agent.Driver/Parser/LogParser.cs b/Engine/Source/Programs/Horde/Horde.Agent.Driver/Parser/LogParser.cs
--- a/Engine/Source/Programs/Horde/Horde.Agent.Driver/Parser/LogParser.cs
+++ b/Engine/Source/Programs/Horde/Horde.Agent.Driver/Parser/LogParser.cs
@@ -5,7 +5,7 @@ using System.Text.RegularExpressions;
 using EpicGames.Core;
 using Microsoft.Extensions.Logging;
 
-namespace Horde.Agent.Parser
+namespace Horde.Agent.Driver.Parser
 {
 	/// <summary>
 	/// Turns raw text output into structured logging events
diff --git a/Engine/Source/Programs/Horde/Horde.Agent.Driver/Utility/LinuxInterop.cs b/Engine/Source/Programs/Horde/Horde.Agent.Driver/Utility/LinuxInterop.cs
--- a/Engine/Source/Programs/Horde/Horde.Agent.Driver/Utility/LinuxInterop.cs
+++ b/Engine/Source/Programs/Horde/Horde.Agent.Driver/Utility/LinuxInterop.cs
@@ -2,24 +2,25 @@
 
 using System.Runtime.InteropServices;
 
-namespace Horde.Agent.Utility;
-
-/// <summary>
-/// Linux specific P/Invoke calls
-/// </summary>
-static class LinuxInterop
+namespace Horde.Agent.Driver.Utility
 {
 	/// <summary>
-	/// Get user identity
+	/// Linux specific P/Invoke calls
 	/// </summary>
-	/// <returns>Real user ID of the calling process</returns>
-	[DllImport("libc", SetLastError = true)]
-	internal static extern uint getuid();
+	static class LinuxInterop
+	{
+		/// <summary>
+		/// Get user identity
+		/// </summary>
+		/// <returns>Real user ID of the calling process</returns>
+		[DllImport("libc", SetLastError = true)]
+		internal static extern uint getuid();
 
-	/// <summary>
-	/// Get group identity
-	/// </summary>
-	/// <returns>Real group ID of the calling process</returns>
-	[DllImport("libc", SetLastError = true)]
-	internal static extern uint getgid();
+		/// <summary>
+		/// Get group identity
+		/// </summary>
+		/// <returns>Real group ID of the calling process</returns>
+		[DllImport("libc", SetLastError = true)]
+		internal static extern uint getgid();
+	}
 }
diff --git a/Engine/Source/Programs/Horde/Horde.Agent.Driver/Utility/TempStorage.cs b/Engine/Source/Programs/Horde/Horde.Agent.Driver/Utility/TempStorage.cs
--- a/Engine/Source/Programs/Horde/Horde.Agent.Driver/Utility/TempStorage.cs
+++ b/Engine/Source/Programs/Horde/Horde.Agent.Driver/Utility/TempStorage.cs
@@ -16,7 +16,7 @@ using Microsoft.Extensions.Logging;
 
 #pragma warning disable CA1819 // Properties should not return arrays
 
-namespace Horde.Storage.Utility
+namespace Horde.Agent.Driver.Utility
 {
 	/// <summary>
 	/// Exception thrown by the temp storage system
diff --git a/Engine/Source/Programs/Horde/Horde.Agent.Driver/Utility/TerminateProcessHelper.cs b/Engine/Source/Programs/Horde/Horde.Agent.Driver/Utility/TerminateProcessHelper.cs
--- a/Engine/Source/Programs/Horde/Horde.Agent.Driver/Utility/TerminateProcessHelper.cs
+++ b/Engine/Source/Programs/Horde/Horde.Agent.Driver/Utility/TerminateProcessHelper.cs
@@ -1,10 +1,9 @@
 // Copyright Epic Games, Inc. All Rights Reserved.
 
 using EpicGames.Core;
-using Horde.Agent.Driver;
 using Microsoft.Extensions.Logging;
 
-namespace Horde.Agent.Utility
+namespace Horde.Agent.Driver.Utility
 {
 	/// <summary>
 	/// Flags for processes to terminate
diff --git a/Engine/Source/Programs/Horde/Horde.Agent.Driver/Utility/WorkspaceInfo.cs b/Engine/Source/Programs/Horde/Horde.Agent.Driver/Utility/WorkspaceInfo.cs
--- a/Engine/Source/Programs/Horde/Horde.Agent.Driver/Utility/WorkspaceInfo.cs
+++ b/Engine/Source/Programs/Horde/Horde.Agent.Driver/Utility/WorkspaceInfo.cs
@@ -11,7 +11,7 @@ using EpicGames.Perforce.Managed;
 using HordeCommon.Rpc.Messages;
 using Microsoft.Extensions.Logging;
 
-namespace Horde.Agent.Utility
+namespace Horde.Agent.Driver.Utility
 {
 	/// <summary>
 	/// Stores information about a managed Perforce workspace
diff --git a/Engine/Source/Programs/Horde/Horde.Agent/AgentApp.cs b/Engine/Source/Programs/Horde/Horde.Agent/AgentApp.cs
--- a/Engine/Source/Programs/Horde/Horde.Agent/AgentApp.cs
+++ b/Engine/Source/Programs/Horde/Horde.Agent/AgentApp.cs
@@ -12,7 +12,7 @@ using Datadog.Trace.OpenTracing;
 using EpicGames.Core;
 using EpicGames.Horde;
 using Horde.Agent.Driver;
-using Horde.Agent.Execution;
+using Horde.Agent.Driver.Execution;
 using Horde.Agent.Leases;
 using Horde.Agent.Leases.Handlers;
 using Horde.Agent.Services;
diff --git a/Engine/Source/Programs/Horde/Horde.Agent/Commands/Execution/ExecuteJobCommand.cs b/Engine/Source/Programs/Horde/Horde.Agent/Commands/Execution/ExecuteJobCommand.cs
deleted file mode 100644
diff --git a/Engine/Source/Programs/Horde/Horde.Agent/Leases/Handlers/ConformHandler.cs b/Engine/Source/Programs/Horde/Horde.Agent/Leases/Handlers/ConformHandler.cs
--- a/Engine/Source/Programs/Horde/Horde.Agent/Leases/Handlers/ConformHandler.cs
+++ b/Engine/Source/Programs/Horde/Horde.Agent/Leases/Handlers/ConformHandler.cs
@@ -1,13 +1,9 @@
 // Copyright Epic Games, Inc. All Rights Reserved.
 
 using EpicGames.Horde.Agents.Leases;
-using EpicGames.Horde.Logs;
 using Horde.Agent.Driver;
-using Horde.Agent.Execution;
+using Horde.Agent.Driver.Execution;
 using Horde.Agent.Services;
-using Horde.Agent.Utility;
-using HordeCommon.Rpc;
-using HordeCommon.Rpc.Messages;
 using HordeCommon.Rpc.Tasks;
 using Microsoft.Extensions.Logging;
 using Microsoft.Extensions.Options;
@@ -29,66 +25,8 @@ namespace Horde.Agent.Leases.Handlers
 		/// <inheritdoc/>
 		public override async Task<LeaseResult> ExecuteAsync(ISession session, LeaseId leaseId, ConformTask conformTask, ILogger localLogger, CancellationToken cancellationToken)
 		{
-			await using IServerLogger serverLogger = session.HordeClient.CreateServerLogger(LogId.Parse(conformTask.LogId)).WithLocalLogger(localLogger);
-			try
-			{
-				LeaseResult result = await ExecuteInternalAsync(session, leaseId, conformTask, serverLogger, cancellationToken);
-				return result;
-			}
-			catch (Exception ex)
-			{
-				serverLogger.LogError(ex, "Unhandled exception while running conform: {Message}", ex.Message);
-				throw;
-			}
-		}
-
-		async Task<LeaseResult> ExecuteInternalAsync(ISession session, LeaseId leaseId, ConformTask conformTask, ILogger logger, CancellationToken cancellationToken)
-		{
-			logger.LogInformation("Conforming, lease {LeaseId}", leaseId);
-			await TerminateProcessHelper.TerminateProcessesAsync(TerminateCondition.BeforeConform, session.WorkingDir, _driverSettings.ProcessesToTerminate, logger, cancellationToken);
-
-			bool removeUntrackedFiles = conformTask.RemoveUntrackedFiles;
-			IList<RpcAgentWorkspace> pendingWorkspaces = conformTask.Workspaces;
-			for (; ; )
-			{
-				bool isPerforceExecutor = _driverSettings.Executor.Equals(PerforceExecutor.Name, StringComparison.OrdinalIgnoreCase);
-				bool isWorkspaceExecutor = _driverSettings.Executor.Equals(WorkspaceExecutor.Name, StringComparison.OrdinalIgnoreCase);
-
-				// When using WorkspaceExecutor, only job options can override exact materializer to use
-				// It will default to ManagedWorkspaceMaterializer, which is compatible with the conform call below
-				// Therefore, compatibility is assumed for now. Exact materializer to use should be changed to a per workspace setting.
-				// See WorkspaceExecutorFactory.CreateExecutor
-				bool isExecutorConformCompatible = isPerforceExecutor || isWorkspaceExecutor;
-
-				// Run the conform task
-				if (isExecutorConformCompatible && _driverSettings.PerforceExecutor.RunConform)
-				{
-					await PerforceExecutor.ConformAsync(session.WorkingDir, pendingWorkspaces, removeUntrackedFiles, logger, cancellationToken);
-				}
-				else
-				{
-					logger.LogInformation("Skipping conform. Executor={Executor} RunConform={RunConform}", _driverSettings.Executor, _driverSettings.PerforceExecutor.RunConform);
-				}
-
-				// Update the new set of workspaces
-				RpcUpdateAgentWorkspacesRequest request = new RpcUpdateAgentWorkspacesRequest();
-				request.AgentId = session.AgentId.ToString();
-				request.Workspaces.AddRange(pendingWorkspaces);
-				request.RemoveUntrackedFiles = removeUntrackedFiles;
-
-				HordeRpc.HordeRpcClient hordeRpc = await session.HordeClient.CreateGrpcClientAsync<HordeRpc.HordeRpcClient>(cancellationToken);
-
-				RpcUpdateAgentWorkspacesResponse response = await hordeRpc.UpdateAgentWorkspacesAsync(request, cancellationToken: cancellationToken);
-				if (!response.Retry)
-				{
-					logger.LogInformation("Conform finished");
-					break;
-				}
-
-				logger.LogInformation("Pending workspaces have changed - running conform again...");
-				pendingWorkspaces = response.PendingWorkspaces;
-				removeUntrackedFiles = response.RemoveUntrackedFiles;
-			}
+			ConformExecutor executor = new ConformExecutor(session.HordeClient, session.WorkingDir, session.AgentId, leaseId, conformTask, _driverSettings, localLogger);
+			await executor.ExecuteAsync(cancellationToken);
 
 			return LeaseResult.Success;
 		}
diff --git a/Engine/Source/Programs/Horde/Horde.Agent/Leases/Handlers/JobHandler.cs b/Engine/Source/Programs/Horde/Horde.Agent/Leases/Handlers/JobHandler.cs
--- a/Engine/Source/Programs/Horde/Horde.Agent/Leases/Handlers/JobHandler.cs
+++ b/Engine/Source/Programs/Horde/Horde.Agent/Leases/Handlers/JobHandler.cs
@@ -5,11 +5,9 @@ using System.Reflection;
 using EpicGames.Core;
 using EpicGames.Horde;
 using EpicGames.Horde.Agents.Leases;
-using EpicGames.Horde.Jobs;
-using EpicGames.Horde.Logs;
 using Google.Protobuf;
 using Horde.Agent.Driver;
-using Horde.Agent.Execution;
+using Horde.Agent.Driver.Execution;
 using Horde.Agent.Services;
 using HordeCommon.Rpc.Tasks;
 using Microsoft.Extensions.Logging;
@@ -127,28 +125,11 @@ namespace Horde.Agent.Leases.Handlers
 
 		internal async Task<LeaseResult> ExecuteInternalAsync(IHordeClient hordeClient, DirectoryReference workingDir, LeaseId leaseId, ExecuteJobTask executeTask, ILogger localLogger, CancellationToken cancellationToken)
 		{
-			// Create an executor for this job
-			string executorName = String.IsNullOrEmpty(executeTask.JobOptions.Executor) ? _driverSettings.Executor : executeTask.JobOptions.Executor;
-
-			IJobExecutorFactory? executorFactory = _executorFactories.FirstOrDefault(x => x.Name.Equals(executorName, StringComparison.OrdinalIgnoreCase));
-			if (executorFactory == null)
-			{
-				throw new InvalidOperationException($"Unable to find executor '{executorName}'");
-			}
-
 			GlobalTracer.Instance.ActiveSpan?.SetTag("jobId", executeTask.JobId.ToString());
 			GlobalTracer.Instance.ActiveSpan?.SetTag("jobName", executeTask.JobName.ToString());
 			GlobalTracer.Instance.ActiveSpan?.SetTag("batchId", executeTask.BatchId.ToString());
 
-			JobId jobId = JobId.Parse(executeTask.JobId);
-			JobStepBatchId batchId = JobStepBatchId.Parse(executeTask.BatchId);
-			LogId logId = LogId.Parse(executeTask.LogId);
-
-			JobExecutorOptions options = new JobExecutorOptions(hordeClient, workingDir, _driverSettings.ProcessesToTerminate, jobId, batchId, leaseId, logId, executeTask.JobOptions);
-
-			using JobExecutor executor = executorFactory.CreateExecutor(executeTask.Workspace, executeTask.AutoSdkWorkspace, options);
-			await executor.ExecuteAsync(localLogger, cancellationToken);
-
+			await JobExecutorHelpers.ExecuteAsync(hordeClient, workingDir, leaseId, executeTask, _executorFactories, _driverSettings, localLogger, cancellationToken);
 			return LeaseResult.Success;
 		}
 	}
