Horde: Separate out job and agent rpc messages.

diff --git a/Engine/Source/Programs/Horde/Horde.Agent.Driver/Execution/ConformExecutor.cs b/Engine/Source/Programs/Horde/Horde.Agent.Driver/Execution/ConformExecutor.cs
--- a/Engine/Source/Programs/Horde/Horde.Agent.Driver/Execution/ConformExecutor.cs
+++ b/Engine/Source/Programs/Horde/Horde.Agent.Driver/Execution/ConformExecutor.cs
@@ -6,7 +6,7 @@ using EpicGames.Horde.Agents;
 using EpicGames.Horde.Agents.Leases;
 using EpicGames.Horde.Logs;
 using Horde.Agent.Driver.Utility;
-using HordeCommon.Rpc;
+using Horde.Common.Rpc;
 using HordeCommon.Rpc.Messages;
 using HordeCommon.Rpc.Tasks;
 using Microsoft.Extensions.Logging;
@@ -82,7 +82,7 @@ namespace Horde.Agent.Driver.Execution
 				request.Workspaces.AddRange(pendingWorkspaces);
 				request.RemoveUntrackedFiles = removeUntrackedFiles;
 
-				HordeRpc.HordeRpcClient hordeRpc = await _hordeClient.CreateGrpcClientAsync<HordeRpc.HordeRpcClient>(cancellationToken);
+				JobRpc.JobRpcClient hordeRpc = await _hordeClient.CreateGrpcClientAsync<JobRpc.JobRpcClient>(cancellationToken);
 
 				RpcUpdateAgentWorkspacesResponse response = await hordeRpc.UpdateAgentWorkspacesAsync(request, cancellationToken: cancellationToken);
 				if (!response.Retry)
diff --git a/Engine/Source/Programs/Horde/Horde.Agent.Driver/Execution/JobExecutor.cs b/Engine/Source/Programs/Horde/Horde.Agent.Driver/Execution/JobExecutor.cs
--- a/Engine/Source/Programs/Horde/Horde.Agent.Driver/Execution/JobExecutor.cs
+++ b/Engine/Source/Programs/Horde/Horde.Agent.Driver/Execution/JobExecutor.cs
@@ -19,7 +19,6 @@ using Horde.Agent.Driver.Parser;
 using Horde.Agent.Driver.Utility;
 using Horde.Common.Rpc;
 using HordeCommon;
-using HordeCommon.Rpc;
 using HordeCommon.Rpc.Messages;
 using HordeCommon.Rpc.Tasks;
 using Microsoft.Extensions.Logging;
diff --git a/Engine/Source/Programs/Horde/Horde.Agent.Driver/Execution/PerforceExecutor.cs b/Engine/Source/Programs/Horde/Horde.Agent.Driver/Execution/PerforceExecutor.cs
--- a/Engine/Source/Programs/Horde/Horde.Agent.Driver/Execution/PerforceExecutor.cs
+++ b/Engine/Source/Programs/Horde/Horde.Agent.Driver/Execution/PerforceExecutor.cs
@@ -6,7 +6,7 @@ using EpicGames.Core;
 using EpicGames.Perforce;
 using EpicGames.Perforce.Managed;
 using Horde.Agent.Driver.Utility;
-using HordeCommon.Rpc;
+using Horde.Common.Rpc;
 using HordeCommon.Rpc.Messages;
 using Microsoft.Extensions.Logging;
 using OpenTracing;
diff --git a/Engine/Source/Programs/Horde/Horde.Agent.Driver/Execution/TestExecutor.cs b/Engine/Source/Programs/Horde/Horde.Agent.Driver/Execution/TestExecutor.cs
--- a/Engine/Source/Programs/Horde/Horde.Agent.Driver/Execution/TestExecutor.cs
+++ b/Engine/Source/Programs/Horde/Horde.Agent.Driver/Execution/TestExecutor.cs
@@ -5,7 +5,7 @@ using EpicGames.Core;
 using EpicGames.Horde.Jobs;
 using EpicGames.Perforce;
 using Horde.Agent.Driver.Parser;
-using HordeCommon.Rpc;
+using Horde.Common.Rpc;
 using HordeCommon.Rpc.Messages;
 using Microsoft.Extensions.Logging;
 
diff --git a/Engine/Source/Programs/Horde/Horde.Agent.Driver/Execution/WorkspaceExecutor.cs b/Engine/Source/Programs/Horde/Horde.Agent.Driver/Execution/WorkspaceExecutor.cs
--- a/Engine/Source/Programs/Horde/Horde.Agent.Driver/Execution/WorkspaceExecutor.cs
+++ b/Engine/Source/Programs/Horde/Horde.Agent.Driver/Execution/WorkspaceExecutor.cs
@@ -3,7 +3,7 @@
 using System.Text.RegularExpressions;
 using EpicGames.Core;
 using Horde.Agent.Driver.Utility;
-using HordeCommon.Rpc;
+using Horde.Common.Rpc;
 using HordeCommon.Rpc.Messages;
 using Microsoft.Extensions.Logging;
 using OpenTracing;
diff --git a/Engine/Source/Programs/Horde/Horde.Agent.Driver/Horde.Agent.Driver.csproj b/Engine/Source/Programs/Horde/Horde.Agent.Driver/Horde.Agent.Driver.csproj
--- a/Engine/Source/Programs/Horde/Horde.Agent.Driver/Horde.Agent.Driver.csproj
+++ b/Engine/Source/Programs/Horde/Horde.Agent.Driver/Horde.Agent.Driver.csproj
@@ -49,9 +49,6 @@
     <Content Include="..\Horde.Server\Protos\horde\tasks\**\*.proto">
       <Link>Protos\horde\tasks\%(RecursiveDir)%(FileName)%(Extension)</Link>
     </Content>
-    <Compile Include="..\Horde.Server\Protos\horde\**\*.cs">
-      <Link>Protos\horde\%(RecursiveDir)%(FileName)%(Extension)</Link>
-    </Compile>
   </ItemGroup>
 
   <ItemGroup>
@@ -94,31 +91,12 @@
   </ItemGroup>
   
   <ItemGroup>
-    <Protobuf Include="..\Horde.Server\Protos\horde\horde_rpc.proto" ProtoRoot="..\Horde.Server\Protos" GrpcServices="Client" Access="Public" />
-    <Protobuf Include="..\Horde.Server\Protos\horde\horde_rpc_messages.proto" ProtoRoot="..\Horde.Server\Protos" Access="Public" />
-    <Protobuf Include="..\Horde.Server\Protos\horde\enrollment_rpc.proto" ProtoRoot="..\Horde.Server\Protos" GrpcServices="Client" Access="Public" />
-    <Protobuf Include="..\Horde.Server\Protos\horde\enrollment_rpc_messages.proto" ProtoRoot="..\Horde.Server\Protos" Access="Public" />
     <Protobuf Include="..\Horde.Server\Protos\horde\job_rpc.proto" ProtoRoot="..\Horde.Server\Protos" GrpcServices="Client" Access="Public" AdditionalImportDirs="..\..\Shared\EpicGames.Horde\Protos"/>
     <Protobuf Include="..\Horde.Server\Protos\horde\job_rpc_messages.proto" ProtoRoot="..\Horde.Server\Protos" Access="Public" />
-    <Protobuf Include="..\Horde.Server\Protos\horde\common\agent_status.proto" ProtoRoot="..\Horde.Server\Protos" Access="Public" />
-    <Protobuf Include="..\Horde.Server\Protos\horde\common\io_hash.proto" ProtoRoot="..\Horde.Server\Protos" Access="Public" />
     <Protobuf Include="..\Horde.Server\Protos\horde\common\label_change.proto" ProtoRoot="..\Horde.Server\Protos" Access="Public" />
-    <Protobuf Include="..\Horde.Server\Protos\horde\common\lease_outcome.proto" ProtoRoot="..\Horde.Server\Protos" Access="Public" />
-    <Protobuf Include="..\Horde.Server\Protos\horde\common\lease_state.proto" ProtoRoot="..\Horde.Server\Protos" Access="Public" />
-    <Protobuf Include="..\Horde.Server\Protos\horde\common\namespace_id.proto" ProtoRoot="..\Horde.Server\Protos" Access="Public" />
-    <Protobuf Include="..\Horde.Server\Protos\horde\messages\agent_capabilities.proto" ProtoRoot="..\Horde.Server\Protos" Access="Public" />
-    <Protobuf Include="..\Horde.Server\Protos\horde\messages\agent_requirements.proto" ProtoRoot="..\Horde.Server\Protos" Access="Public" />
     <Protobuf Include="..\Horde.Server\Protos\horde\messages\agent_workspace.proto" ProtoRoot="..\Horde.Server\Protos" Access="Public" />
-    <Protobuf Include="..\Horde.Server\Protos\horde\messages\lease.proto" ProtoRoot="..\Horde.Server\Protos" Access="Public" />
-    <Protobuf Include="..\Horde.Server\Protos\horde\messages\aws_lambda.proto" ProtoRoot="..\Horde.Server\Protos" Access="Public" />
-    <Protobuf Include="..\Horde.Server\Protos\horde\messages\telemetry.proto" ProtoRoot="..\Horde.Server\Protos" Access="Public" />
-    <Protobuf Include="..\Horde.Server\Protos\horde\tasks\compute_task.proto" ProtoRoot="..\Horde.Server\Protos" Access="Public" />
     <Protobuf Include="..\Horde.Server\Protos\horde\tasks\conform_task.proto" ProtoRoot="..\Horde.Server\Protos" Access="Public" />
     <Protobuf Include="..\Horde.Server\Protos\horde\tasks\job_task.proto" ProtoRoot="..\Horde.Server\Protos" Access="Public" />
-    <Protobuf Include="..\Horde.Server\Protos\horde\tasks\restart_task.proto" ProtoRoot="..\Horde.Server\Protos" Access="Public" />
-    <Protobuf Include="..\Horde.Server\Protos\horde\tasks\shutdown_task.proto" ProtoRoot="..\Horde.Server\Protos" Access="Public" />
-    <Protobuf Include="..\Horde.Server\Protos\horde\tasks\test_task.proto" ProtoRoot="..\Horde.Server\Protos" Access="Public" />
-    <Protobuf Include="..\Horde.Server\Protos\horde\tasks\upgrade_task.proto" ProtoRoot="..\Horde.Server\Protos" Access="Public" />
   </ItemGroup>
 
   <ItemGroup>
diff --git a/Engine/Source/Programs/Horde/Horde.Agent.Driver/Utilities/HordeRpc.cs b/Engine/Source/Programs/Horde/Horde.Agent.Driver/Utilities/HordeRpc.cs
--- a/Engine/Source/Programs/Horde/Horde.Agent.Driver/Utilities/HordeRpc.cs
+++ b/Engine/Source/Programs/Horde/Horde.Agent.Driver/Utilities/HordeRpc.cs
@@ -7,7 +7,7 @@ using EpicGames.Horde.Jobs;
 // Warnings below are disabled to avoid documenting every class touched.
 #pragma warning disable CS1591
 #pragma warning disable CA1716
-namespace HordeCommon.Rpc
+namespace Horde.Common.Rpc
 {
 	partial class RpcGetJobRequest
 	{
@@ -78,14 +78,6 @@ namespace HordeCommon.Rpc
 			AbortRequested = abortRequested;
 		}
 	}
-
-	partial class RpcDownloadSoftwareRequest
-	{
-		public RpcDownloadSoftwareRequest(string version)
-		{
-			Version = version;
-		}
-	}
 }
 
 #pragma warning restore CA1716
diff --git a/Engine/Source/Programs/Horde/Horde.Agent.Driver/Utility/JobStepLogger.cs b/Engine/Source/Programs/Horde/Horde.Agent.Driver/Utility/JobStepLogger.cs
--- a/Engine/Source/Programs/Horde/Horde.Agent.Driver/Utility/JobStepLogger.cs
+++ b/Engine/Source/Programs/Horde/Horde.Agent.Driver/Utility/JobStepLogger.cs
@@ -4,7 +4,6 @@ using EpicGames.Horde;
 using EpicGames.Horde.Jobs;
 using EpicGames.Horde.Logs;
 using Horde.Common.Rpc;
-using HordeCommon.Rpc;
 using Microsoft.Extensions.Logging;
 
 namespace Horde.Agent.Driver.Utility
diff --git a/Engine/Source/Programs/Horde/Horde.Agent.Tests/Execution/FakeWorkspaceMaterializer.cs b/Engine/Source/Programs/Horde/Horde.Agent.Tests/Execution/FakeWorkspaceMaterializer.cs
--- a/Engine/Source/Programs/Horde/Horde.Agent.Tests/Execution/FakeWorkspaceMaterializer.cs
+++ b/Engine/Source/Programs/Horde/Horde.Agent.Tests/Execution/FakeWorkspaceMaterializer.cs
@@ -7,7 +7,7 @@ using System.Linq;
 using System.Threading;
 using System.Threading.Tasks;
 using EpicGames.Core;
-using Horde.Agent.Execution;
+using Horde.Agent.Driver.Execution;
 using Microsoft.Extensions.Logging;
 
 namespace Horde.Agent.Tests.Execution;
diff --git a/Engine/Source/Programs/Horde/Horde.Agent.Tests/Execution/FakeWorkspaceMaterializerTest.cs b/Engine/Source/Programs/Horde/Horde.Agent.Tests/Execution/FakeWorkspaceMaterializerTest.cs
--- a/Engine/Source/Programs/Horde/Horde.Agent.Tests/Execution/FakeWorkspaceMaterializerTest.cs
+++ b/Engine/Source/Programs/Horde/Horde.Agent.Tests/Execution/FakeWorkspaceMaterializerTest.cs
@@ -4,7 +4,7 @@ using System;
 using System.IO;
 using System.Threading;
 using System.Threading.Tasks;
-using Horde.Agent.Execution;
+using Horde.Agent.Driver.Execution;
 using Microsoft.Extensions.Logging.Abstractions;
 using Microsoft.VisualStudio.TestTools.UnitTesting;
 
diff --git a/Engine/Source/Programs/Horde/Horde.Agent.Tests/Execution/WorkspaceExecutorTest.cs b/Engine/Source/Programs/Horde/Horde.Agent.Tests/Execution/WorkspaceExecutorTest.cs
--- a/Engine/Source/Programs/Horde/Horde.Agent.Tests/Execution/WorkspaceExecutorTest.cs
+++ b/Engine/Source/Programs/Horde/Horde.Agent.Tests/Execution/WorkspaceExecutorTest.cs
@@ -9,9 +9,9 @@ using EpicGames.Core;
 using EpicGames.Horde.Jobs;
 using EpicGames.Horde.Streams;
 using Horde.Agent.Driver;
-using Horde.Agent.Execution;
+using Horde.Agent.Driver.Execution;
 using Horde.Agent.Services;
-using HordeCommon.Rpc;
+using Horde.Common.Rpc;
 using HordeCommon.Rpc.Tasks;
 using Microsoft.Extensions.Logging;
 using Microsoft.Extensions.Logging.Abstractions;
diff --git a/Engine/Source/Programs/Horde/Horde.Agent.Tests/Helper.cs b/Engine/Source/Programs/Horde/Horde.Agent.Tests/Helper.cs
--- a/Engine/Source/Programs/Horde/Horde.Agent.Tests/Helper.cs
+++ b/Engine/Source/Programs/Horde/Horde.Agent.Tests/Helper.cs
@@ -7,9 +7,8 @@ using System.Threading.Tasks;
 using EpicGames.Horde.Jobs;
 using Google.Protobuf.WellKnownTypes;
 using Grpc.Core;
-using Horde.Agent.Execution;
+using Horde.Agent.Driver.Execution;
 using Horde.Common.Rpc;
-using HordeCommon.Rpc;
 using HordeCommon.Rpc.Messages;
 using Microsoft.Extensions.Logging;
 using Microsoft.Extensions.Logging.Abstractions;
diff --git a/Engine/Source/Programs/Horde/Horde.Agent.Tests/Services/TelemetryServiceTest.cs b/Engine/Source/Programs/Horde/Horde.Agent.Tests/Services/TelemetryServiceTest.cs
--- a/Engine/Source/Programs/Horde/Horde.Agent.Tests/Services/TelemetryServiceTest.cs
+++ b/Engine/Source/Programs/Horde/Horde.Agent.Tests/Services/TelemetryServiceTest.cs
@@ -5,7 +5,7 @@ using System.Collections.Generic;
 using System.Threading;
 using System.Threading.Tasks;
 using Horde.Agent.Driver;
-using Horde.Agent.Execution;
+using Horde.Agent.Driver.Execution;
 using Horde.Agent.Leases.Handlers;
 using Horde.Agent.Services;
 using Microsoft.Extensions.Logging.Abstractions;
diff --git a/Engine/Source/Programs/Horde/Horde.Agent.Tests/WorkerServiceTest.cs b/Engine/Source/Programs/Horde/Horde.Agent.Tests/WorkerServiceTest.cs
--- a/Engine/Source/Programs/Horde/Horde.Agent.Tests/WorkerServiceTest.cs
+++ b/Engine/Source/Programs/Horde/Horde.Agent.Tests/WorkerServiceTest.cs
@@ -21,7 +21,7 @@ using Google.Protobuf.WellKnownTypes;
 using Grpc.Core;
 using Grpc.Net.Client;
 using Horde.Agent.Driver;
-using Horde.Agent.Execution;
+using Horde.Agent.Driver.Execution;
 using Horde.Agent.Leases;
 using Horde.Agent.Leases.Handlers;
 using Horde.Agent.Services;
diff --git a/Engine/Source/Programs/Horde/Horde.Agent.Tests/WorkspaceInfoTest.cs b/Engine/Source/Programs/Horde/Horde.Agent.Tests/WorkspaceInfoTest.cs
--- a/Engine/Source/Programs/Horde/Horde.Agent.Tests/WorkspaceInfoTest.cs
+++ b/Engine/Source/Programs/Horde/Horde.Agent.Tests/WorkspaceInfoTest.cs
@@ -1,7 +1,7 @@
 // Copyright Epic Games, Inc. All Rights Reserved.
 
 using EpicGames.Perforce.Managed;
-using Horde.Agent.Utility;
+using Horde.Agent.Driver.Utility;
 using Microsoft.VisualStudio.TestTools.UnitTesting;
 
 namespace Horde.Agent.Tests;
@@ -44,7 +44,7 @@ public class WorkspaceInfoTest
 
 	static ManagedWorkspaceOptions GetMwOptions(string? method, long? minScratchSpace = null)
 	{
-		HordeCommon.Rpc.Messages.RpcAgentWorkspace workspace = new () { Method = method ?? "", MinScratchSpace = minScratchSpace ?? 0 };
+		HordeCommon.Rpc.Messages.RpcAgentWorkspace workspace = new() { Method = method ?? "", MinScratchSpace = minScratchSpace ?? 0 };
 		return WorkspaceInfo.GetMwOptions(workspace);
 	}
 }
diff --git a/Engine/Source/Programs/Horde/Horde.Agent/Horde.Agent.csproj b/Engine/Source/Programs/Horde/Horde.Agent/Horde.Agent.csproj
--- a/Engine/Source/Programs/Horde/Horde.Agent/Horde.Agent.csproj
+++ b/Engine/Source/Programs/Horde/Horde.Agent/Horde.Agent.csproj
@@ -137,37 +137,27 @@
     <ProjectReference Include="..\Horde.Agent.Driver\Horde.Agent.Driver.csproj" />
   </ItemGroup>
 
-  <!--
   <ItemGroup>
     <Protobuf Include="..\Horde.Server\Protos\horde\horde_rpc.proto" ProtoRoot="..\Horde.Server\Protos" GrpcServices="Client" Access="Public" />
     <Protobuf Include="..\Horde.Server\Protos\horde\horde_rpc_messages.proto" ProtoRoot="..\Horde.Server\Protos" Access="Public" />
     <Protobuf Include="..\Horde.Server\Protos\horde\enrollment_rpc.proto" ProtoRoot="..\Horde.Server\Protos" GrpcServices="Client" Access="Public" />
     <Protobuf Include="..\Horde.Server\Protos\horde\enrollment_rpc_messages.proto" ProtoRoot="..\Horde.Server\Protos" Access="Public" />
-    <Protobuf Include="..\Horde.Server\Protos\horde\job_rpc.proto" ProtoRoot="..\Horde.Server\Protos" GrpcServices="Client" Access="Public" />
-    <Protobuf Include="..\Horde.Server\Protos\horde\job_rpc_messages.proto" ProtoRoot="..\Horde.Server\Protos" Access="Public" />
-    <Protobuf Include="..\Horde.Server\Protos\horde\log_rpc.proto" ProtoRoot="..\Horde.Server\Protos" GrpcServices="Client" Access="Public" />
-    <Protobuf Include="..\Horde.Server\Protos\horde\log_rpc_messages.proto" ProtoRoot="..\Horde.Server\Protos" Access="Public" />
     <Protobuf Include="..\Horde.Server\Protos\horde\common\agent_status.proto" ProtoRoot="..\Horde.Server\Protos" Access="Public" />
     <Protobuf Include="..\Horde.Server\Protos\horde\common\io_hash.proto" ProtoRoot="..\Horde.Server\Protos" Access="Public" />
-    <Protobuf Include="..\Horde.Server\Protos\horde\common\label_change.proto" ProtoRoot="..\Horde.Server\Protos" Access="Public" />
     <Protobuf Include="..\Horde.Server\Protos\horde\common\lease_outcome.proto" ProtoRoot="..\Horde.Server\Protos" Access="Public" />
     <Protobuf Include="..\Horde.Server\Protos\horde\common\lease_state.proto" ProtoRoot="..\Horde.Server\Protos" Access="Public" />
     <Protobuf Include="..\Horde.Server\Protos\horde\common\namespace_id.proto" ProtoRoot="..\Horde.Server\Protos" Access="Public" />
     <Protobuf Include="..\Horde.Server\Protos\horde\messages\agent_capabilities.proto" ProtoRoot="..\Horde.Server\Protos" Access="Public" />
     <Protobuf Include="..\Horde.Server\Protos\horde\messages\agent_requirements.proto" ProtoRoot="..\Horde.Server\Protos" Access="Public" />
-    <Protobuf Include="..\Horde.Server\Protos\horde\messages\agent_workspace.proto" ProtoRoot="..\Horde.Server\Protos" Access="Public" />
     <Protobuf Include="..\Horde.Server\Protos\horde\messages\lease.proto" ProtoRoot="..\Horde.Server\Protos" Access="Public" />
     <Protobuf Include="..\Horde.Server\Protos\horde\messages\aws_lambda.proto" ProtoRoot="..\Horde.Server\Protos" Access="Public" />
     <Protobuf Include="..\Horde.Server\Protos\horde\messages\telemetry.proto" ProtoRoot="..\Horde.Server\Protos" Access="Public" />
     <Protobuf Include="..\Horde.Server\Protos\horde\tasks\compute_task.proto" ProtoRoot="..\Horde.Server\Protos" Access="Public" />
-    <Protobuf Include="..\Horde.Server\Protos\horde\tasks\conform_task.proto" ProtoRoot="..\Horde.Server\Protos" Access="Public" />
-    <Protobuf Include="..\Horde.Server\Protos\horde\tasks\job_task.proto" ProtoRoot="..\Horde.Server\Protos" Access="Public" />
     <Protobuf Include="..\Horde.Server\Protos\horde\tasks\restart_task.proto" ProtoRoot="..\Horde.Server\Protos" Access="Public" />
     <Protobuf Include="..\Horde.Server\Protos\horde\tasks\shutdown_task.proto" ProtoRoot="..\Horde.Server\Protos" Access="Public" />
     <Protobuf Include="..\Horde.Server\Protos\horde\tasks\test_task.proto" ProtoRoot="..\Horde.Server\Protos" Access="Public" />
     <Protobuf Include="..\Horde.Server\Protos\horde\tasks\upgrade_task.proto" ProtoRoot="..\Horde.Server\Protos" Access="Public" />
   </ItemGroup>
-  -->
 
   <ItemGroup>
     <Compile Update="Properties\Resources.Designer.cs">
diff --git a/Engine/Source/Programs/Horde/Horde.Agent/Leases/Handlers/UpgradeHandler.cs b/Engine/Source/Programs/Horde/Horde.Agent/Leases/Handlers/UpgradeHandler.cs
--- a/Engine/Source/Programs/Horde/Horde.Agent/Leases/Handlers/UpgradeHandler.cs
+++ b/Engine/Source/Programs/Horde/Horde.Agent/Leases/Handlers/UpgradeHandler.cs
@@ -42,7 +42,7 @@ namespace Horde.Agent.Leases.Handlers
 				FileInfo outputFile = new FileInfo(Path.Combine(upgradeDir.FullName, "Agent.zip"));
 
 				HordeRpc.HordeRpcClient hordeRpc = await session.HordeClient.CreateGrpcClientAsync<HordeRpc.HordeRpcClient>(cancellationToken);
-				using (AsyncServerStreamingCall<RpcDownloadSoftwareResponse> cursor = hordeRpc.DownloadSoftware(new RpcDownloadSoftwareRequest(requiredVersion), null, null, cancellationToken))
+				using (AsyncServerStreamingCall<RpcDownloadSoftwareResponse> cursor = hordeRpc.DownloadSoftware(new RpcDownloadSoftwareRequest { Version = requiredVersion }, null, null, cancellationToken))
 				{
 					await using Stream outputStream = outputFile.Open(FileMode.Create);
 					while (await cursor.ResponseStream.MoveNext(cancellationToken))
diff --git a/Engine/Source/Programs/Horde/Horde.Server.Tests/Issues/IssueServiceTests.cs b/Engine/Source/Programs/Horde/Horde.Server.Tests/Issues/IssueServiceTests.cs
--- a/Engine/Source/Programs/Horde/Horde.Server.Tests/Issues/IssueServiceTests.cs
+++ b/Engine/Source/Programs/Horde/Horde.Server.Tests/Issues/IssueServiceTests.cs
@@ -30,7 +30,7 @@ using Horde.Server.Storage;
 using Horde.Server.Streams;
 using Horde.Server.Tests.Stubs.Services;
 using Horde.Server.Users;
-using HordeAgentDriver.Horde.Agent.Parser;
+using HordeAgentDriver.Horde.Agent.Driver.Parser;
 using Microsoft.Extensions.Logging;
 using Microsoft.Extensions.Logging.Abstractions;
 using Microsoft.VisualStudio.TestTools.UnitTesting;
diff --git a/Engine/Source/Programs/Horde/Horde.Server/Jobs/IJob.cs b/Engine/Source/Programs/Horde/Horde.Server/Jobs/IJob.cs
--- a/Engine/Source/Programs/Horde/Horde.Server/Jobs/IJob.cs
+++ b/Engine/Source/Programs/Horde/Horde.Server/Jobs/IJob.cs
@@ -19,13 +19,13 @@ using EpicGames.Horde.Jobs.Templates;
 using EpicGames.Horde.Logs;
 using EpicGames.Horde.Streams;
 using EpicGames.Horde.Users;
+using Horde.Common.Rpc;
 using Horde.Server.Acls;
 using Horde.Server.Jobs.Graphs;
 using Horde.Server.Jobs.Timing;
 using Horde.Server.Server;
 using Horde.Server.Streams;
 using Horde.Server.Ugs;
-using HordeCommon.Rpc;
 using MongoDB.Bson;
 using MongoDB.Driver;
 
diff --git a/Engine/Source/Programs/Horde/Horde.Server/Jobs/JobRpcService.cs b/Engine/Source/Programs/Horde/Horde.Server/Jobs/JobRpcService.cs
--- a/Engine/Source/Programs/Horde/Horde.Server/Jobs/JobRpcService.cs
+++ b/Engine/Source/Programs/Horde/Horde.Server/Jobs/JobRpcService.cs
@@ -13,6 +13,7 @@ using System.Threading;
 using System.Threading.Tasks;
 using EpicGames.Core;
 using EpicGames.Horde;
+using EpicGames.Horde.Agents;
 using EpicGames.Horde.Agents.Sessions;
 using EpicGames.Horde.Artifacts;
 using EpicGames.Horde.Jobs;
@@ -21,6 +22,8 @@ using Google.Protobuf.WellKnownTypes;
 using Grpc.Core;
 using Horde.Common.Rpc;
 using Horde.Server.Acls;
+using Horde.Server.Agents;
+using Horde.Server.Agents.Pools;
 using Horde.Server.Artifacts;
 using Horde.Server.Jobs.Graphs;
 using Horde.Server.Jobs.Templates;
@@ -28,8 +31,8 @@ using Horde.Server.Jobs.TestData;
 using Horde.Server.Logs;
 using Horde.Server.Server;
 using Horde.Server.Streams;
+using Horde.Server.Tasks;
 using Horde.Server.Utilities;
-using HordeCommon.Rpc;
 using Microsoft.AspNetCore.Authorization;
 using Microsoft.Extensions.Logging;
 using Microsoft.Extensions.Options;
@@ -46,6 +49,9 @@ namespace Horde.Server.Jobs
 	{
 		readonly AclService _aclService;
 		readonly JobService _jobService;
+		readonly AgentService _agentService;
+		readonly PoolService _poolService;
+		readonly ConformTaskSource _conformTaskSource;
 		readonly IArtifactCollection _artifactCollection;
 		readonly IJobCollection _jobCollection;
 		readonly ILogCollection _logCollection;
@@ -61,10 +67,13 @@ namespace Horde.Server.Jobs
 		/// <summary>
 		/// Constructor
 		/// </summary>
-		public JobRpcService(AclService aclService, JobService jobService, IArtifactCollection artifactCollection, IJobCollection jobCollection, ILogCollection logCollection, IGraphCollection graphs, ITestDataCollection testData, IJobStepRefCollection jobStepRefCollection, ITemplateCollection templateCollection, HttpClient httpClient, IClock clock, IOptionsSnapshot<GlobalConfig> globalConfig, ILogger<JobRpcService> logger)
+		public JobRpcService(AclService aclService, JobService jobService, AgentService agentService, PoolService poolService, ConformTaskSource conformTaskSource, IArtifactCollection artifactCollection, IJobCollection jobCollection, ILogCollection logCollection, IGraphCollection graphs, ITestDataCollection testData, IJobStepRefCollection jobStepRefCollection, ITemplateCollection templateCollection, HttpClient httpClient, IClock clock, IOptionsSnapshot<GlobalConfig> globalConfig, ILogger<JobRpcService> logger)
 		{
 			_aclService = aclService;
 			_jobService = jobService;
+			_agentService = agentService;
+			_poolService = poolService;
+			_conformTaskSource = conformTaskSource;
 			_jobCollection = jobCollection;
 			_artifactCollection = artifactCollection;
 			_logCollection = logCollection;
@@ -78,6 +87,39 @@ namespace Horde.Server.Jobs
 			_logger = logger;
 		}
 
+		/// <inheritdoc/>
+		public override async Task<RpcUpdateAgentWorkspacesResponse> UpdateAgentWorkspaces(RpcUpdateAgentWorkspacesRequest request, ServerCallContext context)
+		{
+			for (; ; )
+			{
+				// Get the current agent state
+				IAgent? agent = await _agentService.GetAgentAsync(new AgentId(request.AgentId));
+				if (agent == null)
+				{
+					throw new StructuredRpcException(StatusCode.OutOfRange, "Agent {AgentId} does not exist", request.AgentId);
+				}
+
+				// Get the new workspaces
+				List<AgentWorkspaceInfo> newWorkspaces = request.Workspaces.Select(x => new AgentWorkspaceInfo(x)).ToList();
+
+				// Get the set of workspaces that are currently required
+				HashSet<AgentWorkspaceInfo> conformWorkspaces = await _poolService.GetWorkspacesAsync(agent, DateTime.UtcNow, _globalConfig.Value, context.CancellationToken);
+				bool pendingConform = !conformWorkspaces.SetEquals(newWorkspaces) || (agent.RequestFullConform && !request.RemoveUntrackedFiles);
+
+				// Update the workspaces
+				if (await _agentService.TryUpdateWorkspacesAsync(agent, newWorkspaces, pendingConform, context.CancellationToken))
+				{
+					RpcUpdateAgentWorkspacesResponse response = new RpcUpdateAgentWorkspacesResponse();
+					if (pendingConform)
+					{
+						response.Retry = await _conformTaskSource.GetWorkspacesAsync(agent, response.PendingWorkspaces, context.CancellationToken);
+						response.RemoveUntrackedFiles = request.RemoveUntrackedFiles || agent.RequestFullConform;
+					}
+					return response;
+				}
+			}
+		}
+
 		/// <summary>
 		/// Gets information about a job
 		/// </summary>
diff --git a/Engine/Source/Programs/Horde/Horde.Server/Server/RpcService.cs b/Engine/Source/Programs/Horde/Horde.Server/Server/RpcService.cs
--- a/Engine/Source/Programs/Horde/Horde.Server/Server/RpcService.cs
+++ b/Engine/Source/Programs/Horde/Horde.Server/Server/RpcService.cs
@@ -17,12 +17,12 @@ using Google.Protobuf;
 using Google.Protobuf.Reflection;
 using Google.Protobuf.WellKnownTypes;
 using Grpc.Core;
+using Horde.Common.Rpc;
 using Horde.Server.Acls;
 using Horde.Server.Agents;
-using Horde.Server.Agents.Pools;
 using Horde.Server.Agents.Sessions;
 using Horde.Server.Agents.Telemetry;
-using Horde.Server.Tasks;
+using Horde.Server.Jobs;
 using Horde.Server.Telemetry;
 using Horde.Server.Tools;
 using Horde.Server.Utilities;
@@ -30,6 +30,7 @@ using HordeCommon.Rpc;
 using HordeCommon.Rpc.Messages;
 using HordeCommon.Rpc.Messages.Telemetry;
 using Microsoft.AspNetCore.Authorization;
+using Microsoft.Extensions.DependencyInjection;
 using Microsoft.Extensions.Logging;
 using Microsoft.Extensions.Options;
 
@@ -47,29 +48,27 @@ namespace Horde.Server.Server
 		internal TimeSpan _longPollTimeout = TimeSpan.FromMinutes(9);
 
 		readonly AgentService _agentService;
-		readonly PoolService _poolService;
 		readonly LifetimeService _lifetimeService;
 		readonly ITelemetrySink _telemetrySink;
-		readonly ConformTaskSource _conformTaskSource;
 		readonly IToolCollection _toolCollection;
 		readonly IAgentTelemetryCollection _agentTelemetryCollection;
 		readonly AclService _aclService;
+		readonly IServiceProvider _serviceProvider;
 		readonly IOptionsSnapshot<GlobalConfig> _globalConfig;
 		readonly ILogger _logger;
 
 		/// <summary>
 		/// Constructor
 		/// </summary>
-		public RpcService(AgentService agentService, PoolService poolService, LifetimeService lifetimeService, ITelemetrySink telemetrySink, ConformTaskSource conformTaskSource, IToolCollection toolCollection, IAgentTelemetryCollection agentTelemetryCollection, AclService aclService, IOptionsSnapshot<GlobalConfig> globalConfig, ILogger<RpcService> logger)
+		public RpcService(AgentService agentService, LifetimeService lifetimeService, ITelemetrySink telemetrySink, IToolCollection toolCollection, IAgentTelemetryCollection agentTelemetryCollection, AclService aclService, IServiceProvider serviceProvider, IOptionsSnapshot<GlobalConfig> globalConfig, ILogger<RpcService> logger)
 		{
 			_agentService = agentService;
-			_poolService = poolService;
 			_lifetimeService = lifetimeService;
 			_telemetrySink = telemetrySink;
-			_conformTaskSource = conformTaskSource;
 			_toolCollection = toolCollection;
 			_agentTelemetryCollection = agentTelemetryCollection;
 			_aclService = aclService;
+			_serviceProvider = serviceProvider;
 			_globalConfig = globalConfig;
 			_logger = logger;
 		}
@@ -163,42 +162,11 @@ namespace Horde.Server.Server
 			}
 		}
 
-		/// <summary>
-		/// Updates the workspaces synced for an agent
-		/// </summary>
-		/// <param name="request">The request parameters</param>
-		/// <param name="context">Context for the call</param>
-		/// <returns>Response object</returns>
-		public override async Task<RpcUpdateAgentWorkspacesResponse> UpdateAgentWorkspaces(RpcUpdateAgentWorkspacesRequest request, ServerCallContext context)
+		/// <inheritdoc/>
+		public override Task<RpcUpdateAgentWorkspacesResponse> UpdateAgentWorkspaces(RpcUpdateAgentWorkspacesRequest request, ServerCallContext context)
 		{
-			for (; ; )
-			{
-				// Get the current agent state
-				IAgent? agent = await _agentService.GetAgentAsync(new AgentId(request.AgentId));
-				if (agent == null)
-				{
-					throw new StructuredRpcException(StatusCode.OutOfRange, "Agent {AgentId} does not exist", request.AgentId);
-				}
-
-				// Get the new workspaces
-				List<AgentWorkspaceInfo> newWorkspaces = request.Workspaces.Select(x => new AgentWorkspaceInfo(x)).ToList();
-
-				// Get the set of workspaces that are currently required
-				HashSet<AgentWorkspaceInfo> conformWorkspaces = await _poolService.GetWorkspacesAsync(agent, DateTime.UtcNow, _globalConfig.Value, context.CancellationToken);
-				bool pendingConform = !conformWorkspaces.SetEquals(newWorkspaces) || (agent.RequestFullConform && !request.RemoveUntrackedFiles);
-
-				// Update the workspaces
-				if (await _agentService.TryUpdateWorkspacesAsync(agent, newWorkspaces, pendingConform, context.CancellationToken))
-				{
-					RpcUpdateAgentWorkspacesResponse response = new RpcUpdateAgentWorkspacesResponse();
-					if (pendingConform)
-					{
-						response.Retry = await _conformTaskSource.GetWorkspacesAsync(agent, response.PendingWorkspaces, context.CancellationToken);
-						response.RemoveUntrackedFiles = request.RemoveUntrackedFiles || agent.RequestFullConform;
-					}
-					return response;
-				}
-			}
+			JobRpcService jobRpcService = ActivatorUtilities.CreateInstance<JobRpcService>(_serviceProvider);
+			return jobRpcService.UpdateAgentWorkspaces(request, context);
 		}
 
 		static void CopyPropertyToResource(string name, List<string> properties, Dictionary<string, int> resources)
diff --git a/Engine/Source/Programs/Horde/Horde.Server/Utilities/ArtifactChunkStream.cs b/Engine/Source/Programs/Horde/Horde.Server/Utilities/ArtifactChunkStream.cs
--- a/Engine/Source/Programs/Horde/Horde.Server/Utilities/ArtifactChunkStream.cs
+++ b/Engine/Source/Programs/Horde/Horde.Server/Utilities/ArtifactChunkStream.cs
@@ -5,7 +5,7 @@ using System.IO;
 using System.Threading;
 using System.Threading.Tasks;
 using Grpc.Core;
-using HordeCommon.Rpc;
+using Horde.Common.Rpc;
 
 namespace Horde.Server.Utilities
 {
