Horde: Update release notes.

#rnx

diff --git a/Engine/Source/Programs/Horde/Docs/ReleaseNotes.md b/Engine/Source/Programs/Horde/Docs/ReleaseNotes.md
--- a/Engine/Source/Programs/Horde/Docs/ReleaseNotes.md
+++ b/Engine/Source/Programs/Horde/Docs/ReleaseNotes.md
@@ -1,5 +1,23 @@
 # Release Notes
 
+## 2024-05-20
+
+* Add grouping keys and platforms to tools. Intent is for the dashboard to only show one tool for each grouping key by default (the one with the closest matching platform based on the browser's user-agent), but users can manually expand the list of tools to show other platforms if desired.
+  Platforms should be NET RIDs, eg. win-x64, osx-arm64, etc... (https://learn.microsoft.com/en-us/dotnet/core/rid-catalog). (33698418)
+* Use native GrpcChannel instances for agent connection management, rather than IRpcConnection. (33660725, 33638856)
+* Add an AdditionalArguments property to jobs, which can be used to append arbitrary arguments to those derived from a job's parameters. This field is preserved - but not appended - to the arguments list if arguments are specified explicitly. Needs hooking up to the dashboard. (33630321)
+* Prevent default parameters being appended to jobs when an explicit argument list is specified. (33628649)
+* Move Perforce/BuildGraph functionality into a new Horde.Agent.Driver application. Intent is to separate this from the core Agent application over time, making it easier to iterate on Job/BuildGraph related functionality outside of the core agent deployment. (33606684)
+* Move settings for different executors into their own files. (33605105)
+* Trap exceptions when parsing invalid workflow ids from node annotations. (33602812)
+* Enable HTTP compression for server-sent responses (33600392)
+* Add support for agent queries via replica read (33575058)
+* Optimize agent assignment for compute tasks by filtering by pool (33488689)
+* Add command-based tracer for MongoDB. Old tracer operates at the collection level, this instead listens for events emitted by the Mongo client. (33479889)
+* Keep a cached list of current agents in AgentService (33454500)
+* Add an explicit error when the server URL does not have a valid scheme or host name. (33427628)
+* Move expired and ephemeral agent clean up to shared ticker (33424708)
+
 ## 2024-05-01
 
 * Explicitly check that the expiry time is set when searching for refs to expire, so that DocDB will use an indexed query. (33367777)
