Horde: Fix build warnings.

#rnx

diff --git a/Engine/Source/Programs/Horde/Horde.Server/Server/RpcService.cs b/Engine/Source/Programs/Horde/Horde.Server/Server/RpcService.cs
--- a/Engine/Source/Programs/Horde/Horde.Server/Server/RpcService.cs
+++ b/Engine/Source/Programs/Horde/Horde.Server/Server/RpcService.cs
@@ -80,6 +80,7 @@ namespace Horde.Server.Server
 		/// <param name="writer">Response writer</param>
 		/// <param name="context">Context for the call</param>
 		/// <returns>Response object</returns>
+		[Obsolete("RPC connection management now done natively")]
 		public override async Task QueryServerState(IAsyncStreamReader<RpcQueryServerStateRequest> reader, IServerStreamWriter<RpcQueryServerStateResponse> writer, ServerCallContext context)
 		{
 			if (await reader.MoveNext())
@@ -123,6 +124,7 @@ namespace Horde.Server.Server
 		/// <param name="writer">Response writer</param>
 		/// <param name="context">Context for the call</param>
 		/// <returns>Response object</returns>
+		[Obsolete("RPC connection management now done natively")]
 		public override async Task QueryServerStateV2(IAsyncStreamReader<RpcQueryServerStateRequest> reader, IServerStreamWriter<RpcQueryServerStateResponse> writer, ServerCallContext context)
 		{
 			if (await reader.MoveNext())
@@ -163,6 +165,7 @@ namespace Horde.Server.Server
 		}
 
 		/// <inheritdoc/>
+		[Obsolete("Use JobRpc instead")]
 		public override Task<RpcUpdateAgentWorkspacesResponse> UpdateAgentWorkspaces(RpcUpdateAgentWorkspacesRequest request, ServerCallContext context)
 		{
 			JobRpcService jobRpcService = ActivatorUtilities.CreateInstance<JobRpcService>(_serviceProvider);
