Horde: Fix another obsolete warning.

#rnx

diff --git a/Engine/Source/Programs/Horde/Horde.Server.Tests/Server/RpcServiceTest.cs b/Engine/Source/Programs/Horde/Horde.Server.Tests/Server/RpcServiceTest.cs
--- a/Engine/Source/Programs/Horde/Horde.Server.Tests/Server/RpcServiceTest.cs
+++ b/Engine/Source/Programs/Horde/Horde.Server.Tests/Server/RpcServiceTest.cs
@@ -401,28 +401,6 @@ namespace Horde.Server.Tests.Server
 			Assert.IsTrue(re.Message.Contains("Invalid agent name", StringComparison.OrdinalIgnoreCase));
 		}
 
-		[TestMethod]
-		public async Task QueryServerSessionTestAsync()
-		{
-			RpcService._longPollTimeout = TimeSpan.FromMilliseconds(200);
-
-			TestAsyncStreamReader<RpcQueryServerStateRequest> requestStream =
-				new TestAsyncStreamReader<RpcQueryServerStateRequest>(_adminContext);
-			TestServerStreamWriter<RpcQueryServerStateResponse> responseStream =
-				new TestServerStreamWriter<RpcQueryServerStateResponse>(_adminContext);
-			Task call = RpcService.QueryServerState(requestStream, responseStream, _adminContext);
-
-			requestStream.AddMessage(new RpcQueryServerStateRequest { Name = "bogusAgentName" });
-			RpcQueryServerStateResponse? res = await responseStream.ReadNextAsync();
-			Assert.IsNotNull(res);
-
-			res = await responseStream.ReadNextAsync();
-			Assert.IsNotNull(res);
-
-			// Should timeout after LongPollTimeout specified above
-			await call;
-		}
-
 		[TestMethod]
 		public async Task FinishBatchTestAsync()
 		{
