Horde: Tweak job argument display in new build advanced tab
#rnx

diff --git a/Engine/Source/Programs/Horde/HordeDashboard/src/components/build/NewBuildV2.tsx b/Engine/Source/Programs/Horde/HordeDashboard/src/components/build/NewBuildV2.tsx
--- a/Engine/Source/Programs/Horde/HordeDashboard/src/components/build/NewBuildV2.tsx
+++ b/Engine/Source/Programs/Horde/HordeDashboard/src/components/build/NewBuildV2.tsx
@@ -1281,10 +1281,11 @@ const AdvancedPanel: React.FC = observer(() => {
    let fargumentsHeight = 0;
 
    if (options.readOnly && a.length > 0) {
-      farguments = (JSON.stringify(a, null, 2));
+
+      farguments = a.join("\n")
       const flines = (farguments.match(/\n/g) || '').length + 1
       fargumentsHeight = Math.min(flines * 18, 240) + 16
-      height += (fargumentsHeight + 8)
+      height += (fargumentsHeight + 24)
    }
 
    let showAdditionalArgs = !options.readOnly || !!options.advAdditionalArgs;
@@ -1299,7 +1300,7 @@ const AdvancedPanel: React.FC = observer(() => {
       fparameters = (JSON.stringify(b, null, 2));
       const plines = (fparameters.match(/\n/g) || '').length + 1
       fparametersHeight = Math.min(plines * 18, 240) + 16
-      height += (fparametersHeight + 8)
+      height += (fparametersHeight + 24)
    }
 
    let showArgumentClipboardButton = !!options.jobDetails?.jobData?.arguments?.length;
