Horde: Guard against downloading an artifact with an invalid id
#rnx

diff --git a/Engine/Source/Programs/Horde/HordeDashboard/src/components/jobDetailsV2/JobDetailArtifacts.tsx b/Engine/Source/Programs/Horde/HordeDashboard/src/components/jobDetailsV2/JobDetailArtifacts.tsx
--- a/Engine/Source/Programs/Horde/HordeDashboard/src/components/jobDetailsV2/JobDetailArtifacts.tsx
+++ b/Engine/Source/Programs/Horde/HordeDashboard/src/components/jobDetailsV2/JobDetailArtifacts.tsx
@@ -149,7 +149,7 @@ export const JobArtifactsPanel: React.FC<{ jobDetails: JobDetailsV2 }> = observe
                />
                
                <PrimaryButton split text="Download as Zip" menuProps={downloadProps}
-                  disabled={!stepFinished}
+                  disabled={!stepFinished || !item.id}
                   style={{ fontFamily: "Horde Open Sans SemiBold" }}
                   onClick={() => {
                      window.location.assign(`/api/v2/artifacts/${item.id}/download?format=zip`);
@@ -235,5 +235,5 @@ export const JobArtifactsPanel: React.FC<{ jobDetails: JobDetailsV2 }> = observe
             </Stack>
          </Stack>
       </Stack>
-   </Stack>);
+   </Stack>)
 });
