Horde: Missing test files.

#rnx

diff --git a/Engine/Source/Programs/Horde/Horde.Agent.Driver.Tests/JobRpcClientStub.cs b/Engine/Source/Programs/Horde/Horde.Agent.Driver.Tests/JobRpcClientStub.cs
new file mode 100644
--- /dev/null
+++ b/Engine/Source/Programs/Horde/Horde.Agent.Driver.Tests/JobRpcClientStub.cs
@@ -0,0 +1,100 @@
+using System;
+using System.Collections.Generic;
+using System.Threading.Tasks;
+using Google.Protobuf.WellKnownTypes;
+using Grpc.Core;
+using Horde.Common.Rpc;
+using Microsoft.Extensions.Logging;
+
+namespace Horde.Agent.Driver.Tests
+{
+	class JobRpcClientStub : JobRpc.JobRpcClient
+	{
+		public readonly Queue<RpcBeginStepResponse> BeginStepResponses = new Queue<RpcBeginStepResponse>();
+		public readonly List<RpcUpdateStepRequest> UpdateStepRequests = new List<RpcUpdateStepRequest>();
+		public readonly Dictionary<RpcGetStepRequest, RpcGetStepResponse> GetStepResponses = new Dictionary<RpcGetStepRequest, RpcGetStepResponse>();
+		public Func<RpcGetStepRequest, RpcGetStepResponse>? _getStepFunc = null;
+		private readonly ILogger _logger;
+
+		public JobRpcClientStub(ILogger logger)
+		{
+			_logger = logger;
+		}
+
+		public override AsyncUnaryCall<RpcBeginBatchResponse> BeginBatchAsync(RpcBeginBatchRequest request,
+			CallOptions options)
+		{
+			_logger.LogDebug("HordeRpcClientStub.BeginBatchAsync()");
+			RpcBeginBatchResponse res = new RpcBeginBatchResponse();
+
+			res.AgentType = "agentType1";
+			res.LogId = "logId1";
+			res.Change = 1;
+
+			return Wrap(res);
+		}
+
+		public override AsyncUnaryCall<Empty> FinishBatchAsync(RpcFinishBatchRequest request, CallOptions options)
+		{
+			Empty res = new Empty();
+			return Wrap(res);
+		}
+
+		public override AsyncUnaryCall<RpcGetJobResponse> GetJobAsync(RpcGetJobRequest request, CallOptions options)
+		{
+			RpcGetJobResponse res = new RpcGetJobResponse();
+			return Wrap(res);
+		}
+
+		public override AsyncUnaryCall<RpcBeginStepResponse> BeginStepAsync(RpcBeginStepRequest request, CallOptions options)
+		{
+			if (BeginStepResponses.Count == 0)
+			{
+				RpcBeginStepResponse completeRes = new RpcBeginStepResponse();
+				completeRes.State = RpcBeginStepResponse.Types.Result.Complete;
+				return Wrap(completeRes);
+			}
+
+			RpcBeginStepResponse res = BeginStepResponses.Dequeue();
+			res.State = RpcBeginStepResponse.Types.Result.Ready;
+			return Wrap(res);
+		}
+
+		public override AsyncUnaryCall<Empty> UpdateStepAsync(RpcUpdateStepRequest request, CallOptions options)
+		{
+			_logger.LogDebug("UpdateStepAsync(Request: {Request})", request);
+			UpdateStepRequests.Add(request);
+			Empty res = new Empty();
+			return Wrap(res);
+		}
+
+		[Obsolete("Use LogRpc.CreateLogEvents instead")]
+		public override AsyncUnaryCall<Empty> CreateEventsAsync(RpcCreateLogEventsRequest request, CallOptions options)
+		{
+			_logger.LogDebug("CreateEventsAsync: {Request}", request);
+			Empty res = new Empty();
+			return Wrap(res);
+		}
+
+		public override AsyncUnaryCall<RpcGetStepResponse> GetStepAsync(RpcGetStepRequest request, CallOptions options)
+		{
+			if (_getStepFunc != null)
+			{
+				return Wrap(_getStepFunc(request));
+			}
+
+			if (GetStepResponses.TryGetValue(request, out RpcGetStepResponse? res))
+			{
+				return Wrap(res);
+			}
+
+			return Wrap(new RpcGetStepResponse());
+		}
+
+		public static AsyncUnaryCall<T> Wrap<T>(T res)
+		{
+			return new AsyncUnaryCall<T>(Task.FromResult(res), Task.FromResult(Metadata.Empty),
+				() => Status.DefaultSuccess, () => Metadata.Empty, null!);
+		}
+	}
+}
diff --git a/Engine/Source/Programs/Horde/Horde.Agent.Driver.Tests/SimpleTestExecutor.cs b/Engine/Source/Programs/Horde/Horde.Agent.Driver.Tests/SimpleTestExecutor.cs
new file mode 100644
--- /dev/null
+++ b/Engine/Source/Programs/Horde/Horde.Agent.Driver.Tests/SimpleTestExecutor.cs
@@ -0,0 +1,69 @@
+// Copyright Epic Games, Inc. All Rights Reserved.
+
+using System;
+using System.Threading;
+using System.Threading.Tasks;
+using EpicGames.Horde.Jobs;
+using Horde.Agent.Driver.Execution;
+using Horde.Common.Rpc;
+using HordeCommon.Rpc.Messages;
+using Microsoft.Extensions.Logging;
+using Microsoft.Extensions.Logging.Abstractions;
+
+namespace Horde.Agent.Driver.Tests
+{
+	class SimpleTestExecutor : JobExecutor
+	{
+		public const string Name = "Simple";
+
+		private readonly Func<JobStepInfo, ILogger, CancellationToken, Task<JobStepOutcome>> _func;
+
+		public SimpleTestExecutor(Func<JobStepInfo, ILogger, CancellationToken, Task<JobStepOutcome>> func)
+			: base(new JobExecutorOptions(null!, null!, null, default, default, default, default, null!), NullLogger.Instance)
+		{
+			_func = func;
+		}
+
+		public override Task InitializeAsync(RpcBeginBatchResponse batch, ILogger logger, CancellationToken cancellationToken)
+		{
+			logger.LogDebug("SimpleTestExecutor.InitializeAsync()");
+			return Task.CompletedTask;
+		}
+
+		public override Task<JobStepOutcome> RunAsync(JobStepInfo step, ILogger logger, CancellationToken cancellationToken)
+		{
+			logger.LogDebug("SimpleTestExecutor.RunAsync(Step: {Step})", step);
+			return _func(step, logger, cancellationToken);
+		}
+
+		public override Task FinalizeAsync(ILogger logger, CancellationToken cancellationToken)
+		{
+			logger.LogDebug("SimpleTestExecutor.FinalizeAsync()");
+			return Task.CompletedTask;
+		}
+
+		protected override Task<bool> SetupAsync(JobStepInfo step, ILogger logger, CancellationToken cancellationToken)
+		{
+			throw new NotImplementedException();
+		}
+
+		protected override Task<bool> ExecuteAsync(JobStepInfo step, ILogger logger, CancellationToken cancellationToken)
+		{
+			throw new NotImplementedException();
+		}
+	}
+
+	class SimpleTestExecutorFactory : IJobExecutorFactory
+	{
+		readonly JobExecutor _executor;
+
+		public string Name => SimpleTestExecutor.Name;
+
+		public SimpleTestExecutorFactory(JobExecutor executor)
+		{
+			_executor = executor;
+		}
+
+		public JobExecutor CreateExecutor(RpcAgentWorkspace? workspaceInfo, RpcAgentWorkspace? autoSdkWorkspaceInfo, JobExecutorOptions options) => _executor;
+	}
+}
