Horde: Missing copyright notice.

#rnx

diff --git a/Engine/Source/Programs/Horde/Horde.Agent.Driver.Tests/JobRpcClientStub.cs b/Engine/Source/Programs/Horde/Horde.Agent.Driver.Tests/JobRpcClientStub.cs
--- a/Engine/Source/Programs/Horde/Horde.Agent.Driver.Tests/JobRpcClientStub.cs
+++ b/Engine/Source/Programs/Horde/Horde.Agent.Driver.Tests/JobRpcClientStub.cs
@@ -1,3 +1,5 @@
+// Copyright Epic Games, Inc. All Rights Reserved.
+
 using System;
 using System.Collections.Generic;
 using System.Threading.Tasks;
