Horde: Move job driver into a separate executable.

diff --git a/Engine/Source/Programs/Horde/BuildHorde.xml b/Engine/Source/Programs/Horde/BuildHorde.xml
--- a/Engine/Source/Programs/Horde/BuildHorde.xml
+++ b/Engine/Source/Programs/Horde/BuildHorde.xml
@@ -120,8 +120,10 @@
 		<Node Name="Build HordeAgent" Produces="#HordeAgentFiles">
 			<Delete Files="$(ServerToolsStagingDir_HordeAgent)/..."/>
 			<DotNet Arguments="publish &quot;$(RootDir)/Engine/Source/Programs/Horde/Horde.Agent/Horde.Agent.csproj&quot; --output &quot;$(ServerToolsStagingDir_HordeAgent)&quot; $(VersionArguments)"/>
+			<DotNet Arguments="publish &quot;$(RootDir)/Engine/Source/Programs/Horde/Horde.Agent.Driver/Horde.Agent.Driver.csproj&quot; --output &quot;$(ServerToolsStagingDir_HordeAgent)/Driver&quot; $(VersionArguments)"/>
 			<Delete Files="$(ServerToolsStagingDir_HordeAgent)/appsettings.Local.json"/>
 			<Sign Files="$(ServerToolsStagingDir_HordeAgent)/HordeAgent.exe;$(ServerToolsStagingDir_HordeAgent)/HordeAgent.dll;$(ServerToolsStagingDir_HordeAgent)/EpicGames.*.dll" If="'$(SignOutputs)' == 'true'"/>
+			<Sign Files="$(ServerToolsStagingDir_HordeAgent)/Driver/HordeAgent.Driver.exe;$(ServerToolsStagingDir_HordeAgent)/Driver/HordeAgent.Driver.dll;$(ServerToolsStagingDir_HordeAgent)/Driver/EpicGames.*.dll" If="'$(SignOutputs)' == 'true'"/>
 			<Tag Files="$(ServerToolsStagingDir_HordeAgent)/..." With="#HordeAgentFiles"/>
 		</Node>
 		<Node Name="Build HordeAgent Win64" Produces="#HordeAgentWin64Files">
@@ -131,6 +133,10 @@
 			<DotNet Arguments="publish &quot;$(RootDir)/Engine/Source/Programs/Horde/Horde.Agent/Horde.Agent.csproj&quot; --output &quot;$(ServerToolsStagingDir_HordeAgentWin64)&quot; --runtime win7-x64 --self-contained $(VersionArguments)"/>
 			<Sign Files="$(ServerToolsStagingDir_HordeAgentWin64)/HordeAgent.exe;$(ServerToolsStagingDir_HordeAgentWin64)/HordeAgent.dll;$(ServerToolsStagingDir_HordeAgentWin64)/EpicGames.*.dll" If="'$(SignOutputs)' == 'true'"/>
 
+			<!-- Build the driver -->
+			<DotNet Arguments="publish &quot;$(RootDir)/Engine/Source/Programs/Horde/Horde.Agent.Driver/Horde.Agent.Driver.csproj&quot; --output &quot;$(ServerToolsStagingDir_HordeAgentWin64)/Driver&quot; --runtime win7-x64 --self-contained $(VersionArguments)"/>
+			<Sign Files="$(ServerToolsStagingDir_HordeAgentWin64)/Driver/HordeAgent.Driver.exe;$(ServerToolsStagingDir_HordeAgentWin64)/Driver/HordeAgent.Driver.dll;$(ServerToolsStagingDir_HordeAgentWin64)/Driver/EpicGames.*.dll" If="'$(SignOutputs)' == 'true'"/>
+
 			<!-- Build the tray app -->
 			<DotNet Arguments="publish &quot;$(RootDir)/Engine/Source/Programs/Horde/Horde.Agent.TrayApp/Horde.Agent.TrayApp.csproj&quot; --output &quot;$(ServerToolsStagingDir_HordeAgentWin64)/TrayApp&quot; --runtime win7-x64 --self-contained $(VersionArguments)"/>
 			<Sign Files="$(ServerToolsStagingDir_HordeAgentWin64)/TrayApp/Horde.Agent.TrayApp.exe;$(ServerToolsStagingDir_HordeAgentWin64)/TrayApp/EpicGames.*.dll" If="'$(SignOutputs)' == 'true'"/>
diff --git a/Engine/Source/Programs/Horde/Horde.Agent.Driver.Tests/Execution/WorkspaceExecutorTest.cs b/Engine/Source/Programs/Horde/Horde.Agent.Driver.Tests/Execution/WorkspaceExecutorTest.cs
--- a/Engine/Source/Programs/Horde/Horde.Agent.Driver.Tests/Execution/WorkspaceExecutorTest.cs
+++ b/Engine/Source/Programs/Horde/Horde.Agent.Driver.Tests/Execution/WorkspaceExecutorTest.cs
@@ -150,7 +150,7 @@ namespace Horde.Agent.Driver.Tests.Execution
 			_workspace.SetFile(1, "main.cpp", "main");
 			_workspace.SetFile(1, "foo/bar/baz.h", "baz");
 
-			JobExecutorOptions executorOptions = new JobExecutorOptions(_hordeClient, _workingDir, null, _jobId, _batchId, default, default, new RpcJobOptions());
+			JobExecutorOptions executorOptions = new JobExecutorOptions(_hordeClient, _workingDir, null, _jobId, _batchId, default, new RpcJobOptions());
 			_executor = new(executorOptions, _workspace, null, NullLogger.Instance);
 		}
 
@@ -175,7 +175,7 @@ namespace Horde.Agent.Driver.Tests.Execution
 		[TestMethod]
 		public async Task RegularAndAutoSdkWorkspaceAsync()
 		{
-			JobExecutorOptions executorOptions = new JobExecutorOptions(_hordeClient, _workingDir, Array.Empty<ProcessToTerminate>(), _jobId, _batchId, default, default, new RpcJobOptions());
+			JobExecutorOptions executorOptions = new JobExecutorOptions(_hordeClient, _workingDir, Array.Empty<ProcessToTerminate>(), _jobId, _batchId, default, new RpcJobOptions());
 			using WorkspaceExecutor executor = new(executorOptions, _workspace, _autoSdkWorkspace, NullLogger.Instance);
 
 			RpcBeginBatchResponse batch = new RpcBeginBatchResponse { Change = 1 };
@@ -189,7 +189,7 @@ namespace Horde.Agent.Driver.Tests.Execution
 		[TestMethod]
 		public async Task EnvVarsAsync()
 		{
-			JobExecutorOptions executorOptions = new JobExecutorOptions(_hordeClient, _workingDir, Array.Empty<ProcessToTerminate>(), _jobId, _batchId, default, default, new RpcJobOptions());
+			JobExecutorOptions executorOptions = new JobExecutorOptions(_hordeClient, _workingDir, Array.Empty<ProcessToTerminate>(), _jobId, _batchId, default, new RpcJobOptions());
 			using WorkspaceExecutor executor = new(executorOptions, _workspace, _autoSdkWorkspace, NullLogger.Instance);
 
 			RpcBeginBatchResponse batch = new RpcBeginBatchResponse { Change = 1, StreamName = "//UE5/Main" };
@@ -217,7 +217,7 @@ namespace Horde.Agent.Driver.Tests.Execution
 			_hordeClient.AddJob(preflightJobId, _streamId, 1, 1000);
 			_workspace.SetFile(1000, "New/Feature/Foo.cs", "foo");
 
-			JobExecutorOptions executorOptions = new JobExecutorOptions(_hordeClient, _workingDir, Array.Empty<ProcessToTerminate>(), preflightJobId, _batchId, default, default, new RpcJobOptions());
+			JobExecutorOptions executorOptions = new JobExecutorOptions(_hordeClient, _workingDir, Array.Empty<ProcessToTerminate>(), preflightJobId, _batchId, default, new RpcJobOptions());
 			using WorkspaceExecutor executor = new(executorOptions, _workspace, null, NullLogger.Instance);
 
 			RpcBeginBatchResponse batch = new RpcBeginBatchResponse { Change = 1, PreflightChange = 1000 };
@@ -235,7 +235,7 @@ namespace Horde.Agent.Driver.Tests.Execution
 
 			_hordeClient.AddJob(noChangeJobId, _streamId, 0, 0);
 
-			JobExecutorOptions executorOptions = new JobExecutorOptions(_hordeClient, _workingDir, Array.Empty<ProcessToTerminate>(), noChangeJobId, _batchId, default, default, new RpcJobOptions());
+			JobExecutorOptions executorOptions = new JobExecutorOptions(_hordeClient, _workingDir, Array.Empty<ProcessToTerminate>(), noChangeJobId, _batchId, default, new RpcJobOptions());
 			using WorkspaceExecutor executor = new(executorOptions, _workspace, null, NullLogger.Instance);
 
 			RpcBeginBatchResponse batch = new RpcBeginBatchResponse { };
diff --git a/Engine/Source/Programs/Horde/Horde.Agent.Driver.Tests/SimpleTestExecutor.cs b/Engine/Source/Programs/Horde/Horde.Agent.Driver.Tests/SimpleTestExecutor.cs
--- a/Engine/Source/Programs/Horde/Horde.Agent.Driver.Tests/SimpleTestExecutor.cs
+++ b/Engine/Source/Programs/Horde/Horde.Agent.Driver.Tests/SimpleTestExecutor.cs
@@ -19,7 +19,7 @@ namespace Horde.Agent.Driver.Tests
 		private readonly Func<JobStepInfo, ILogger, CancellationToken, Task<JobStepOutcome>> _func;
 
 		public SimpleTestExecutor(Func<JobStepInfo, ILogger, CancellationToken, Task<JobStepOutcome>> func)
-			: base(new JobExecutorOptions(null!, null!, null, default, default, default, default, null!), NullLogger.Instance)
+			: base(new JobExecutorOptions(null!, null!, null, default, default, default, null!), NullLogger.Instance)
 		{
 			_func = func;
 		}
diff --git a/Engine/Source/Programs/Horde/Horde.Agent.Driver/Commands/Execution/ExecuteConformCommand.cs b/Engine/Source/Programs/Horde/Horde.Agent.Driver/Commands/Execution/ExecuteConformCommand.cs
--- a/Engine/Source/Programs/Horde/Horde.Agent.Driver/Commands/Execution/ExecuteConformCommand.cs
+++ b/Engine/Source/Programs/Horde/Horde.Agent.Driver/Commands/Execution/ExecuteConformCommand.cs
@@ -41,6 +41,8 @@ namespace Horde.Agent.Driver.Commands.Execution
 
 		public override async Task<int> ExecuteAsync(ILogger logger)
 		{
+			logger.LogInformation("Running conform command in driver");
+
 			ConformTask conformTask = ConformTask.Parser.ParseFrom(Convert.FromBase64String(Task));
 
 			string accessToken = null!;
diff --git a/Engine/Source/Programs/Horde/Horde.Agent.Driver/Commands/Execution/ExecuteJobCommand.cs b/Engine/Source/Programs/Horde/Horde.Agent.Driver/Commands/Execution/ExecuteJobCommand.cs
--- a/Engine/Source/Programs/Horde/Horde.Agent.Driver/Commands/Execution/ExecuteJobCommand.cs
+++ b/Engine/Source/Programs/Horde/Horde.Agent.Driver/Commands/Execution/ExecuteJobCommand.cs
@@ -43,6 +43,8 @@ namespace Horde.Agent.Driver.Commands.Execution
 
 		public override async Task<int> ExecuteAsync(ILogger logger)
 		{
+			logger.LogInformation("Running job command in driver");
+
 			ExecuteJobTask executeTask = ExecuteJobTask.Parser.ParseFrom(Convert.FromBase64String(Task));
 
 			IHordeClient hordeClient = _hordeClientFactory.Create(executeTask.Token);
diff --git a/Engine/Source/Programs/Horde/Horde.Agent.Driver/DriverApp.cs b/Engine/Source/Programs/Horde/Horde.Agent.Driver/DriverApp.cs
--- a/Engine/Source/Programs/Horde/Horde.Agent.Driver/DriverApp.cs
+++ b/Engine/Source/Programs/Horde/Horde.Agent.Driver/DriverApp.cs
@@ -2,21 +2,67 @@
 
 using System.Reflection;
 using EpicGames.Core;
+using EpicGames.Horde;
+using Horde.Agent.Driver.Execution;
+using Microsoft.Extensions.Configuration;
 using Microsoft.Extensions.DependencyInjection;
+using Microsoft.Extensions.Logging;
+using Serilog;
+using Serilog.Extensions.Logging;
+using Serilog.Formatting.Json;
 
 namespace Horde.Agent.Driver
 {
 	class DriverApp
 	{
-		public static async Task<int> MainAsync(string[] args)
+		public static async Task<int> Main(string[] args)
 		{
 			CommandLineArguments arguments = new CommandLineArguments(args);
 
+			// Read the driver config
+			IConfiguration configuration = new ConfigurationBuilder()
+				.AddJsonFile("appsettings.json", optional: false)
+				.AddEnvironmentVariables()
+				.Build();
+
+			// Create the services 
 			IServiceCollection services = new ServiceCollection();
+			services.AddOptions<DriverSettings>().Configure(options => configuration.GetSection("Driver").Bind(options)).ValidateDataAnnotations();
+			services.AddLogging(builder => builder.AddProvider(CreateLoggerProvider(configuration)));
+			services.AddHorde(options => options.AllowAuthPrompt = false);
+
+			services.AddSingleton<IJobExecutorFactory, PerforceExecutorFactory>();
+			services.AddSingleton<IJobExecutorFactory, WorkspaceExecutorFactory>();
+			services.AddSingleton<IJobExecutorFactory, LocalExecutorFactory>();
+			services.AddSingleton<IJobExecutorFactory, TestExecutorFactory>();
+
+			services.AddSingleton<IWorkspaceMaterializerFactory, WorkspaceMaterializerFactory>();
+
 			services.AddCommandsFromAssembly(Assembly.GetExecutingAssembly());
 
+			// Run the host
 			await using ServiceProvider serviceProvider = services.BuildServiceProvider();
 			return await CommandHost.RunAsync(arguments, serviceProvider, null);
 		}
+
+		public static ILoggerProvider CreateLoggerProvider(IConfiguration configuration)
+		{
+			//			ConsoleTheme theme;
+			//			if (RuntimeInformation.IsOSPlatform(OSPlatform.Windows) && Environment.OSVersion.Version < new Version(10, 0))
+			//			{
+			//				theme = SystemConsoleTheme.Literate;
+			//			}
+			//			else
+			//			{
+			//				theme = AnsiConsoleTheme.Code;
+			//			}
+
+			LoggerConfiguration loggerConfiguration = new LoggerConfiguration()
+				.WriteTo.Console(new JsonFormatter(renderMessage: true))
+				.ReadFrom.Configuration(configuration)
+				.Enrich.FromLogContext();
+
+			return new SerilogLoggerProvider(loggerConfiguration.CreateLogger());
+		}
 	}
 }
diff --git a/Engine/Source/Programs/Horde/Horde.Agent.Driver/Execution/ConformExecutor.cs b/Engine/Source/Programs/Horde/Horde.Agent.Driver/Execution/ConformExecutor.cs
--- a/Engine/Source/Programs/Horde/Horde.Agent.Driver/Execution/ConformExecutor.cs
+++ b/Engine/Source/Programs/Horde/Horde.Agent.Driver/Execution/ConformExecutor.cs
@@ -4,7 +4,6 @@ using EpicGames.Core;
 using EpicGames.Horde;
 using EpicGames.Horde.Agents;
 using EpicGames.Horde.Agents.Leases;
-using EpicGames.Horde.Logs;
 using Horde.Agent.Driver.Utility;
 using Horde.Common.Rpc;
 using HordeCommon.Rpc.Messages;
@@ -21,9 +20,9 @@ namespace Horde.Agent.Driver.Execution
 		readonly LeaseId _leaseId;
 		readonly ConformTask _conformTask;
 		readonly DriverSettings _driverSettings;
-		readonly ILogger _localLogger;
+		readonly ILogger _logger;
 
-		public ConformExecutor(IHordeClient hordeClient, DirectoryReference workingDir, AgentId agentId, LeaseId leaseId, ConformTask conformTask, DriverSettings driverSettings, ILogger localLogger)
+		public ConformExecutor(IHordeClient hordeClient, DirectoryReference workingDir, AgentId agentId, LeaseId leaseId, ConformTask conformTask, DriverSettings driverSettings, ILogger logger)
 		{
 			_hordeClient = hordeClient;
 			_workingDir = workingDir;
@@ -31,27 +30,26 @@ namespace Horde.Agent.Driver.Execution
 			_leaseId = leaseId;
 			_conformTask = conformTask;
 			_driverSettings = driverSettings;
-			_localLogger = localLogger;
+			_logger = logger;
 		}
 
 		public async Task ExecuteAsync(CancellationToken cancellationToken)
 		{
-			await using IServerLogger serverLogger = _hordeClient.CreateServerLogger(LogId.Parse(_conformTask.LogId)).WithLocalLogger(_localLogger);
 			try
 			{
-				await ExecuteInternalAsync(serverLogger, cancellationToken);
+				await ExecuteInternalAsync(cancellationToken);
 			}
 			catch (Exception ex)
 			{
-				serverLogger.LogError(ex, "Unhandled exception while running conform: {Message}", ex.Message);
+				_logger.LogError(ex, "Unhandled exception while running conform: {Message}", ex.Message);
 				throw;
 			}
 		}
 
-		async Task ExecuteInternalAsync(ILogger logger, CancellationToken cancellationToken)
+		async Task ExecuteInternalAsync(CancellationToken cancellationToken)
 		{
-			logger.LogInformation("Conforming, lease {LeaseId}", _leaseId);
-			await TerminateProcessHelper.TerminateProcessesAsync(TerminateCondition.BeforeConform, _workingDir, _driverSettings.ProcessesToTerminate, logger, cancellationToken);
+			_logger.LogInformation("Conforming, lease {LeaseId}", _leaseId);
+			await TerminateProcessHelper.TerminateProcessesAsync(TerminateCondition.BeforeConform, _workingDir, _driverSettings.ProcessesToTerminate, _logger, cancellationToken);
 
 			bool removeUntrackedFiles = _conformTask.RemoveUntrackedFiles;
 			IList<RpcAgentWorkspace> pendingWorkspaces = _conformTask.Workspaces;
@@ -69,11 +67,11 @@ namespace Horde.Agent.Driver.Execution
 				// Run the conform task
 				if (isExecutorConformCompatible && _driverSettings.PerforceExecutor.RunConform)
 				{
-					await PerforceExecutor.ConformAsync(_workingDir, pendingWorkspaces, removeUntrackedFiles, logger, cancellationToken);
+					await PerforceExecutor.ConformAsync(_workingDir, pendingWorkspaces, removeUntrackedFiles, _logger, cancellationToken);
 				}
 				else
 				{
-					logger.LogInformation("Skipping conform. Executor={Executor} RunConform={RunConform}", _driverSettings.Executor, _driverSettings.PerforceExecutor.RunConform);
+					_logger.LogInformation("Skipping conform. Executor={Executor} RunConform={RunConform}", _driverSettings.Executor, _driverSettings.PerforceExecutor.RunConform);
 				}
 
 				// Update the new set of workspaces
@@ -87,11 +85,11 @@ namespace Horde.Agent.Driver.Execution
 				RpcUpdateAgentWorkspacesResponse response = await hordeRpc.UpdateAgentWorkspacesAsync(request, cancellationToken: cancellationToken);
 				if (!response.Retry)
 				{
-					logger.LogInformation("Conform finished");
+					_logger.LogInformation("Conform finished");
 					break;
 				}
 
-				logger.LogInformation("Pending workspaces have changed - running conform again...");
+				_logger.LogInformation("Pending workspaces have changed - running conform again...");
 				pendingWorkspaces = response.PendingWorkspaces;
 				removeUntrackedFiles = response.RemoveUntrackedFiles;
 			}
diff --git a/Engine/Source/Programs/Horde/Horde.Agent.Driver/Execution/JobExecutor.cs b/Engine/Source/Programs/Horde/Horde.Agent.Driver/Execution/JobExecutor.cs
--- a/Engine/Source/Programs/Horde/Horde.Agent.Driver/Execution/JobExecutor.cs
+++ b/Engine/Source/Programs/Horde/Horde.Agent.Driver/Execution/JobExecutor.cs
@@ -22,6 +22,7 @@ using HordeCommon;
 using HordeCommon.Rpc.Messages;
 using HordeCommon.Rpc.Tasks;
 using Microsoft.Extensions.Logging;
+using Microsoft.Extensions.Logging.Abstractions;
 using OpenTracing;
 using OpenTracing.Util;
 
@@ -56,10 +57,9 @@ namespace Horde.Agent.Driver.Execution
 		public JobId JobId { get; }
 		public JobStepBatchId BatchId { get; }
 		public LeaseId LeaseId { get; }
-		public LogId LogId { get; }
 		public RpcJobOptions JobOptions { get; }
 
-		public JobExecutorOptions(IHordeClient hordeClient, DirectoryReference workingDir, IReadOnlyList<ProcessToTerminate>? processesToTerminate, JobId jobId, JobStepBatchId batchId, LeaseId leaseId, LogId logId, RpcJobOptions jobOptions)
+		public JobExecutorOptions(IHordeClient hordeClient, DirectoryReference workingDir, IReadOnlyList<ProcessToTerminate>? processesToTerminate, JobId jobId, JobStepBatchId batchId, LeaseId leaseId, RpcJobOptions jobOptions)
 		{
 			HordeClient = hordeClient;
 			WorkingDir = workingDir;
@@ -67,7 +67,6 @@ namespace Horde.Agent.Driver.Execution
 			JobId = jobId;
 			BatchId = batchId;
 			LeaseId = leaseId;
-			LogId = logId;
 			JobOptions = jobOptions;
 		}
 	}
@@ -220,7 +219,6 @@ namespace Horde.Agent.Driver.Execution
 		protected IReadOnlyList<ProcessToTerminate>? ProcessesToTerminate { get; }
 
 		protected RpcBeginBatchResponse Batch { get; private set; }
-		private readonly LogId _logId;
 
 		protected List<string> _additionalArguments = new List<string>();
 		protected bool _allowTargetChanges;
@@ -252,7 +250,6 @@ namespace Horde.Agent.Driver.Execution
 			JobId = options.JobId;
 			BatchId = options.BatchId;
 			LeaseId = options.LeaseId;
-			_logId = options.LogId;
 			ProcessesToTerminate = options.ProcessesToTerminate;
 
 			Batch = null!; // Set in InitializeAsync()
@@ -274,19 +271,18 @@ namespace Horde.Agent.Driver.Execution
 		{
 		}
 
-		public async Task ExecuteAsync(ILogger localLogger, CancellationToken cancellationToken)
+		public async Task ExecuteAsync(ILogger logger, CancellationToken cancellationToken)
 		{
 			JobRpc.JobRpcClient jobRpc = await HordeClient.CreateGrpcClientAsync<JobRpc.JobRpcClient>(cancellationToken);
 
 			// Create a storage client for this session
-			await using IServerLogger logger = HordeClient.CreateServerLogger(_logId).WithLocalLogger(localLogger);
 			logger.LogInformation("Executing jobId {JobId}, batchId {BatchId}, leaseId {LeaseId} with executor {Name}", JobId, BatchId, LeaseId, JobOptions.Executor);
 
 			// Start executing the current batch
 			RpcBeginBatchResponse batch = await jobRpc.BeginBatchAsync(new RpcBeginBatchRequest(JobId, BatchId, LeaseId), cancellationToken: cancellationToken);
 			try
 			{
-				await ExecuteBatchAsync(batch, logger, localLogger, cancellationToken);
+				await ExecuteBatchAsync(batch, logger, NullLogger.Instance, cancellationToken);
 			}
 			catch (Exception ex)
 			{
@@ -2214,9 +2210,8 @@ namespace Horde.Agent.Driver.Execution
 
 			JobId jobId = JobId.Parse(executeTask.JobId);
 			JobStepBatchId batchId = JobStepBatchId.Parse(executeTask.BatchId);
-			LogId logId = LogId.Parse(executeTask.LogId);
 
-			JobExecutorOptions options = new JobExecutorOptions(hordeClient, workingDir, driverSettings.ProcessesToTerminate, jobId, batchId, leaseId, logId, executeTask.JobOptions);
+			JobExecutorOptions options = new JobExecutorOptions(hordeClient, workingDir, driverSettings.ProcessesToTerminate, jobId, batchId, leaseId, executeTask.JobOptions);
 
 			using JobExecutor executor = executorFactory.CreateExecutor(executeTask.Workspace, executeTask.AutoSdkWorkspace, options);
 			await executor.ExecuteAsync(logger, cancellationToken);
diff --git a/Engine/Source/Programs/Horde/Horde.Agent.Driver/Horde.Agent.Driver.csproj b/Engine/Source/Programs/Horde/Horde.Agent.Driver/Horde.Agent.Driver.csproj
--- a/Engine/Source/Programs/Horde/Horde.Agent.Driver/Horde.Agent.Driver.csproj
+++ b/Engine/Source/Programs/Horde/Horde.Agent.Driver/Horde.Agent.Driver.csproj
@@ -1,7 +1,7 @@
 <Project Sdk="Microsoft.NET.Sdk">
 
   <PropertyGroup>
-    <OutputType>Library</OutputType>
+    <OutputType>Exe</OutputType>
     <TargetFramework>net6.0</TargetFramework>
     <Nullable>enable</Nullable>
     <ImplicitUsings>enable</ImplicitUsings>
@@ -27,7 +27,7 @@
   </PropertyGroup>
 
   <ItemGroup>
-    <None Remove="HordeAgent.xml" />
+    <None Remove="HordeAgentDriver.xml" />
   </ItemGroup>
 
   <ItemGroup>
@@ -106,4 +106,12 @@
     <ProjectReference Include="..\..\Shared\EpicGames.Perforce\EpicGames.Perforce.csproj" />
   </ItemGroup>
 
+  <ItemGroup>
+    <Content Include="appsettings.json" Condition="Exists('appsettings.json')">
+      <Generator></Generator>
+      <CopyToOutputDirectory>Always</CopyToOutputDirectory>
+      <CopyToPublishDirectory>Always</CopyToPublishDirectory>
+    </Content>
+  </ItemGroup>
+
 </Project>
diff --git a/Engine/Source/Programs/Horde/Horde.Agent.Driver/appsettings.json b/Engine/Source/Programs/Horde/Horde.Agent.Driver/appsettings.json
new file mode 100644
--- /dev/null
+++ b/Engine/Source/Programs/Horde/Horde.Agent.Driver/appsettings.json
@@ -0,0 +1,23 @@
+{
+	"Driver":
+	{
+	},
+	"Serilog": {
+		"MinimumLevel": {
+			"Default": "Information",
+			"Override": {
+				"Microsoft": "Information",
+				"System.Net.Http.HttpClient": "Warning",
+				"Microsoft.AspNetCore.Routing.EndpointMiddleware": "Warning",
+				"Microsoft.AspNetCore.Authorization.DefaultAuthorizationService": "Warning",
+				"HordeServer.Authentication.HordeJwtBearerHandler": "Warning",
+				"HordeServer.Authentication.OktaHandler": "Warning",
+				"Microsoft.AspNetCore.Hosting.Diagnostics": "Warning",
+				"Microsoft.AspNetCore.Mvc.Infrastructure.ControllerActionInvoker": "Warning",
+				"Serilog.AspNetCore.RequestLoggingMiddleware": "Warning",
+				"Grpc.Net.Client.Internal.GrpcCall": "Warning",
+				"EpicGames.Horde.Storage": "Debug"
+			}
+		}
+	}
+}
diff --git a/Engine/Source/Programs/Horde/Horde.Agent.Tests/Services/TelemetryServiceTest.cs b/Engine/Source/Programs/Horde/Horde.Agent.Tests/Services/TelemetryServiceTest.cs
--- a/Engine/Source/Programs/Horde/Horde.Agent.Tests/Services/TelemetryServiceTest.cs
+++ b/Engine/Source/Programs/Horde/Horde.Agent.Tests/Services/TelemetryServiceTest.cs
@@ -4,8 +4,6 @@ using System;
 using System.Collections.Generic;
 using System.Threading;
 using System.Threading.Tasks;
-using Horde.Agent.Driver;
-using Horde.Agent.Driver.Execution;
 using Horde.Agent.Leases.Handlers;
 using Horde.Agent.Services;
 using Microsoft.Extensions.Logging.Abstractions;
@@ -26,7 +24,7 @@ public sealed class TelemetryServiceTest : IDisposable
 		OptionsWrapper<AgentSettings> settingsOpt = new(settings);
 
 		using WorkerService workerService = new WorkerService(null!, null!, null!, null!);
-		JobHandler jobHandler = new JobHandler(new List<IJobExecutorFactory>(), settingsOpt, new OptionsWrapper<DriverSettings>(new DriverSettings()));
+		JobHandler jobHandler = new JobHandler();
 		GrpcService grpcService = new(settingsOpt, NullLogger<GrpcService>.Instance, _loggerFactory);
 
 		_telemetryService = new TelemetryService(workerService, jobHandler, grpcService, null!, settingsOpt, NullLogger<TelemetryService>.Instance);
diff --git a/Engine/Source/Programs/Horde/Horde.Agent/AgentApp.cs b/Engine/Source/Programs/Horde/Horde.Agent/AgentApp.cs
--- a/Engine/Source/Programs/Horde/Horde.Agent/AgentApp.cs
+++ b/Engine/Source/Programs/Horde/Horde.Agent/AgentApp.cs
@@ -11,8 +11,6 @@ using Datadog.Trace.Configuration;
 using Datadog.Trace.OpenTracing;
 using EpicGames.Core;
 using EpicGames.Horde;
-using Horde.Agent.Driver;
-using Horde.Agent.Driver.Execution;
 using Horde.Agent.Leases;
 using Horde.Agent.Leases.Handlers;
 using Horde.Agent.Services;
@@ -158,7 +156,6 @@ namespace Horde.Agent
 			// Add all the default 
 			IConfigurationSection configSection = configuration.GetSection(AgentSettings.SectionName);
 			services.AddOptions<AgentSettings>().Configure(options => configSection.Bind(options)).ValidateDataAnnotations();
-			services.AddOptions<DriverSettings>().Configure(options => configSection.Bind(options)).ValidateDataAnnotations();
 
 			ServerProfile serverProfile = settings.GetCurrentServerProfile();
 			if (String.IsNullOrEmpty(serverProfile.Url.Scheme) || String.IsNullOrEmpty(serverProfile.Url.Host))
@@ -196,13 +193,6 @@ namespace Horde.Agent
 			services.AddSingleton<TelemetryService>();
 			services.AddHostedService(sp => sp.GetRequiredService<TelemetryService>());
 
-			services.AddSingleton<IJobExecutorFactory, PerforceExecutorFactory>();
-			services.AddSingleton<IJobExecutorFactory, WorkspaceExecutorFactory>();
-			services.AddSingleton<IJobExecutorFactory, LocalExecutorFactory>();
-			services.AddSingleton<IJobExecutorFactory, TestExecutorFactory>();
-
-			services.AddSingleton<IWorkspaceMaterializerFactory, WorkspaceMaterializerFactory>();
-
 			services.AddSingleton<JobHandler>();
 			services.AddSingleton<StatusService>();
 			services.AddHostedService<StatusService>(sp => sp.GetRequiredService<StatusService>());
diff --git a/Engine/Source/Programs/Horde/Horde.Agent/Horde.Agent.csproj b/Engine/Source/Programs/Horde/Horde.Agent/Horde.Agent.csproj
--- a/Engine/Source/Programs/Horde/Horde.Agent/Horde.Agent.csproj
+++ b/Engine/Source/Programs/Horde/Horde.Agent/Horde.Agent.csproj
@@ -134,7 +134,6 @@
     <ProjectReference Include="..\..\Shared\EpicGames.Horde\EpicGames.Horde.csproj" />
     <ProjectReference Include="..\..\Shared\EpicGames.Perforce.Managed\EpicGames.Perforce.Managed.csproj" />
     <ProjectReference Include="..\..\Shared\EpicGames.Perforce\EpicGames.Perforce.csproj" />
-    <ProjectReference Include="..\Horde.Agent.Driver\Horde.Agent.Driver.csproj" />
   </ItemGroup>
 
   <ItemGroup>
@@ -157,6 +156,12 @@
     <Protobuf Include="..\Horde.Server\Protos\horde\tasks\shutdown_task.proto" ProtoRoot="..\Horde.Server\Protos" Access="Public" />
     <Protobuf Include="..\Horde.Server\Protos\horde\tasks\test_task.proto" ProtoRoot="..\Horde.Server\Protos" Access="Public" />
     <Protobuf Include="..\Horde.Server\Protos\horde\tasks\upgrade_task.proto" ProtoRoot="..\Horde.Server\Protos" Access="Public" />
+
+    <Protobuf Include="..\Horde.Server\Protos\horde\job_rpc_messages.proto" ProtoRoot="..\Horde.Server\Protos" Access="Public" />
+    <Protobuf Include="..\Horde.Server\Protos\horde\common\label_change.proto" ProtoRoot="..\Horde.Server\Protos" Access="Public" />
+    <Protobuf Include="..\Horde.Server\Protos\horde\messages\agent_workspace.proto" ProtoRoot="..\Horde.Server\Protos" Access="Public" />
+    <Protobuf Include="..\Horde.Server\Protos\horde\tasks\conform_task.proto" ProtoRoot="..\Horde.Server\Protos" Access="Public" />
+    <Protobuf Include="..\Horde.Server\Protos\horde\tasks\job_task.proto" ProtoRoot="..\Horde.Server\Protos" Access="Public" />
   </ItemGroup>
 
   <ItemGroup>
diff --git a/Engine/Source/Programs/Horde/Horde.Agent/Leases/Handlers/ConformHandler.cs b/Engine/Source/Programs/Horde/Horde.Agent/Leases/Handlers/ConformHandler.cs
--- a/Engine/Source/Programs/Horde/Horde.Agent/Leases/Handlers/ConformHandler.cs
+++ b/Engine/Source/Programs/Horde/Horde.Agent/Leases/Handlers/ConformHandler.cs
@@ -1,34 +1,43 @@
 // Copyright Epic Games, Inc. All Rights Reserved.
 
+using System.Reflection;
+using EpicGames.Core;
+using EpicGames.Horde;
 using EpicGames.Horde.Agents.Leases;
-using Horde.Agent.Driver;
-using Horde.Agent.Driver.Execution;
+using EpicGames.Horde.Logs;
+using Google.Protobuf;
 using Horde.Agent.Services;
 using HordeCommon.Rpc.Tasks;
 using Microsoft.Extensions.Logging;
-using Microsoft.Extensions.Options;
 
 namespace Horde.Agent.Leases.Handlers
 {
 	class ConformHandler : LeaseHandler<ConformTask>
 	{
-		readonly DriverSettings _driverSettings;
-
-		/// <summary>
-		/// Constructor
-		/// </summary>
-		public ConformHandler(IOptions<DriverSettings> driverSettings)
-		{
-			_driverSettings = driverSettings.Value;
-		}
-
 		/// <inheritdoc/>
 		public override async Task<LeaseResult> ExecuteAsync(ISession session, LeaseId leaseId, ConformTask conformTask, ILogger localLogger, CancellationToken cancellationToken)
 		{
-			ConformExecutor executor = new ConformExecutor(session.HordeClient, session.WorkingDir, session.AgentId, leaseId, conformTask, _driverSettings, localLogger);
-			await executor.ExecuteAsync(cancellationToken);
+			await using IServerLogger serverLogger = session.HordeClient.CreateServerLogger(LogId.Parse(conformTask.LogId)).WithLocalLogger(localLogger);
+
+			List<string> arguments = new List<string>();
+			arguments.Add("execute");
+			arguments.Add("conform");
+			arguments.Add($"-AgentId={session.AgentId}");
+			arguments.Add($"-SessionId={session.SessionId}");
+			arguments.Add($"-LeaseId={leaseId}");
+			arguments.Add($"-WorkingDir={session.WorkingDir}");
+			arguments.Add($"-Task={Convert.ToBase64String(conformTask.ToByteArray())}");
+
+			FileReference driverAssembly = FileReference.Combine(new FileReference(Assembly.GetExecutingAssembly().Location).Directory, "Driver", "HordeAgent.Driver.dll");
+
+			Dictionary<string, string> environment = ManagedProcess.GetCurrentEnvVars();
+			environment[HordeHttpClient.HordeUrlEnvVarName] = session.HordeClient.ServerUrl.ToString();
+			environment[HordeHttpClient.HordeTokenEnvVarName] = await session.HordeClient.GetAccessTokenAsync(false, cancellationToken) ?? String.Empty;
+
+			int exitCode = await RunDotNetProcessAsync(driverAssembly, arguments, environment, false, serverLogger, cancellationToken);
+			serverLogger.LogInformation("Driver finished with exit code {ExitCode}", exitCode);
 
-			return LeaseResult.Success;
+			return (exitCode == 0) ? LeaseResult.Success : LeaseResult.Failed;
 		}
 	}
 }
diff --git a/Engine/Source/Programs/Horde/Horde.Agent/Leases/Handlers/JobHandler.cs b/Engine/Source/Programs/Horde/Horde.Agent/Leases/Handlers/JobHandler.cs
--- a/Engine/Source/Programs/Horde/Horde.Agent/Leases/Handlers/JobHandler.cs
+++ b/Engine/Source/Programs/Horde/Horde.Agent/Leases/Handlers/JobHandler.cs
@@ -1,17 +1,14 @@
 // Copyright Epic Games, Inc. All Rights Reserved.
 
-using System.Diagnostics;
 using System.Reflection;
 using EpicGames.Core;
 using EpicGames.Horde;
 using EpicGames.Horde.Agents.Leases;
+using EpicGames.Horde.Logs;
 using Google.Protobuf;
-using Horde.Agent.Driver;
-using Horde.Agent.Driver.Execution;
 using Horde.Agent.Services;
 using HordeCommon.Rpc.Tasks;
 using Microsoft.Extensions.Logging;
-using Microsoft.Extensions.Options;
 using OpenTracing.Util;
 
 namespace Horde.Agent.Leases.Handlers
@@ -33,20 +30,6 @@ namespace Horde.Agent.Leases.Handlers
 		/// </summary>
 		public string? CurrentBatchId { get; private set; } = null;
 
-		readonly IEnumerable<IJobExecutorFactory> _executorFactories;
-		readonly AgentSettings _agentSettings;
-		readonly DriverSettings _driverSettings;
-
-		/// <summary>
-		/// Constructor
-		/// </summary>
-		public JobHandler(IEnumerable<IJobExecutorFactory> executorFactories, IOptions<AgentSettings> agentSettings, IOptions<DriverSettings> driverSettings)
-		{
-			_executorFactories = executorFactories;
-			_agentSettings = agentSettings.Value;
-			_driverSettings = driverSettings.Value;
-		}
-
 		/// <inheritdoc/>
 		public override async Task<LeaseResult> ExecuteAsync(ISession session, LeaseId leaseId, ExecuteJobTask executeTask, ILogger localLogger, CancellationToken cancellationToken)
 		{
@@ -56,64 +39,33 @@ namespace Horde.Agent.Leases.Handlers
 				CurrentJobId = executeTask.JobId;
 				CurrentBatchId = executeTask.BatchId;
 
-				executeTask.JobOptions ??= new RpcJobOptions();
-
-				if (executeTask.JobOptions.RunInSeparateProcess ?? false)
-				{
-					if (AgentApp.IsSelfContained)
-					{
-						// TODO: Implement handling for invoking a self-contained agent process (i.e handle "dotnet" below)
-						throw new NotSupportedException("Running job in a separate process not supported for self-contained agents");
-					}
-
-					using (ManagedProcessGroup processGroup = new ManagedProcessGroup())
-					{
-						List<string> arguments = new List<string>();
-#pragma warning disable IL3000 // Avoid accessing Assembly file path when publishing as a single file
-						arguments.Add(Assembly.GetExecutingAssembly().Location);
-#pragma warning restore IL3000 // Avoid accessing Assembly file path when publishing as a single file						
-						arguments.Add("execute");
-						arguments.Add("job");
-						arguments.Add($"-Server={_agentSettings.GetCurrentServerProfile().Name}");
-						arguments.Add($"-AgentId={session.AgentId}");
-						arguments.Add($"-SessionId={session.SessionId}");
-						arguments.Add($"-LeaseId={leaseId}");
-						arguments.Add($"-WorkingDir={session.WorkingDir}");
-						arguments.Add($"-Task={Convert.ToBase64String(executeTask.ToByteArray())}");
+				GlobalTracer.Instance.ActiveSpan?.SetTag("jobId", executeTask.JobId.ToString());
+				GlobalTracer.Instance.ActiveSpan?.SetTag("jobName", executeTask.JobName.ToString());
+				GlobalTracer.Instance.ActiveSpan?.SetTag("batchId", executeTask.BatchId.ToString());
 
-						string commandLine = CommandLineArguments.Join(arguments);
-						localLogger.LogInformation("Running child process with arguments: {CommandLine}", commandLine);
+				await using IServerLogger logger = session.HordeClient.CreateServerLogger(LogId.Parse(executeTask.LogId)).WithLocalLogger(localLogger);
 
-						using (ManagedProcess process = new ManagedProcess(processGroup, "dotnet", commandLine, null,
-								   null, ProcessPriorityClass.Normal))
-						{
-							using (LogEventParser parser = new LogEventParser(localLogger))
-							{
-								for (; ; )
-								{
-									string? line = await process.ReadLineAsync(cancellationToken);
-									if (line == null)
-									{
-										break;
-									}
+				executeTask.JobOptions ??= new RpcJobOptions();
 
-									parser.WriteLine(line);
-								}
-							}
+				List<string> arguments = new List<string>();
+				arguments.Add("execute");
+				arguments.Add("job");
+				arguments.Add($"-AgentId={session.AgentId}");
+				arguments.Add($"-SessionId={session.SessionId}");
+				arguments.Add($"-LeaseId={leaseId}");
+				arguments.Add($"-WorkingDir={session.WorkingDir}");
+				arguments.Add($"-Task={Convert.ToBase64String(executeTask.ToByteArray())}");
 
-							await process.WaitForExitAsync(CancellationToken.None);
+				FileReference driverAssembly = FileReference.Combine(new FileReference(Assembly.GetExecutingAssembly().Location).Directory, "Driver", "HordeAgent.Driver.dll");
 
-							if (process.ExitCode != 0)
-							{
-								return LeaseResult.Failed;
-							}
-						}
-					}
+				Dictionary<string, string> environment = ManagedProcess.GetCurrentEnvVars();
+				environment[HordeHttpClient.HordeUrlEnvVarName] = session.HordeClient.ServerUrl.ToString();
+				environment[HordeHttpClient.HordeTokenEnvVarName] = executeTask.Token;
 
-					return LeaseResult.Success;
-				}
+				int exitCode = await RunDotNetProcessAsync(driverAssembly, arguments, environment, false, logger, cancellationToken);
+				logger.LogInformation("Driver finished with exit code {ExitCode}", exitCode);
 
-				return await ExecuteInternalAsync(session.HordeClient, session.WorkingDir, leaseId, executeTask, localLogger, cancellationToken);
+				return (exitCode == 0) ? LeaseResult.Success : LeaseResult.Failed;
 			}
 			finally
 			{
@@ -122,16 +74,6 @@ namespace Horde.Agent.Leases.Handlers
 				CurrentBatchId = null;
 			}
 		}
-
-		internal async Task<LeaseResult> ExecuteInternalAsync(IHordeClient hordeClient, DirectoryReference workingDir, LeaseId leaseId, ExecuteJobTask executeTask, ILogger localLogger, CancellationToken cancellationToken)
-		{
-			GlobalTracer.Instance.ActiveSpan?.SetTag("jobId", executeTask.JobId.ToString());
-			GlobalTracer.Instance.ActiveSpan?.SetTag("jobName", executeTask.JobName.ToString());
-			GlobalTracer.Instance.ActiveSpan?.SetTag("batchId", executeTask.BatchId.ToString());
-
-			await JobExecutorHelpers.ExecuteAsync(hordeClient, workingDir, leaseId, executeTask, _executorFactories, _driverSettings, localLogger, cancellationToken);
-			return LeaseResult.Success;
-		}
 	}
 }
 
diff --git a/Engine/Source/Programs/Horde/Horde.Agent/Leases/LeaseHandler.cs b/Engine/Source/Programs/Horde/Horde.Agent/Leases/LeaseHandler.cs
--- a/Engine/Source/Programs/Horde/Horde.Agent/Leases/LeaseHandler.cs
+++ b/Engine/Source/Programs/Horde/Horde.Agent/Leases/LeaseHandler.cs
@@ -1,5 +1,7 @@
 // Copyright Epic Games, Inc. All Rights Reserved.
 
+using System.Diagnostics;
+using EpicGames.Core;
 using EpicGames.Horde.Agents.Leases;
 using Google.Protobuf;
 using Google.Protobuf.Reflection;
@@ -24,6 +26,71 @@ namespace Horde.Agent.Leases
 		/// </summary>
 		/// <returns>Result for the lease</returns>
 		public abstract Task<LeaseResult> ExecuteAsync(ISession session, LeaseId leaseId, Any message, ILogger logger, CancellationToken cancellationToken);
+
+		/// <summary>
+		/// Runs a child process, piping the output to the given logger
+		/// </summary>
+		/// <param name="executable">Executable to launch</param>
+		/// <param name="arguments">Command line arguments for the new process</param>
+		/// <param name="environment">Environment for the new process</param>
+		/// <param name="logger">Logger for output</param>
+		/// <param name="cancellationToken">Cancellation token for the operation</param>
+		/// <returns>Exit code of the process</returns>
+		protected static async Task<int> RunProcessAsync(string executable, IEnumerable<string> arguments, IReadOnlyDictionary<string, string>? environment, ILogger logger, CancellationToken cancellationToken)
+		{
+			string commandLine = CommandLineArguments.Join(arguments);
+			logger.LogInformation("Running child process with arguments: {CommandLine}", commandLine);
+
+			using (ManagedProcessGroup processGroup = new ManagedProcessGroup())
+			using (ManagedProcess process = new ManagedProcess(processGroup, executable, commandLine, null, environment, ProcessPriorityClass.Normal))
+			{
+				for (; ; )
+				{
+					string? line = await process.ReadLineAsync(cancellationToken);
+					if (line == null)
+					{
+						break;
+					}
+
+					JsonLogEvent jsonLogEvent;
+					if (JsonLogEvent.TryParse(line, out jsonLogEvent))
+					{
+						logger.LogJsonLogEvent(jsonLogEvent);
+					}
+					else
+					{
+						logger.LogInformation("{Line}", line);
+					}
+				}
+
+				await process.WaitForExitAsync(CancellationToken.None);
+				return process.ExitCode;
+			}
+		}
+
+		/// <summary>
+		/// Runs a .NET assembly as a child process, piping the output to the given logger
+		/// </summary>
+		/// <param name="entryAssembly">Assembly to launch</param>
+		/// <param name="arguments">Command line arguments for the new process</param>
+		/// <param name="environment">Environment for the new process</param>
+		/// <param name="useNativeHost">Whether to use a native host process</param>
+		/// <param name="logger">Logger for output</param>
+		/// <param name="cancellationToken">Cancellation token for the operation</param>
+		/// <returns>Exit code of the process</returns>
+		protected static async Task<int> RunDotNetProcessAsync(FileReference entryAssembly, IEnumerable<string> arguments, IReadOnlyDictionary<string, string>? environment, bool useNativeHost, ILogger logger, CancellationToken cancellationToken)
+		{
+			if (useNativeHost)
+			{
+				FileReference nativeHost = entryAssembly.ChangeExtension(OperatingSystem.IsWindows() ? ".exe" : null);
+				return await RunProcessAsync(nativeHost.FullName, arguments, environment, logger, cancellationToken);
+			}
+			else
+			{
+				IEnumerable<string> allArguments = arguments.Prepend(entryAssembly.FullName);
+				return await RunProcessAsync("dotnet", allArguments, environment, logger, cancellationToken);
+			}
+		}
 	}
 
 	/// <summary>
diff --git a/Engine/Source/Programs/Shared/EpicGames.Core/JsonLogEvent.cs b/Engine/Source/Programs/Shared/EpicGames.Core/JsonLogEvent.cs
--- a/Engine/Source/Programs/Shared/EpicGames.Core/JsonLogEvent.cs
+++ b/Engine/Source/Programs/Shared/EpicGames.Core/JsonLogEvent.cs
@@ -90,7 +90,7 @@ namespace EpicGames.Core
 		}
 
 		/// <summary>
-		/// Tries to parse a Json log event from the given 
+		/// Tries to parse a Json log event from the given string
 		/// </summary>
 		/// <param name="data"></param>
 		/// <param name="logEvent"></param>
@@ -108,6 +108,18 @@ namespace EpicGames.Core
 			}
 		}
 
+		/// <summary>
+		/// Tries to parse a Json log event from the given string
+		/// </summary>
+		/// <param name="text">Text to parse</param>
+		/// <param name="logEvent"></param>
+		/// <returns></returns>
+		public static bool TryParse(string text, out JsonLogEvent logEvent)
+		{
+			byte[] data = Encoding.UTF8.GetBytes(text);
+			return TryParse(data, out logEvent);
+		}
+
 		static bool TryParseInternal(ReadOnlyMemory<byte> data, out JsonLogEvent logEvent)
 		{
 			LogLevel level = LogLevel.None;
