Horde: Add optionm to use legacy build dialog
#rnx

diff --git a/Engine/Source/Programs/Horde/HordeDashboard/src/components/build/NewBuildV2.tsx b/Engine/Source/Programs/Horde/HordeDashboard/src/components/build/NewBuildV2.tsx
--- a/Engine/Source/Programs/Horde/HordeDashboard/src/components/build/NewBuildV2.tsx
+++ b/Engine/Source/Programs/Horde/HordeDashboard/src/components/build/NewBuildV2.tsx
@@ -1,21 +1,21 @@
+import { Checkbox, ComboBox, ContextualMenuItemType, DefaultButton, DirectionalHint, Dropdown, DropdownMenuItemType, IComboBoxOption, IContextualMenuItem, IContextualMenuProps, IDropdownOption, Icon, IconButton, Label, MessageBar, MessageBarType, Modal, Pivot, PivotItem, PrimaryButton, ScrollablePane, ScrollbarVisibility, Spinner, SpinnerSize, Stack, TagPicker, Text, TextField, TooltipHost } from "@fluentui/react";
+import { useConst } from '@fluentui/react-hooks';
+import { ITextField } from "@fluentui/react/lib-commonjs/TextField";
+import Markdown from "markdown-to-jsx";
+import { action, makeObservable, observable } from "mobx";
 import { observer } from "mobx-react-lite";
-import { JobDetailsV2 } from "../jobDetailsV2/JobDetailsViewCommon";
+import moment from "moment";
+import React, { useEffect, useState } from "react";
 import backend, { useBackend } from "../../backend";
-import { Checkbox, ComboBox, ContextualMenuItemType, DefaultButton, DirectionalHint, Dropdown, DropdownMenuItemType, IComboBoxOption, IContextualMenuItem, IContextualMenuProps, IDropdownOption, Icon, IconButton, Label, MaskedTextField, MessageBar, MessageBarType, Modal, Pivot, PivotItem, PrimaryButton, ScrollablePane, ScrollbarVisibility, Spinner, SpinnerSize, Stack, TagPicker, Text, TextField, ThemeSettingName, TooltipHost } from "@fluentui/react";
-import { getHordeStyling } from "../../styles/Styles";
-import { useEffect, useState } from "react";
-import { action, makeObservable, observable } from "mobx";
-import { useConst } from '@fluentui/react-hooks';
-import { BoolParameterData, ChangeQueryConfig, CreateJobRequest, GetJobsTabResponse, GetTemplateRefResponse, GroupParameterData, JobsTabData, ListParameterData, ListParameterItemData, ListParameterStyle, ParameterData, ParameterType, Priority, StreamData, TabType, TextParameterData } from "../../backend/Api";
+import { BoolParameterData, ChangeQueryConfig, CreateJobRequest, GetJobsTabResponse, GetTemplateRefResponse, JobsTabData, ListParameterData, ListParameterItemData, ListParameterStyle, ParameterData, ParameterType, Priority, StreamData, TabType, TextParameterData } from "../../backend/Api";
+import dashboard from "../../backend/Dashboard";
 import { ProjectStore } from "../../backend/ProjectStore";
 import templateCache from "../../backend/TemplateCache";
-import ErrorHandler from "../ErrorHandler";
-import dashboard from "../../backend/Dashboard";
-import moment from "moment";
-import Markdown from "markdown-to-jsx";
-import React from "react";
-import { ITextField } from "@fluentui/react/lib-commonjs/TextField";
 import { copyToClipboard } from "../../base/utilities/clipboard";
+import { getHordeStyling } from "../../styles/Styles";
+import ErrorHandler from "../ErrorHandler";
+import { NewBuild } from "../NewBuild";
+import { JobDetailsV2 } from "../jobDetailsV2/JobDetailsViewCommon";
 
 type ValidationError = {
    paramString: string;
@@ -1322,7 +1322,7 @@ const AdvancedPanel: React.FC = observer(() => {
          <Stack>
             <Dropdown disabled={options.readOnly} key={"key_adv_priority"} defaultValue={options.advJobPriority} label="Priority" options={priorityOptions} onChange={(ev, option) => {
                options.advJobPriority = option?.key as Priority;
-            options.setChanged();
+               options.setChanged();
             }
             } />
          </Stack>
@@ -1354,7 +1354,7 @@ const AdvancedPanel: React.FC = observer(() => {
                   return arg;
                }
             }).join(" ")
-         )}/>}
+         )} />}
          {!!farguments && <TextField style={{ height: fargumentsHeight }} key={"key_adv_job_arguments"} defaultValue={farguments} readOnly={true} label="Job Arguments" multiline resizable={false} />}
          {!!fparameters && <TextField style={{ height: fparametersHeight }} key={"key_adv_job_parameters]"} defaultValue={fparameters} readOnly={true} label="Job Parameters" multiline resizable={false} />}
 
@@ -1375,7 +1375,7 @@ const AdvancedPanel: React.FC = observer(() => {
 
 })
 
-const BuildModal: React.FC = observer(() => {
+const BuildModal: React.FC<{ setUseLegacyDialog: (value: boolean) => void }> = observer(({ setUseLegacyDialog }) => {
 
    const options = BuildOptions.get();
    options.subscribe();
@@ -1432,7 +1432,7 @@ const BuildModal: React.FC = observer(() => {
          }
       }
 
-      let additionalArgs:string[] = [];
+      let additionalArgs: string[] = [];
       if (options!.advAdditionalArgs) {
          const argRegex = /"(\\"|[^"])*?"|[^ ]+/g;
          options.advAdditionalArgs.trim().match(argRegex)?.forEach(arg => additionalArgs.push(arg.replace(/"/g, "")));
@@ -1630,6 +1630,10 @@ const BuildModal: React.FC = observer(() => {
                   {options.mode === "Advanced" && <AdvancedPanel />}
                </Stack>
                <Stack horizontal tokens={{ childrenGap: 16 }} styles={{ root: { paddingTop: 32, paddingLeft: 8, paddingBottom: 8 } }}>
+                  <Stack>
+                     <DefaultButton disabled={ options.readOnly || !!options.jobDetails} text="Use Legacy Dialog" style={{ width: 160 }} onClick={() => setUseLegacyDialog(true)} />
+                  </Stack>
+
                   <Stack grow />
                   <PrimaryButton text="Start Job" disabled={!template || options.readOnly || options.preview} onClick={() => { onSubmit(); }} />
                   <DefaultButton text="Cancel" disabled={false} onClick={() => { options.onClose(); }} />
@@ -1642,14 +1646,14 @@ const BuildModal: React.FC = observer(() => {
 })
 
 
-const NewBuildV2Inner: React.FC = observer(() => {
+const NewBuildV2Inner: React.FC<{ setUseLegacyDialog: (value: boolean) => void }> = observer(({ setUseLegacyDialog }) => {
 
    const options = BuildOptions.get();
    options.subscribe();
 
    return <Stack>
       {!!options.validationErrors.length && <ValidationErrorModal errors={options.validationErrors} onClose={() => options.setValidationErrors([])} />}
-      {!options.submitting && <BuildModal />}
+      {!options.submitting && <BuildModal setUseLegacyDialog={setUseLegacyDialog} />}
       {!!options.submitting && <SubmittingModal />}
    </Stack>
 
@@ -1658,8 +1662,8 @@ const NewBuildV2Inner: React.FC = observer(() => {
 export const NewBuildV2: React.FC<{ streamId: string; show: boolean; onClose: (newJobId: string | undefined) => void, jobKey?: string; jobDetails?: JobDetailsV2, readOnly?: boolean }> = observer(({ streamId, jobKey, show, onClose, jobDetails, readOnly }) => {
 
    const { projectStore } = useBackend();
-
    const options = useConst(new BuildOptions(streamId, projectStore, onClose, jobDetails, jobKey, readOnly));
+   const [useNewBuildV1, setNewBuildV1] = useState<boolean>(false);
 
    useEffect(() => {
       return () => {
@@ -1667,9 +1671,12 @@ export const NewBuildV2: React.FC<{ streamId: string; show: boolean; onClose: (n
       };
    }, [options]);
 
+   if (useNewBuildV1) {
+      return <NewBuild show={true} streamId={streamId} onClose={onClose} jobKey={jobKey} jobDetails={jobDetails} readOnly={readOnly} />
+   }
 
    return <Stack>
-      <NewBuildV2Inner />
+      <NewBuildV2Inner setUseLegacyDialog={(value:boolean) => { setNewBuildV1(value)}}/>
    </Stack>
 })
 
