Horde: Prevent AWS status queries from causing session update failures.

diff --git a/Engine/Source/Programs/Horde/Horde.Agent/Services/CapabilitiesService.cs b/Engine/Source/Programs/Horde/Horde.Agent/Services/CapabilitiesService.cs
--- a/Engine/Source/Programs/Horde/Horde.Agent/Services/CapabilitiesService.cs
+++ b/Engine/Source/Programs/Horde/Horde.Agent/Services/CapabilitiesService.cs
@@ -48,21 +48,29 @@ namespace Horde.Agent.Services
 		/// <returns>Worker object for advertising to the server</returns>
 		public async Task<RpcAgentCapabilities> GetCapabilitiesAsync(DirectoryReference? workingDir)
 		{
-			_logger.LogInformation("Querying agent capabilities... (may take up to 30 seconds)");
-			Stopwatch timer = Stopwatch.StartNew();
-
-			Task<RpcAgentCapabilities> task = GetCapabilitiesInternalAsync(workingDir);
-			while (!task.IsCompleted)
+			try
 			{
-				Task delayTask = Task.Delay(TimeSpan.FromSeconds(30.0));
-				if (Task.WhenAny(task, delayTask) == delayTask)
+				_logger.LogInformation("Querying agent capabilities... (may take up to 30 seconds)");
+				Stopwatch timer = Stopwatch.StartNew();
+
+				Task<RpcAgentCapabilities> task = GetCapabilitiesInternalAsync(workingDir);
+				while (!task.IsCompleted)
 				{
-					_logger.LogWarning("GetCapabilitiesInternalAsync() has been running for {Time}", timer.Elapsed);
+					Task delayTask = Task.Delay(TimeSpan.FromSeconds(30.0));
+					if (Task.WhenAny(task, delayTask) == delayTask)
+					{
+						_logger.LogWarning("GetCapabilitiesInternalAsync() has been running for {Time}", timer.Elapsed);
+					}
 				}
-			}
-			_logger.LogInformation("Agent capabilities queried in {Time} ms", timer.ElapsedMilliseconds);
+				_logger.LogInformation("Agent capabilities queried in {Time} ms", timer.ElapsedMilliseconds);
 
-			return await task;
+				return await task;
+			}
+			catch (Exception ex)
+			{
+				_logger.LogError(ex, "Unable to query agent properties: {Message}", ex.Message);
+				throw;
+			}
 		}
 
 		async Task<RpcAgentCapabilities> GetCapabilitiesInternalAsync(DirectoryReference? workingDir)
@@ -561,7 +569,19 @@ namespace Horde.Agent.Services
 			return records;
 		}
 
-		static async Task AddAwsPropertiesAsync(IList<string> properties, ILogger logger)
+		async Task AddAwsPropertiesAsync(IList<string> properties, ILogger logger)
+		{
+			try
+			{
+				await AddAwsPropertiesInternalAsync(properties, logger);
+			}
+			catch (Exception ex)
+			{
+				_logger.LogDebug(ex, "Exception while querying EC2 metadata: {Message}", ex.Message);
+			}
+		}
+
+		static async Task AddAwsPropertiesInternalAsync(IList<string> properties, ILogger logger)
 		{
 			if (EC2InstanceMetadata.IdentityDocument != null)
 			{
