Horde: Clean up some naming post-move of agent driver.

#rnx

diff --git a/Engine/Source/Programs/Horde/Drivers/JobDriver.Tests/Execution/FakeWorkspaceMaterializer.cs b/Engine/Source/Programs/Horde/Drivers/JobDriver.Tests/Execution/FakeWorkspaceMaterializer.cs
--- a/Engine/Source/Programs/Horde/Drivers/JobDriver.Tests/Execution/FakeWorkspaceMaterializer.cs
+++ b/Engine/Source/Programs/Horde/Drivers/JobDriver.Tests/Execution/FakeWorkspaceMaterializer.cs
@@ -7,10 +7,10 @@ using System.Linq;
 using System.Threading;
 using System.Threading.Tasks;
 using EpicGames.Core;
-using Horde.Agent.Driver.Execution;
+using JobDriver.Execution;
 using Microsoft.Extensions.Logging;
 
-namespace Horde.Agent.Driver.Tests.Execution;
+namespace JobDriver.Tests.Execution;
 
 class FakeWorkspaceMaterializer : IWorkspaceMaterializer
 {
diff --git a/Engine/Source/Programs/Horde/Drivers/JobDriver.Tests/Execution/FakeWorkspaceMaterializerTest.cs b/Engine/Source/Programs/Horde/Drivers/JobDriver.Tests/Execution/FakeWorkspaceMaterializerTest.cs
--- a/Engine/Source/Programs/Horde/Drivers/JobDriver.Tests/Execution/FakeWorkspaceMaterializerTest.cs
+++ b/Engine/Source/Programs/Horde/Drivers/JobDriver.Tests/Execution/FakeWorkspaceMaterializerTest.cs
@@ -4,11 +4,11 @@ using System;
 using System.IO;
 using System.Threading;
 using System.Threading.Tasks;
-using Horde.Agent.Driver.Execution;
+using JobDriver.Execution;
 using Microsoft.Extensions.Logging.Abstractions;
 using Microsoft.VisualStudio.TestTools.UnitTesting;
 
-namespace Horde.Agent.Driver.Tests.Execution;
+namespace JobDriver.Tests.Execution;
 
 [TestClass]
 public sealed class FakeWorkspaceMaterializerTest : IDisposable
diff --git a/Engine/Source/Programs/Horde/Drivers/JobDriver.Tests/Execution/WorkspaceExecutorTest.cs b/Engine/Source/Programs/Horde/Drivers/JobDriver.Tests/Execution/WorkspaceExecutorTest.cs
--- a/Engine/Source/Programs/Horde/Drivers/JobDriver.Tests/Execution/WorkspaceExecutorTest.cs
+++ b/Engine/Source/Programs/Horde/Drivers/JobDriver.Tests/Execution/WorkspaceExecutorTest.cs
@@ -14,14 +14,14 @@ using EpicGames.Horde.Storage;
 using EpicGames.Horde.Streams;
 using Grpc.Core;
 using Grpc.Net.Client;
-using Horde.Agent.Driver.Execution;
+using JobDriver.Execution;
 using Horde.Common.Rpc;
 using HordeCommon.Rpc.Tasks;
 using Microsoft.Extensions.Logging;
 using Microsoft.Extensions.Logging.Abstractions;
 using Microsoft.VisualStudio.TestTools.UnitTesting;
 
-namespace Horde.Agent.Driver.Tests.Execution
+namespace JobDriver.Tests.Execution
 {
 	[TestClass]
 	public sealed class WorkspaceExecutorTest : IAsyncDisposable
diff --git a/Engine/Source/Programs/Horde/Drivers/JobDriver.Tests/JobRpcClientStub.cs b/Engine/Source/Programs/Horde/Drivers/JobDriver.Tests/JobRpcClientStub.cs
--- a/Engine/Source/Programs/Horde/Drivers/JobDriver.Tests/JobRpcClientStub.cs
+++ b/Engine/Source/Programs/Horde/Drivers/JobDriver.Tests/JobRpcClientStub.cs
@@ -8,7 +8,7 @@ using Grpc.Core;
 using Horde.Common.Rpc;
 using Microsoft.Extensions.Logging;
 
-namespace Horde.Agent.Driver.Tests
+namespace JobDriver.Tests
 {
 	class JobRpcClientStub : JobRpc.JobRpcClient
 	{
diff --git a/Engine/Source/Programs/Horde/Drivers/JobDriver.Tests/SimpleTestExecutor.cs b/Engine/Source/Programs/Horde/Drivers/JobDriver.Tests/SimpleTestExecutor.cs
--- a/Engine/Source/Programs/Horde/Drivers/JobDriver.Tests/SimpleTestExecutor.cs
+++ b/Engine/Source/Programs/Horde/Drivers/JobDriver.Tests/SimpleTestExecutor.cs
@@ -4,13 +4,13 @@ using System;
 using System.Threading;
 using System.Threading.Tasks;
 using EpicGames.Horde.Jobs;
-using Horde.Agent.Driver.Execution;
+using JobDriver.Execution;
 using Horde.Common.Rpc;
 using HordeCommon.Rpc.Messages;
 using Microsoft.Extensions.Logging;
 using Microsoft.Extensions.Logging.Abstractions;
 
-namespace Horde.Agent.Driver.Tests
+namespace JobDriver.Tests
 {
 	class SimpleTestExecutor : JobExecutor
 	{
diff --git a/Engine/Source/Programs/Horde/Drivers/JobDriver.Tests/WorkerServiceTest.cs b/Engine/Source/Programs/Horde/Drivers/JobDriver.Tests/WorkerServiceTest.cs
--- a/Engine/Source/Programs/Horde/Drivers/JobDriver.Tests/WorkerServiceTest.cs
+++ b/Engine/Source/Programs/Horde/Drivers/JobDriver.Tests/WorkerServiceTest.cs
@@ -11,7 +11,7 @@ using EpicGames.Horde.Storage.Backends;
 using EpicGames.Horde.Storage.Bundles;
 using EpicGames.Horde.Storage.Clients;
 using Grpc.Core;
-using Horde.Agent.Driver.Execution;
+using JobDriver.Execution;
 using Horde.Common.Rpc;
 using HordeCommon.Rpc.Messages;
 using HordeCommon.Rpc.Tasks;
@@ -20,7 +20,7 @@ using Microsoft.Extensions.Logging;
 using Microsoft.Extensions.Logging.Abstractions;
 using Microsoft.VisualStudio.TestTools.UnitTesting;
 
-namespace Horde.Agent.Driver.Tests
+namespace JobDriver.Tests
 {
 	[TestClass]
 	public class WorkerServiceTest
diff --git a/Engine/Source/Programs/Horde/Drivers/JobDriver.Tests/WorkspaceInfoTest.cs b/Engine/Source/Programs/Horde/Drivers/JobDriver.Tests/WorkspaceInfoTest.cs
--- a/Engine/Source/Programs/Horde/Drivers/JobDriver.Tests/WorkspaceInfoTest.cs
+++ b/Engine/Source/Programs/Horde/Drivers/JobDriver.Tests/WorkspaceInfoTest.cs
@@ -1,10 +1,10 @@
 // Copyright Epic Games, Inc. All Rights Reserved.
 
 using EpicGames.Perforce.Managed;
-using Horde.Agent.Driver.Utility;
+using JobDriver.Utility;
 using Microsoft.VisualStudio.TestTools.UnitTesting;
 
-namespace Horde.Agent.Driver.Tests
+namespace JobDriver.Tests
 {
 	using MWO = ManagedWorkspaceOptions;
 
diff --git a/Engine/Source/Programs/Horde/Drivers/JobDriver/Commands/Execution/ExecuteConformCommand.cs b/Engine/Source/Programs/Horde/Drivers/JobDriver/Commands/Execution/ExecuteConformCommand.cs
--- a/Engine/Source/Programs/Horde/Drivers/JobDriver/Commands/Execution/ExecuteConformCommand.cs
+++ b/Engine/Source/Programs/Horde/Drivers/JobDriver/Commands/Execution/ExecuteConformCommand.cs
@@ -5,12 +5,12 @@ using EpicGames.Horde;
 using EpicGames.Horde.Agents;
 using EpicGames.Horde.Agents.Leases;
 using EpicGames.Horde.Agents.Sessions;
-using Horde.Agent.Driver.Execution;
+using JobDriver.Execution;
 using HordeCommon.Rpc.Tasks;
 using Microsoft.Extensions.Logging;
 using Microsoft.Extensions.Options;
 
-namespace Horde.Agent.Driver.Commands.Execution
+namespace JobDriver.Commands.Execution
 {
 	[Command("Execute", "Conform", "Executes a conform")]
 	class ExecuteConformCommand : Command
diff --git a/Engine/Source/Programs/Horde/Drivers/JobDriver/Commands/Execution/ExecuteJobCommand.cs b/Engine/Source/Programs/Horde/Drivers/JobDriver/Commands/Execution/ExecuteJobCommand.cs
--- a/Engine/Source/Programs/Horde/Drivers/JobDriver/Commands/Execution/ExecuteJobCommand.cs
+++ b/Engine/Source/Programs/Horde/Drivers/JobDriver/Commands/Execution/ExecuteJobCommand.cs
@@ -5,12 +5,12 @@ using EpicGames.Horde;
 using EpicGames.Horde.Agents;
 using EpicGames.Horde.Agents.Leases;
 using EpicGames.Horde.Agents.Sessions;
-using Horde.Agent.Driver.Execution;
+using JobDriver.Execution;
 using HordeCommon.Rpc.Tasks;
 using Microsoft.Extensions.Logging;
 using Microsoft.Extensions.Options;
 
-namespace Horde.Agent.Driver.Commands.Execution
+namespace JobDriver.Commands.Execution
 {
 	[Command("Execute", "Job", "Executes a job")]
 	class ExecuteJobCommand : Command
diff --git a/Engine/Source/Programs/Horde/Drivers/JobDriver/Commands/Parse/ParseCommand.cs b/Engine/Source/Programs/Horde/Drivers/JobDriver/Commands/Parse/ParseCommand.cs
--- a/Engine/Source/Programs/Horde/Drivers/JobDriver/Commands/Parse/ParseCommand.cs
+++ b/Engine/Source/Programs/Horde/Drivers/JobDriver/Commands/Parse/ParseCommand.cs
@@ -3,10 +3,10 @@
 using System.ComponentModel;
 using EpicGames.Core;
 using EpicGames.Perforce;
-using Horde.Agent.Driver.Parser;
+using JobDriver.Parser;
 using Microsoft.Extensions.Logging;
 
-namespace Horde.Agent.Driver.Commands.Parse
+namespace JobDriver.Commands.Parse
 {
 	/// <summary>
 	/// Installs the agent as a service
diff --git a/Engine/Source/Programs/Horde/Drivers/JobDriver/Commands/Utilities/BuildGraphCommand.cs b/Engine/Source/Programs/Horde/Drivers/JobDriver/Commands/Utilities/BuildGraphCommand.cs
--- a/Engine/Source/Programs/Horde/Drivers/JobDriver/Commands/Utilities/BuildGraphCommand.cs
+++ b/Engine/Source/Programs/Horde/Drivers/JobDriver/Commands/Utilities/BuildGraphCommand.cs
@@ -2,10 +2,10 @@
 
 using System.Diagnostics;
 using EpicGames.Core;
-using Horde.Agent.Driver.Parser;
+using JobDriver.Parser;
 using Microsoft.Extensions.Logging;
 
-namespace Horde.Agent.Driver.Commands.Utilities
+namespace JobDriver.Commands.Utilities
 {
 	/// <summary>
 	/// Runs a BuildGraph script and captures the processed log output
diff --git a/Engine/Source/Programs/Horde/Drivers/JobDriver/DriverApp.cs b/Engine/Source/Programs/Horde/Drivers/JobDriver/DriverApp.cs
--- a/Engine/Source/Programs/Horde/Drivers/JobDriver/DriverApp.cs
+++ b/Engine/Source/Programs/Horde/Drivers/JobDriver/DriverApp.cs
@@ -3,7 +3,7 @@
 using System.Reflection;
 using EpicGames.Core;
 using EpicGames.Horde;
-using Horde.Agent.Driver.Execution;
+using JobDriver.Execution;
 using Microsoft.Extensions.Configuration;
 using Microsoft.Extensions.DependencyInjection;
 using Microsoft.Extensions.Logging;
@@ -11,7 +11,7 @@ using Serilog;
 using Serilog.Extensions.Logging;
 using Serilog.Formatting.Json;
 
-namespace Horde.Agent.Driver
+namespace JobDriver
 {
 	class DriverApp
 	{
diff --git a/Engine/Source/Programs/Horde/Drivers/JobDriver/DriverSettings.cs b/Engine/Source/Programs/Horde/Drivers/JobDriver/DriverSettings.cs
--- a/Engine/Source/Programs/Horde/Drivers/JobDriver/DriverSettings.cs
+++ b/Engine/Source/Programs/Horde/Drivers/JobDriver/DriverSettings.cs
@@ -1,9 +1,9 @@
 ﻿// Copyright Epic Games, Inc. All Rights Reserved.
 
-using Horde.Agent.Driver.Execution;
-using Horde.Agent.Driver.Utility;
+using JobDriver.Execution;
+using JobDriver.Utility;
 
-namespace Horde.Agent.Driver
+namespace JobDriver
 {
 	/// <summary>
 	/// Settings for the driver
diff --git a/Engine/Source/Programs/Horde/Drivers/JobDriver/Execution/ConformExecutor.cs b/Engine/Source/Programs/Horde/Drivers/JobDriver/Execution/ConformExecutor.cs
--- a/Engine/Source/Programs/Horde/Drivers/JobDriver/Execution/ConformExecutor.cs
+++ b/Engine/Source/Programs/Horde/Drivers/JobDriver/Execution/ConformExecutor.cs
@@ -4,13 +4,13 @@ using EpicGames.Core;
 using EpicGames.Horde;
 using EpicGames.Horde.Agents;
 using EpicGames.Horde.Agents.Leases;
-using Horde.Agent.Driver.Utility;
+using JobDriver.Utility;
 using Horde.Common.Rpc;
 using HordeCommon.Rpc.Messages;
 using HordeCommon.Rpc.Tasks;
 using Microsoft.Extensions.Logging;
 
-namespace Horde.Agent.Driver.Execution
+namespace JobDriver.Execution
 {
 	class ConformExecutor
 	{
diff --git a/Engine/Source/Programs/Horde/Drivers/JobDriver/Execution/JobExecutor.cs b/Engine/Source/Programs/Horde/Drivers/JobDriver/Execution/JobExecutor.cs
--- a/Engine/Source/Programs/Horde/Drivers/JobDriver/Execution/JobExecutor.cs
+++ b/Engine/Source/Programs/Horde/Drivers/JobDriver/Execution/JobExecutor.cs
@@ -15,8 +15,8 @@ using EpicGames.Horde.Logs;
 using EpicGames.Horde.Storage;
 using EpicGames.Horde.Storage.Nodes;
 using Grpc.Core;
-using Horde.Agent.Driver.Parser;
-using Horde.Agent.Driver.Utility;
+using JobDriver.Parser;
+using JobDriver.Utility;
 using Horde.Common.Rpc;
 using HordeCommon;
 using HordeCommon.Rpc.Messages;
@@ -26,7 +26,7 @@ using Microsoft.Extensions.Logging.Abstractions;
 using OpenTracing;
 using OpenTracing.Util;
 
-namespace Horde.Agent.Driver.Execution
+namespace JobDriver.Execution
 {
 	record class JobStepInfo
 	(
diff --git a/Engine/Source/Programs/Horde/Drivers/JobDriver/Execution/LocalExecutor.cs b/Engine/Source/Programs/Horde/Drivers/JobDriver/Execution/LocalExecutor.cs
--- a/Engine/Source/Programs/Horde/Drivers/JobDriver/Execution/LocalExecutor.cs
+++ b/Engine/Source/Programs/Horde/Drivers/JobDriver/Execution/LocalExecutor.cs
@@ -5,7 +5,7 @@ using HordeCommon.Rpc.Messages;
 using Microsoft.Extensions.Logging;
 using Microsoft.Extensions.Options;
 
-namespace Horde.Agent.Driver.Execution
+namespace JobDriver.Execution
 {
 	/// <summary>
 	/// Settings for the local executor
diff --git a/Engine/Source/Programs/Horde/Drivers/JobDriver/Execution/ManagedWorkspaceMaterializer.cs b/Engine/Source/Programs/Horde/Drivers/JobDriver/Execution/ManagedWorkspaceMaterializer.cs
--- a/Engine/Source/Programs/Horde/Drivers/JobDriver/Execution/ManagedWorkspaceMaterializer.cs
+++ b/Engine/Source/Programs/Horde/Drivers/JobDriver/Execution/ManagedWorkspaceMaterializer.cs
@@ -2,13 +2,13 @@
 
 using EpicGames.Core;
 using EpicGames.Perforce.Managed;
-using Horde.Agent.Driver.Utility;
+using JobDriver.Utility;
 using HordeCommon.Rpc.Messages;
 using Microsoft.Extensions.Logging;
 using OpenTracing;
 using OpenTracing.Util;
 
-namespace Horde.Agent.Driver.Execution;
+namespace JobDriver.Execution;
 
 /// <summary>
 /// Workspace materializer wrapping ManagedWorkspace and WorkspaceInfo
diff --git a/Engine/Source/Programs/Horde/Drivers/JobDriver/Execution/PerforceExecutor.cs b/Engine/Source/Programs/Horde/Drivers/JobDriver/Execution/PerforceExecutor.cs
--- a/Engine/Source/Programs/Horde/Drivers/JobDriver/Execution/PerforceExecutor.cs
+++ b/Engine/Source/Programs/Horde/Drivers/JobDriver/Execution/PerforceExecutor.cs
@@ -5,14 +5,14 @@ using System.Text;
 using EpicGames.Core;
 using EpicGames.Perforce;
 using EpicGames.Perforce.Managed;
-using Horde.Agent.Driver.Utility;
+using JobDriver.Utility;
 using Horde.Common.Rpc;
 using HordeCommon.Rpc.Messages;
 using Microsoft.Extensions.Logging;
 using OpenTracing;
 using OpenTracing.Util;
 
-namespace Horde.Agent.Driver.Execution
+namespace JobDriver.Execution
 {
 	/// <summary>
 	/// Settings for the perforce executor
diff --git a/Engine/Source/Programs/Horde/Drivers/JobDriver/Execution/TestExecutor.cs b/Engine/Source/Programs/Horde/Drivers/JobDriver/Execution/TestExecutor.cs
--- a/Engine/Source/Programs/Horde/Drivers/JobDriver/Execution/TestExecutor.cs
+++ b/Engine/Source/Programs/Horde/Drivers/JobDriver/Execution/TestExecutor.cs
@@ -4,12 +4,12 @@ using System.Diagnostics.CodeAnalysis;
 using EpicGames.Core;
 using EpicGames.Horde.Jobs;
 using EpicGames.Perforce;
-using Horde.Agent.Driver.Parser;
+using JobDriver.Parser;
 using Horde.Common.Rpc;
 using HordeCommon.Rpc.Messages;
 using Microsoft.Extensions.Logging;
 
-namespace Horde.Agent.Driver.Execution
+namespace JobDriver.Execution
 {
 	class TestExecutor : JobExecutor
 	{
diff --git a/Engine/Source/Programs/Horde/Drivers/JobDriver/Execution/WorkspaceExecutor.cs b/Engine/Source/Programs/Horde/Drivers/JobDriver/Execution/WorkspaceExecutor.cs
--- a/Engine/Source/Programs/Horde/Drivers/JobDriver/Execution/WorkspaceExecutor.cs
+++ b/Engine/Source/Programs/Horde/Drivers/JobDriver/Execution/WorkspaceExecutor.cs
@@ -2,14 +2,14 @@
 
 using System.Text.RegularExpressions;
 using EpicGames.Core;
-using Horde.Agent.Driver.Utility;
+using JobDriver.Utility;
 using Horde.Common.Rpc;
 using HordeCommon.Rpc.Messages;
 using Microsoft.Extensions.Logging;
 using OpenTracing;
 using OpenTracing.Util;
 
-namespace Horde.Agent.Driver.Execution
+namespace JobDriver.Execution
 {
 	class WorkspaceExecutor : JobExecutor
 	{
diff --git a/Engine/Source/Programs/Horde/Drivers/JobDriver/Execution/WorkspaceMaterializer.cs b/Engine/Source/Programs/Horde/Drivers/JobDriver/Execution/WorkspaceMaterializer.cs
--- a/Engine/Source/Programs/Horde/Drivers/JobDriver/Execution/WorkspaceMaterializer.cs
+++ b/Engine/Source/Programs/Horde/Drivers/JobDriver/Execution/WorkspaceMaterializer.cs
@@ -4,7 +4,7 @@ using EpicGames.Core;
 using HordeCommon.Rpc.Messages;
 using Microsoft.Extensions.Logging;
 
-namespace Horde.Agent.Driver.Execution;
+namespace JobDriver.Execution;
 
 /// <summary>
 /// Exception for workspace materializer
@@ -138,7 +138,7 @@ public interface IWorkspaceMaterializer : IDisposable
 	/// <param name="preflightChangeNum">Preflight change number to add</param>
 	/// <param name="options">Additional options</param>
 	/// <param name="cancellationToken">Cancellation token for the call</param>
-	/// <exception cref="Horde.Agent.Driver.Execution.WorkspaceMaterializationException">Thrown if syncing fails</exception>
+	/// <exception cref="JobDriver.Execution.WorkspaceMaterializationException">Thrown if syncing fails</exception>
 	/// <returns>Async task</returns>
 	public Task SyncAsync(int changeNum, int preflightChangeNum, SyncOptions options, CancellationToken cancellationToken);
 }
diff --git a/Engine/Source/Programs/Horde/Drivers/JobDriver/Parser/LogParser.cs b/Engine/Source/Programs/Horde/Drivers/JobDriver/Parser/LogParser.cs
--- a/Engine/Source/Programs/Horde/Drivers/JobDriver/Parser/LogParser.cs
+++ b/Engine/Source/Programs/Horde/Drivers/JobDriver/Parser/LogParser.cs
@@ -5,7 +5,7 @@ using System.Text.RegularExpressions;
 using EpicGames.Core;
 using Microsoft.Extensions.Logging;
 
-namespace Horde.Agent.Driver.Parser
+namespace JobDriver.Parser
 {
 	/// <summary>
 	/// Turns raw text output into structured logging events
diff --git a/Engine/Source/Programs/Horde/Drivers/JobDriver/Utility/JobStepLogger.cs b/Engine/Source/Programs/Horde/Drivers/JobDriver/Utility/JobStepLogger.cs
--- a/Engine/Source/Programs/Horde/Drivers/JobDriver/Utility/JobStepLogger.cs
+++ b/Engine/Source/Programs/Horde/Drivers/JobDriver/Utility/JobStepLogger.cs
@@ -6,7 +6,7 @@ using EpicGames.Horde.Logs;
 using Horde.Common.Rpc;
 using Microsoft.Extensions.Logging;
 
-namespace Horde.Agent.Driver.Utility
+namespace JobDriver.Utility
 {
 	/// <summary>
 	/// Logger which updates the job step outcome whenever a warning or error is detected.
diff --git a/Engine/Source/Programs/Horde/Drivers/JobDriver/Utility/LinuxInterop.cs b/Engine/Source/Programs/Horde/Drivers/JobDriver/Utility/LinuxInterop.cs
--- a/Engine/Source/Programs/Horde/Drivers/JobDriver/Utility/LinuxInterop.cs
+++ b/Engine/Source/Programs/Horde/Drivers/JobDriver/Utility/LinuxInterop.cs
@@ -2,7 +2,7 @@
 
 using System.Runtime.InteropServices;
 
-namespace Horde.Agent.Driver.Utility
+namespace JobDriver.Utility
 {
 	/// <summary>
 	/// Linux specific P/Invoke calls
diff --git a/Engine/Source/Programs/Horde/Drivers/JobDriver/Utility/TempStorage.cs b/Engine/Source/Programs/Horde/Drivers/JobDriver/Utility/TempStorage.cs
--- a/Engine/Source/Programs/Horde/Drivers/JobDriver/Utility/TempStorage.cs
+++ b/Engine/Source/Programs/Horde/Drivers/JobDriver/Utility/TempStorage.cs
@@ -16,7 +16,7 @@ using Microsoft.Extensions.Logging;
 
 #pragma warning disable CA1819 // Properties should not return arrays
 
-namespace Horde.Agent.Driver.Utility
+namespace JobDriver.Utility
 {
 	/// <summary>
 	/// Exception thrown by the temp storage system
diff --git a/Engine/Source/Programs/Horde/Drivers/JobDriver/Utility/TerminateProcessHelper.cs b/Engine/Source/Programs/Horde/Drivers/JobDriver/Utility/TerminateProcessHelper.cs
--- a/Engine/Source/Programs/Horde/Drivers/JobDriver/Utility/TerminateProcessHelper.cs
+++ b/Engine/Source/Programs/Horde/Drivers/JobDriver/Utility/TerminateProcessHelper.cs
@@ -3,7 +3,7 @@
 using EpicGames.Core;
 using Microsoft.Extensions.Logging;
 
-namespace Horde.Agent.Driver.Utility
+namespace JobDriver.Utility
 {
 	/// <summary>
 	/// Flags for processes to terminate
diff --git a/Engine/Source/Programs/Horde/Drivers/JobDriver/Utility/WorkspaceInfo.cs b/Engine/Source/Programs/Horde/Drivers/JobDriver/Utility/WorkspaceInfo.cs
--- a/Engine/Source/Programs/Horde/Drivers/JobDriver/Utility/WorkspaceInfo.cs
+++ b/Engine/Source/Programs/Horde/Drivers/JobDriver/Utility/WorkspaceInfo.cs
@@ -11,7 +11,7 @@ using EpicGames.Perforce.Managed;
 using HordeCommon.Rpc.Messages;
 using Microsoft.Extensions.Logging;
 
-namespace Horde.Agent.Driver.Utility
+namespace JobDriver.Utility
 {
 	/// <summary>
 	/// Stores information about a managed Perforce workspace
diff --git a/Engine/Source/Programs/Horde/Horde.Agent/Leases/Handlers/ConformHandler.cs b/Engine/Source/Programs/Horde/Horde.Agent/Leases/Handlers/ConformHandler.cs
--- a/Engine/Source/Programs/Horde/Horde.Agent/Leases/Handlers/ConformHandler.cs
+++ b/Engine/Source/Programs/Horde/Horde.Agent/Leases/Handlers/ConformHandler.cs
@@ -28,7 +28,7 @@ namespace Horde.Agent.Leases.Handlers
 			arguments.Add($"-WorkingDir={session.WorkingDir}");
 			arguments.Add($"-Task={Convert.ToBase64String(conformTask.ToByteArray())}");
 
-			FileReference driverAssembly = FileReference.Combine(new FileReference(Assembly.GetExecutingAssembly().Location).Directory, "JobDriver", "JobDriver.dll");
+			FileReference driverAssembly = FileReference.Combine(new DirectoryReference(AppContext.BaseDirectory), "JobDriver", "JobDriver.dll");
 
 			Dictionary<string, string> environment = ManagedProcess.GetCurrentEnvVars();
 			environment[HordeHttpClient.HordeUrlEnvVarName] = session.HordeClient.ServerUrl.ToString();
diff --git a/Engine/Source/Programs/Horde/Horde.Agent/Leases/Handlers/JobHandler.cs b/Engine/Source/Programs/Horde/Horde.Agent/Leases/Handlers/JobHandler.cs
--- a/Engine/Source/Programs/Horde/Horde.Agent/Leases/Handlers/JobHandler.cs
+++ b/Engine/Source/Programs/Horde/Horde.Agent/Leases/Handlers/JobHandler.cs
@@ -56,7 +56,7 @@ namespace Horde.Agent.Leases.Handlers
 				arguments.Add($"-WorkingDir={session.WorkingDir}");
 				arguments.Add($"-Task={Convert.ToBase64String(executeTask.ToByteArray())}");
 
-				FileReference driverAssembly = FileReference.Combine(new FileReference(Assembly.GetExecutingAssembly().Location).Directory, "JobDriver", "JobDriver.dll");
+				FileReference driverAssembly = FileReference.Combine(new DirectoryReference(AppContext.BaseDirectory), "JobDriver", "JobDriver.dll");
 
 				Dictionary<string, string> environment = ManagedProcess.GetCurrentEnvVars();
 				environment[HordeHttpClient.HordeUrlEnvVarName] = session.HordeClient.ServerUrl.ToString();
diff --git a/Engine/Source/Programs/Horde/Horde.Server.Tests/Horde.Server.Tests.csproj b/Engine/Source/Programs/Horde/Horde.Server.Tests/Horde.Server.Tests.csproj
--- a/Engine/Source/Programs/Horde/Horde.Server.Tests/Horde.Server.Tests.csproj
+++ b/Engine/Source/Programs/Horde/Horde.Server.Tests/Horde.Server.Tests.csproj
@@ -36,8 +36,8 @@
     <ProjectReference Include="..\Horde.Agent\Horde.Agent.csproj">
       <Aliases>HordeAgent</Aliases>
     </ProjectReference>
-    <ProjectReference Include="..\Horde.Agent.Driver\Horde.Agent.Driver.csproj">
-      <Aliases>HordeAgentDriver</Aliases>
+    <ProjectReference Include="..\Drivers\JobDriver\JobDriver.csproj">
+      <Aliases>JobDriver</Aliases>
     </ProjectReference>
     <ProjectReference Include="..\Horde.Server\Horde.Server.csproj" />
   </ItemGroup>
diff --git a/Engine/Source/Programs/Horde/Horde.Server.Tests/Issues/IssueServiceTests.cs b/Engine/Source/Programs/Horde/Horde.Server.Tests/Issues/IssueServiceTests.cs
--- a/Engine/Source/Programs/Horde/Horde.Server.Tests/Issues/IssueServiceTests.cs
+++ b/Engine/Source/Programs/Horde/Horde.Server.Tests/Issues/IssueServiceTests.cs
@@ -1,7 +1,6 @@
 // Copyright Epic Games, Inc. All Rights Reserved.
 
-extern alias HordeAgent;
-extern alias HordeAgentDriver;
+extern alias JobDriver;
 using System;
 using System.Buffers;
 using System.Collections.Generic;
@@ -30,7 +29,7 @@ using Horde.Server.Storage;
 using Horde.Server.Streams;
 using Horde.Server.Tests.Stubs.Services;
 using Horde.Server.Users;
-using HordeAgentDriver.Horde.Agent.Driver.Parser;
+using JobDriver.Horde.Agent.Driver.Parser;
 using Microsoft.Extensions.Logging;
 using Microsoft.Extensions.Logging.Abstractions;
 using Microsoft.VisualStudio.TestTools.UnitTesting;
