Horde: Fix path to JobDriver in docker file.

#rnx

diff --git a/Engine/Source/Programs/Horde/Horde.Server/Dockerfile b/Engine/Source/Programs/Horde/Horde.Server/Dockerfile
--- a/Engine/Source/Programs/Horde/Horde.Server/Dockerfile
+++ b/Engine/Source/Programs/Horde/Horde.Server/Dockerfile
@@ -37,8 +37,8 @@ COPY Source/Programs/Shared/EpicGames.Perforce.Managed/*.csproj ./Source/Program
 COPY Source/Programs/Shared/EpicGames.Redis/*.csproj ./Source/Programs/Shared/EpicGames.Redis/
 COPY Source/Programs/Shared/EpicGames.Serialization/*.csproj ./Source/Programs/Shared/EpicGames.Serialization/
 COPY Source/Programs/Shared/EpicGames.Serialization.Tests/*.csproj ./Source/Programs/Shared/EpicGames.Serialization.Tests/
+COPY Source/Programs/Horde/Drivers/JobDriver/*.csproj ./Source/Programs/Horde/Drivers/JobDriver/
 COPY Source/Programs/Horde/Horde.Agent/*.csproj ./Source/Programs/Horde/Horde.Agent/
-COPY Source/Programs/Horde/Horde.Agent.Driver/*.csproj ./Source/Programs/Horde/Horde.Agent.Driver/
 COPY Source/Programs/Horde/Horde.Server/*.csproj ./Source/Programs/Horde/Horde.Server/
 COPY Source/Programs/Horde/Horde.Server.Tests/*.csproj ./Source/Programs/Horde/Horde.Server.Tests/
 RUN dotnet restore Source/Programs/Horde/Horde.Server/Horde.Server.csproj -p:WithHordeStorage=false
@@ -54,8 +54,8 @@ ARG msbuild_args
 RUN dotnet publish Source/Programs/Horde/Horde.Server/Horde.Server.csproj -c Release -o out -p:WithHordeStorage=false $msbuild_args -p:RunAnalyzersDuringBuild=True
 
 # Run tests
+COPY Source/Programs/Horde/Drivers/JobDriver/ Source/Programs/Horde/Drivers/JobDriver/
 COPY Source/Programs/Horde/Horde.Agent/ Source/Programs/Horde/Horde.Agent/
-COPY Source/Programs/Horde/Horde.Agent.Driver/ Source/Programs/Horde/Horde.Agent.Driver/
 COPY Source/Programs/Horde/Horde.Server.Tests/ Source/Programs/Horde/Horde.Server.Tests/
 COPY Source/Programs/Shared/EpicGames.Serialization.Tests/ Source/Programs/Shared/EpicGames.Serialization.Tests/
 
