Horde: Allow overriding the driver for jobs through the JobOptions object.

diff --git a/Engine/Source/Programs/Horde/Horde.Agent/Leases/Handlers/JobHandler.cs b/Engine/Source/Programs/Horde/Horde.Agent/Leases/Handlers/JobHandler.cs
--- a/Engine/Source/Programs/Horde/Horde.Agent/Leases/Handlers/JobHandler.cs
+++ b/Engine/Source/Programs/Horde/Horde.Agent/Leases/Handlers/JobHandler.cs
@@ -55,7 +55,8 @@ namespace Horde.Agent.Leases.Handlers
 				arguments.Add($"-WorkingDir={session.WorkingDir}");
 				arguments.Add($"-Task={Convert.ToBase64String(executeTask.ToByteArray())}");
 
-				FileReference driverAssembly = FileReference.Combine(new DirectoryReference(AppContext.BaseDirectory), "JobDriver", "JobDriver.dll");
+				string driverName = String.IsNullOrEmpty(executeTask.JobOptions.Driver) ? "JobDriver" : executeTask.JobOptions.Driver;
+				FileReference driverAssembly = FileReference.Combine(new DirectoryReference(AppContext.BaseDirectory), driverName, $"{driverName}.dll");
 
 				Dictionary<string, string> environment = ManagedProcess.GetCurrentEnvVars();
 				environment[HordeHttpClient.HordeUrlEnvVarName] = session.HordeClient.ServerUrl.ToString();
diff --git a/Engine/Source/Programs/Horde/Horde.Server/Jobs/JobTaskSource.cs b/Engine/Source/Programs/Horde/Horde.Server/Jobs/JobTaskSource.cs
--- a/Engine/Source/Programs/Horde/Horde.Server/Jobs/JobTaskSource.cs
+++ b/Engine/Source/Programs/Horde/Horde.Server/Jobs/JobTaskSource.cs
@@ -952,6 +952,7 @@ namespace Horde.Server.Jobs
 			{
 				rpcOptions.Container = GetRpcJobContainerOptions(options.Container);
 			}
+			rpcOptions.Driver = options.Driver;
 			return rpcOptions;
 		}
 
