Horde: Remove server OS version info from anonymous endpoint responses. (#11909)

diff --git a/Engine/Source/Programs/Horde/Horde.Server/Server/ServerController.cs b/Engine/Source/Programs/Horde/Horde.Server/Server/ServerController.cs
--- a/Engine/Source/Programs/Horde/Horde.Server/Server/ServerController.cs
+++ b/Engine/Source/Programs/Horde/Horde.Server/Server/ServerController.cs
@@ -80,7 +80,6 @@ namespace Horde.Server.Server
 
 			FileVersionInfo versionInfo = FileVersionInfo.GetVersionInfo(Assembly.GetExecutingAssembly().Location);
 			response.ServerVersion = versionInfo.ProductVersion ?? String.Empty;
-			response.OsDescription = RuntimeInformation.OSDescription;
 
 			ITool? tool = await _toolCollection.GetAsync(AgentExtensions.AgentToolId, HttpContext.RequestAborted);
 			if (tool != null)
diff --git a/Engine/Source/Programs/Shared/EpicGames.Horde/Server/ServerMessages.cs b/Engine/Source/Programs/Shared/EpicGames.Horde/Server/ServerMessages.cs
--- a/Engine/Source/Programs/Shared/EpicGames.Horde/Server/ServerMessages.cs
+++ b/Engine/Source/Programs/Shared/EpicGames.Horde/Server/ServerMessages.cs
@@ -27,11 +27,6 @@ namespace EpicGames.Horde.Server
 		/// The current agent version string
 		/// </summary>
 		public string? AgentVersion { get; set; }
-
-		/// <summary>
-		/// The operating system server is hosted on
-		/// </summary>
-		public string OsDescription { get; set; } = String.Empty;
 	}
 
 	/// <summary>
