Horde: Pass -Target arguments to individual steps to ensure that BuildGraph can correctly filter out which blocks to embed in temp storage manifests.

diff --git a/Engine/Source/Programs/Horde/Drivers/JobDriver/Execution/JobExecutor.cs b/Engine/Source/Programs/Horde/Drivers/JobDriver/Execution/JobExecutor.cs
--- a/Engine/Source/Programs/Horde/Drivers/JobDriver/Execution/JobExecutor.cs
+++ b/Engine/Source/Programs/Horde/Drivers/JobDriver/Execution/JobExecutor.cs
@@ -1096,6 +1096,15 @@ namespace JobDriver.Execution
 			{
 				arguments.AppendArgument(ScriptArgumentPrefix, _scriptFileName);
 			}
+
+			if (!_allowTargetChanges)
+			{
+				foreach (string target in _targets)
+				{
+					arguments.AppendArgument(TargetArgumentPrefix + target);
+				}
+			}
+
 			arguments.AppendArgument("-SingleNode=", step.Name);
 			//			Arguments.AppendArgument("-TokenSignature=", JobId.ToString());
 
