Horde: Allow overriding the targets to execute for a job.

The optional 'Targets' property can be specified when creating a job, and will replace any -Target= arguments in the command line if set.

diff --git a/Engine/Source/Programs/Horde/Horde.Server/Jobs/IJob.cs b/Engine/Source/Programs/Horde/Horde.Server/Jobs/IJob.cs
--- a/Engine/Source/Programs/Horde/Horde.Server/Jobs/IJob.cs
+++ b/Engine/Source/Programs/Horde/Horde.Server/Jobs/IJob.cs
@@ -181,6 +181,11 @@ namespace Horde.Server.Jobs
 		/// </summary>
 		public IReadOnlyList<string> Arguments { get; }
 
+		/// <summary>
+		/// Custom list of targets for the job. If null or empty, the list of targets is determined from the command line.
+		/// </summary>
+		public IReadOnlyList<string>? Targets { get; }
+
 		/// <summary>
 		/// Additional arguments for the job, when a set of parameters are applied.
 		/// </summary>
diff --git a/Engine/Source/Programs/Horde/Horde.Server/Jobs/IJobCollection.cs b/Engine/Source/Programs/Horde/Horde.Server/Jobs/IJobCollection.cs
--- a/Engine/Source/Programs/Horde/Horde.Server/Jobs/IJobCollection.cs
+++ b/Engine/Source/Programs/Horde/Horde.Server/Jobs/IJobCollection.cs
@@ -77,6 +77,9 @@ namespace Horde.Server.Jobs
 		/// <inheritdoc cref="IJob.Arguments"/>
 		public List<string> Arguments { get; } = new List<string>();
 
+		/// <inheritdoc cref="IJob.Targets"/>
+		public List<string>? Targets { get; set; }
+
 		/// <inheritdoc cref="IJob.AdditionalArguments"/>
 		public List<string> AdditionalArguments { get; } = new List<string>();
 
diff --git a/Engine/Source/Programs/Horde/Horde.Server/Jobs/JobCollection.cs b/Engine/Source/Programs/Horde/Horde.Server/Jobs/JobCollection.cs
--- a/Engine/Source/Programs/Horde/Horde.Server/Jobs/JobCollection.cs
+++ b/Engine/Source/Programs/Horde/Horde.Server/Jobs/JobCollection.cs
@@ -82,6 +82,7 @@ namespace Horde.Server.Jobs
 			IReadOnlyList<IJobStepBatch> IJob.Batches => Batches;
 			IReadOnlyDictionary<ParameterId, string> IJob.Parameters => _document.Parameters;
 			IReadOnlyList<string> IJob.Arguments => _document.Arguments;
+			IReadOnlyList<string>? IJob.Targets => _document.Targets;
 			IReadOnlyList<string> IJob.AdditionalArguments => _document.AdditionalArguments;
 			IReadOnlyDictionary<string, string> IJob.Environment => _document.Environment;
 			IReadOnlyList<int> IJob.Issues => _document.ReferencedByIssues;
@@ -453,6 +454,7 @@ namespace Horde.Server.Jobs
 			public Dictionary<ParameterId, string> Parameters { get; set; } = new Dictionary<ParameterId, string>();
 
 			public List<string> Arguments { get; set; } = new List<string>();
+			public List<string>? Targets { get; set; }
 			public List<string> AdditionalArguments { get; set; } = new List<string>();
 
 			[BsonDictionaryOptions(DictionaryRepresentation.ArrayOfDocuments)]
@@ -516,6 +518,7 @@ namespace Horde.Server.Jobs
 				NotificationChannelFilter = options.NotificationChannelFilter;
 				Parameters = new Dictionary<ParameterId, string>(options.Parameters);
 				Arguments.AddRange(options.Arguments);
+				Targets = (options.Targets != null && options.Targets.Count > 0) ? new List<string>(options.Targets) : null;
 				AdditionalArguments.AddRange(options.AdditionalArguments);
 
 				foreach (KeyValuePair<string, string> pair in options.Environment)
diff --git a/Engine/Source/Programs/Horde/Horde.Server/Jobs/JobsController.cs b/Engine/Source/Programs/Horde/Horde.Server/Jobs/JobsController.cs
--- a/Engine/Source/Programs/Horde/Horde.Server/Jobs/JobsController.cs
+++ b/Engine/Source/Programs/Horde/Horde.Server/Jobs/JobsController.cs
@@ -236,6 +236,13 @@ namespace Horde.Server.Jobs
 				options.Arguments.AddRange(options.AdditionalArguments);
 			}
 
+			// Override the targets for the job if specified
+			if (create.Targets != null && create.Targets.Count > 0)
+			{
+				options.Arguments.RemoveAll(x => x.StartsWith("-Target=", StringComparison.OrdinalIgnoreCase));
+				options.Arguments.AddRange(create.Targets.Select(x => $"-Target={x}"));
+			}
+
 			// Merge the environment variables
 			foreach ((string key, string value) in environment)
 			{
@@ -608,6 +615,7 @@ namespace Horde.Server.Jobs
 			response.Parameters = job.Parameters.ToDictionary();
 			response.Arguments = job.Arguments.ToList();
 			response.AdditionalArguments = job.AdditionalArguments.ToList();
+			response.Targets = (job.Targets != null && job.Targets.Count > 0) ? job.Targets.ToList() : null;
 			response.UpdateTime = new DateTimeOffset(job.UpdateTimeUtc);
 			response.UseArtifactsV2 = true;
 			response.UpdateIssues = job.UpdateIssues;
diff --git a/Engine/Source/Programs/Shared/EpicGames.Horde/Jobs/JobMessages.cs b/Engine/Source/Programs/Shared/EpicGames.Horde/Jobs/JobMessages.cs
--- a/Engine/Source/Programs/Shared/EpicGames.Horde/Jobs/JobMessages.cs
+++ b/Engine/Source/Programs/Shared/EpicGames.Horde/Jobs/JobMessages.cs
@@ -248,6 +248,11 @@ namespace EpicGames.Horde.Jobs
 		/// </summary>
 		public List<string>? AdditionalArguments { get; set; }
 
+		/// <summary>
+		/// Targets for the job. Will override any parameters specified in the Arguments or Parameters section if specified.
+		/// </summary>
+		public List<string>? Targets { get; set; }
+
 		/// <summary>
 		/// Private constructor for serialization
 		/// </summary>
@@ -520,6 +525,11 @@ namespace EpicGames.Horde.Jobs
 		/// </summary>
 		public List<string> AdditionalArguments { get; set; } = new List<string>();
 
+		/// <summary>
+		/// Custom list of targets for the job
+		/// </summary>
+		public List<string>? Targets { get; set; }
+
 		/// <summary>
 		/// The last update time for this job
 		/// </summary>
