Horde: Fix targets argument not being stored on jobs.

#rnx

diff --git a/Engine/Source/Programs/Horde/Horde.Server/Jobs/JobCollection.cs b/Engine/Source/Programs/Horde/Horde.Server/Jobs/JobCollection.cs
--- a/Engine/Source/Programs/Horde/Horde.Server/Jobs/JobCollection.cs
+++ b/Engine/Source/Programs/Horde/Horde.Server/Jobs/JobCollection.cs
@@ -454,9 +454,11 @@ namespace Horde.Server.Jobs
 			public Dictionary<ParameterId, string> Parameters { get; set; } = new Dictionary<ParameterId, string>();
 
 			public List<string> Arguments { get; set; } = new List<string>();
-			public List<string>? Targets { get; set; }
 			public List<string> AdditionalArguments { get; set; } = new List<string>();
 
+			[BsonIgnoreIfNull]
+			public List<string>? Targets { get; set; }
+
 			[BsonDictionaryOptions(DictionaryRepresentation.ArrayOfDocuments)]
 			public Dictionary<string, string> Environment { get; set; } = new Dictionary<string, string>();
 
@@ -518,8 +520,8 @@ namespace Horde.Server.Jobs
 				NotificationChannelFilter = options.NotificationChannelFilter;
 				Parameters = new Dictionary<ParameterId, string>(options.Parameters);
 				Arguments.AddRange(options.Arguments);
-				Targets = (options.Targets != null && options.Targets.Count > 0) ? new List<string>(options.Targets) : null;
 				AdditionalArguments.AddRange(options.AdditionalArguments);
+				Targets = (options.Targets != null && options.Targets.Count > 0) ? new List<string>(options.Targets) : null;
 
 				foreach (KeyValuePair<string, string> pair in options.Environment)
 				{
diff --git a/Engine/Source/Programs/Horde/Horde.Server/Jobs/JobsController.cs b/Engine/Source/Programs/Horde/Horde.Server/Jobs/JobsController.cs
--- a/Engine/Source/Programs/Horde/Horde.Server/Jobs/JobsController.cs
+++ b/Engine/Source/Programs/Horde/Horde.Server/Jobs/JobsController.cs
@@ -239,6 +239,8 @@ namespace Horde.Server.Jobs
 			// Override the targets for the job if specified
 			if (create.Targets != null && create.Targets.Count > 0)
 			{
+				options.Targets = create.Targets.ToList();
+
 				options.Arguments.RemoveAll(x => x.StartsWith("-Target=", StringComparison.OrdinalIgnoreCase));
 				options.Arguments.AddRange(create.Targets.Select(x => $"-Target={x}"));
 			}
