Horde: New Build target handling
#rnx

diff --git a/Engine/Source/Programs/Horde/HordeDashboard/src/components/build/NewBuildV2.tsx b/Engine/Source/Programs/Horde/HordeDashboard/src/components/build/NewBuildV2.tsx
--- a/Engine/Source/Programs/Horde/HordeDashboard/src/components/build/NewBuildV2.tsx
+++ b/Engine/Source/Programs/Horde/HordeDashboard/src/components/build/NewBuildV2.tsx
@@ -1,4 +1,4 @@
-import { Checkbox, ComboBox, ContextualMenuItemType, DefaultButton, DirectionalHint, Dropdown, DropdownMenuItemType, IComboBoxOption, IContextualMenuItem, IContextualMenuProps, IDropdownOption, Icon, IconButton, Label, MessageBar, MessageBarType, Modal, Pivot, PivotItem, PrimaryButton, ScrollablePane, ScrollbarVisibility, Spinner, SpinnerSize, Stack, TagPicker, Text, TextField, TooltipHost } from "@fluentui/react";
+import { Checkbox, ComboBox, ContextualMenuItemType, DefaultButton, DirectionalHint, Dropdown, DropdownMenuItemType, IComboBoxOption, IContextualMenuItem, IContextualMenuProps, IDropdownOption, Icon, IconButton, Label, MessageBar, MessageBarType, Modal, Pivot, PivotItem, PrimaryButton, ScrollablePane, ScrollbarVisibility, Spinner, SpinnerSize, Stack, TagPicker, Text, TextField, TooltipHost, ValidationState } from "@fluentui/react";
 import { useConst } from '@fluentui/react-hooks';
 import { ITextField } from "@fluentui/react/lib-commonjs/TextField";
 import Markdown from "markdown-to-jsx";
@@ -104,8 +104,10 @@ class BuildOptions {
    advJobPriority?: Priority;
    advUpdateIssues?: boolean;
    advAdditionalArgs?: string;
+   advTargets?: string[];
 
    parameters: Record<string, string> = {};
+   disabledParameters: Set<string> = new Set();
 
    validationErrors: ValidationError[] = [];
 
@@ -153,7 +155,7 @@ class BuildOptions {
 
    get advancedModified(): boolean {
 
-      return !!((this.advJobPriority && this.advJobPriority !== Priority.Normal) || this.advUpdateIssues || !!this.advJobName || !!this.advAdditionalArgs)
+      return !!((this.advJobPriority && this.advJobPriority !== Priority.Normal) || this.advUpdateIssues || !!this.advJobName || !!this.advAdditionalArgs || !!this.advTargets)
    }
 
    change?: number;
@@ -363,6 +365,10 @@ class BuildOptions {
          return true;
       }
 
+      if (this.advTargets && !this.advTargets?.length) {
+         errors.push({ error: `No targets specified`, paramString: "Target", value: "" })
+      }
+
       if (this.preflightChange && !this.template.allowPreflights) {
          errors.push({ error: `Template "${this.template.name}" does not allow preflights`, paramString: "Shelved Change", value: this.preflightChange.toString() })
       }
@@ -492,6 +498,10 @@ class BuildOptions {
          }
       });
 
+      if (this.jobDetails?.jobData?.targets?.length) {
+         this.setTargets(this.jobDetails?.jobData?.targets);
+      }
+
    }
 
    private setPreviewTemplate() {
@@ -505,6 +515,189 @@ class BuildOptions {
 
    }
 
+   setTargets(targets: string[]) {
+      if (!this.advTargets) {
+         this.initOverrideTargets();
+      }
+      this.advTargets = targets.sort((a, b) => a.localeCompare(b));
+      this.setChanged();
+   }
+
+   addTarget(target: string) {
+      if (!this.advTargets) {
+         this.initOverrideTargets();
+         this.advTargets = this.targets;
+      }
+      this.advTargets.push(target);
+      this.advTargets = Array.from(new Set(this.advTargets)).sort((a, b) => a.localeCompare(b));
+
+      this.setChanged();
+   }
+
+   get overrideTargets(): boolean {
+      return !!this.advTargets;
+   }
+
+   get targets(): string[] {
+
+      if (this.advTargets) {
+         return this.advTargets
+      }
+
+      const template = this.template;
+
+      if (!template) {
+         return [];
+      }
+
+      let targets: Set<string> = new Set();
+
+      template?.arguments.forEach(a => {
+         if (a.toLowerCase().startsWith("-target=")) {
+            targets.add(a.slice(8).trim())
+         }
+      });
+
+      template.parameters.forEach(p => {
+
+         let allArgs: string[] = [];
+
+         switch (p.type) {
+            case ParameterType.Bool:
+               const b = p as BoolParameterData;
+               if (this.parameters[b.id] === "true") {
+                  if (b.argumentIfEnabled) {
+                     allArgs.push(b.argumentIfEnabled);
+                  }
+                  else if (b.argumentsIfEnabled?.length) {
+                     allArgs.push(...b.argumentsIfEnabled);
+                  }
+               } else {
+                  if (b.argumentIfDisabled) {
+                     allArgs.push(b.argumentIfDisabled);
+                  }
+                  else if (b.argumentsIfDisabled?.length) {
+                     allArgs.push(...b.argumentsIfDisabled);
+                  }
+               }
+
+               allArgs.filter(a => a.toLowerCase().startsWith("-target=")).forEach(t => {
+                  targets.add(t.slice(8).trim())
+               })
+
+               break;
+            case ParameterType.Text:
+               const t = p as TextParameterData;
+               if (t.argument?.toLowerCase().startsWith("-target=")) {
+                  if (this.parameters[t.id]?.length) {
+                     targets.add(t.id);
+                  }
+               }
+               break;
+            case ParameterType.List:
+               const list = p as ListParameterData;
+               list.items.forEach(i => {
+                  allArgs = [];
+
+                  if (this.parameters[i.id] === "true") {
+                     if (i.argumentIfEnabled) {
+                        allArgs.push(i.argumentIfEnabled);
+                     }
+                     else if (i.argumentsIfEnabled?.length) {
+                        allArgs.push(...i.argumentsIfEnabled);
+                     }
+                  } else {
+                     if (i.argumentIfDisabled) {
+                        allArgs.push(i.argumentIfDisabled);
+                     }
+                     else if (i.argumentsIfDisabled?.length) {
+                        allArgs.push(...i.argumentsIfDisabled);
+                     }
+                  }
+
+                  allArgs.filter(a => a.toLowerCase().startsWith("-target=")).forEach(t => {
+                     targets.add(t.slice(8).trim())
+                  })
+
+               });
+               break;
+         }
+      })
+
+      return Array.from(targets).filter(t => !!t).sort((a, b) => a.localeCompare(b));
+   }
+
+
+   private initOverrideTargets() {
+
+      const template = this.template;
+
+      if (!template) {
+         return;
+      }
+
+      template.parameters.forEach(p => {
+
+         let allArgs: string[] = [];
+
+         switch (p.type) {
+            case ParameterType.Bool:
+               const b = p as BoolParameterData;
+
+               if (b.argumentIfEnabled) {
+                  allArgs.push(b.argumentIfEnabled);
+               }
+               else if (b.argumentsIfEnabled?.length) {
+                  allArgs.push(...b.argumentsIfEnabled);
+               }
+
+               if (b.argumentIfDisabled) {
+                  allArgs.push(b.argumentIfDisabled);
+               }
+               else if (b.argumentsIfDisabled?.length) {
+                  allArgs.push(...b.argumentsIfDisabled);
+               }
+
+               if (allArgs.find(a => a.toLowerCase().startsWith("-target="))) {
+                  this.disabledParameters.add(b.id);
+               }
+
+               break;
+            case ParameterType.Text:
+               const t = p as TextParameterData;
+               if (t.argument?.toLowerCase().startsWith("-target=")) {
+                  this.disabledParameters.add(t.id);
+               }
+               break;
+            case ParameterType.List:
+               const list = p as ListParameterData;
+               list.items.forEach(i => {
+                  allArgs = [];
+                  if (i.argumentIfEnabled) {
+                     allArgs.push(i.argumentIfEnabled);
+                  }
+                  else if (i.argumentsIfEnabled?.length) {
+                     allArgs.push(...i.argumentsIfEnabled);
+                  }
+
+                  if (i.argumentIfDisabled) {
+                     allArgs.push(i.argumentIfDisabled);
+                  }
+                  else if (i.argumentsIfDisabled?.length) {
+                     allArgs.push(...i.argumentsIfDisabled);
+                  }
+
+                  if (allArgs.find(a => a.toLowerCase().startsWith("-target="))) {
+                     this.disabledParameters.add(i.id);
+                  }
+
+               });
+               break;
+         }
+
+      })
+   }
+
    private setTemplate(templateId: string, notifyChanged = true) {
 
       const template = this.allTemplates.find(t => t.id === templateId);
@@ -516,6 +709,8 @@ class BuildOptions {
          throw `Unable to find template ${templateId}`;
       }
 
+      this.advTargets = undefined;
+
       const source = { ...template };
       source.hash = "";
 
@@ -753,7 +948,7 @@ const BoolParameter: React.FC<{ param: BoolParameterData }> = observer(({ param
    return <Stack>
       <Checkbox key={key}
          label={param.label}
-         disabled={options.readOnly}
+         disabled={options.readOnly || options.disabledParameters.has(param.id)}
          checked={options.parameters[param.id] == "true"}
          onChange={(ev, value) => {
             ev?.preventDefault();
@@ -772,13 +967,18 @@ const TextParameter: React.FC<{ param: TextParameterData }> = observer(({ param
 
    const key = param.id.replaceAll(".", "-");
 
+   let value = options.parameters[param.id] ?? "";
+   if (options.disabledParameters.has(param.id)) {
+      value = "";
+   }
+
    return <Stack>
       <TextField key={key}
          placeholder={options.jobDetails ? "" : param.hint}
          label={param.label}
          spellCheck={false}
-         value={options.parameters[param.id] ?? ""}
-         disabled={options.readOnly}
+         value={value}
+         disabled={options.readOnly || options.disabledParameters.has(param.id)}
          onChange={(ev, value) => {
             options.onTextChanged(param.id, value ?? "");
          }}
@@ -863,7 +1063,8 @@ const BasicListParameter: React.FC<{ param: ListParameterData }> = observer(({ p
       doptions.push({
          key: item.id,
          text: item.text,
-         selected: options.parameters[item.id] === "true"
+         disabled: options.disabledParameters.has(item.id),
+         selected: options.parameters[item.id] === "true" && !options.disabledParameters.has(item.id)
       });
    });
 
@@ -932,7 +1133,7 @@ const MultiListParameter: React.FC<{ param: ListParameterData }> = observer(({ p
       param.items.forEach(item => {
          if (item.group === group) {
             const key = item.id;
-            const selected = options.parameters[item.id] === "true";
+            const selected = options.parameters[item.id] === "true" && !options.disabledParameters.has(item.id);
             if (selected) {
                selectedKeys.push(key);
             }
@@ -940,7 +1141,8 @@ const MultiListParameter: React.FC<{ param: ListParameterData }> = observer(({ p
                key: key,
                data: item,
                text: (item.group !== "__nogroup" && dupes.get(item.text)! > 1) ? `${item.text} - ${item.group}` : item.text,
-               selected: selected
+               selected: selected,
+               disabled: options.disabledParameters.has(item.id)
             });
          }
       });
@@ -1250,6 +1452,8 @@ const PreviewPanel: React.FC = observer(() => {
 
 const AdvancedPanel: React.FC = observer(() => {
 
+   const targetPicker = React.useRef(null)
+
    const options = BuildOptions.get();
    options.subscribe();
 
@@ -1308,6 +1512,21 @@ const AdvancedPanel: React.FC = observer(() => {
       height += 32;
    }
 
+   // targets
+   height += 92;
+
+   // target options for picker
+   type TargetPickerItem = {
+      key: string;
+      name: string;
+   }
+
+   const targetItems: TargetPickerItem[] = options.targets.map(t => {
+      return {
+         key: t,
+         name: t
+      }
+   });
    return <Stack style={{
       height: height,
       position: 'relative',
@@ -1341,6 +1560,64 @@ const AdvancedPanel: React.FC = observer(() => {
                options.setChanged();
             }} />
          </Stack>
+         <Stack>
+            <Label>Targets</Label>
+            <TagPicker
+
+               componentRef={targetPicker}
+
+               disabled={options.readOnly}
+
+               onBlur={(ev) => {
+                  if (ev.target.value && ev.target.value.trim()) {
+                     options.addTarget(ev.target.value.trim())
+                  }
+
+                  // This is using undocumented behavior to clear the input when you lose focus
+                  // It could very well break
+                  if (targetPicker.current) {
+                     try {
+                        (targetPicker.current as any).input.current._updateValue("");
+                     } catch (reason) {
+                        console.error("There was an error adding target to list and clearing the input in process\n" + reason);
+                     }
+
+                  }
+
+               }}
+
+               onResolveSuggestions={(filter, selected) => {
+                  return [];
+               }}
+
+               onEmptyResolveSuggestions={(selected) => {
+                  return [];
+               }}
+
+               onRemoveSuggestion={(item) => { }}
+
+               createGenericItem={(input: string, ValidationState: ValidationState) => {
+                  return {
+                     item: { name: input, key: input },
+                     selected: true
+                  };
+               }}
+
+               onChange={(items) => { if (items) options.setTargets(items.map(i => i.name)) }}
+
+               onValidateInput={(input?: string) => input ? ValidationState.valid : ValidationState.invalid}
+
+               selectedItems={targetItems}
+
+               onItemSelected={(item) => {
+                  if (!item || !item.name?.trim()) {
+                     return null;
+                  }
+                  options.addTarget(item.name.trim());
+                  return null;
+               }}
+            />
+         </Stack>
          {showAdditionalArgs && <TextField key={"key_adv_add_args"} multiline style={{ height: 48 }} resizable={false} readOnly={options.readOnly} spellCheck={false} defaultValue={options.advAdditionalArgs} label="Additional Arguments" onChange={(ev, newValue) => options.advAdditionalArgs = newValue} />}
          {showArgumentClipboardButton && <DefaultButton text="Copy Job Arguments to Clipboard" style={{ width: 240 }} onClick={() => copyToClipboard(
             options.jobDetails?.jobData?.arguments.map(arg => {
@@ -1383,7 +1660,6 @@ const BuildModal: React.FC<{ setUseLegacyDialog: (value: boolean) => void }> = o
    const template = options.template;
 
    if (!template) {
-      console.error("Build Modal has no template");
       return null;
    }
 
@@ -1447,7 +1723,8 @@ const BuildModal: React.FC<{ setUseLegacyDialog: (value: boolean) => void }> = o
          updateIssues: updateIssues,
          changeQueries: changeQueries,
          parameters: options.parameters,
-         additionalArguments: additionalArgs?.length ? additionalArgs : undefined
+         additionalArguments: additionalArgs?.length ? additionalArgs : undefined,
+         targets: options.advTargets
       };
 
       if (typeof (options.change) === 'number') {
@@ -1631,7 +1908,7 @@ const BuildModal: React.FC<{ setUseLegacyDialog: (value: boolean) => void }> = o
                </Stack>
                <Stack horizontal tokens={{ childrenGap: 16 }} styles={{ root: { paddingTop: 32, paddingLeft: 8, paddingBottom: 8 } }}>
                   <Stack>
-                     <DefaultButton disabled={ options.readOnly || !!options.jobDetails} text="Use Legacy Dialog" style={{ width: 160 }} onClick={() => setUseLegacyDialog(true)} />
+                     <DefaultButton disabled={options.readOnly || !!options.jobDetails} text="Use Legacy Dialog" style={{ width: 160 }} onClick={() => setUseLegacyDialog(true)} />
                   </Stack>
 
                   <Stack grow />
@@ -1676,7 +1953,7 @@ export const NewBuildV2: React.FC<{ streamId: string; show: boolean; onClose: (n
    }
 
    return <Stack>
-      <NewBuildV2Inner setUseLegacyDialog={(value:boolean) => { setNewBuildV1(value)}}/>
+      <NewBuildV2Inner setUseLegacyDialog={(value: boolean) => { setNewBuildV1(value) }} />
    </Stack>
 })
 
