Horde: Clear disabled parameters when changing template
#skipci
#rnx

diff --git a/Engine/Source/Programs/Horde/HordeDashboard/src/components/build/NewBuildV2.tsx b/Engine/Source/Programs/Horde/HordeDashboard/src/components/build/NewBuildV2.tsx
--- a/Engine/Source/Programs/Horde/HordeDashboard/src/components/build/NewBuildV2.tsx
+++ b/Engine/Source/Programs/Horde/HordeDashboard/src/components/build/NewBuildV2.tsx
@@ -710,6 +710,7 @@ class BuildOptions {
       }
 
       this.advTargets = undefined;
+      this.disabledParameters.clear();
 
       const source = { ...template };
       source.hash = "";
