Horde: Add API changes for targets and tools
#rnx
#skipci

diff --git a/Engine/Source/Programs/Horde/HordeDashboard/src/backend/Api.ts b/Engine/Source/Programs/Horde/HordeDashboard/src/backend/Api.ts
--- a/Engine/Source/Programs/Horde/HordeDashboard/src/backend/Api.ts
+++ b/Engine/Source/Programs/Horde/HordeDashboard/src/backend/Api.ts
@@ -1142,6 +1142,9 @@ export type CreateJobRequest = {
 
 	/** Additional arguments for the job, job dialog v2 */
 	additionalArguments?: string[];
+
+	/** Custom list of targets for the job, will take precedence over all argument and parameter targets */
+	targets?: string[];
 }
 
 /**Response from creating a new job */
@@ -1304,6 +1307,9 @@ export type GetJobResponse = {
 
 	/** Additional arguments for the job, job dialog v2 */
 	additionalArguments?: string[];
+
+	/** Custom list of targets for the job, will take precedence over all argument and parameter targets */
+	targets?: string[];
 }
 
 /**Request used to update a jobstep */
@@ -5018,27 +5024,40 @@ export type GetTestDataRefResponse = {
 
 	/// The number of suite tests swith errors
 	suiteSuccessCount?: number;
-
 }
+	
 
 /** Summary for a particular tool */
 export type GetToolSummaryResponse = {
 
-	/** Unique id of tool */
+	/** Unique identifier for the tool */
 	id: string;
 
-	/** Name of tool */
+	/** Name of the tool */
 	name: string;
 
-	/** Category of the tool */
+	/** Description of the tool */
+	description: string;
+
+	/** Category to display the tool in on the dashboard */
 	category?: string;
 
-	/** Description of tool */
-	description: string;
+	/** Grouping key to control how different tools should be merged on the dashboard */
+	group?: string;
+
+	/** List of platforms that this tool supports, as NET runtime identifiers */
+	platforms?: string;
 
 	/** Version of tool */
 	version?: string;
 
+	/** Identifier for the current deployment */
+	deploymentId?: string;
+
+	/** Whether to show this tool for download inside UGS */
+	showInUGS: boolean;
+
+	/**Whether to show this tool for download on the dashboard */
 	showInDashboard: boolean;
 }
 
