Horde: Build graph to use DotNet 8.0.300

#rnx

diff --git a/Engine/Source/Programs/Horde/BuildHorde.xml b/Engine/Source/Programs/Horde/BuildHorde.xml
--- a/Engine/Source/Programs/Horde/BuildHorde.xml
+++ b/Engine/Source/Programs/Horde/BuildHorde.xml
@@ -28,9 +28,14 @@
 	<Property Name="InformationalVersion" Value="$(InformationalVersion)-PF-$(PreflightChange)" If="'$(PreflightChange)' != ''"/>
 	<Property Name="VersionArguments" Value="/p:Version=$(Version).0 /p:InformationalVersion=$(InformationalVersion)"/>
 
+	<Property Name="DotNetRootPath" Value="$(RootDir)/Engine/Binaries/ThirdParty/DotNet/8.0.300"/>
 	<Property Name="DotNetPath" Value=""/>
-	<Property Name="DotNetPath" Value="$(RootDir)/Engine/Binaries/ThirdParty/DotNet/8.0.100/windows/dotnet.exe" If="$(HostPlatform) == 'Win64'"/>
-	<Property Name="DotNetPath" Value="$(RootDir)/Engine/Binaries/ThirdParty/DotNet/8.0.100/linux/dotnet" If="$(HostPlatform) == 'Linux'"/>
+	<Property Name="DotNetPath" Value="$(DotNetRootPath)/win-x64/dotnet.exe" If="$(HostPlatform) == 'Win64' And $(HostArchitecture) == 'X64'"/>
+	<Property Name="DotNetPath" Value="$(DotNetRootPath)/win-arm64/dotnet.exe" If="$(HostPlatform) == 'Win64' And $(HostArchitecture) == 'Arm64'"/>
+	<Property Name="DotNetPath" Value="$(DotNetRootPath)/linux-x64/dotnet" If="$(HostPlatform) == 'Linux' And $(HostArchitecture) == 'X64'"/>
+	<Property Name="DotNetPath" Value="$(DotNetRootPath)/linux-arm64/dotnet" If="$(HostPlatform) == 'Linux' And $(HostArchitecture) == 'Arm64'"/>
+	<Property Name="DotNetPath" Value="$(DotNetRootPath)/mac-x64/dotnet" If="$(HostPlatform) == 'Mac' And $(HostArchitecture) == 'X64'"/>
+	<Property Name="DotNetPath" Value="$(DotNetRootPath)/mac-arm64/dotnet" If="$(HostPlatform) == 'Mac' And $(HostArchitecture) == 'Arm64'"/>
 
 	<!-- SERVER -->
 
