Horde: Updating schema documentation

#rnx

diff --git a/Engine/Source/Programs/Horde/Docs/Config/Schema/Projects.md b/Engine/Source/Programs/Horde/Docs/Config/Schema/Projects.md
--- a/Engine/Source/Programs/Horde/Docs/Config/Schema/Projects.md
+++ b/Engine/Source/Programs/Horde/Docs/Config/Schema/Projects.md
@@ -178,6 +178,7 @@ Name | Description
 `workspaceMaterializer` | `string`<br>What workspace materializer to use in WorkspaceExecutor. Will override any value from workspace config.
 `container` | [`JobContainerOptions`](#jobcontaineroptions)<br>Options for executing a job inside a container
 `expireAfterDays` | `integer`<br>Number of days after which to expire jobs
+`driver` | `string`<br>Name of the driver to use
 
 ## JobContainerOptions
 
diff --git a/Engine/Source/Programs/Horde/Docs/Config/Schema/Streams.md b/Engine/Source/Programs/Horde/Docs/Config/Schema/Streams.md
--- a/Engine/Source/Programs/Horde/Docs/Config/Schema/Streams.md
+++ b/Engine/Source/Programs/Horde/Docs/Config/Schema/Streams.md
@@ -64,6 +64,7 @@ Name | Description
 `workspaceMaterializer` | `string`<br>What workspace materializer to use in WorkspaceExecutor. Will override any value from workspace config.
 `container` | [`JobContainerOptions`](#jobcontaineroptions)<br>Options for executing a job inside a container
 `expireAfterDays` | `integer`<br>Number of days after which to expire jobs
+`driver` | `string`<br>Name of the driver to use
 
 ## JobContainerOptions
 
