Horde: Reverse direction on issue audit logs
#rnx

diff --git a/Engine/Source/Programs/Horde/HordeDashboard/src/components/AuditLog.tsx b/Engine/Source/Programs/Horde/HordeDashboard/src/components/AuditLog.tsx
--- a/Engine/Source/Programs/Horde/HordeDashboard/src/components/AuditLog.tsx
+++ b/Engine/Source/Programs/Horde/HordeDashboard/src/components/AuditLog.tsx
@@ -396,6 +396,10 @@ export const AuditLogPanel: React.FC<{ agentId?: string, issueId?: string }> = o
 
    });
 
+   if (issueId) {
+      logItems.reverse();
+   }
+
    const typeItems: IDropdownOption[] = Array.from(entryTypes.keys()).map(type => {
       return {
          key: type,
