Horde: Tools view improvements
#rnx
#skipci

diff --git a/Engine/Source/Programs/Horde/HordeDashboard/src/components/ToolView.tsx b/Engine/Source/Programs/Horde/HordeDashboard/src/components/ToolView.tsx
--- a/Engine/Source/Programs/Horde/HordeDashboard/src/components/ToolView.tsx
+++ b/Engine/Source/Programs/Horde/HordeDashboard/src/components/ToolView.tsx
@@ -1,24 +1,41 @@
 // Copyright Epic Games, Inc. All Rights Reserved.
 
-import { DetailsList, DetailsListLayoutMode, IColumn, Pivot, PivotItem, PrimaryButton, SelectionMode, Stack, Text } from "@fluentui/react";
+import { DetailsList, DetailsListLayoutMode, IColumn, IDetailsGroupDividerProps, IDetailsGroupRenderProps, IDetailsListProps, IGroup, IconButton, Pivot, PivotItem, PrimaryButton, SelectionMode, Stack, Text } from "@fluentui/react";
 import { observer } from "mobx-react-lite";
 import React, { useEffect, useState } from "react";
 import backend from "../backend";
 import { GetToolSummaryResponse } from "../backend/Api";
 import { PollBase } from "../backend/PollBase";
 import { useWindowSize } from "../base/utilities/hooks";
+import { getHordeStyling } from "../styles/Styles";
 import { Breadcrumbs } from "./Breadcrumbs";
 import { TopNav } from "./TopNav";
-import { getHordeStyling } from "../styles/Styles";
 
 const defaultCategory = "General";
 
+enum Platform {
+   Windows = "win-x64",
+   MacOS = "osx-x64",
+   Linux = "linux-x64",
+   Unknown = "Unknown"
+}
+
 class ToolHandler extends PollBase {
 
    constructor(pollTime = 30000) {
-
       super(pollTime);
 
+      const userAgent = window.navigator.userAgent;
+
+      if (userAgent.indexOf("Win") != -1) this.clientPlatform = Platform.Windows;
+      else if (userAgent.indexOf("Mac") != -1) this.clientPlatform = Platform.MacOS;
+      else if (userAgent.indexOf("X11") != -1) this.clientPlatform = Platform.Linux;
+      else if (userAgent.indexOf("Linux") != -1) this.clientPlatform = Platform.Linux;
+      else {
+         // guess windows
+         this.clientPlatform = Platform.Windows;
+      }
+
    }
 
    clear() {
@@ -29,6 +46,10 @@ class ToolHandler extends PollBase {
 
    async poll(): Promise<void> {
 
+      if (this.loaded) {
+         return;
+      }
+
       try {
 
          this.tools = await backend.getTools();
@@ -36,13 +57,63 @@ class ToolHandler extends PollBase {
          this.categories = new Map();
 
          this.tools.forEach(t => {
+
+            // categories
             const cat = t.category ?? defaultCategory;
             if (!this.categories.has(cat)) {
                this.categories.set(cat, []);
             }
             this.categories.get(cat)!.push(t);
+
          })
 
+
+         // sort by prefered platform
+         this.tools = this.tools.sort((a, b) => {
+
+            if (a.group && !b.group) {
+               return -1;
+            }
+
+            if (!a.group && b.group) {
+               return 1;
+            }
+
+            if (a.group && b.group && (a.group !== b.group)) {
+               return a.group.localeCompare(b.group);
+            }
+
+
+            if (a.platforms?.length && !b.platforms?.length) {
+               return -1;
+            }
+
+            if (!a.platforms?.length && b.platforms?.length) {
+               return 1;
+            }
+
+            if (a.platforms?.length && b.platforms?.length) {
+
+               if ((a.platforms.indexOf(this.clientPlatform) !== -1) && (b.platforms.indexOf(this.clientPlatform) === -1)) {
+                  return -1;
+               }
+               if ((a.platforms.indexOf(this.clientPlatform) === -1) && (b.platforms.indexOf(this.clientPlatform) !== -1)) {
+                  return 1;
+               }
+
+               if ((a.platforms.indexOf("any")) !== -1 && (b.platforms.indexOf("any") === -1)) {
+                  return -1;
+               }
+
+               if ((a.platforms.indexOf("any")) === -1 && (b.platforms.indexOf("any") !== -1)) {
+                  return 1;
+               }
+            }
+
+            return a.name.localeCompare(b.name);
+         })
+
+
          this.loaded = true;
          this.setUpdated();
 
@@ -52,6 +123,8 @@ class ToolHandler extends PollBase {
 
    }
 
+   clientPlatform: Platform;
+
    categories: Map<string, GetToolSummaryResponse[]> = new Map();
 
    loaded = false;
@@ -81,51 +154,107 @@ const ToolPanel: React.FC<{ selectedKey: string }> = observer(({ selectedKey })
       { key: 'column_name', name: 'Name', minWidth: 240, maxWidth: 240, isResizable: false },
       { key: 'column_desc', name: 'Description', fieldName: 'description', minWidth: 580, maxWidth: 580, isResizable: false, isMultiline: true },
       { key: 'column_version', name: 'Version', fieldName: 'version', minWidth: 280, maxWidth: 280, isResizable: false, headerClassName: hordeClasses.detailsHeader },
-      { key: 'column_download', name: '', minWidth: 160, maxWidth: 160, isResizable: false }
+      { key: 'column_download', name: '', minWidth: 16, maxWidth: 16, isResizable: false }
    ];
 
    let tools = [...handler.tools];
 
    tools = tools.filter(t => (t.category ?? defaultCategory) === selectedKey)
 
-   tools = tools.sort((a, b) => a.name.localeCompare(b.name));
+   const groups: IGroup[] = [];
+
+   let cgroup: string | undefined;
 
-   const renderItem = (item: any, index?: number, column?: IColumn) => {
+   const toolHide = new Set<string>();
 
-      if (!column) {
-         return null;
+   // emit groups
+   for (let i = 0; i < tools.length; i++) {
+      const tool = tools[i];
+
+      if (tool.group != cgroup) {
+
+         if (cgroup) {
+            groups[groups.length - 1].count = i - groups[groups.length - 1].startIndex;
+         }
+
+         cgroup = tool.group;
+
+         if (cgroup) {
+            groups.push({ startIndex: i, name: tool.name, key: `group_key_${cgroup}`, count: 0, isCollapsed: true, data: tool });
+            toolHide.add(tool.id);
+         }
       }
+   }
 
-      if (column.key === 'column_name') {
-         return <Stack verticalAlign="center" verticalFill={true}>
-            <Text style={{ fontFamily: "Horde Open Sans SemiBold", color: modeColors.text }}>{item.name}</Text>
-         </Stack>
+   if (cgroup) {
+      groups[groups.length - 1].count = tools.length - groups[groups.length - 1].startIndex;
+   }
+
+   const renderTool = (tool: GetToolSummaryResponse, groupProps?: IDetailsGroupDividerProps) => {
+
+      let pad = 0;
+      let width = 330;
+      if (groupProps && groupProps.group!.count < 2) {
+         pad = 36;
+         width += pad
       }
 
-      if (column.key === 'column_version') {
-         if (!item.version) {
-            return null;
-         }
-         return <Stack verticalAlign="center" verticalFill={true}>
-            <Text style={{ color: modeColors.text }}>{item.version}</Text>
-         </Stack>
+      if (!groupProps) {
+         pad = 36 + 12;
+         width = 366;
       }
 
+      return <Stack horizontal verticalAlign="center" verticalFill={true} style={{ padding: 12 }}>
+         {!!groupProps && groupProps.group!.count > 1 && <IconButton style={{ color: modeColors.text, fontSize: 18, marginRight: 4 }} iconProps={{ iconName: groupProps.group?.isCollapsed ? 'ChevronRight' : 'ChevronDown' }} onClick={(event: any) => {
+            event?.stopPropagation();
+            groupProps.onToggleCollapse!(groupProps.group!);
+         }}
+         />}
 
-      if (column.key === 'column_download') {
-         return <Stack horizontalAlign="center" verticalAlign="center" verticalFill={true}>
-            <PrimaryButton style={{ width: 120, color: "#FFFFFF" }} text="Download" href={`/api/v1/tools/${item.id}?action=download`} />
+         <Stack style={{ width: width, paddingLeft: pad}}>
+            <Text style={{ fontFamily: (groupProps || !tool.group) ? "Horde Open Sans SemiBold" : undefined, color: modeColors.text }}>{tool.name}</Text>
          </Stack>
+         <Stack style={{ width: 580 }}>
+            <Text style={{ color: modeColors.text }}>{tool.description ?? ""}</Text>
+         </Stack>
+         <Stack style={{ width: 260 }}>
+            <Text style={{ color: modeColors.text }}>{tool.version}</Text>
+         </Stack>
+         <Stack >
+            <PrimaryButton style={{ width: 120, color: "#FFFFFF" }} text="Download" href={`/api/v1/tools/${tool.id}?action=download`} />
+         </Stack>
+      </Stack>
+   }
+
+   const onRenderGroupHeader: IDetailsGroupRenderProps['onRenderHeader'] = (props: IDetailsGroupDividerProps) => {
+      if (props) {
+
+         const group = props.group!;
+
+         return renderTool(group.data as GetToolSummaryResponse, props)
+
       }
 
-      if (!column?.fieldName) {
-         return null;
+      return null;
+   };
+
+   const renderRow: IDetailsListProps['onRenderRow'] = (props) => {
+
+      if (props) {
+
+         const item = props!.item as GetToolSummaryResponse;
+         if (toolHide.has(item.id)) {
+            return null;
+         }
+
+         return renderTool(item)
+
       }
-      return <Stack verticalAlign="center" verticalFill={true}>
-         <Text style={{ color: modeColors.text }}>{item[column?.fieldName]}</Text>
-      </Stack>
+      return null;
    };
 
+
+
    return <Stack>
       {!tools.length && handler.loaded && <Stack style={{ paddingBottom: 12 }}>
          <Stack verticalAlign="center">
@@ -136,16 +265,32 @@ const ToolPanel: React.FC<{ selectedKey: string }> = observer(({ selectedKey })
       </Stack>}
 
       {!!tools.length && <Stack className={hordeClasses.raised} >
-         <Stack styles={{ root: { paddingLeft: 12, paddingRight: 12, paddingBottom: 12, width: "100%", selectors: { ".ms-DetailsHeader": { "paddingTop": "0px" } } } }} >
+         <Stack styles={{
+            root: {
+               paddingLeft: 12, paddingRight: 12, paddingBottom: 12, width: "100%", selectors: {
+                  ".ms-DetailsHeader": { "paddingTop": "0px" }, '.ms-List-cell': {
+                     minHeight: "0px !important"
+                  }
+               }
+            }
+         }} >
             <DetailsList
-               isHeaderVisible={true}
-               styles={{ headerWrapper: { paddingTop: 0 } }}
+               isHeaderVisible={false}
+               styles={{
+                  headerWrapper: {
+                     paddingTop: 0
+                  }
+               }}
                items={tools}
+               groups={groups.length ? groups : undefined}
+               groupProps={{
+                  onRenderHeader: onRenderGroupHeader,
+               }}
                columns={columns}
                selectionMode={SelectionMode.none}
                layoutMode={DetailsListLayoutMode.justified}
-               compact={false}
-               onRenderItemColumn={renderItem}
+               compact={true}
+               onRenderRow={renderRow}
             />
          </Stack>
       </Stack>}
