Horde: New Build target fixes, reset functionality
#rnx
#skipci

diff --git a/Engine/Source/Programs/Horde/HordeDashboard/src/components/build/NewBuildV2.tsx b/Engine/Source/Programs/Horde/HordeDashboard/src/components/build/NewBuildV2.tsx
--- a/Engine/Source/Programs/Horde/HordeDashboard/src/components/build/NewBuildV2.tsx
+++ b/Engine/Source/Programs/Horde/HordeDashboard/src/components/build/NewBuildV2.tsx
@@ -30,6 +30,15 @@ const parameterGap = 6;
 
 type NewBuildMode = "Basic" | "Advanced";
 
+type OptionsConstruction = {
+   streamId: string;
+   projectStore: ProjectStore;
+   onClose: (newJobId: string | undefined) => void;
+   jobDetails?: JobDetailsV2;
+   jobKey?: string;
+   readOnly?: boolean;
+}
+
 class BuildOptions {
 
    constructor(streamId: string, projectStore: ProjectStore, onClose: (newJobId: string | undefined) => void, jobDetails?: JobDetailsV2, jobKey?: string, readOnly?: boolean) {
@@ -42,6 +51,42 @@ class BuildOptions {
       BuildOptions.instance = this;
       makeObservable(this);
 
+      this.construction = {
+         streamId: streamId,
+         projectStore: projectStore,
+         onClose: onClose,
+         jobDetails: jobDetails,
+         jobKey: jobKey,
+         readOnly: readOnly
+      }
+
+      this.init();
+
+   }
+
+   init() {
+
+      this.change = undefined;
+      this.changeOption = undefined;
+      this.preflightChange = undefined;
+      this.queryShelvedChange = undefined;
+      this.autoSubmit = undefined;
+      this.advJobName = undefined;
+      this.advJobPriority = undefined;
+      this.advUpdateIssues = undefined;
+      this.advAdditionalArgs = undefined;
+      this.advTargets = undefined;
+      this.template = undefined;
+      this.parameters = {};
+      this.disabledParameters = new Set();
+
+      const streamId = this.construction.streamId;
+      const projectStore = this.construction.projectStore;
+      const onClose = this.construction.onClose;
+      const jobDetails = this.construction.jobDetails;
+      const jobKey = this.construction.jobKey;
+      const readOnly = this.construction.readOnly;
+
       const query = new URLSearchParams(window.location.search);
 
       this.isPreflightSubmit = !!query.get("shelvedchange");
@@ -81,8 +126,11 @@ class BuildOptions {
       this.advAdditionalArgs = jobDetails?.jobData?.additionalArguments?.join(" ");
 
       this.load();
+
    }
 
+   construction: OptionsConstruction;
+
    onClose: (newJobId?: string) => void;
 
    loaded = false;
@@ -90,6 +138,7 @@ class BuildOptions {
    jobKey?: string;
    jobDetails?: JobDetailsV2;
    readOnly: boolean;
+   currentRenderKey = 1;
 
    isPreflightSubmit?: boolean;
    isFromP4V?: boolean;
@@ -276,6 +325,20 @@ class BuildOptions {
       this.setParameterChanged();
    }
 
+   templateIdOverride?: string;
+
+   onResetToDefaults() {
+
+      const template = this.template;
+
+      if (!template) {
+         return;
+      }
+
+      this.templateIdOverride = template.id;
+      this.init();
+   }
+
    setValidationErrors(errors: ValidationError[]) {
       this.validationErrors = errors;
       this.setChanged();
@@ -304,6 +367,7 @@ class BuildOptions {
 
    @action
    setChanged() {
+      this.currentRenderKey++;
       this.changed++;
    }
 
@@ -590,7 +654,7 @@ class BuildOptions {
                const t = p as TextParameterData;
                if (t.argument?.toLowerCase().startsWith("-target=")) {
                   if (this.parameters[t.id]?.length) {
-                     targets.add(t.id);
+                     targets.add(this.parameters[t.id]);
                   }
                }
                break;
@@ -799,7 +863,12 @@ class BuildOptions {
 
       let t: GetTemplateRefResponse | undefined;
 
-      if (jobDetails?.template) {
+      if (this.templateIdOverride) {
+         t = this.allTemplates.find(t => t.id === this.templateIdOverride);
+      }
+      this.templateIdOverride = undefined;
+
+      if (!t && jobDetails?.template) {
          t = this.allTemplates.find(t => t.name === jobDetails.template!.name);
       }
       // handle preflight redirect case
@@ -938,18 +1007,18 @@ class BuildOptions {
 
 }
 
-const BoolParameter: React.FC<{ param: BoolParameterData }> = observer(({ param }) => {
+const BoolParameter: React.FC<{ param: BoolParameterData, disabled?: boolean }> = observer(({ param, disabled }) => {
 
    const options = BuildOptions.get();
 
    options.subscribeToParameterChange();
 
-   const key = param.id.replaceAll(".", "-");
+   const key = param.id.replaceAll(".", "-") + `_${options.currentRenderKey}`;
 
    return <Stack>
       <Checkbox key={key}
          label={param.label}
-         disabled={options.readOnly || options.disabledParameters.has(param.id)}
+         disabled={options.readOnly || disabled}
          checked={options.parameters[param.id] == "true"}
          onChange={(ev, value) => {
             ev?.preventDefault();
@@ -960,13 +1029,13 @@ const BoolParameter: React.FC<{ param: BoolParameterData }> = observer(({ param
 
 })
 
-const TextParameter: React.FC<{ param: TextParameterData }> = observer(({ param }) => {
+const TextParameter: React.FC<{ param: TextParameterData, disabled?: boolean }> = observer(({ param, disabled }) => {
 
    const options = BuildOptions.get();
 
    options.subscribeToParameterChange();
 
-   const key = param.id.replaceAll(".", "-");
+   const key = param.id.replaceAll(".", "-") + `_${options.currentRenderKey}`;
 
    let value = options.parameters[param.id] ?? "";
    if (options.disabledParameters.has(param.id)) {
@@ -978,8 +1047,9 @@ const TextParameter: React.FC<{ param: TextParameterData }> = observer(({ param
          placeholder={options.jobDetails ? "" : param.hint}
          label={param.label}
          spellCheck={false}
+         autoComplete="off"
          value={value}
-         disabled={options.readOnly || options.disabledParameters.has(param.id)}
+         disabled={options.readOnly || disabled}
          onChange={(ev, value) => {
             options.onTextChanged(param.id, value ?? "");
          }}
@@ -988,13 +1058,13 @@ const TextParameter: React.FC<{ param: TextParameterData }> = observer(({ param
 })
 
 
-const TagPickerParameter: React.FC<{ param: ListParameterData }> = observer(({ param }) => {
+const TagPickerParameter: React.FC<{ param: ListParameterData, disabled?: boolean }> = observer(({ param, disabled }) => {
 
    const options = BuildOptions.get();
 
    options.subscribeToParameterChange();
 
-   const key = `parameter_key_${param.label}`
+   const key = `parameter_key_${param.label}` + `_${options.currentRenderKey}`;
 
    type PickerItem = {
       itemData: ListParameterItemData;
@@ -1018,7 +1088,7 @@ const TagPickerParameter: React.FC<{ param: ListParameterData }> = observer(({ p
    return <Stack key={key}>
       <Label> {param.label}</Label>
       <TagPicker
-         disabled={options.readOnly}
+         disabled={options.readOnly || disabled}
          onResolveSuggestions={(filter, selected) => {
             return allItems.filter(i => {
                return !selected?.find(s => i.key === s.key) && i.name.toLowerCase().indexOf(filter.toLowerCase()) !== -1;
@@ -1051,12 +1121,12 @@ const TagPickerParameter: React.FC<{ param: ListParameterData }> = observer(({ p
    </Stack>;
 })
 
-const BasicListParameter: React.FC<{ param: ListParameterData }> = observer(({ param }) => {
+const BasicListParameter: React.FC<{ param: ListParameterData, disabled?: boolean }> = observer(({ param, disabled }) => {
 
    const options = BuildOptions.get();
    options.subscribeToParameterChange();
 
-   const key = `parameter_key_${param.label}`
+   const key = `parameter_key_${param.label}` + `_${options.currentRenderKey}`;
 
    const doptions: IDropdownOption[] = [];
 
@@ -1064,15 +1134,14 @@ const BasicListParameter: React.FC<{ param: ListParameterData }> = observer(({ p
       doptions.push({
          key: item.id,
          text: item.text,
-         disabled: options.disabledParameters.has(item.id),
-         selected: options.parameters[item.id] === "true" && !options.disabledParameters.has(item.id)
+         selected: options.parameters[item.id] === "true"
       });
    });
 
    return <Dropdown key={key}
       label={param.label}
       options={doptions}
-      disabled={options.readOnly}
+      disabled={options.readOnly || disabled}
       placeholder={options.jobDetails ? "" : "Select option"}
       onChange={(ev, option, index) => {
 
@@ -1086,7 +1155,7 @@ const BasicListParameter: React.FC<{ param: ListParameterData }> = observer(({ p
       }} />
 })
 
-const MultiListParameter: React.FC<{ param: ListParameterData }> = observer(({ param }) => {
+const MultiListParameter: React.FC<{ param: ListParameterData, disabled?: boolean }> = observer(({ param, disabled }) => {
 
    const options = BuildOptions.get();
    options.subscribeToParameterChange();
@@ -1095,7 +1164,7 @@ const MultiListParameter: React.FC<{ param: ListParameterData }> = observer(({ p
 
    const jobDetails = options.jobDetails;
 
-   const key = `parameter_key_${param.label}`
+   const key = `parameter_key_${param.label}` + `_${options.currentRenderKey}`;
 
    const gset: Set<string> = new Set();
 
@@ -1134,7 +1203,7 @@ const MultiListParameter: React.FC<{ param: ListParameterData }> = observer(({ p
       param.items.forEach(item => {
          if (item.group === group) {
             const key = item.id;
-            const selected = options.parameters[item.id] === "true" && !options.disabledParameters.has(item.id);
+            const selected = options.parameters[item.id] === "true"
             if (selected) {
                selectedKeys.push(key);
             }
@@ -1142,8 +1211,7 @@ const MultiListParameter: React.FC<{ param: ListParameterData }> = observer(({ p
                key: key,
                data: item,
                text: (item.group !== "__nogroup" && dupes.get(item.text)! > 1) ? `${item.text} - ${item.group}` : item.text,
-               selected: selected,
-               disabled: options.disabledParameters.has(item.id)
+               selected: selected
             });
          }
       });
@@ -1151,7 +1219,7 @@ const MultiListParameter: React.FC<{ param: ListParameterData }> = observer(({ p
 
    return <Dropdown
       key={key}
-      disabled={options.readOnly}
+      disabled={options.readOnly || disabled}
       placeholder={jobDetails ? "" : "Select options"}
       styles={{
          callout: {
@@ -1201,15 +1269,15 @@ const MultiListParameter: React.FC<{ param: ListParameterData }> = observer(({ p
 })
 
 
-const ListParameter: React.FC<{ param: ListParameterData }> = ({ param }) => {
+const ListParameter: React.FC<{ param: ListParameterData, disabled?: boolean }> = ({ param, disabled }) => {
 
 
    if (param.style === ListParameterStyle.TagPicker) {
-      return <TagPickerParameter param={param} />
+      return <TagPickerParameter param={param} disabled={disabled} />
    } else if (param.style === ListParameterStyle.MultiList) {
-      return <MultiListParameter param={param} />
+      return <MultiListParameter param={param} disabled={disabled} />
    } else if (param.style === ListParameterStyle.List) {
-      return <BasicListParameter param={param} />
+      return <BasicListParameter param={param} disabled={disabled} />
    }
 
    return null;
@@ -1232,18 +1300,18 @@ const BuildParametersPanel: React.FC = observer(() => {
 
    const estimatedHeight = options.estimateHeight();
 
-   const renderParameter = (param: ParameterData) => {
+   const renderParameter = (param: ParameterData, disabled?: boolean) => {
 
       switch (param.type) {
          case ParameterType.Bool:
             const b = param as BoolParameterData;
-            return <BoolParameter key={`parameter_key_${b.id}`} param={b} />
+            return <BoolParameter key={`parameter_key_${b.id}_${options.currentRenderKey}`} param={b} disabled={disabled} />
          case ParameterType.Text:
             const t = param as TextParameterData;
-            return <TextParameter key={`parameter_key_${t.id}`} param={t} />
+            return <TextParameter key={`parameter_key_${t.id}_${options.currentRenderKey}`} param={t} disabled={disabled} />
          case ParameterType.List:
             const list = param as ListParameterData;
-            return <ListParameter key={`parameter_key_${list.label}`} param={list} />
+            return <ListParameter key={`parameter_key_${list.label}_${options.currentRenderKey}`} param={list} disabled={disabled} />
          default:
             return <Text>Unknown Parameter Type</Text>;
       }
@@ -1268,30 +1336,50 @@ const BuildParametersPanel: React.FC = observer(() => {
 
 
                   let key = "";
+                  let disabled = false;
 
                   switch (p.type) {
                      case ParameterType.Bool:
                         const b = p as BoolParameterData;
                         key = `boolean_parameter_key_${b.id}`;
+                        if (options.disabledParameters.has(b.id)) {
+                           disabled = true;
+                        }
                         break;
                      case ParameterType.Text:
                         const t = p as TextParameterData;
                         key = `text_parameter_key_${t.id}`;
+                        if (options.disabledParameters.has(t.id)) {
+                           disabled = true;
+                        }
                         break;
                      case ParameterType.List:
                         const list = p as ListParameterData;
                         key = `list_parameter_key_${list.label}`;
+                        if (list.items.find(i => options.disabledParameters.has(i.id))) {
+                           disabled = true;
+                        }
                         break;
                      default:
                         return <Text>Unknown Parameter Type</Text>;
                   }
 
-                  if ((p as any).toolTip) {
-                     return <TooltipHost key={key} content={(p as any).toolTip} id={`unique_tooltip_${toolTipId++}`} directionalHint={DirectionalHint.leftCenter}>
-                        {renderParameter(p)}
+                  let toolTip = (p as any).toolTip ?? "";
+                  if (disabled) {
+                     if (toolTip) {
+                        toolTip += "\n(Disabled due to advanced option target modification)";
+                     } else {
+                        toolTip = "Disabled due to advanced option target modification";
+                     }
+
+                  }
+
+                  if (toolTip) {
+                     return <TooltipHost key={key} content={toolTip} id={`unique_tooltip_${toolTipId++}`} directionalHint={DirectionalHint.leftCenter}>
+                        {renderParameter(p, disabled)}
                      </TooltipHost>
                   } else {
-                     return renderParameter(p);
+                     return renderParameter(p, disabled);
                   }
 
                })}
@@ -1540,14 +1628,14 @@ const AdvancedPanel: React.FC = observer(() => {
          </Stack>
          }
          <Stack>
-            <Dropdown disabled={options.readOnly} key={"key_adv_priority"} defaultValue={options.advJobPriority} label="Priority" options={priorityOptions} onChange={(ev, option) => {
+            <Dropdown disabled={options.readOnly} key={`key_adv_priority_${options.currentRenderKey}`} defaultValue={options.advJobPriority} label="Priority" options={priorityOptions} onChange={(ev, option) => {
                options.advJobPriority = option?.key as Priority;
                options.setChanged();
             }
             } />
          </Stack>
          <Stack>
-            <Checkbox disabled={options.readOnly || options.template?.updateIssues} key="key_adv_update_issues"
+            <Checkbox disabled={options.readOnly || options.template?.updateIssues} key={`key_adv_update_issues_${options.currentRenderKey}`}
                label="Update Build Health Issues"
                defaultChecked={options.template?.updateIssues ? true : options.advUpdateIssues}
                onChange={(ev, checked) => {
@@ -1556,7 +1644,7 @@ const AdvancedPanel: React.FC = observer(() => {
                }} />
          </Stack>
          <Stack>
-            <TextField key={"key_adv_job_name"} disabled={options.readOnly} spellCheck={false} defaultValue={options.advJobName} label="Job Name" onChange={(ev, newValue) => {
+            <TextField key={`key_adv_job_name_${options.currentRenderKey}`} disabled={options.readOnly} spellCheck={false} defaultValue={options.advJobName} label="Job Name" onChange={(ev, newValue) => {
                options.advJobName = newValue;
                options.setChanged();
             }} />
@@ -1604,7 +1692,11 @@ const AdvancedPanel: React.FC = observer(() => {
                   };
                }}
 
-               onChange={(items) => { if (items) options.setTargets(items.map(i => i.name)) }}
+               onChange={(items) => {
+                  if (items) {
+                     options.setTargets(items.map(i => i.name))
+                  }
+               }}
 
                onValidateInput={(input?: string) => input ? ValidationState.valid : ValidationState.invalid}
 
@@ -1619,7 +1711,7 @@ const AdvancedPanel: React.FC = observer(() => {
                }}
             />
          </Stack>
-         {showAdditionalArgs && <TextField key={"key_adv_add_args"} multiline style={{ height: 48 }} resizable={false} readOnly={options.readOnly} spellCheck={false} defaultValue={options.advAdditionalArgs} label="Additional Arguments" onChange={(ev, newValue) => options.advAdditionalArgs = newValue} />}
+         {showAdditionalArgs && <TextField key={`key_adv_add_args_${options.currentRenderKey}`} multiline style={{ height: 48 }} resizable={false} readOnly={options.readOnly} spellCheck={false} defaultValue={options.advAdditionalArgs} label="Additional Arguments" onChange={(ev, newValue) => options.advAdditionalArgs = newValue} />}
          {showArgumentClipboardButton && <DefaultButton text="Copy Job Arguments to Clipboard" style={{ width: 240 }} onClick={() => copyToClipboard(
             options.jobDetails?.jobData?.arguments.map(arg => {
                if (arg.indexOf("=") !== -1) {
@@ -1633,8 +1725,8 @@ const AdvancedPanel: React.FC = observer(() => {
                }
             }).join(" ")
          )} />}
-         {!!farguments && <TextField style={{ height: fargumentsHeight }} key={"key_adv_job_arguments"} defaultValue={farguments} readOnly={true} label="Job Arguments" multiline resizable={false} />}
-         {!!fparameters && <TextField style={{ height: fparametersHeight }} key={"key_adv_job_parameters]"} defaultValue={fparameters} readOnly={true} label="Job Parameters" multiline resizable={false} />}
+         {!!farguments && <TextField style={{ height: fargumentsHeight }} key={`key_adv_job_arguments_${options.currentRenderKey}`} defaultValue={farguments} readOnly={true} label="Job Arguments" multiline resizable={false} />}
+         {!!fparameters && <TextField style={{ height: fparametersHeight }} key={`key_adv_job_parameters_${options.currentRenderKey}`} defaultValue={fparameters} readOnly={true} label="Job Parameters" multiline resizable={false} />}
 
          <Stack>
             <Checkbox disabled={options.readOnly}
@@ -1908,9 +2000,9 @@ const BuildModal: React.FC<{ setUseLegacyDialog: (value: boolean) => void }> = o
                   {options.mode === "Advanced" && <AdvancedPanel />}
                </Stack>
                <Stack horizontal tokens={{ childrenGap: 16 }} styles={{ root: { paddingTop: 32, paddingLeft: 8, paddingBottom: 8 } }}>
-                  <Stack>
-                     <DefaultButton disabled={options.readOnly || !!options.jobDetails} text="Use Legacy Dialog" style={{ width: 160 }} onClick={() => setUseLegacyDialog(true)} />
-                  </Stack>
+                  {!options.jobDetails && !options.readOnly && <Stack>
+                     <DefaultButton text="Reset" style={{ width: 160 }} onClick={() => options.onResetToDefaults()} />
+                  </Stack>}
 
                   <Stack grow />
                   <PrimaryButton text="Start Job" disabled={!template || options.readOnly || options.preview} onClick={() => { onSubmit(); }} />
