Horde: Add a fingerprint description as part of issue response

diff --git a/Engine/Source/Programs/Horde/Horde.Server/Issues/IssuesController.cs b/Engine/Source/Programs/Horde/Horde.Server/Issues/IssuesController.cs
--- a/Engine/Source/Programs/Horde/Horde.Server/Issues/IssuesController.cs
+++ b/Engine/Source/Programs/Horde/Horde.Server/Issues/IssuesController.cs
@@ -535,6 +535,15 @@ namespace Horde.Server.Issues
 			response.QuarantineTimeUtc = details.QuarantineTimeUtc;
 			response.ForceClosedByUserInfo = details.ForceClosedBy?.ToThinApiResponse();
 			response.WorkflowThreadUrl = issue.WorkflowThreadUrl;
+			
+			if (issue.Fingerprints != null && issue.Fingerprints.Count > 0)
+			{
+				response.FingerprintDescription = String.Join(", ", issue.Fingerprints.Select(x =>
+				{
+					return $"(Type: {x.Type} / Keys: {String.Join(", ", x.Keys)} / RejectKeys: {String.Join(", ", x.RejectKeys ?? new HashSet<IssueKey>())})";
+				}));
+			}
+
 			return response;
 		}
 
diff --git a/Engine/Source/Programs/Shared/EpicGames.Horde/Issues/IssueMessages.cs b/Engine/Source/Programs/Shared/EpicGames.Horde/Issues/IssueMessages.cs
--- a/Engine/Source/Programs/Shared/EpicGames.Horde/Issues/IssueMessages.cs
+++ b/Engine/Source/Programs/Shared/EpicGames.Horde/Issues/IssueMessages.cs
@@ -274,6 +274,11 @@ namespace EpicGames.Horde.Issues
 		/// </summary>
 		public string? Description { get; set; }
 
+		/// <summary>
+		/// Description of the current fingerprint used for issue identification
+		/// </summary>
+		public string? FingerprintDescription { get; set; }
+
 		/// <summary>
 		/// Severity of this issue
 		/// </summary>
