Horde: Add fingerprint type to issue modal
#rnx
#skipci

diff --git a/Engine/Source/Programs/Horde/HordeDashboard/src/backend/Api.ts b/Engine/Source/Programs/Horde/HordeDashboard/src/backend/Api.ts
--- a/Engine/Source/Programs/Horde/HordeDashboard/src/backend/Api.ts
+++ b/Engine/Source/Programs/Horde/HordeDashboard/src/backend/Api.ts
@@ -3134,6 +3134,9 @@ export type GetIssueResponse = {
 	/**Description of the issue*/
 	description?: string;
 
+	/**Description of the fingerprint used to identify this issue*/
+	fingerprintDescription?: string;
+
 	/** Severity of this issue	*/
 	severity: IssueSeverity;
 
diff --git a/Engine/Source/Programs/Horde/HordeDashboard/src/components/IssueViewV2.tsx b/Engine/Source/Programs/Horde/HordeDashboard/src/components/IssueViewV2.tsx
--- a/Engine/Source/Programs/Horde/HordeDashboard/src/components/IssueViewV2.tsx
+++ b/Engine/Source/Programs/Horde/HordeDashboard/src/components/IssueViewV2.tsx
@@ -1330,9 +1330,18 @@ const IssueSummaryPanel: React.FC = () => {
       }
    }
 
+   if (issue.fingerprintDescription) {
+      const components = issue.fingerprintDescription.split(" / ");
+      if (components.length) {         
+         items.push({
+            title: "Fingerprint",
+            text: components[0].replace("Type:", "").trim()
+         })   
+      }
+   }
+
    if (!items.length && !issue.description) {
       return null;
-
    }
 
    return <Stack style={{ flexBasis: "70px", flexShrink: 0 }}>
