Horde: Update render key only when resetting optons
#rnx
#submitool bypass
#skipci

diff --git a/Engine/Source/Programs/Horde/HordeDashboard/src/components/build/NewBuildV2.tsx b/Engine/Source/Programs/Horde/HordeDashboard/src/components/build/NewBuildV2.tsx
--- a/Engine/Source/Programs/Horde/HordeDashboard/src/components/build/NewBuildV2.tsx
+++ b/Engine/Source/Programs/Horde/HordeDashboard/src/components/build/NewBuildV2.tsx
@@ -329,6 +329,8 @@ class BuildOptions {
 
    onResetToDefaults() {
 
+      this.currentRenderKey++;
+
       const template = this.template;
 
       if (!template) {
@@ -366,8 +368,7 @@ class BuildOptions {
    private changed = 0;
 
    @action
-   setChanged() {
-      this.currentRenderKey++;
+   setChanged() {      
       this.changed++;
    }
 
