Horde: Fix issue with change option not resetting
#rnx
#skipci

diff --git a/Engine/Source/Programs/Horde/HordeDashboard/src/components/build/NewBuildV2.tsx b/Engine/Source/Programs/Horde/HordeDashboard/src/components/build/NewBuildV2.tsx
--- a/Engine/Source/Programs/Horde/HordeDashboard/src/components/build/NewBuildV2.tsx
+++ b/Engine/Source/Programs/Horde/HordeDashboard/src/components/build/NewBuildV2.tsx
@@ -1954,7 +1954,7 @@ const BuildModal: React.FC<{ setUseLegacyDialog: (value: boolean) => void }> = o
                         <TemplateSelector />
                      </Stack>
                      <Stack>
-                        <ComboBox style={{ width: 204 }} label="Change" text={changeText} options={changeItems} disabled={options.readOnly} allowFreeform autoComplete="off" placeholder="Latest Change" defaultValue={"-1"} onChange={(ev, option, index, value) => {
+                        <ComboBox style={{ width: 204 }} key={`build_change_key_${options.currentRenderKey}`} label="Change" text={changeText} options={changeItems} disabled={options.readOnly} allowFreeform autoComplete="off" placeholder="Latest Change" defaultValue={"-1"} onChange={(ev, option, index, value) => {
                            ev.preventDefault();
                            if (option) {
                               options.change = undefined;
@@ -1976,6 +1976,7 @@ const BuildModal: React.FC<{ setUseLegacyDialog: (value: boolean) => void }> = o
 
                      <Stack>
                         <TextField
+                           key={`build_shelved_change_key_${options.currentRenderKey}`} 
                            label="Shelved Change"
                            style={{ width: 146 }}
                            placeholder={!template!.allowPreflights ? "Disabled by template" : "None"}
