Horde: Only enable debug controllers if enabled in settings

diff --git a/Engine/Source/Programs/Horde/Horde.Server.Tests/Server/DebugControllerTests.cs b/Engine/Source/Programs/Horde/Horde.Server.Tests/Server/DebugControllerTests.cs
new file mode 100644
--- /dev/null
+++ b/Engine/Source/Programs/Horde/Horde.Server.Tests/Server/DebugControllerTests.cs
@@ -0,0 +1,51 @@
+// Copyright Epic Games, Inc. All Rights Reserved.
+
+using System;
+using System.Collections.Generic;
+using System.Net;
+using System.Net.Http;
+using System.Threading.Tasks;
+using Microsoft.VisualStudio.TestTools.UnitTesting;
+
+#pragma warning disable CA2234 // Pass system uri objects instead of strings
+
+namespace Horde.Server.Tests.Server;
+
+[TestClass]
+public class DebugControllerTests
+{
+	[TestMethod]
+	public async Task DebugEndpointsEnabledAsync()
+	{
+		await using FakeHordeWebApp app = CreateApp(true);
+		using HttpResponseMessage response = await app.CreateHttpClient().GetAsync("/api/v1/debug/aclscopes");
+		response.EnsureSuccessStatusCode();
+	}
+	
+	[TestMethod]
+	public async Task DebugEndpointsDisabledAsync()
+	{
+		await using FakeHordeWebApp app = CreateApp(null);
+		using HttpResponseMessage response = await app.CreateHttpClient().GetAsync("/api/v1/debug/aclscopes");
+		Assert.AreEqual(HttpStatusCode.Forbidden, response.StatusCode);
+	}
+	
+	[TestMethod]
+	public async Task DebugEndpointsDisabledByDefaultAsync()
+	{
+		await using FakeHordeWebApp app = CreateApp(null);
+		using HttpResponseMessage response = await app.CreateHttpClient().GetAsync("/api/v1/debug/aclscopes");
+		Assert.AreEqual(HttpStatusCode.Forbidden, response.StatusCode);
+	}
+	
+	private static FakeHordeWebApp CreateApp(bool? enableDebugEndpoints)
+	{
+		Dictionary<string, string> settings = new();
+		if (enableDebugEndpoints != null)
+		{
+			settings["Horde:EnableDebugEndpoints"] = Convert.ToString(enableDebugEndpoints.Value);
+		}
+		return new FakeHordeWebApp(settings);
+	}
+}
+
diff --git a/Engine/Source/Programs/Horde/Horde.Server/Server/DebugController.cs b/Engine/Source/Programs/Horde/Horde.Server/Server/DebugController.cs
--- a/Engine/Source/Programs/Horde/Horde.Server/Server/DebugController.cs
+++ b/Engine/Source/Programs/Horde/Horde.Server/Server/DebugController.cs
@@ -33,6 +33,7 @@ using JetBrains.Profiler.SelfApi;
 using Microsoft.AspNetCore.Authorization;
 using Microsoft.AspNetCore.Http;
 using Microsoft.AspNetCore.Mvc;
+using Microsoft.AspNetCore.Mvc.Filters;
 using Microsoft.Extensions.Logging;
 using Microsoft.Extensions.Options;
 using MongoDB.Bson;
@@ -245,17 +246,54 @@ namespace Horde.Server.Server
 		/**/
 	}
 #endif
+
+	/// <summary>
+	/// Only requests to attached controller to pass if debug endpoint is enabled in settings
+	/// Adds extra security for not enabling these admin endpoints by accident.
+	/// </summary>
+	[AttributeUsage(AttributeTargets.Class)]
+	public sealed class DebugEndpointCheckAttribute : Attribute, IActionFilter
+	{
+		/// <inheritdoc />
+		public void OnActionExecuting(ActionExecutingContext context)
+		{
+			if (context.Controller is SecureDebugController controller)
+			{
+				if (!controller.ServerSettings.Value.EnableDebugEndpoints)
+				{
+					context.Result = new ForbidResult();
+				}
+			}
+			else
+			{
+				// Assume forbidden if controller is not resolved
+				context.Result = new ForbidResult();
+			}
+		}
+
+		/// <inheritdoc />
+		public void OnActionExecuted(ActionExecutedContext context)
+		{
+		}
+	}
+	
 #if ENABLE_SECURE_DEBUG_CONTROLLER
 	/// <summary>
 	/// Controller managing account status
 	/// </summary>
 	[ApiController]
 	[Authorize]
+	[DebugEndpointCheck]
 	[Tags("Debug")]
 	public class SecureDebugController : HordeControllerBase
 	{
 		private static readonly Random s_random = new();
 
+		/// <summary>
+		/// Server settings (exposed for debug endpoint check filter)
+		/// </summary>
+		public IOptionsSnapshot<ServerSettings> ServerSettings { get; private set; }
+		
 		private readonly MongoService _mongoService;
 		private readonly ConfigService _configService;
 		private readonly AgentRelayService _agentRelayService;
@@ -275,7 +313,11 @@ namespace Horde.Server.Server
 			AgentRelayService agentRelayService,
 			JobService jobService,
 			JobTaskSource jobTaskSource,
-			ILogCollection logCollection, ILeaseCollection leaseCollection, IOptionsSnapshot<GlobalConfig> globalConfig, ILogger<SecureDebugController> logger)
+			ILogCollection logCollection,
+			ILeaseCollection leaseCollection,
+			IOptionsSnapshot<ServerSettings> serverSettings,
+			IOptionsSnapshot<GlobalConfig> globalConfig,
+			ILogger<SecureDebugController> logger)
 		{
 			_mongoService = mongoService;
 			_configService = configService;
@@ -284,6 +326,7 @@ namespace Horde.Server.Server
 			_jobTaskSource = jobTaskSource;
 			_logCollection = logCollection;
 			_leaseCollection = leaseCollection;
+			ServerSettings = serverSettings;
 			_globalConfig = globalConfig;
 			_logger = logger;
 		}
diff --git a/Engine/Source/Programs/Horde/Horde.Server/ServerSettings.cs b/Engine/Source/Programs/Horde/Horde.Server/ServerSettings.cs
--- a/Engine/Source/Programs/Horde/Horde.Server/ServerSettings.cs
+++ b/Engine/Source/Programs/Horde/Horde.Server/ServerSettings.cs
@@ -645,6 +645,11 @@ namespace Horde.Server
 		/// </summary>
 		public string CorsOrigin { get; set; } = null!;
 
+		/// <summary>
+		/// Whether to enable debug/administrative REST API endpoints
+		/// </summary>
+		public bool EnableDebugEndpoints { get; set; } = false;
+
 		/// <summary>
 		/// Whether to automatically enable new agents by default. If false, new agents must manually be enabled before they can take on work.
 		/// </summary>
