Horde: Skip cleaning files for AutoSDK workspace

A significant amount of time is spent on cleaning these workspaces with files never changing. They cannot be preflighted either. By disabling this for AutoSDK workspaces, we can increase utilization.

diff --git a/Engine/Source/Programs/Horde/Drivers/JobDriver/Execution/ManagedWorkspaceMaterializer.cs b/Engine/Source/Programs/Horde/Drivers/JobDriver/Execution/ManagedWorkspaceMaterializer.cs
--- a/Engine/Source/Programs/Horde/Drivers/JobDriver/Execution/ManagedWorkspaceMaterializer.cs
+++ b/Engine/Source/Programs/Horde/Drivers/JobDriver/Execution/ManagedWorkspaceMaterializer.cs
@@ -18,6 +18,7 @@ public sealed class ManagedWorkspaceMaterializer : IWorkspaceMaterializer
 	private readonly RpcAgentWorkspace _agentWorkspace;
 	private readonly DirectoryReference _workingDir;
 	private readonly bool _useCacheFile;
+	private readonly bool _cleanDuringFinalize;
 	private WorkspaceInfo? _workspace;
 
 	/// <summary>
@@ -26,14 +27,17 @@ public sealed class ManagedWorkspaceMaterializer : IWorkspaceMaterializer
 	/// <param name="agentWorkspace">Workspace configuration</param>
 	/// <param name="workingDir">Where to put synced Perforce files and any cached data/metadata</param>
 	/// <param name="useCacheFile">Whether to use a cache file during syncs</param>
+	/// <param name="cleanDuringFinalize">Whether to clean and revert files during finalize</param>
 	public ManagedWorkspaceMaterializer(
 		RpcAgentWorkspace agentWorkspace,
 		DirectoryReference workingDir,
-		bool useCacheFile)
+		bool useCacheFile,
+		bool cleanDuringFinalize)
 	{
 		_agentWorkspace = agentWorkspace;
 		_workingDir = workingDir;
 		_useCacheFile = useCacheFile;
+		_cleanDuringFinalize = cleanDuringFinalize;
 	}
 
 	/// <inheritdoc/>
@@ -57,7 +61,7 @@ public sealed class ManagedWorkspaceMaterializer : IWorkspaceMaterializer
 	{
 		using IScope scope = CreateTraceSpan("ManagedWorkspaceMaterializer.FinalizeAsync");
 
-		if (_workspace != null)
+		if (_workspace != null && _cleanDuringFinalize)
 		{
 			await _workspace.CleanAsync(cancellationToken);
 		}
@@ -146,6 +150,7 @@ public sealed class ManagedWorkspaceMaterializer : IWorkspaceMaterializer
 		scope.Span.SetTag("Stream", _agentWorkspace.Stream);
 		scope.Span.SetTag("Partitioned", _agentWorkspace.Partitioned);
 		scope.Span.SetTag("UseCacheFile", _useCacheFile);
+		scope.Span.SetTag("CleanDuringFinalize", _cleanDuringFinalize);
 		return scope;
 	}
 }
\ No newline at end of file
diff --git a/Engine/Source/Programs/Horde/Drivers/JobDriver/Execution/WorkspaceMaterializer.cs b/Engine/Source/Programs/Horde/Drivers/JobDriver/Execution/WorkspaceMaterializer.cs
--- a/Engine/Source/Programs/Horde/Drivers/JobDriver/Execution/WorkspaceMaterializer.cs
+++ b/Engine/Source/Programs/Horde/Drivers/JobDriver/Execution/WorkspaceMaterializer.cs
@@ -169,15 +169,11 @@ class WorkspaceMaterializerFactory : IWorkspaceMaterializerFactory
 	/// <inheritdoc/>
 	public IWorkspaceMaterializer CreateMaterializer(WorkspaceMaterializerType type, RpcAgentWorkspace workspaceInfo, JobExecutorOptions options, bool forAutoSdk)
 	{
-		switch (type)
+		return type switch
 		{
-			case WorkspaceMaterializerType.ManagedWorkspace:
-				return forAutoSdk
-					? new ManagedWorkspaceMaterializer(workspaceInfo, options.WorkingDir, true)
-					: new ManagedWorkspaceMaterializer(workspaceInfo, options.WorkingDir, false);
-
-			default:
-				throw new Exception("Unhandled materializer option: " + type);
-		}
+			WorkspaceMaterializerType.ManagedWorkspace when forAutoSdk => new ManagedWorkspaceMaterializer(workspaceInfo, options.WorkingDir, true, false),
+			WorkspaceMaterializerType.ManagedWorkspace => new ManagedWorkspaceMaterializer(workspaceInfo, options.WorkingDir, false, true),
+			_ => throw new Exception("Unhandled materializer option: " + type)
+		};
 	}
 }
