Horde: Remove deprecated setting for processes to terminate.

#rnx

diff --git a/Engine/Source/Programs/Horde/Drivers/JobDriver/DriverSettings.cs b/Engine/Source/Programs/Horde/Drivers/JobDriver/DriverSettings.cs
--- a/Engine/Source/Programs/Horde/Drivers/JobDriver/DriverSettings.cs
+++ b/Engine/Source/Programs/Horde/Drivers/JobDriver/DriverSettings.cs
@@ -25,12 +25,6 @@ namespace JobDriver
 		/// </summary>
 		public PerforceExecutorSettings PerforceExecutor { get; set; } = new PerforceExecutorSettings();
 
-		/// <summary>
-		/// List of process names to terminate after a job
-		/// </summary>
-		[Obsolete("Prefer using the ProcessesToTerminate list instead")]
-		public List<string> ProcessNamesToTerminate { get; } = new List<string>();
-
 		/// <summary>
 		/// List of process names to terminate after a lease completes, but not after a job step
 		/// </summary>
