Horde: Add more handlers for horde ids in log render
#skipci
#rnx

diff --git a/Engine/Source/Programs/Horde/HordeDashboard/src/components/LogRender.tsx b/Engine/Source/Programs/Horde/HordeDashboard/src/components/LogRender.tsx
--- a/Engine/Source/Programs/Horde/HordeDashboard/src/components/LogRender.tsx
+++ b/Engine/Source/Programs/Horde/HordeDashboard/src/components/LogRender.tsx
@@ -9,10 +9,15 @@ enum TagType {
    None,
    SourceFile,
    MSDNCode,
+   Link,
+   JobId,
+   LogId,
+   ArtifactId,
    AgentId,
+   PoolId,
    LeaseId,
-   Link,
-   ArtifactId
+   ProjectId,
+   StreamId,
 }
 
 export type LogItem = {
@@ -153,6 +158,27 @@ const renderTags = (navigate: NavigateFunction, line: LogLine, lineNumber: numbe
             tagType = TagType.AgentId;
          }
 
+         if (type === "PoolId") {
+            tagType = TagType.PoolId;
+         }
+
+         if (type === "StreamId") {
+            tagType = TagType.StreamId;
+         }
+
+         if (type === "ProjectId") {
+            tagType = TagType.ProjectId;
+         }
+
+         if (type === "JobId") {
+            tagType = TagType.JobId;
+         }
+
+         if (type === "LogId") {
+            debugger;
+            tagType = TagType.LogId;
+         }
+
          if (type === "ArtifactId") {
             tagType = TagType.ArtifactId;
          }
@@ -163,6 +189,21 @@ const renderTags = (navigate: NavigateFunction, line: LogLine, lineNumber: numbe
 
          return <Highlight key={key} search={search ? search : ""} className={logStyle.logLine}>{text}</Highlight>;
 
+      } else if (tagType === TagType.JobId) {
+         const url = `/job/${text}`
+         return <a key={key} href={url} onClick={(ev) => { ev.stopPropagation(); ev.preventDefault(); navigate(url) }}><Highlight search={search ? search : ""} className={logStyle.logLine}>{text}</Highlight></a>;
+      } else if (tagType === TagType.LogId) {
+         const url = `/log/${text}`
+         return <a key={key} href={url} onClick={(ev) => { ev.stopPropagation(); ev.preventDefault(); navigate(url) }}><Highlight search={search ? search : ""} className={logStyle.logLine}>{text}</Highlight></a>;
+      } else if (tagType === TagType.ProjectId) {
+         const url = `/project/${text}`
+         return <a key={key} href={url} onClick={(ev) => { ev.stopPropagation(); ev.preventDefault(); navigate(url) }}><Highlight search={search ? search : ""} className={logStyle.logLine}>{text}</Highlight></a>;
+      } else if (tagType === TagType.StreamId) {
+         const url = `/stream/${text}`
+         return <a key={key} href={url} onClick={(ev) => { ev.stopPropagation(); ev.preventDefault(); navigate(url) }}><Highlight search={search ? search : ""} className={logStyle.logLine}>{text}</Highlight></a>;
+      } else if (tagType === TagType.PoolId) {
+         const url = `/pools?poolid=${text}`
+         return <a key={key} href={url} onClick={(ev) => { ev.stopPropagation(); ev.preventDefault(); navigate(url) }}><Highlight search={search ? search : ""} className={logStyle.logLine}>{text}</Highlight></a>;
       } else if (tagType === TagType.LeaseId) {
 
          const navigateToLeaseLog = async (toplevel: boolean) => {
@@ -226,9 +267,9 @@ const renderTags = (navigate: NavigateFunction, line: LogLine, lineNumber: numbe
       } else if (tagType === TagType.Link) {
 
          return <a key={key} rel="noreferrer" href={record.target} onClick={(ev) => ev.stopPropagation()}><Highlight search={search ? search : ""} className={logStyle.logLine}>{text}</Highlight></a>;
-      } else if (tagType === TagType.ArtifactId) {         
+      } else if (tagType === TagType.ArtifactId) {
          const artifactType = properties.ArtifactType;
-         if (artifactType?.length) {            
+         if (artifactType?.length) {
             const search = new URLSearchParams(window.location.search);
             search.set("artifactContext", encodeURIComponent(artifactType as string));
             const url = `${window.location.pathname}?` + search.toString();
@@ -259,7 +300,7 @@ const renderTags = (navigate: NavigateFunction, line: LogLine, lineNumber: numbe
          return t;
       }
 
-      const rtags:any = [];
+      const rtags: any = [];
 
       const key = `log_line_${lineNumber}_${idx}_${index}_fragment`;
 
