Horde: Add optional dark theme logo support
#rnx

diff --git a/Engine/Source/Programs/Horde/Horde.Server/Projects/ProjectConfig.cs b/Engine/Source/Programs/Horde/Horde.Server/Projects/ProjectConfig.cs
--- a/Engine/Source/Programs/Horde/Horde.Server/Projects/ProjectConfig.cs
+++ b/Engine/Source/Programs/Horde/Horde.Server/Projects/ProjectConfig.cs
@@ -71,6 +71,11 @@ namespace Horde.Server.Projects
 		/// </summary>
 		public ConfigResource? Logo { get; set; }
 
+		/// <summary>
+		/// Optional path to the project logo for the dark theme 
+		/// </summary>
+		public ConfigResource? LogoDarkTheme { get; set; }
+
 		/// <summary>
 		/// List of pools for this project
 		/// </summary>
diff --git a/Engine/Source/Programs/Horde/Horde.Server/Projects/ProjectsController.cs b/Engine/Source/Programs/Horde/Horde.Server/Projects/ProjectsController.cs
--- a/Engine/Source/Programs/Horde/Horde.Server/Projects/ProjectsController.cs
+++ b/Engine/Source/Programs/Horde/Horde.Server/Projects/ProjectsController.cs
@@ -102,10 +102,11 @@ namespace Horde.Server.Projects
 		/// Retrieve information about a specific project
 		/// </summary>
 		/// <param name="projectId">Id of the project to get information about</param>
+		/// <param name="darkTheme">Whether to prefer darktheme logo, if available</param>
 		/// <returns>Information about the requested project</returns>
 		[HttpGet]
 		[Route("/api/v1/projects/{projectId}/logo")]
-		public ActionResult<object> GetProjectLogo(ProjectId projectId)
+		public ActionResult<object> GetProjectLogo(ProjectId projectId, [FromQuery] bool darkTheme = false)
 		{
 			ProjectConfig? projectConfig;
 			if (!_globalConfig.Value.TryGetProject(projectId, out projectConfig))
@@ -117,7 +118,7 @@ namespace Horde.Server.Projects
 				return Forbid(ProjectAclAction.ViewProject, projectId);
 			}
 
-			ConfigResource? logoResource = projectConfig.Logo;
+			ConfigResource? logoResource = darkTheme ? projectConfig.LogoDarkTheme ?? projectConfig.Logo : projectConfig.Logo;
 			if (logoResource == null || logoResource.Path == null || logoResource.Data.Length == 0)
 			{
 				return NotFound("Missing logo resource data");
