Horde: Add support for dark theme logos to dashboard
#rnx
#skipci

diff --git a/Engine/Source/Programs/Horde/HordeDashboard/src/components/HomeView.tsx b/Engine/Source/Programs/Horde/HordeDashboard/src/components/HomeView.tsx
--- a/Engine/Source/Programs/Horde/HordeDashboard/src/components/HomeView.tsx
+++ b/Engine/Source/Programs/Horde/HordeDashboard/src/components/HomeView.tsx
@@ -9,6 +9,7 @@ import { projectStore } from '../backend/ProjectStore';
 import { ProjectData } from '../backend/Api';
 import { Breadcrumbs } from './Breadcrumbs';
 import { getHordeStyling } from '../styles/Styles';
+import dashboard from '../backend/Dashboard';
 
 export const HomeView: React.FC = observer(() => {
 
@@ -26,11 +27,12 @@ export const HomeView: React.FC = observer(() => {
 						{
 							projectStore.projects.sort((a: ProjectData, b: ProjectData) => {
 								return a.order - b.order;
-							}).map((project) => {
+                     }).map((project) => {
+                        const logoSrc = `/api/v1/projects/${project.id}/logo${dashboard.darktheme ? "?darkTheme=true" : ""}`;
 								return (
 									<Stack.Item key={project.id} className={hordeClasses.projectLogoCardDropShadow}>
 										<Link onClick={() => { projectStore.setActive(project.id); }} to={`/project/${project.id}`}>
-											<img src={`/api/v1/projects/${project.id}/logo`} alt="Project logo" width={560} height={280} />
+											<img src={logoSrc} alt="Project logo" width={560} height={280} />
 										</Link>
 									</Stack.Item>
 								);
diff --git a/Engine/Source/Programs/Horde/HordeDashboard/src/components/ProjectHome.tsx b/Engine/Source/Programs/Horde/HordeDashboard/src/components/ProjectHome.tsx
--- a/Engine/Source/Programs/Horde/HordeDashboard/src/components/ProjectHome.tsx
+++ b/Engine/Source/Programs/Horde/HordeDashboard/src/components/ProjectHome.tsx
@@ -15,144 +15,145 @@ import { getHordeTheme } from '../styles/theme';
 import { getHordeStyling } from '../styles/Styles';
 
 type FilteredCategory = {
-    key: string;
-    name: string;
-    streams: StreamData[];
+   key: string;
+   name: string;
+   streams: StreamData[];
 };
 
 class LocalState {
    constructor() {
       makeObservable(this);
    }
-    filterString = "";
-    activeProject: ProjectData | undefined = undefined;
-    @observable.shallow filteredCategories: FilteredCategory[] = [];
-
-    async setActiveProject(project?: ProjectData) {
-        if (this.activeProject?.id !== project?.id) {
-            backend.getProject(project!).then(projectWithMetadata => {
-                this.activeProject = projectWithMetadata;
-                this._filterStreams();
+   filterString = "";
+   activeProject: ProjectData | undefined = undefined;
+   @observable.shallow filteredCategories: FilteredCategory[] = [];
+
+   async setActiveProject(project?: ProjectData) {
+      if (this.activeProject?.id !== project?.id) {
+         backend.getProject(project!).then(projectWithMetadata => {
+            this.activeProject = projectWithMetadata;
+            this._filterStreams();
+         });
+      }
+   }
+
+   setFilterString(newFilter: string) {
+      if (this.filterString.toLowerCase() !== newFilter.toLowerCase()) {
+         this.filterString = newFilter;
+         this._filterStreams();
+      }
+   }
+
+   @action
+   private _filterStreams() {
+      if (this.activeProject?.streams) {
+         this.filteredCategories = [];
+         if (this.activeProject.categories) {
+            this.activeProject.categories?.forEach(category => {
+               let streams: StreamData[] = [];
+               category.streams!.forEach(stream => {
+                  let foundStream = this.activeProject!.streams!.find(s => s.id === stream && s.name.toLowerCase().indexOf(this.filterString.toLowerCase()) !== -1);
+                  if (foundStream) {
+                     streams.push(foundStream);
+                  }
+               });
+               if (streams.length !== 0) {
+                  this.filteredCategories.push({ key: category.name!, name: category.name!, streams: streams });
+               }
+               streams.sort((a, b) => a.name < b.name ? -1 : 1)
             });
-        }
-    }
+         }
+      }
 
-    setFilterString(newFilter: string) {
-        if (this.filterString.toLowerCase() !== newFilter.toLowerCase()) {
-            this.filterString = newFilter;
-            this._filterStreams();
-        }
-    }
-
-    @action
-    private _filterStreams() {
-        if (this.activeProject?.streams) {
-            this.filteredCategories = [];
-            if (this.activeProject.categories) {
-                this.activeProject.categories?.forEach(category => {
-                    let streams: StreamData[] = [];
-                    category.streams!.forEach(stream => {
-                        let foundStream = this.activeProject!.streams!.find(s => s.id === stream && s.name.toLowerCase().indexOf(this.filterString.toLowerCase()) !== -1);
-                        if (foundStream) {
-                            streams.push(foundStream);
-                        }
-                    });
-                    if (streams.length !== 0) {
-                        this.filteredCategories.push({ key: category.name!, name: category.name!, streams: streams });
-                    }
-                    streams.sort((a, b) => a.name < b.name ? -1 : 1)
-                });
-            }
-        }
-
-    }
+   }
 }
 
 
 const localState = new LocalState();
 
 export const ProjectHome: React.FC = observer(() => {
-    const { projectId } = useParams<{ projectId: string }>();
+   const { projectId } = useParams<{ projectId: string }>();
    const { projectStore } = useBackend();
 
-   const { hordeClasses, detailClasses } = getHordeStyling();   
+   const { hordeClasses, detailClasses } = getHordeStyling();
    const hordeTheme = getHordeTheme();
 
-    const activeProject = projectStore.byId(projectId);
+   const activeProject = projectStore.byId(projectId);
 
-    if (!activeProject || activeProject.id !== projectId) {
-        projectStore.setActive(projectId);
-    }
-    localState.setActiveProject(activeProject);
+   if (!activeProject || activeProject.id !== projectId) {
+      projectStore.setActive(projectId);
+   }
+   localState.setActiveProject(activeProject);
 
-    const crumbs: BreadcrumbItem[] = [];
-    if (activeProject) {
-        crumbs.push({ text: activeProject.name });
-    }
+   const crumbs: BreadcrumbItem[] = [];
+   if (activeProject) {
+      crumbs.push({ text: activeProject.name });
+   }
 
-    let categoryCount = localState.filteredCategories.length;
-    let width = (100 / categoryCount);
+   let categoryCount = localState.filteredCategories.length;
+   let width = (100 / categoryCount);
 
-    let categoryStackItems: any[] = [];
+   let categoryStackItems: any[] = [];
 
-    for (let idx = 0; idx < categoryCount; idx++) {
-        let category = localState.filteredCategories[idx];
+   for (let idx = 0; idx < categoryCount; idx++) {
+      let category = localState.filteredCategories[idx];
 
-        let links: any[] = [];
-        category.streams.forEach(stream => {
-            links.push(<Link style={{ marginLeft: 10, fontSize: 15 }} to={`/stream/${stream.id}`}>{stream.name}</Link>);
-        })
+      let links: any[] = [];
+      category.streams.forEach(stream => {
+         links.push(<Link style={{ marginLeft: 10, fontSize: 15 }} to={`/stream/${stream.id}`}>{stream.name}</Link>);
+      })
 
-        categoryStackItems.push(
-            <Stack key={'stack_' + idx} styles={{ root: { width: `${width}%` } }} tokens={{ childrenGap: 5 }}>
-                <Text variant="mediumPlus" styles={{ root: { fontFamily: "Horde Open Sans SemiBold" } }}>{category.name}</Text>
-                { links }
-            </Stack>
-        );
-    }
-
-    const imgStyle = dashboard.darktheme ? { filter: "invert(1)" } : {};
-
-    const windowSize = useWindowSize();
-
-    const vw = Math.max(document.documentElement.clientWidth, window.innerWidth || 0);
-
-    return (
-        <Stack className={hordeClasses.horde}>
-            <TopNav />
-            <Breadcrumbs items={crumbs} />
-            <Stack horizontal style={{backgroundColor: hordeTheme.horde.neutralBackground}}>
-                <div key={`windowsize_streamview_${windowSize.width}_${windowSize.height}`} style={{ width: vw / 2 - (1440/2), flexShrink: 0 }} />
-                <Stack tokens={{ childrenGap: 0 }} styles={{ root: { width: "100%" } }}>
-                    <Stack style={{ maxWidth: 1440, paddingTop: 6, marginLeft: 4 }}>
-                        <FocusZone direction={FocusZoneDirection.vertical} style={{ padding: 0 }}>
-                            <div className={detailClasses.container} style={{ width: "100%", height: 'calc(100vh - 208px)', position: 'relative' }} data-is-scrollable={true}>
-                                <ScrollablePane scrollbarVisibility={ScrollbarVisibility.auto} onScroll={() => { }}>
-                                    <Stack tokens={{ childrenGap: 18 }} style={{ padding: 0 }}>
-                                        <Stack.Item align={'center'} className={hordeClasses.projectLogoCard}>
-                                            {
-                                                activeProject ? <img style={imgStyle} src={`/api/v1/projects/${activeProject!.id}/logo`} alt="Project logo" height={300} /> : <div></div>            
-                                            }
-                                        </Stack.Item>
-                                        <Stack tokens={{ childrenGap: 20 }} className={hordeClasses.raised}>
-                                            <SearchBox
-                                                placeholder="Search"
-                                                onChange={(event?: React.ChangeEvent<HTMLInputElement> | undefined, newValue?: string | undefined) => { localState.setFilterString(newValue ?? ""); }}
-                                                onClear={() => { localState.setFilterString(""); }}
-                                            />
-                                            <Stack>
-                                                <Stack horizontal>
-                                                    {categoryStackItems}
-                                                </Stack>
-                                            </Stack>
-                                        </Stack>
+      categoryStackItems.push(
+         <Stack key={'stack_' + idx} styles={{ root: { width: `${width}%` } }} tokens={{ childrenGap: 5 }}>
+            <Text variant="mediumPlus" styles={{ root: { fontFamily: "Horde Open Sans SemiBold" } }}>{category.name}</Text>
+            {links}
+         </Stack>
+      );
+   }
+
+   let logoSrc = "";
+   if (activeProject) {
+      logoSrc = `/api/v1/projects/${activeProject.id}/logo${dashboard.darktheme ? "?darkTheme=true" : ""}`;
+   }
+
+   const windowSize = useWindowSize();
+
+   const vw = Math.max(document.documentElement.clientWidth, window.innerWidth || 0);
+
+   return (
+      <Stack className={hordeClasses.horde}>
+         <TopNav />
+         <Breadcrumbs items={crumbs} />
+         <Stack horizontal style={{ backgroundColor: hordeTheme.horde.neutralBackground }}>
+            <div key={`windowsize_streamview_${windowSize.width}_${windowSize.height}`} style={{ width: vw / 2 - (1440 / 2), flexShrink: 0 }} />
+            <Stack tokens={{ childrenGap: 0 }} styles={{ root: { width: "100%" } }}>
+               <Stack style={{ maxWidth: 1440, paddingTop: 6, marginLeft: 4 }}>
+                  <FocusZone direction={FocusZoneDirection.vertical} style={{ padding: 0 }}>
+                     <div className={detailClasses.container} style={{ width: "100%", height: 'calc(100vh - 208px)', position: 'relative' }} data-is-scrollable={true}>
+                        <ScrollablePane scrollbarVisibility={ScrollbarVisibility.auto} onScroll={() => { }}>
+                           <Stack tokens={{ childrenGap: 18 }} style={{ padding: 0 }}>
+                              <Stack.Item align={'center'} className={hordeClasses.projectLogoCard}>
+                                 { !!logoSrc && <img src={logoSrc} alt="Project logo" height={300} /> }
+                              </Stack.Item>
+                              <Stack tokens={{ childrenGap: 20 }} className={hordeClasses.raised}>
+                                 <SearchBox
+                                    placeholder="Search"
+                                    onChange={(event?: React.ChangeEvent<HTMLInputElement> | undefined, newValue?: string | undefined) => { localState.setFilterString(newValue ?? ""); }}
+                                    onClear={() => { localState.setFilterString(""); }}
+                                 />
+                                 <Stack>
+                                    <Stack horizontal>
+                                       {categoryStackItems}
                                     </Stack>
-                                </ScrollablePane>
-                            </div>
-                        </FocusZone>
-                    </Stack>
-                </Stack>
+                                 </Stack>
+                              </Stack>
+                           </Stack>
+                        </ScrollablePane>
+                     </div>
+                  </FocusZone>
+               </Stack>
             </Stack>
-        </Stack>
-    );
+         </Stack>
+      </Stack>
+   );
 });
\ No newline at end of file
diff --git a/Engine/Source/Programs/Horde/HordeDashboard/src/components/UserHome.tsx b/Engine/Source/Programs/Horde/HordeDashboard/src/components/UserHome.tsx
--- a/Engine/Source/Programs/Horde/HordeDashboard/src/components/UserHome.tsx
+++ b/Engine/Source/Programs/Horde/HordeDashboard/src/components/UserHome.tsx
@@ -80,14 +80,13 @@ const ProjectsPanel: React.FC = observer(() => {
        projects = projects.slice(0, 3);
    }
    */
-
-   const imgStyle = dashboard.darktheme ? { filter: "invert(1)" } : {};
-
+   
    return (
       <Stack style={{ width: homeWidth, marginLeft: 4 }}>
          <Stack tokens={{}} horizontal verticalFill wrap horizontalAlign='center' >
             {
                projects.map((project) => {
+                  const logoSrc = `/api/v1/projects/${project.id}/logo${dashboard.darktheme ? "?darkTheme=true" : ""}`
                   return (
                      <Stack key={project.id} style={{
                         width: 560 * .65,
@@ -98,7 +97,7 @@ const ProjectsPanel: React.FC = observer(() => {
                      }}
                      >
                         <Link onClick={() => { projectStore.setActive(project.id); }} to={`/project/${project.id}`}>
-                           <img style={imgStyle} src={`/api/v1/projects/${project.id}/logo`} alt="" width={560 * .65} height={280 * .65} />
+                           <img src={logoSrc} alt="" width={560 * .65} height={280 * .65} />
                         </Link>
                      </Stack>
                   );
