Avoid Gauntlet issues to be grouped all together if they are from different type or contains callstack

#rnx

diff --git a/Engine/Source/Programs/Horde/Horde.Server.Tests/Issues/IssueServiceTests.cs b/Engine/Source/Programs/Horde/Horde.Server.Tests/Issues/IssueServiceTests.cs
--- a/Engine/Source/Programs/Horde/Horde.Server.Tests/Issues/IssueServiceTests.cs
+++ b/Engine/Source/Programs/Horde/Horde.Server.Tests/Issues/IssueServiceTests.cs
@@ -1966,7 +1966,7 @@ namespace Horde.Server.Tests.Issues
 
 				IReadOnlyList<IIssue> issues = await IssueCollection.FindIssuesAsync();
 				Assert.AreEqual(1, issues.Count);
-				Assert.AreEqual("Gauntlet", issues[0].Fingerprints[0].Type);
+				Assert.AreEqual("Gauntlet:test", issues[0].Fingerprints[0].Type);
 				Assert.AreEqual(new IssueKey($"{job.StreamId}:Update Version Files", IssueKeyType.None), issues[0].Fingerprints[0].Keys.First());
 				Assert.AreEqual("Automation test errors in Update Version Files", issues[0].Summary);
 			}
@@ -2007,14 +2007,14 @@ namespace Horde.Server.Tests.Issues
 
 				IReadOnlyList<IIssue> issues = await IssueCollection.FindIssuesAsync();
 				Assert.AreEqual(1, issues.Count);
-				Assert.AreEqual("Gauntlet", issues[0].Fingerprints[0].Type);
+				Assert.AreEqual("Gauntlet:fatal:with-callstack", issues[0].Fingerprints[0].Type);
 				Assert.AreEqual("hash:", issues[0].Fingerprints[0].Keys.First().Name.Substring(0, 5));
 				Assert.AreEqual(37, issues[0].Fingerprints[0].Keys.First().Name.Length);
 				Assert.AreEqual("Automation fatal errors in Update Version Files", issues[0].Summary);
 			}
 			// #3
-			// Scenario: Gauntlet Test event
-			// Expected: Gauntlet fingerprint using hash prefix and job step salt
+			// Scenario: Gauntlet Test + Fatal event
+			// Expected: Gauntlet Test fingerprint using hash prefix and job step salt, Gauntlet Fatal is a separated issue
 			{
 				string[] logErrors =
 				{
@@ -2033,6 +2033,7 @@ namespace Horde.Server.Tests.Issues
 				IJob job = CreateJob(_mainStreamId, 130, "Test Build", _graph);
 				await using (TestJsonLogger logger = await CreateLoggerAsync(job, 0, 0))
 				{
+					logger.LogError(KnownLogEvents.Gauntlet_FatalEvent, "Some critical error");
 					foreach (string error in logErrors)
 					{
 						logger.LogError(KnownLogEvents.Gauntlet_TestEvent, "{Error}", error);
@@ -2041,8 +2042,9 @@ namespace Horde.Server.Tests.Issues
 				await UpdateCompleteStepAsync(job, 0, 0, JobStepOutcome.Failure);
 
 				IReadOnlyList<IIssue> issues = await IssueCollection.FindIssuesAsync();
-				Assert.AreEqual(1, issues.Count);
-				Assert.AreEqual("Gauntlet", issues[0].Fingerprints[0].Type);
+				Assert.AreEqual(2, issues.Count);
+				Assert.AreEqual("Gauntlet:test", issues[0].Fingerprints[0].Type);
+				Assert.AreEqual("Gauntlet:fatal", issues[1].Fingerprints[0].Type);
 				Assert.AreEqual(10, issues[0].Fingerprints[0].Keys.Count);
 				Assert.AreEqual("hash:", issues[0].Fingerprints[0].Keys.First().Name.Substring(0, 5));
 				Assert.AreEqual($":{job.StreamId}:Update Version Files", issues[0].Fingerprints[0].Keys.First().Name.Substring(37));
diff --git a/Engine/Source/Programs/Shared/EpicGames.Horde/Issues/Handlers/GauntletIssueHandler.cs b/Engine/Source/Programs/Shared/EpicGames.Horde/Issues/Handlers/GauntletIssueHandler.cs
--- a/Engine/Source/Programs/Shared/EpicGames.Horde/Issues/Handlers/GauntletIssueHandler.cs
+++ b/Engine/Source/Programs/Shared/EpicGames.Horde/Issues/Handlers/GauntletIssueHandler.cs
@@ -98,12 +98,15 @@ namespace EpicGames.Horde.Issues.Handlers
 		/// <param name="issueEvent">The issue event</param>
 		/// <param name="keys">Receives a set of the keys</param>
 		/// <param name="metadata">Receives a set of metadata</param>
-		private void GetHash(IssueEvent issueEvent, HashSet<IssueKey> keys, HashSet<IssueMetadata> metadata)
+		/// <param name="hasCallstack">Set true if a callstack property was found</param>
+		private void GetHash(IssueEvent issueEvent, HashSet<IssueKey> keys, HashSet<IssueMetadata> metadata, out bool hasCallstack)
 		{
+			hasCallstack = false;
 			if (TryGetHash(issueEvent, out Md5Hash hash))
 			{
 				string key = $"hash:{hash}";
-				if (!EventHasCallstackProperty(issueEvent))
+				hasCallstack = EventHasCallstackProperty(issueEvent);
+				if (!hasCallstack)
 				{
 					// add job step salt if no Callstack property was found
 					key += $":{_context.StreamId}:{_context.NodeName}";
@@ -166,10 +169,16 @@ namespace EpicGames.Horde.Issues.Handlers
 		{
 			if (issueEvent.EventId != null && IsMatchingEventId(issueEvent.EventId.Value))
 			{
-				IssueEventGroup issue = new IssueEventGroup("Gauntlet", "Automation {Meta:GauntletType} {Severity} in {Meta:Node}", IssueChangeFilter.All);
+				string gauntletType = GetEventPrefix(issueEvent.EventId!.Value);
+				IssueEventGroup issue = new IssueEventGroup($"Gauntlet:{gauntletType}", "Automation {Meta:GauntletType} {Severity} in {Meta:Node}", IssueChangeFilter.All);
 				issue.Events.Add(issueEvent);
-				GetHash(issueEvent, issue.Keys, issue.Metadata);
-				issue.Metadata.Add("GauntletType", GetEventPrefix(issueEvent.EventId.Value));
+				bool hasCallstack;
+				GetHash(issueEvent, issue.Keys, issue.Metadata, out hasCallstack);
+				issue.Metadata.Add("GauntletType", gauntletType);
+				if (hasCallstack)
+				{
+					issue.Type = $"{issue.Type}:with-callstack";
+				}
 				_issues.Add(issue);
 
 				return true;
