Horde: Updating schema documentation

#rnx

diff --git a/Engine/Source/Programs/Horde/Docs/Config/Schema/Projects.md b/Engine/Source/Programs/Horde/Docs/Config/Schema/Projects.md
--- a/Engine/Source/Programs/Horde/Docs/Config/Schema/Projects.md
+++ b/Engine/Source/Programs/Horde/Docs/Config/Schema/Projects.md
@@ -13,6 +13,7 @@ Name | Description
 `macros` | [`ConfigMacro`](#configmacro)`[]`<br>Macros within the global scope
 `order` | `integer`<br>Order of this project on the dashboard
 `logo` | `string`<br>Path to the project logo
+`logoDarkTheme` | `string`<br>Optional path to the project logo for the dark theme
 `pools` | [`PoolConfig`](#poolconfig)`[]`<br>List of pools for this project
 `categories` | [`ProjectCategoryConfig`](#projectcategoryconfig)`[]`<br>Categories to include in this project
 `jobOptions` | [`JobOptions`](#joboptions)<br>Default settings for executing jobs
diff --git a/Engine/Source/Programs/Horde/Docs/Deployment/ServerSettings.md b/Engine/Source/Programs/Horde/Docs/Deployment/ServerSettings.md
--- a/Engine/Source/Programs/Horde/Docs/Deployment/ServerSettings.md
+++ b/Engine/Source/Programs/Horde/Docs/Deployment/ServerSettings.md
@@ -45,6 +45,7 @@ Name | Description
 `adminClaimValue` | `string`<br>Value of the claim type for administrators
 `corsEnabled` | `boolean`<br>Whether to enable Cors, generally for development purposes
 `corsOrigin` | `string`<br>Allowed Cors origin
+`enableDebugEndpoints` | `boolean`<br>Whether to enable debug/administrative REST API endpoints
 `enableNewAgentsByDefault` | `boolean`<br>Whether to automatically enable new agents by default. If false, new agents must manually be enabled before they can take on work.
 `testDataRetainMonths` | `integer`<br>The number of months to retain test data
 `schedulePollingInterval` | `string`<br>Interval between rebuilding the schedule queue with a DB query.
