Horde: Additional logging for creation/disposing of Perforce connections.

#rnx

diff --git a/Engine/Source/Programs/Horde/Drivers/JobDriver/Execution/ManagedWorkspaceMaterializer.cs b/Engine/Source/Programs/Horde/Drivers/JobDriver/Execution/ManagedWorkspaceMaterializer.cs
--- a/Engine/Source/Programs/Horde/Drivers/JobDriver/Execution/ManagedWorkspaceMaterializer.cs
+++ b/Engine/Source/Programs/Horde/Drivers/JobDriver/Execution/ManagedWorkspaceMaterializer.cs
@@ -49,6 +49,11 @@ public sealed class ManagedWorkspaceMaterializer : IWorkspaceMaterializer
 	/// <inheritdoc/>
 	public async Task<WorkspaceMaterializerSettings> InitializeAsync(ILogger logger, CancellationToken cancellationToken)
 	{
+		if (_workspace != null)
+		{
+			throw new WorkspaceMaterializationException("Materializer initialized twice");
+		}
+
 		using IScope scope = CreateTraceSpan("ManagedWorkspaceMaterializer.InitializeAsync");
 
 		ManagedWorkspaceOptions options = WorkspaceInfo.GetMwOptions(_agentWorkspace);
diff --git a/Engine/Source/Programs/Horde/Drivers/JobDriver/Utility/WorkspaceInfo.cs b/Engine/Source/Programs/Horde/Drivers/JobDriver/Utility/WorkspaceInfo.cs
--- a/Engine/Source/Programs/Horde/Drivers/JobDriver/Utility/WorkspaceInfo.cs
+++ b/Engine/Source/Programs/Horde/Drivers/JobDriver/Utility/WorkspaceInfo.cs
@@ -81,6 +81,8 @@ namespace JobDriver.Utility
 		/// </summary>
 		public ManagedWorkspace Repository { get; }
 
+		readonly ILogger _logger;
+
 		/// <summary>
 		/// Constructor
 		/// </summary>
@@ -93,7 +95,7 @@ namespace JobDriver.Utility
 		/// <param name="view">View for files to be synced</param>
 		/// <param name="removeUntrackedFiles">Whether to remove untracked files when syncing</param>
 		/// <param name="repository">The repository instance</param>
-		public WorkspaceInfo(IPerforceConnection perforce, string hostName, string streamName, PerforceViewMap streamView, DirectoryReference metadataDir, DirectoryReference workspaceDir, IList<string>? view, bool removeUntrackedFiles, ManagedWorkspace repository)
+		public WorkspaceInfo(IPerforceConnection perforce, string hostName, string streamName, PerforceViewMap streamView, DirectoryReference metadataDir, DirectoryReference workspaceDir, IList<string>? view, bool removeUntrackedFiles, ManagedWorkspace repository, ILogger logger)
 		{
 			PerforceClient = perforce;
 
@@ -110,11 +112,15 @@ namespace JobDriver.Utility
 			View = (view == null) ? new List<string>() : new List<string>(view);
 			RemoveUntrackedFiles = removeUntrackedFiles;
 			Repository = repository;
+			_logger = logger;
+
+			_logger.LogInformation("Created WorkspaceInfo for {ClientName}", PerforceClient.Settings.ClientName);
 		}
 
 		/// <inheritdoc/>
 		public void Dispose()
 		{
+			_logger.LogInformation("Disposing WorkspaceInfo for {ClientName}", PerforceClient.Settings.ClientName);
 			PerforceClient.Dispose();
 		}
 
@@ -250,7 +256,7 @@ namespace JobDriver.Utility
 
 				// Create the workspace info
 				logger.LogInformation("Syncing {ClientName} to {BaseDir} from {Server}, using stream {Stream} and view:{View}", clientName, workspaceDir, info.ServerAddress, streamName, String.Join("", view.Select(x => $"\n  {x}")));
-				return new WorkspaceInfo(perforceClient, hostName, streamName, streamView, metadataDir, workspaceDir, view, removeUntrackedFiles, newRepository);
+				return new WorkspaceInfo(perforceClient, hostName, streamName, streamView, metadataDir, workspaceDir, view, removeUntrackedFiles, newRepository, logger);
 			}
 			catch
 			{
