Horde: Add disk space to agent log telemetry
#rnx

diff --git a/Engine/Source/Programs/Horde/HordeDashboard/src/backend/Api.ts b/Engine/Source/Programs/Horde/HordeDashboard/src/backend/Api.ts
--- a/Engine/Source/Programs/Horde/HordeDashboard/src/backend/Api.ts
+++ b/Engine/Source/Programs/Horde/HordeDashboard/src/backend/Api.ts
@@ -2116,6 +2116,8 @@ export type GetAgentTelemetrySampleResponse = {
 	freeRam: number;
 	usedRam: number;
 	totalRam: number;
+	freeDisk: number;
+	totalDisk: number;
 }
 
 /**Information about a workspace type */
diff --git a/Engine/Source/Programs/Horde/HordeDashboard/src/components/agents/AgentTelemetrySparkline.tsx b/Engine/Source/Programs/Horde/HordeDashboard/src/components/agents/AgentTelemetrySparkline.tsx
--- a/Engine/Source/Programs/Horde/HordeDashboard/src/components/agents/AgentTelemetrySparkline.tsx
+++ b/Engine/Source/Programs/Horde/HordeDashboard/src/components/agents/AgentTelemetrySparkline.tsx
@@ -70,7 +70,7 @@ export class AgentTelemetryHandler extends PollBase {
 
          x *= sparkWidth;
          x = Math.floor(x);
-         
+
          if (this.lastX === x) {
             return;
          }
@@ -94,7 +94,7 @@ export class AgentTelemetryHandler extends PollBase {
                .attr("y1", () => 0)
                .attr("y2", () => sparkHeight)
                .attr("stroke-width", () => 2)
-               .attr("stroke", dashboard.darktheme ? "#6D6C6B" : "#4D4C4BAA")            
+               .attr("stroke", dashboard.darktheme ? "#6D6C6B" : "#4D4C4BAA")
 
          })
 
@@ -118,7 +118,7 @@ export class AgentTelemetryHandler extends PollBase {
          let endTime = this.endTime;
          if (!endTime) {
             endTime = new Date();
-         }         
+         }
 
          // clamp to 24 hours
          if ((endTime.getTime() - this.startTime.getTime()) > 86400000) {
@@ -170,13 +170,13 @@ export const AgentTelemetrySparkline: React.FC<{ handler: AgentTelemetryHandler
    const data = handler.data;
 
    if (data === undefined) {
-      return <Stack horizontal horizontalAlign="center" style={{padding: 32}}>
+      return <Stack horizontal horizontalAlign="center" style={{ padding: 32 }}>
          <Spinner size={SpinnerSize.large} />
       </Stack>
    }
 
    if (!data.length) {
-      return <Stack horizontal horizontalAlign="center" style={{padding: 32}}>
+      return <Stack horizontal horizontalAlign="center" style={{ padding: 32 }}>
          <Text variant="mediumPlus">No Telemetry Data</Text>
       </Stack>
    }
@@ -193,14 +193,24 @@ export const AgentTelemetrySparkline: React.FC<{ handler: AgentTelemetryHandler
       return (d.userCpu + d.systemCpu) / 100;
    })
 
+   const diskData = data.map(d => {
+      if (typeof (d.totalDisk) === "number") {
+         return d.freeDisk / d.totalDisk;
+      }
+
+      return 0;
+   })
+
    let cpuText = "";
    let ramText = "";
+   let diskText = "";
 
    let device: any;
    const devices = handler.agent?.capabilities?.devices;
    if (devices?.length) {
       device = devices[0];
    }
+
    device?.properties?.forEach(v => {
       if (v.startsWith("CPU=")) {
          cpuText = `${v.replace("CPU=", "")}`;
@@ -209,34 +219,70 @@ export const AgentTelemetrySparkline: React.FC<{ handler: AgentTelemetryHandler
       if (v.startsWith("RAM=")) {
          ramText = `${v.replace("RAM=", "")} GB`;
       }
+
+      if (v.startsWith("DiskTotalSize=")) {
+         const elements = v.split("=");
+         if (elements.length === 2) {
+
+            function formatBytes(bytes: number, decimals = 2) {
+               if (!+bytes) return '0 Bytes'
+
+               const k = 1024
+               const dm = decimals < 0 ? 0 : decimals
+               const sizes = ['Bytes', 'KB', 'MB', 'GB', 'TB', 'PB', 'EB', 'ZB', 'YB']
+
+               const i = Math.floor(Math.log(bytes) / Math.log(k))
+
+               return `${parseFloat(Math.ceil((bytes / Math.pow(k, i))).toFixed(dm))} ${sizes[i]}`
+            }
+
+            diskText = formatBytes(parseInt(elements[1]))
+         }
+      }
    })
 
-   return <Stack horizontal horizontalAlign="center" tokens={{ childrenGap: 24 }} style={{paddingBottom: 8}}>
-      <Stack style={{ width: sparkWidth }}>
-         <Stack horizontal>
-            <Label>CPU</Label>
-            <Stack grow></Stack>
-            {!!cpuText && <Label>{cpuText}</Label>}
+   return <Stack>
+      <Stack horizontal horizontalAlign="center" tokens={{ childrenGap: 24 }} style={{ paddingBottom: 8 }}>
+         <Stack style={{ width: sparkWidth }}>
+            <Stack horizontal>
+               <Label>CPU</Label>
+               <Stack grow></Stack>
+               {!!cpuText && <Label>{cpuText}</Label>}
+            </Stack>
+            <div id="telemetry_cpu">
+               <Sparklines width={sparkWidth} height={sparkHeight} data={cpuData} max={1} style={{ backgroundColor: dashboard.darktheme ? "#060709" : "#F3F2F1", padding: 8, border: "solid 1px #181A1B" }}>
+                  <SparklinesLine color={dashboard.darktheme ? "lightblue" : "#1E90FF"} />
+                  <SparklinesReferenceLine type="avg" />
+               </Sparklines>
+            </div>
+         </Stack>
+         <Stack style={{ width: sparkWidth }}>
+            <Stack horizontal>
+               <Label>RAM</Label>
+               <Stack grow></Stack>
+               {!!ramText && <Label>{ramText}</Label>}
+            </Stack>
+            <div id="telemetry_ram">
+               <Sparklines width={sparkWidth} height={sparkHeight} data={ramData} max={1} style={{ backgroundColor: dashboard.darktheme ? "#060709" : "#F3F2F1", padding: 8, border: "solid 1px #181A1B" }}>
+                  <SparklinesLine color={dashboard.darktheme ? "lightblue" : "#1E90FF"} />
+                  <SparklinesReferenceLine type="avg" />
+               </Sparklines>
+            </div>
          </Stack>
-         <div id="telemetry_cpu">
-            <Sparklines width={sparkWidth} height={sparkHeight} data={cpuData} max={1} style={{ backgroundColor: dashboard.darktheme ? "#060709" : "#F3F2F1", padding: 8, border: "solid 1px #181A1B" }}>
-               <SparklinesLine color={dashboard.darktheme ? "lightblue" : "#1E90FF"} />
-               <SparklinesReferenceLine type="avg" />
-            </Sparklines>
-         </div>
       </Stack>
-      <Stack style={{ width: sparkWidth }}>
+      <Stack style={{ width: sparkWidth, paddingLeft: 4, paddingBottom: 4 }}>
          <Stack horizontal>
-            <Label>RAM</Label>
+            <Label>Disk</Label>
             <Stack grow></Stack>
-            {!!ramText && <Label>{ramText}</Label>}
+            {!!diskText && <Label>{diskText}</Label>}
          </Stack>
-         <div id="telemetry_ram">
-            <Sparklines width={sparkWidth} height={sparkHeight} data={ramData} max={1} style={{ backgroundColor: dashboard.darktheme ? "#060709" : "#F3F2F1", padding: 8, border: "solid 1px #181A1B" }}>
+         <div id="telemetry_disk">
+            <Sparklines width={sparkWidth} height={sparkHeight} data={diskData} max={1} style={{ backgroundColor: dashboard.darktheme ? "#060709" : "#F3F2F1", padding: 8, border: "solid 1px #181A1B" }}>
                <SparklinesLine color={dashboard.darktheme ? "lightblue" : "#1E90FF"} />
                <SparklinesReferenceLine type="avg" />
             </Sparklines>
          </div>
       </Stack>
+
    </Stack>
 })
\ No newline at end of file
