Horde: Fix dashboard build
#rnx

diff --git a/Engine/Source/Programs/Horde/HordeDashboard/src/components/HistoryModal.tsx b/Engine/Source/Programs/Horde/HordeDashboard/src/components/HistoryModal.tsx
--- a/Engine/Source/Programs/Horde/HordeDashboard/src/components/HistoryModal.tsx
+++ b/Engine/Source/Programs/Horde/HordeDashboard/src/components/HistoryModal.tsx
@@ -323,7 +323,9 @@ class HistoryModalState {
                      systemCpu: 0,
                      freeRam: 0,
                      usedRam: 0,
-                     totalRam: 1
+                     totalRam: 1,
+                     freeDisk: 0,
+                     totalDisk: 0
                   });
                }
 
