Horde: Additional diagnostics for P4 connection errors when fetching config files.

#rnx

diff --git a/Engine/Source/Programs/Horde/Horde.Server/Configuration/ConfigSource.cs b/Engine/Source/Programs/Horde/Horde.Server/Configuration/ConfigSource.cs
--- a/Engine/Source/Programs/Horde/Horde.Server/Configuration/ConfigSource.cs
+++ b/Engine/Source/Programs/Horde/Horde.Server/Configuration/ConfigSource.cs
@@ -299,22 +299,29 @@ namespace Horde.Server.Configuration
 			{
 				using (IPerforceConnection perforce = await ConnectAsync(group.Key, cancellationToken))
 				{
-					FileSpecList fileSpec = group.Select(x => x.AbsolutePath).Distinct(StringComparer.OrdinalIgnoreCase).ToList();
+					try
+					{
+						FileSpecList fileSpec = group.Select(x => x.AbsolutePath).Distinct(StringComparer.OrdinalIgnoreCase).ToList();
 
-					List<FStatRecord> records = await perforce.FStatAsync(FStatOptions.ShortenOutput, fileSpec, cancellationToken).ToListAsync(cancellationToken);
-					records.RemoveAll(x => x.HeadAction == FileAction.Delete || x.HeadAction == FileAction.MoveDelete);
+						List<FStatRecord> records = await perforce.FStatAsync(FStatOptions.ShortenOutput, fileSpec, cancellationToken).ToListAsync(cancellationToken);
+						records.RemoveAll(x => x.HeadAction == FileAction.Delete || x.HeadAction == FileAction.MoveDelete);
 
-					Dictionary<string, FStatRecord> absolutePathToRecord = records.ToDictionary(x => x.DepotFile ?? String.Empty, x => x, StringComparer.OrdinalIgnoreCase);
-					foreach (Uri uri in group)
-					{
-						FStatRecord? record;
-						if (!absolutePathToRecord.TryGetValue(uri.AbsolutePath, out record))
+						Dictionary<string, FStatRecord> absolutePathToRecord = records.ToDictionary(x => x.DepotFile ?? String.Empty, x => x, StringComparer.OrdinalIgnoreCase);
+						foreach (Uri uri in group)
 						{
-							throw new FileNotFoundException($"Unable to read {uri}. No matching files found.");
+							FStatRecord? record;
+							if (!absolutePathToRecord.TryGetValue(uri.AbsolutePath, out record))
+							{
+								throw new FileNotFoundException($"Unable to read {uri}. No matching files found.");
+							}
+
+							IUser? author = await GetAuthorAsync(perforce, group.Key, record.HeadChange, cancellationToken);
+							results[uri] = new ConfigFileImpl(uri, record.HeadChange, author, this);
 						}
-
-						IUser? author = await GetAuthorAsync(perforce, group.Key, record.HeadChange, cancellationToken);
-						results[uri] = new ConfigFileImpl(uri, record.HeadChange, author, this);
+					}
+					catch (PerforceException ex) when (ex.Error != null && ex.Error.Generic == PerforceGenericCode.Config)
+					{
+						throw new PerforceException($"{ex.Message} [server: {perforce.Settings.ServerAndPort}, user: {perforce.Settings.UserName}]", ex);
 					}
 				}
 			}
diff --git a/Engine/Source/Programs/Shared/EpicGames.Perforce/PerforceException.cs b/Engine/Source/Programs/Shared/EpicGames.Perforce/PerforceException.cs
--- a/Engine/Source/Programs/Shared/EpicGames.Perforce/PerforceException.cs
+++ b/Engine/Source/Programs/Shared/EpicGames.Perforce/PerforceException.cs
@@ -42,5 +42,14 @@ namespace EpicGames.Perforce
 		{
 			Error = error;
 		}
+
+		/// <summary>
+		/// Constructor
+		/// </summary>
+		public PerforceException(string message, PerforceException innerException)
+			: base(message, innerException)
+		{
+			Error = innerException.Error;
+		}
 	}
 }
