UBA: Prevent Horde being disabled when configured via the environment.

#rnx

diff --git a/Engine/Source/Programs/Horde/Drivers/JobDriver/Utility/WorkspaceInfo.cs b/Engine/Source/Programs/Horde/Drivers/JobDriver/Utility/WorkspaceInfo.cs
--- a/Engine/Source/Programs/Horde/Drivers/JobDriver/Utility/WorkspaceInfo.cs
+++ b/Engine/Source/Programs/Horde/Drivers/JobDriver/Utility/WorkspaceInfo.cs
@@ -95,6 +95,7 @@ namespace JobDriver.Utility
 		/// <param name="view">View for files to be synced</param>
 		/// <param name="removeUntrackedFiles">Whether to remove untracked files when syncing</param>
 		/// <param name="repository">The repository instance</param>
+		/// <param name="logger">Logger for diagnostic messages</param>
 		public WorkspaceInfo(IPerforceConnection perforce, string hostName, string streamName, PerforceViewMap streamView, DirectoryReference metadataDir, DirectoryReference workspaceDir, IList<string>? view, bool removeUntrackedFiles, ManagedWorkspace repository, ILogger logger)
 		{
 			PerforceClient = perforce;
diff --git a/Engine/Source/Programs/UnrealBuildTool/Executors/UnrealBuildAccelerator/UBAAgentCoordinatorHorde.cs b/Engine/Source/Programs/UnrealBuildTool/Executors/UnrealBuildAccelerator/UBAAgentCoordinatorHorde.cs
--- a/Engine/Source/Programs/UnrealBuildTool/Executors/UnrealBuildAccelerator/UBAAgentCoordinatorHorde.cs
+++ b/Engine/Source/Programs/UnrealBuildTool/Executors/UnrealBuildAccelerator/UBAAgentCoordinatorHorde.cs
@@ -364,7 +364,7 @@ namespace UnrealBuildTool
 				return null;
 			}
 
-			if (String.IsNullOrEmpty(hordeConfig.HordeServer))
+			if (String.IsNullOrEmpty(hordeConfig.HordeServer) && HordeOptions.GetServerUrlFromEnvironment() == null && HordeOptions.GetDefaultServerUrl() == null)
 			{
 				logger.LogInformation("Horde disabled. Url not set.");
 				return null;
