Horde: Log in to the server when querying for config files from Perforce.

diff --git a/Engine/Source/Programs/Horde/Horde.Server/Configuration/ConfigSource.cs b/Engine/Source/Programs/Horde/Horde.Server/Configuration/ConfigSource.cs
--- a/Engine/Source/Programs/Horde/Horde.Server/Configuration/ConfigSource.cs
+++ b/Engine/Source/Programs/Horde/Horde.Server/Configuration/ConfigSource.cs
@@ -396,6 +396,10 @@ namespace Horde.Server.Configuration
 			}
 
 			IPerforceConnection connection = await PerforceConnection.CreateAsync(connectionSettings.ToPerforceSettings(), _logger);
+			if (connectionSettings.Credentials != null && !String.IsNullOrEmpty(connectionSettings.Credentials.Password) && String.IsNullOrEmpty(connectionSettings.Credentials.Ticket))
+			{
+				await connection.LoginAsync(connectionSettings.Credentials.Password, cancellationToken);
+			}
 			return connection;
 		}
 
