Horde: Only bind to the loopback adapter when spawning a local Redis instance for running tests.

#rnx

diff --git a/Engine/Source/Programs/Horde/Horde.Server.Tests/DatabaseIntegrationTest.cs b/Engine/Source/Programs/Horde/Horde.Server.Tests/DatabaseIntegrationTest.cs
--- a/Engine/Source/Programs/Horde/Horde.Server.Tests/DatabaseIntegrationTest.cs
+++ b/Engine/Source/Programs/Horde/Horde.Server.Tests/DatabaseIntegrationTest.cs
@@ -161,7 +161,7 @@ namespace Horde.Server.Tests
 			if (RuntimeInformation.IsOSPlatform(OSPlatform.Windows))
 			{
 				s_redisProcess = new RedisProcess(NullLogger.Instance);
-				s_redisProcess.Start("--save \"\" --appendonly no");
+				s_redisProcess.Start("");
 
 				return s_redisProcess.Port;
 			}
diff --git a/Engine/Source/Programs/Horde/Horde.Server/Server/RedisProcess.cs b/Engine/Source/Programs/Horde/Horde.Server/Server/RedisProcess.cs
--- a/Engine/Source/Programs/Horde/Horde.Server/Server/RedisProcess.cs
+++ b/Engine/Source/Programs/Horde/Horde.Server/Server/RedisProcess.cs
@@ -81,7 +81,7 @@ namespace Horde.Server.Server
 
 			// Launch the child process
 			_processGroup = new ManagedProcessGroup();
-			_process = new ManagedProcess(_processGroup, tempRedisExe.FullName, $"{arguments} --port {_port}", redisExe.Directory.FullName, null, ProcessPriorityClass.Normal);
+			_process = new ManagedProcess(_processGroup, tempRedisExe.FullName, $"--bind 127.0.0.1 --port {_port} --save \"\" --appendonly no {arguments}", redisExe.Directory.FullName, null, ProcessPriorityClass.Normal);
 			_process.StdIn.Close();
 
 			_relayTask = BackgroundTask.StartNew(RelayOutputAsync);
