Horde: Fix static analyzer warnings.

#rnx

diff --git a/Engine/Source/Programs/Horde/Drivers/JobDriver.Tests/Execution/FakeWorkspaceMaterializer.cs b/Engine/Source/Programs/Horde/Drivers/JobDriver.Tests/Execution/FakeWorkspaceMaterializer.cs
--- a/Engine/Source/Programs/Horde/Drivers/JobDriver.Tests/Execution/FakeWorkspaceMaterializer.cs
+++ b/Engine/Source/Programs/Horde/Drivers/JobDriver.Tests/Execution/FakeWorkspaceMaterializer.cs
@@ -8,7 +8,6 @@ using System.Threading;
 using System.Threading.Tasks;
 using EpicGames.Core;
 using JobDriver.Execution;
-using Microsoft.Extensions.Logging;
 
 namespace JobDriver.Tests.Execution;
 
diff --git a/Engine/Source/Programs/Horde/Drivers/JobDriver.Tests/Execution/FakeWorkspaceMaterializerTest.cs b/Engine/Source/Programs/Horde/Drivers/JobDriver.Tests/Execution/FakeWorkspaceMaterializerTest.cs
--- a/Engine/Source/Programs/Horde/Drivers/JobDriver.Tests/Execution/FakeWorkspaceMaterializerTest.cs
+++ b/Engine/Source/Programs/Horde/Drivers/JobDriver.Tests/Execution/FakeWorkspaceMaterializerTest.cs
@@ -5,7 +5,6 @@ using System.IO;
 using System.Threading;
 using System.Threading.Tasks;
 using JobDriver.Execution;
-using Microsoft.Extensions.Logging.Abstractions;
 using Microsoft.VisualStudio.TestTools.UnitTesting;
 
 namespace JobDriver.Tests.Execution;
