Horde: Fix incorrect base directory being used for launching UAT.

#rnx

diff --git a/Engine/Source/Programs/Horde/Drivers/JobDriver.Tests/Execution/FakeWorkspaceMaterializer.cs b/Engine/Source/Programs/Horde/Drivers/JobDriver.Tests/Execution/FakeWorkspaceMaterializer.cs
--- a/Engine/Source/Programs/Horde/Drivers/JobDriver.Tests/Execution/FakeWorkspaceMaterializer.cs
+++ b/Engine/Source/Programs/Horde/Drivers/JobDriver.Tests/Execution/FakeWorkspaceMaterializer.cs
@@ -20,8 +20,6 @@ class FakeWorkspaceMaterializer : IWorkspaceMaterializer
 
 	public string Identifier => "fakeWorkspaceIdentifier";
 
-	public string StreamRoot => "fakeWorkspaceStreamRoot";
-
 	public IReadOnlyDictionary<string, string> EnvironmentVariables { get; } = new Dictionary<string, string>();
 
 	public bool IsPerforceWorkspace => false;
diff --git a/Engine/Source/Programs/Horde/Drivers/JobDriver/Execution/ManagedWorkspaceMaterializer.cs b/Engine/Source/Programs/Horde/Drivers/JobDriver/Execution/ManagedWorkspaceMaterializer.cs
--- a/Engine/Source/Programs/Horde/Drivers/JobDriver/Execution/ManagedWorkspaceMaterializer.cs
+++ b/Engine/Source/Programs/Horde/Drivers/JobDriver/Execution/ManagedWorkspaceMaterializer.cs
@@ -17,21 +17,17 @@ namespace JobDriver.Execution;
 public sealed class ManagedWorkspaceMaterializer : IWorkspaceMaterializer
 {
 	private readonly RpcAgentWorkspace _agentWorkspace;
-	private readonly DirectoryReference _workingDir;
 	private readonly bool _useCacheFile;
 	private readonly bool _cleanDuringFinalize;
 	private readonly WorkspaceInfo _workspace;
 	private readonly ILogger _logger;
 
 	/// <inheritdoc/>
-	public DirectoryReference DirectoryPath => _workingDir;
+	public DirectoryReference DirectoryPath => _workspace.WorkspaceDir;
 
 	/// <inheritdoc/>
 	public string Identifier => _agentWorkspace.Identifier;
 
-	/// <inheritdoc/>
-	public string StreamRoot => _agentWorkspace.Stream;
-
 	/// <inheritdoc/>
 	public IReadOnlyDictionary<string, string> EnvironmentVariables { get; }
 
@@ -42,21 +38,18 @@ public sealed class ManagedWorkspaceMaterializer : IWorkspaceMaterializer
 	/// Constructor
 	/// </summary>
 	/// <param name="agentWorkspace">Workspace configuration</param>
-	/// <param name="workingDir">Where to put synced Perforce files and any cached data/metadata</param>
 	/// <param name="useCacheFile">Whether to use a cache file during syncs</param>
 	/// <param name="cleanDuringFinalize">Whether to clean and revert files during finalize</param>
 	/// <param name="workspace"></param>
 	/// <param name="logger"></param>
 	private ManagedWorkspaceMaterializer(
 		RpcAgentWorkspace agentWorkspace,
-		DirectoryReference workingDir,
 		bool useCacheFile,
 		bool cleanDuringFinalize,
 		WorkspaceInfo workspace,
 		ILogger logger)
 	{
 		_agentWorkspace = agentWorkspace;
-		_workingDir = workingDir;
 		_useCacheFile = useCacheFile;
 		_cleanDuringFinalize = cleanDuringFinalize;
 		_workspace = workspace;
@@ -84,7 +77,7 @@ public sealed class ManagedWorkspaceMaterializer : IWorkspaceMaterializer
 	{
 		ManagedWorkspaceOptions options = WorkspaceInfo.GetMwOptions(agentWorkspace);
 		WorkspaceInfo workspace = await WorkspaceInfo.CreateWorkspaceInfoAsync(agentWorkspace, workingDir, options, logger, cancellationToken);
-		return new ManagedWorkspaceMaterializer(agentWorkspace, workingDir, useCacheFile, cleanDuringFinalize, workspace, logger);
+		return new ManagedWorkspaceMaterializer(agentWorkspace, useCacheFile, cleanDuringFinalize, workspace, logger);
 	}
 
 	/// <inheritdoc/>
diff --git a/Engine/Source/Programs/Horde/Drivers/JobDriver/Execution/WorkspaceMaterializer.cs b/Engine/Source/Programs/Horde/Drivers/JobDriver/Execution/WorkspaceMaterializer.cs
--- a/Engine/Source/Programs/Horde/Drivers/JobDriver/Execution/WorkspaceMaterializer.cs
+++ b/Engine/Source/Programs/Horde/Drivers/JobDriver/Execution/WorkspaceMaterializer.cs
@@ -68,11 +68,6 @@ public interface IWorkspaceMaterializer : IDisposable
 	/// </summary>
 	string Identifier { get; }
 
-	/// <summary>
-	/// Stream path inside Perforce
-	/// </summary>
-	string StreamRoot { get; }
-
 	/// <summary>
 	/// Environment variables expected to be set for applications executing inside the workspace
 	/// Mostly intended for Perforce-specific variables when <see cref="IsPerforceWorkspace" /> is set to true
