Horde: Alow retrying skipped steps from the main job detail page
#rnx

diff --git a/Engine/Source/Programs/Horde/HordeDashboard/src/components/jobDetailsV2/JobOperationsBar.tsx b/Engine/Source/Programs/Horde/HordeDashboard/src/components/jobDetailsV2/JobOperationsBar.tsx
--- a/Engine/Source/Programs/Horde/HordeDashboard/src/components/jobDetailsV2/JobOperationsBar.tsx
+++ b/Engine/Source/Programs/Horde/HordeDashboard/src/components/jobDetailsV2/JobOperationsBar.tsx
@@ -62,7 +62,7 @@ export const JobOperations: React.FC<{ jobDetails: JobDetailsV2 }> = observer(({
 
    const failedSteps = jobDetails.getSteps().filter(s => {
 
-      if (!s.finishTime || (s.state !== JobStepState.Aborted && s.outcome !== JobStepOutcome.Failure)) {
+      if (s.state !== JobStepState.Skipped && (!s.finishTime || (s.state !== JobStepState.Aborted && s.outcome !== JobStepOutcome.Failure))) {
          return false;
       }
 
diff --git a/Engine/Source/Programs/Horde/HordeDashboard/src/components/jobDetailsV2/StepRetryModal.tsx b/Engine/Source/Programs/Horde/HordeDashboard/src/components/jobDetailsV2/StepRetryModal.tsx
--- a/Engine/Source/Programs/Horde/HordeDashboard/src/components/jobDetailsV2/StepRetryModal.tsx
+++ b/Engine/Source/Programs/Horde/HordeDashboard/src/components/jobDetailsV2/StepRetryModal.tsx
@@ -140,7 +140,24 @@ export const RetryStepsModal: React.FC<{ stepIds: string[]; jobDetails: JobDetai
       stepNames.push(stepName);
    });
 
-   const stepElements = stepNames.map(name => {
+   const stepElements: JSX.Element[] = [];
+
+   if (stepIds.length > 1) {
+
+      const allSelected = retrySteps.size === stepIds.length;
+
+      stepElements.push(<Stack horizontal tokens={{ childrenGap: 12 }}><Checkbox checked={allSelected} onChange={(ev, checked) => {
+         const newSteps = new Set<string>();
+         if (checked) {
+            stepIds.forEach(s => newSteps.add(s));
+         } else {
+         }
+         setRetrySteps(newSteps);
+      }} /><Text>{retrySteps.size === stepIds.length ? "Deselect All" : "Select All"}</Text></Stack>)
+   }
+
+
+   stepElements.push(...stepNames.map(name => {
       const stepId = stepLookup.get(name)!;
       return <Stack horizontal tokens={{ childrenGap: 12 }}><Checkbox checked={retrySteps.has(stepId)} onChange={(ev, checked) => {
          const newSteps = new Set(retrySteps);
@@ -151,7 +168,7 @@ export const RetryStepsModal: React.FC<{ stepIds: string[]; jobDetails: JobDetai
          }
          setRetrySteps(newSteps);
       }} /><Text>{name}</Text></Stack>
-   })
+   }))
 
    return <Modal className={hordeClasses.modal} isOpen={true} styles={{ main: { padding: 8, width: 800 } }} onDismiss={() => { onClose() }}>
       <Stack horizontal styles={{ root: { padding: 8 } }}>
@@ -214,14 +231,14 @@ export const StepRetryModal: React.FC<{ stepId: string; jobDetails: JobDetailsV2
          retry: true
       }).then(async (response) => {
 
-         if (response.stepId) {            
+         if (response.stepId) {
             if (jobDetails.jobData?.id) {
                // get new job data so we have the new batch/step
                jobDetails.jobData = await backend.getJob(jobDetails.jobData?.id);
                jobDetails.processGraph();
             }
-            
-            navigate(`/job/${jobData.id}?step=${response.stepId!}`)            
+
+            navigate(`/job/${jobData.id}?step=${response.stepId!}`)
          }
 
       }).catch((reason) => {
@@ -242,7 +259,7 @@ export const StepRetryModal: React.FC<{ stepId: string; jobDetails: JobDetailsV2
 
       const job = jobData;
 
-      const args:string[] = [];
+      const args: string[] = [];
       args.push(`-Target=Setup Build`);
       args.push(`-Target=${step!.name}`);
 
