Horde: Default retry steps to off
#rnx

diff --git a/Engine/Source/Programs/Horde/HordeDashboard/src/components/jobDetailsV2/StepRetryModal.tsx b/Engine/Source/Programs/Horde/HordeDashboard/src/components/jobDetailsV2/StepRetryModal.tsx
--- a/Engine/Source/Programs/Horde/HordeDashboard/src/components/jobDetailsV2/StepRetryModal.tsx
+++ b/Engine/Source/Programs/Horde/HordeDashboard/src/components/jobDetailsV2/StepRetryModal.tsx
@@ -24,7 +24,7 @@ export const RetryStepsModal: React.FC<{ stepIds: string[]; jobDetails: JobDetai
 
    const [submitting, setSubmitting] = useState(false);
    const [submitResults, setSubmitResults] = useState<RunAgainResult[] | undefined>(undefined);
-   const [retrySteps, setRetrySteps] = useState(new Set(stepIds));
+   const [retrySteps, setRetrySteps] = useState(new Set());
 
    const { hordeClasses } = getHordeStyling();
 
