Horde: Fix disk telemetry
#rnx
#skipci

diff --git a/Engine/Source/Programs/Horde/HordeDashboard/src/components/agents/AgentTelemetrySparkline.tsx b/Engine/Source/Programs/Horde/HordeDashboard/src/components/agents/AgentTelemetrySparkline.tsx
--- a/Engine/Source/Programs/Horde/HordeDashboard/src/components/agents/AgentTelemetrySparkline.tsx
+++ b/Engine/Source/Programs/Horde/HordeDashboard/src/components/agents/AgentTelemetrySparkline.tsx
@@ -55,7 +55,7 @@ export class AgentTelemetryHandler extends PollBase {
       this.currentTime = currentTime;
 
       // chart selectors
-      const charts = ["#telemetry_cpu > svg", "#telemetry_ram > svg"]
+      const charts = ["#telemetry_cpu > svg", "#telemetry_ram > svg", "#telemetry_disk > svg"]
 
       const startTime = this.lastStartTime;
       const endTime = this.lastEndTime;
@@ -194,8 +194,8 @@ export const AgentTelemetrySparkline: React.FC<{ handler: AgentTelemetryHandler
    })
 
    const diskData = data.map(d => {
-      if (typeof (d.totalDisk) === "number") {
-         return d.freeDisk / d.totalDisk;
+      if (typeof (d.totalDisk) === "number" && d.totalDisk) {
+         return 1.0 - (d.freeDisk / d.totalDisk)
       }
 
       return 0;
@@ -224,7 +224,7 @@ export const AgentTelemetrySparkline: React.FC<{ handler: AgentTelemetryHandler
          const elements = v.split("=");
          if (elements.length === 2) {
 
-            function formatBytes(bytes: number, decimals = 2) {
+            function formatBytes(bytes: number, decimals = 0) {
                if (!+bytes) return '0 Bytes'
 
                const k = 1024
@@ -233,7 +233,7 @@ export const AgentTelemetrySparkline: React.FC<{ handler: AgentTelemetryHandler
 
                const i = Math.floor(Math.log(bytes) / Math.log(k))
 
-               return `${parseFloat(Math.ceil((bytes / Math.pow(k, i))).toFixed(dm))} ${sizes[i]}`
+               return `${parseFloat((bytes / Math.pow(k, i)).toFixed(dm))} ${sizes[i]}`
             }
 
             diskText = formatBytes(parseInt(elements[1]))
@@ -272,11 +272,11 @@ export const AgentTelemetrySparkline: React.FC<{ handler: AgentTelemetryHandler
       </Stack>
       <Stack style={{ width: sparkWidth, paddingLeft: 4, paddingBottom: 4 }}>
          <Stack horizontal>
-            <Label>Disk</Label>
+            <Label>Free Disk</Label>
             <Stack grow></Stack>
             {!!diskText && <Label>{diskText}</Label>}
          </Stack>
-         <div id="telemetry_disk">
+         <div id="telemetry_disk" style={{width: 690}}>
             <Sparklines width={sparkWidth} height={sparkHeight} data={diskData} max={1} style={{ backgroundColor: dashboard.darktheme ? "#060709" : "#F3F2F1", padding: 8, border: "solid 1px #181A1B" }}>
                <SparklinesLine color={dashboard.darktheme ? "lightblue" : "#1E90FF"} />
                <SparklinesReferenceLine type="avg" />
