Horde: Fix renaming default administrator account
#rnx

diff --git a/Engine/Source/Programs/Horde/Horde.Server/Dashboard/DashboardController.cs b/Engine/Source/Programs/Horde/Horde.Server/Dashboard/DashboardController.cs
--- a/Engine/Source/Programs/Horde/Horde.Server/Dashboard/DashboardController.cs
+++ b/Engine/Source/Programs/Horde/Horde.Server/Dashboard/DashboardController.cs
@@ -71,8 +71,8 @@ namespace Horde.Server.Dashboard
 			bool needsFirstTimeSetup = false;
 			if (_settings.AuthMethod == AuthMethod.Horde)
 			{
-				IAccount? account = await _hordeAccounts.FindByLoginAsync("Admin");
-				if (account == null)
+				IAccount? defaultAdminAccount = await _hordeAccounts.GetAsync(AccountId.Parse("65d4f282ff286703e0609ccd"));
+				if (defaultAdminAccount == null)
 				{
 					needsFirstTimeSetup = true;
 				}
