Horde: AutomationHub autoexpansion of suite tests
#rnx

diff --git a/Engine/Source/Programs/Horde/HordeDashboard/src/components/AutomationViewSummary.tsx b/Engine/Source/Programs/Horde/HordeDashboard/src/components/AutomationViewSummary.tsx
--- a/Engine/Source/Programs/Horde/HordeDashboard/src/components/AutomationViewSummary.tsx
+++ b/Engine/Source/Programs/Horde/HordeDashboard/src/components/AutomationViewSummary.tsx
@@ -82,7 +82,7 @@ const TestSummaryButton: React.FC<{ test: GetTestResponse, handler: TestDataHand
       const streamExpand = new Map<string, boolean>();
       Array.from(status.keys()).forEach(s => streamExpand.set(s, true));
       setStreamExpanded(streamExpand);
-      setExpanded(true);      
+      setExpanded(true);
       setAutoExpanded(false);
       return null;
    }
@@ -255,7 +255,7 @@ const TestSummaryButton: React.FC<{ test: GetTestResponse, handler: TestDataHand
                            <Stack style={{ width: 124 }}>
                               <Text variant="xSmall" style={{ fontWeight: fontWeight }}>{metaName}</Text>
                            </Stack>
-                           <Stack horizontal style={{ width: 84 }} verticalFill verticalAlign="center" tokens={{childrenGap: 4}}>
+                           <Stack horizontal style={{ width: 84 }} verticalFill verticalAlign="center" tokens={{ childrenGap: 4 }}>
                               <Text variant="xSmall" style={{ fontWeight: fontWeight }}>CL {last.buildChangeList}</Text>
                               {last.buildChangeList === latestCL && <FontIcon style={{ fontSize: 11, color: color }} iconName="Star" />}
                            </Stack>
@@ -446,11 +446,14 @@ const SuiteSummaryModal: React.FC<{ suite: GetTestSuiteResponse, handler: TestDa
 
 const SuiteSummaryButton: React.FC<{ suite: GetTestSuiteResponse, handler: TestDataHandler }> = ({ suite, handler }) => {
 
+   const search = new URLSearchParams(window.location.search);
+
    const [expanded, setExpanded] = useState(false);
    const [streamExpanded, setStreamExpanded] = useState<Map<string, boolean>>(new Map());
    const [historyShow, setHistoryShown] = useState(false);
    const statusColors = dashboard.getStatusColors();
    const { hordeClasses, modeColors } = getHordeStyling();
+   const [autoExpand, setAutoExpanded] = useState(search.get("autoexpand") === "true");
 
    const styles = mergeStyleSets({
       metaitem: {
@@ -465,13 +468,23 @@ const SuiteSummaryButton: React.FC<{ suite: GetTestSuiteResponse, handler: TestD
    const colorA = dashboard.darktheme ? "#181A1B" : "#e8e8e8";
    const colorB = dashboard.darktheme ? "#242729" : "#f8f8f8";
 
-
    const status = handler.getStatus(suite.id);
 
    if (!status) {
       return <div />;
    }
 
+   if (autoExpand) {
+
+      const streamExpand = new Map<string, boolean>();
+      Array.from(status.keys()).forEach(s => streamExpand.set(s, true));
+      setStreamExpanded(streamExpand);
+      setExpanded(true);
+      setAutoExpanded(false);
+      return null;
+   }
+
+
    const streamElements: JSX.Element[] = [];
 
    let anyFailed = false;
