Horde: Fix overly accurate percentages on suite and other status bars
#rnx

diff --git a/Engine/Source/Programs/Horde/HordeDashboard/src/components/AutomationCommon.tsx b/Engine/Source/Programs/Horde/HordeDashboard/src/components/AutomationCommon.tsx
--- a/Engine/Source/Programs/Horde/HordeDashboard/src/components/AutomationCommon.tsx
+++ b/Engine/Source/Programs/Horde/HordeDashboard/src/components/AutomationCommon.tsx
@@ -24,7 +24,7 @@ export const StatusBar = (stack: StatusBarStack[], width: number, height: number
    stack = stack.filter(s => s.value > 0);
 
    const mainTitle = stack.map((item) => {
-      return item.titleValue === undefined ? `${item.value}% ${item.title}` : `${item.titleValue} ${item.title}`
+      return item.titleValue === undefined ? `${Math.ceil(item.value)}% ${item.title}` : `${item.titleValue} ${item.title}`
    }).join(' ');
 
    return (
