Horde: Update test projects to NET 8.

diff --git a/Engine/Source/Programs/Horde/Drivers/JobDriver.Tests/JobDriver.Tests.csproj b/Engine/Source/Programs/Horde/Drivers/JobDriver.Tests/JobDriver.Tests.csproj
--- a/Engine/Source/Programs/Horde/Drivers/JobDriver.Tests/JobDriver.Tests.csproj
+++ b/Engine/Source/Programs/Horde/Drivers/JobDriver.Tests/JobDriver.Tests.csproj
@@ -1,7 +1,7 @@
 <Project Sdk="Microsoft.NET.Sdk">
 
   <PropertyGroup>
-    <TargetFramework>net6.0</TargetFramework>
+    <TargetFramework>net8.0</TargetFramework>
     <Nullable>enable</Nullable>
     <IsPackable>false</IsPackable>
     <DebugType>full</DebugType>
diff --git a/Engine/Source/Programs/Horde/Horde.Agent.Tests/Horde.Agent.Tests.csproj b/Engine/Source/Programs/Horde/Horde.Agent.Tests/Horde.Agent.Tests.csproj
--- a/Engine/Source/Programs/Horde/Horde.Agent.Tests/Horde.Agent.Tests.csproj
+++ b/Engine/Source/Programs/Horde/Horde.Agent.Tests/Horde.Agent.Tests.csproj
@@ -1,7 +1,7 @@
 <Project Sdk="Microsoft.NET.Sdk">
 
   <PropertyGroup>
-    <TargetFramework>net6.0</TargetFramework>
+    <TargetFramework>net8.0</TargetFramework>
     <Nullable>enable</Nullable>
     <IsPackable>false</IsPackable>
     <DebugType>full</DebugType>
