Horde: Update release notes.

#rnx

diff --git a/Engine/Source/Programs/Horde/Docs/ReleaseNotes.md b/Engine/Source/Programs/Horde/Docs/ReleaseNotes.md
--- a/Engine/Source/Programs/Horde/Docs/ReleaseNotes.md
+++ b/Engine/Source/Programs/Horde/Docs/ReleaseNotes.md
@@ -1,5 +1,23 @@
 # Release Notes
 
+## 2024-05-30
+
+* Update test projects to NET 8. (34002443)
+* Refactor WorkspaceInfo to avoid keeping a long-lived connection. Keeping idle connections to the Perforce server after the initial workspace setup is unnecessary and wastes server resources. (33949575)
+* Log in to the server when querying for config files from Perforce. (33945114)
+* Fix http auth handler not using correct configuration for server when run locally. (33940162)
+* Skip cleaning files for AutoSDK workspace. (33858727)
+* Only enable debug controllers if enabled in settings (33857362)
+* Add a fingerprint description as part of issue response (33833567)
+* EpicGames.Perforce: Fix parsing of array fields from Perforce records. (33807413)
+* EpicGames.Perforce: Fix parsing of records which have duplicate field names in child records. (33804970)
+* Allow overriding the targets to execute for a job. The optional 'Targets' property can be specified when creating a job, and will replace any -Target= arguments in the command line if set. (33799426)
+* Pass -Target arguments to individual steps to ensure that BuildGraph can correctly filter out which blocks to embed in temp storage manifests. (33799034)
+* Remove server OS version info from anonymous endpoint responses. (#11909) (33793366)
+* Allow overriding the driver for jobs through the JobOptions object. (33791494)
+* Prevent AWS status queries from causing session update failures. (33782757)
+* Move job driver into a separate executable. (33775640, 33782224)
+
 ## 2024-05-20
 
 * Add grouping keys and platforms to tools. Intent is for the dashboard to only show one tool for each grouping key by default (the one with the closest matching platform based on the browser's user-agent), but users can manually expand the list of tools to show other platforms if desired.
