Horde: Show retry status inline
#rnx

diff --git a/Engine/Source/Programs/Horde/HordeDashboard/src/components/jobDetailsV2/JobDetailViewSteps.tsx b/Engine/Source/Programs/Horde/HordeDashboard/src/components/jobDetailsV2/JobDetailViewSteps.tsx
--- a/Engine/Source/Programs/Horde/HordeDashboard/src/components/jobDetailsV2/JobDetailViewSteps.tsx
+++ b/Engine/Source/Programs/Horde/HordeDashboard/src/components/jobDetailsV2/JobDetailViewSteps.tsx
@@ -68,13 +68,19 @@ const RenderDynamic: React.FC<{ jobDetails: JobDetailsV2, dataView: StepsDataVie
    }
 
    if (step && column === "Progress") {
+
       if (step.state === JobStepState.Running && !step.abortRequested) {
          return <Stack horizontalAlign={"center"}> {step.startTime && !step.finishTime && <ProgressIndicator percentComplete={getStepPercent(step)} barHeight={2} styles={{ root: { paddingTop: 2, width: 120 } }} />}</Stack>;
       }
 
-      const message = getStepStatusMessage(step);
+      const retries = jobDetails.getStepRetries(step.id);
+      const retry = retries.length > 1 ? retries[retries.length - 1] : undefined;
+
+      if (retry && retry.id !== step.id) {
+         return <Stack horizontal horizontalAlign={"center"} tokens={{ childrenGap: 0, padding: 0 }} style={{ fontSize: "13px" }}><Link to={`/job/${jobDetails.jobId!}?step=${retry.id}`}>{`Retried - ${getStepStatusMessage(step)}`}</Link></Stack>;
+      }
 
-      return <Stack horizontal horizontalAlign={"center"} tokens={{ childrenGap: 0, padding: 0 }}><Text>{message}</Text></Stack>;
+      return <Stack horizontal horizontalAlign={"center"} tokens={{ childrenGap: 0, padding: 0 }}><Text>{getStepStatusMessage(step)}</Text></Stack>;
    }
 
 
