Horde: Add a log enricher which adds the server version.

diff --git a/Engine/Source/Programs/Horde/Horde.Server/ServerApp.cs b/Engine/Source/Programs/Horde/Horde.Server/ServerApp.cs
--- a/Engine/Source/Programs/Horde/Horde.Server/ServerApp.cs
+++ b/Engine/Source/Programs/Horde/Horde.Server/ServerApp.cs
@@ -27,6 +27,15 @@ namespace Horde.Server
 {
 	static class LoggerExtensions
 	{
+		public class ServerVersionLogEnricher : Serilog.Core.ILogEventEnricher
+		{
+			/// <inheritdoc />
+			public void Enrich(Serilog.Events.LogEvent logEvent, Serilog.Core.ILogEventPropertyFactory propertyFactory)
+			{
+				logEvent.AddPropertyIfAbsent(propertyFactory.CreateProperty("hordeserver.version", ServerApp.Version));
+			}
+		}
+
 		public static LoggerConfiguration Console(this LoggerSinkConfiguration sinkConfig, ServerSettings settings)
 		{
 			if (settings.LogJsonToStdOut)
@@ -50,6 +59,8 @@ namespace Horde.Server
 
 		public static LoggerConfiguration WithHordeConfig(this LoggerConfiguration configuration, ServerSettings settings)
 		{
+			configuration.Enrich.With<ServerVersionLogEnricher>();
+
 			if (settings.OpenTelemetry.EnableDatadogCompatibility)
 			{
 				configuration = configuration.Enrich.With<OpenTelemetryDatadogLogEnricher>();
