Horde: Set the service.version property for log events.

diff --git a/Engine/Source/Programs/Horde/Horde.Server/ServerApp.cs b/Engine/Source/Programs/Horde/Horde.Server/ServerApp.cs
--- a/Engine/Source/Programs/Horde/Horde.Server/ServerApp.cs
+++ b/Engine/Source/Programs/Horde/Horde.Server/ServerApp.cs
@@ -32,7 +32,7 @@ namespace Horde.Server
 			/// <inheritdoc />
 			public void Enrich(Serilog.Events.LogEvent logEvent, Serilog.Core.ILogEventPropertyFactory propertyFactory)
 			{
-				logEvent.AddPropertyIfAbsent(propertyFactory.CreateProperty("hordeserver.version", ServerApp.Version));
+				logEvent.AddPropertyIfAbsent(propertyFactory.CreateProperty("service.version", ServerApp.Version));
 			}
 		}
 
