Horde: Fix wrong status on retried step
#rnx

diff --git a/Engine/Source/Programs/Horde/HordeDashboard/src/components/jobDetailsV2/JobDetailViewSteps.tsx b/Engine/Source/Programs/Horde/HordeDashboard/src/components/jobDetailsV2/JobDetailViewSteps.tsx
--- a/Engine/Source/Programs/Horde/HordeDashboard/src/components/jobDetailsV2/JobDetailViewSteps.tsx
+++ b/Engine/Source/Programs/Horde/HordeDashboard/src/components/jobDetailsV2/JobDetailViewSteps.tsx
@@ -77,7 +77,7 @@ const RenderDynamic: React.FC<{ jobDetails: JobDetailsV2, dataView: StepsDataVie
       const retry = retries.length > 1 ? retries[retries.length - 1] : undefined;
 
       if (retry && retry.id !== step.id) {
-         return <Stack horizontal horizontalAlign={"center"} tokens={{ childrenGap: 0, padding: 0 }} style={{ fontSize: "13px" }}><Link to={`/job/${jobDetails.jobId!}?step=${retry.id}`}>{`Retried - ${getStepStatusMessage(step)}`}</Link></Stack>;
+         return <Stack horizontal horizontalAlign={"center"} tokens={{ childrenGap: 0, padding: 0 }} style={{ fontSize: "13px" }}><Link to={`/job/${jobDetails.jobId!}?step=${retry.id}`}>{`Retried - ${getStepStatusMessage(retry)}`}</Link></Stack>;
       }
 
       return <Stack horizontal horizontalAlign={"center"} tokens={{ childrenGap: 0, padding: 0 }}><Text>{getStepStatusMessage(step)}</Text></Stack>;
