Horde: Use the standard version field for log events. This is automatically mapped to the Datadog version field by our pipeline.

diff --git a/Engine/Source/Programs/Horde/Horde.Server/ServerApp.cs b/Engine/Source/Programs/Horde/Horde.Server/ServerApp.cs
--- a/Engine/Source/Programs/Horde/Horde.Server/ServerApp.cs
+++ b/Engine/Source/Programs/Horde/Horde.Server/ServerApp.cs
@@ -27,12 +27,11 @@ namespace Horde.Server
 {
 	static class LoggerExtensions
 	{
-		public class ServerVersionLogEnricher : Serilog.Core.ILogEventEnricher
+		class DatadogVersionLogEnricher : Serilog.Core.ILogEventEnricher
 		{
-			/// <inheritdoc />
 			public void Enrich(Serilog.Events.LogEvent logEvent, Serilog.Core.ILogEventPropertyFactory propertyFactory)
 			{
-				logEvent.AddPropertyIfAbsent(propertyFactory.CreateProperty("service.version", ServerApp.Version));
+				logEvent.AddPropertyIfAbsent(propertyFactory.CreateProperty("dd.version", ServerApp.Version));
 			}
 		}
 
@@ -59,10 +58,9 @@ namespace Horde.Server
 
 		public static LoggerConfiguration WithHordeConfig(this LoggerConfiguration configuration, ServerSettings settings)
 		{
-			configuration.Enrich.With<ServerVersionLogEnricher>();
-
 			if (settings.OpenTelemetry.EnableDatadogCompatibility)
 			{
+				configuration = configuration.Enrich.With<DatadogVersionLogEnricher>();
 				configuration = configuration.Enrich.With<OpenTelemetryDatadogLogEnricher>();
 			}
 
