Horde: Remove legacy AgentCapabilities object.

#rnx

diff --git a/Engine/Source/Programs/Horde/Horde.Server/Agents/AgentCollection.cs b/Engine/Source/Programs/Horde/Horde.Server/Agents/AgentCollection.cs
--- a/Engine/Source/Programs/Horde/Horde.Server/Agents/AgentCollection.cs
+++ b/Engine/Source/Programs/Horde/Horde.Server/Agents/AgentCollection.cs
@@ -32,6 +32,9 @@ namespace Horde.Server.Agents
 	{
 		class Agent : IAgent
 		{
+			static IReadOnlyList<string> DefaultProperties { get; } = new List<string>();
+			static IReadOnlyDictionary<string, int> DefaultResources { get; } = new Dictionary<string, int>();
+
 			readonly AgentCollection _collection;
 			readonly AgentDocument _document;
 
@@ -45,8 +48,8 @@ namespace Horde.Server.Agents
 			bool IAgent.Deleted => _document.Deleted;
 			string? IAgent.Version => _document.Version;
 			string? IAgent.Comment => _document.Comment;
-			IReadOnlyList<string> IAgent.Properties => _document.Properties ?? _document.Capabilities.Devices.FirstOrDefault()?.Properties?.ToList() ?? new List<string>();
-			IReadOnlyDictionary<string, int> IAgent.Resources => _document.Resources ?? _document.Capabilities.Devices.FirstOrDefault()?.Resources ?? new Dictionary<string, int>();
+			IReadOnlyList<string> IAgent.Properties => _document.Properties ?? DefaultProperties;
+			IReadOnlyDictionary<string, int> IAgent.Resources => _document.Resources ?? DefaultResources;
 			string? IAgent.LastUpgradeVersion => _document.LastUpgradeVersion;
 			DateTime? IAgent.LastUpgradeTime => _document.LastUpgradeTime;
 			int? IAgent.UpgradeAttemptCount => _document.UpgradeAttemptCount;
@@ -128,29 +131,6 @@ namespace Horde.Server.Agents
 			}
 		}
 
-		/// <summary>
-		/// Legacy information about a device attached to an agent
-		/// </summary>
-		class DeviceCapabilities
-		{
-			[BsonIgnoreIfNull]
-			public HashSet<string>? Properties { get; set; }
-
-			[BsonIgnoreIfNull]
-			public Dictionary<string, int>? Resources { get; set; }
-		}
-
-		/// <summary>
-		/// Legacy capabilities of an agent
-		/// </summary>
-		class AgentCapabilities
-		{
-			public List<DeviceCapabilities> Devices { get; set; } = new List<DeviceCapabilities>();
-
-			[BsonIgnoreIfNull]
-			public HashSet<string>? Properties { get; set; }
-		}
-
 		/// <summary>
 		/// Concrete implementation of an agent document
 		/// </summary>
@@ -215,7 +195,6 @@ namespace Horde.Server.Agents
 			[BsonIgnoreIfNull]
 			public int? ConformAttemptCount { get; set; }
 
-			public AgentCapabilities Capabilities { get; set; } = new AgentCapabilities();
 			public List<AgentLease>? Leases { get; set; }
 			public DateTime UpdateTime { get; set; }
 			public uint UpdateIndex { get; set; }
