Horde: Add an agent document property for the combined list of pools, and upgrade existing agent documents on read to include it. This allows indexed searches for dynamic pools.

diff --git a/Engine/Source/Programs/Horde/Horde.Server.Tests/Agents/AgentCollectionTests.cs b/Engine/Source/Programs/Horde/Horde.Server.Tests/Agents/AgentCollectionTests.cs
--- a/Engine/Source/Programs/Horde/Horde.Server.Tests/Agents/AgentCollectionTests.cs
+++ b/Engine/Source/Programs/Horde/Horde.Server.Tests/Agents/AgentCollectionTests.cs
@@ -60,7 +60,7 @@ public class AgentCollectionTests : TestSetup
 		await UpdateAgentAsync();
 
 		await _agent.TryCreateSessionAsync(new CreateSessionOptions(SessionIdUtils.GenerateNewId(), DateTime.UtcNow, AgentStatus.Ok,
-			new List<string>(), new Dictionary<string, int>(), new List<PoolId>(), new List<PoolId>(), DateTime.UtcNow, null));
+			new List<string>(), new Dictionary<string, int>(), new List<PoolId>(), DateTime.UtcNow, null));
 
 		List<LeaseId> leases = await AgentCollection.FindActiveLeaseIdsAsync();
 
diff --git a/Engine/Source/Programs/Horde/Horde.Server.Tests/Jobs/JobServiceTests.cs b/Engine/Source/Programs/Horde/Horde.Server.Tests/Jobs/JobServiceTests.cs
--- a/Engine/Source/Programs/Horde/Horde.Server.Tests/Jobs/JobServiceTests.cs
+++ b/Engine/Source/Programs/Horde/Horde.Server.Tests/Jobs/JobServiceTests.cs
@@ -241,7 +241,7 @@ namespace Horde.Server.Tests.Jobs
 			IAgent? agent = await AgentService.CreateAgentAsync("TestAgent", false, "");
 			Assert.IsNotNull(agent);
 
-			agent = await agent.TryUpdateAsync(new UpdateAgentOptions{ Enabled = true, Pools = new List<PoolId> { new PoolId("win") } });
+			agent = await agent.TryUpdateAsync(new UpdateAgentOptions{ Enabled = true, ExplicitPools = new List<PoolId> { new PoolId("win") } });
 			Assert.IsNotNull(agent);
 
 			await AgentService.CreateSessionAsync(agent, AgentStatus.Ok, new List<string>(), new Dictionary<string, int>(), null);
diff --git a/Engine/Source/Programs/Horde/Horde.Server.Tests/Jobs/JobTaskSourceTests.cs b/Engine/Source/Programs/Horde/Horde.Server.Tests/Jobs/JobTaskSourceTests.cs
--- a/Engine/Source/Programs/Horde/Horde.Server.Tests/Jobs/JobTaskSourceTests.cs
+++ b/Engine/Source/Programs/Horde/Horde.Server.Tests/Jobs/JobTaskSourceTests.cs
@@ -166,7 +166,7 @@ namespace Horde.Server.Tests.Jobs
 				IAgent? agent = await AgentService.CreateAgentAsync("TestAgent", false, "");
 				Assert.IsNotNull(agent);
 
-				agent = await agent.TryUpdateAsync(new UpdateAgentOptions { Enabled = isAgentEnabled, Pools = new List<PoolId> { pool.Id } });
+				agent = await agent.TryUpdateAsync(new UpdateAgentOptions { Enabled = isAgentEnabled, ExplicitPools = new List<PoolId> { pool.Id } });
 				Assert.IsNotNull(agent);
 
 				await AgentService.CreateSessionAsync(agent, AgentStatus.Ok, new List<string>(), new Dictionary<string, int>(), null);
diff --git a/Engine/Source/Programs/Horde/Horde.Server.Tests/Server/RpcServiceTest.cs b/Engine/Source/Programs/Horde/Horde.Server.Tests/Server/RpcServiceTest.cs
--- a/Engine/Source/Programs/Horde/Horde.Server.Tests/Server/RpcServiceTest.cs
+++ b/Engine/Source/Programs/Horde/Horde.Server.Tests/Server/RpcServiceTest.cs
@@ -336,7 +336,7 @@ namespace Horde.Server.Tests.Server
 			RpcCreateSessionResponse res = await RpcService.CreateSession(req, _adminContext);
 
 			IAgent agent = (await AgentService.GetAgentAsync(new AgentId(res.AgentId)))!;
-			CollectionAssert.AreEquivalent(new List<PoolId> { new("fooPool"), new("barPool") }, agent.GetPools().ToList());
+			CollectionAssert.AreEquivalent(new List<PoolId> { new("fooPool"), new("barPool") }, agent.Pools.ToList());
 
 			// Connect a second time, when the agent has already been created
 			req = new() { Id = new AgentId("bogusAgentName").ToString(), Capabilities = new RpcAgentCapabilities() };
@@ -344,7 +344,7 @@ namespace Horde.Server.Tests.Server
 			res = await RpcService.CreateSession(req, _adminContext);
 
 			agent = (await AgentService.GetAgentAsync(new AgentId(res.AgentId)))!;
-			CollectionAssert.AreEquivalent(new List<PoolId> { new("fooPool"), new("barPool"), new("bazPool") }, agent.GetPools().ToList());
+			CollectionAssert.AreEquivalent(new List<PoolId> { new("fooPool"), new("barPool"), new("bazPool") }, agent.Pools.ToList());
 		}
 
 		[TestMethod]
diff --git a/Engine/Source/Programs/Horde/Horde.Server.Tests/TestSetup.cs b/Engine/Source/Programs/Horde/Horde.Server.Tests/TestSetup.cs
--- a/Engine/Source/Programs/Horde/Horde.Server.Tests/TestSetup.cs
+++ b/Engine/Source/Programs/Horde/Horde.Server.Tests/TestSetup.cs
@@ -412,7 +412,7 @@ namespace Horde.Server.Tests
 			IAgent? agent = await AgentService.CreateAgentAsync("TestAgent" + s_agentIdCounter++, ephemeral, "");
 			Assert.IsNotNull(agent);
 
-			agent = await agent.TryUpdateAsync(new UpdateAgentOptions { Enabled = enabled, Pools = new List<PoolId> { poolId } });
+			agent = await agent.TryUpdateAsync(new UpdateAgentOptions { Enabled = enabled, ExplicitPools = new List<PoolId> { poolId } });
 			Assert.IsNotNull(agent);
 
 			agent = await AgentService.CreateSessionAsync(agent, AgentStatus.Ok, tempProps, resources, null);
diff --git a/Engine/Source/Programs/Horde/Horde.Server/Agents/AgentCollection.cs b/Engine/Source/Programs/Horde/Horde.Server/Agents/AgentCollection.cs
--- a/Engine/Source/Programs/Horde/Horde.Server/Agents/AgentCollection.cs
+++ b/Engine/Source/Programs/Horde/Horde.Server/Agents/AgentCollection.cs
@@ -8,6 +8,7 @@ using System.Linq;
 using System.Linq.Expressions;
 using System.Threading;
 using System.Threading.Tasks;
+using Amazon.EC2.Model;
 using EpicGames.Core;
 using EpicGames.Horde.Agents;
 using EpicGames.Horde.Agents.Leases;
@@ -18,6 +19,7 @@ using Google.Protobuf.WellKnownTypes;
 using Horde.Server.Auditing;
 using Horde.Server.Server;
 using HordeCommon.Rpc.Tasks;
+using Microsoft.Extensions.Azure;
 using MongoDB.Bson;
 using MongoDB.Bson.Serialization.Attributes;
 using MongoDB.Driver;
@@ -53,8 +55,9 @@ namespace Horde.Server.Agents
 			string? IAgent.LastUpgradeVersion => _document.LastUpgradeVersion;
 			DateTime? IAgent.LastUpgradeTime => _document.LastUpgradeTime;
 			int? IAgent.UpgradeAttemptCount => _document.UpgradeAttemptCount;
+			IReadOnlyList<PoolId> IAgent.Pools => _document.Pools;
 			IReadOnlyList<PoolId> IAgent.DynamicPools => _document.DynamicPools;
-			IReadOnlyList<PoolId> IAgent.ExplicitPools => _document.Pools;
+			IReadOnlyList<PoolId> IAgent.ExplicitPools => _document.ExplicitPools;
 			bool IAgent.RequestConform => _document.RequestConform;
 			bool IAgent.RequestFullConform => _document.RequestFullConform;
 			bool IAgent.RequestRestart => _document.RequestRestart;
@@ -168,8 +171,9 @@ namespace Horde.Server.Agents
 			[BsonIgnoreIfNull]
 			public int? UpgradeAttemptCount { get; set; }
 
-			public List<PoolId> DynamicPools { get; set; } = new List<PoolId>();
 			public List<PoolId> Pools { get; set; } = new List<PoolId>();
+			public List<PoolId> DynamicPools { get; set; } = new List<PoolId>();
+			public List<PoolId> ExplicitPools { get; set; } = new List<PoolId>();
 
 			[BsonIgnoreIfDefault, BsonDefaultValue(false)]
 			public bool RequestConform { get; set; }
@@ -201,6 +205,9 @@ namespace Horde.Server.Agents
 			public string EnrollmentKey { get; set; } = String.Empty;
 			public string? Comment { get; set; }
 
+			[BsonElement("dv")]
+			public int DocumentVersion { get; set; } = 0;
+
 			[BsonConstructor]
 			private AgentDocument()
 			{
@@ -236,6 +243,44 @@ namespace Horde.Server.Agents
 			_auditLog = auditLog;
 		}
 
+		async ValueTask<AgentDocument?> PostLoadAsync(AgentDocument? document, CancellationToken cancellationToken)
+		{
+			while (document != null)
+			{
+				AgentDocument? newDocument = null;
+				if (document.DocumentVersion == 0)
+				{
+					UpdateDefinition<AgentDocument> updateDefinition = Builders<AgentDocument>.Update
+						.Set(x => x.Pools, CreatePoolsList(document.Pools, document.DynamicPools, document.Properties))
+						.Set(x => x.DynamicPools, CreatePoolsList(document.DynamicPools))
+						.Set(x => x.ExplicitPools, CreatePoolsList(document.Pools))
+						.Set(x => x.DocumentVersion, 1);
+
+					newDocument = await TryUpdateAsync(document, updateDefinition, cancellationToken);
+				}
+				else
+				{
+					break;
+				}
+				document = newDocument ?? await _agents.Find<AgentDocument>(x => x.Id == document.Id).FirstOrDefaultAsync(cancellationToken);
+			}
+			return document;
+		}
+
+		async ValueTask<List<AgentDocument>> PostLoadAsync(List<AgentDocument> documents, CancellationToken cancellationToken)
+		{
+			List<AgentDocument> newDocuments = new List<AgentDocument>();
+			foreach (AgentDocument document in documents)
+			{
+				AgentDocument? newDocument = await PostLoadAsync(document, cancellationToken);
+				if (newDocument != null)
+				{
+					newDocuments.Add(newDocument);
+				}
+			}
+			return newDocuments;
+		}
+
 		[return: NotNullIfNotNull("document")]
 		Agent? CreateAgentObject(AgentDocument? document)
 		{
@@ -296,14 +341,16 @@ namespace Horde.Server.Agents
 		/// <inheritdoc/>
 		public async Task<IAgent?> GetAsync(AgentId agentId, CancellationToken cancellationToken)
 		{
-			AgentDocument? agent = await _agents.Find<AgentDocument>(x => x.Id == agentId).FirstOrDefaultAsync(cancellationToken);
-			return CreateAgentObject(agent);
+			AgentDocument? document = await _agents.Find<AgentDocument>(x => x.Id == agentId).FirstOrDefaultAsync(cancellationToken);
+			document = await PostLoadAsync(document, cancellationToken);
+			return CreateAgentObject(document);
 		}
 
 		/// <inheritdoc/>
 		public async Task<IReadOnlyList<IAgent>> GetManyAsync(List<AgentId> agentIds, CancellationToken cancellationToken)
 		{
 			List<AgentDocument> documents = await _agents.Find(p => agentIds.Contains(p.Id)).ToListAsync(cancellationToken);
+			documents = await PostLoadAsync(documents, cancellationToken);
 			return documents.ConvertAll(x => CreateAgentObject(x));
 		}
 
@@ -320,7 +367,7 @@ namespace Horde.Server.Agents
 
 			if (poolId != null)
 			{
-				filter &= filterBuilder.Eq(nameof(AgentDocument.Pools), poolId);
+				filter &= filterBuilder.AnyEq(x => x.Pools, poolId.Value);
 			}
 
 			if (modifiedAfter != null)
@@ -355,6 +402,7 @@ namespace Horde.Server.Agents
 			}
 
 			List<AgentDocument> documents = await search.ToListAsync(cancellationToken);
+			documents = await PostLoadAsync(documents, cancellationToken);
 			return documents.ConvertAll(x => CreateAgentObject(x));
 		}
 
@@ -362,7 +410,9 @@ namespace Horde.Server.Agents
 		public async Task<IReadOnlyList<IAgent>> FindExpiredAsync(DateTime utcNow, int maxAgents, CancellationToken cancellationToken)
 		{
 			FilterDefinition<AgentDocument> filter = Builders<AgentDocument>.Filter.Exists(x => x.SessionExpiresAt) & Builders<AgentDocument>.Filter.Lt(x => x.SessionExpiresAt, utcNow);
+
 			List<AgentDocument> documents = await _agents.Find(filter).Limit(maxAgents).ToListAsync(cancellationToken);
+			documents = await PostLoadAsync(documents, cancellationToken);
 			return documents.ConvertAll(x => CreateAgentObject(x));
 		}
 
@@ -370,6 +420,7 @@ namespace Horde.Server.Agents
 		public async Task<IReadOnlyList<IAgent>> FindDeletedAsync(CancellationToken cancellationToken)
 		{
 			List<AgentDocument> documents = await _agents.Find(x => x.Deleted).ToListAsync(cancellationToken);
+			documents = await PostLoadAsync(documents, cancellationToken);
 			return documents.ConvertAll(x => CreateAgentObject(x));
 		}
 
@@ -412,9 +463,13 @@ namespace Horde.Server.Agents
 			UpdateDefinitionBuilder<AgentDocument> updateBuilder = new UpdateDefinitionBuilder<AgentDocument>();
 
 			List<UpdateDefinition<AgentDocument>> updates = new List<UpdateDefinition<AgentDocument>>();
-			if (options.Pools != null)
+			if (options.ExplicitPools != null)
 			{
-				updates.Add(updateBuilder.Set(x => x.Pools, options.Pools));
+				List<PoolId> pools = CreatePoolsList(agent.DynamicPools, options.ExplicitPools, agent.Properties);
+				updates.Add(updateBuilder.Set(x => x.Pools, pools));
+
+				List<PoolId> explicitPools = CreatePoolsList(options.ExplicitPools).ToList();
+				updates.Add(updateBuilder.Set(x => x.ExplicitPools, explicitPools));
 			}
 			if (options.Enabled != null)
 			{
@@ -512,9 +567,10 @@ namespace Horde.Server.Agents
 			{
 				updates.Add(updateBuilder.Set(x => x.Resources, new Dictionary<string, int>(options.Resources)));
 			}
-			if (options.DynamicPools != null && !options.DynamicPools.SequenceEqual(agent.DynamicPools))
+			if (options.DynamicPools != null)
 			{
-				updates.Add(updateBuilder.Set(x => x.DynamicPools, new List<PoolId>(options.DynamicPools)));
+				List<PoolId> dynamicPools = CreatePoolsList(options.DynamicPools).ToList();
+				updates.Add(updateBuilder.Set(x => x.DynamicPools, dynamicPools));
 			}
 			if (options.Leases != null)
 			{
@@ -544,6 +600,13 @@ namespace Horde.Server.Agents
 				updates.Add(updateBuilder.Set(x => x.Leases, options.Leases));
 			}
 
+			// Update the pools
+			List<PoolId> pools = CreatePoolsList(options.DynamicPools ?? agent.DynamicPools, agent.ExplicitPools, options.Properties ?? agent.Properties ?? Enumerable.Empty<string>());
+			if (!Enumerable.SequenceEqual(pools, agent.Pools))
+			{
+				updates.Add(updateBuilder.Set(x => x.Pools, pools));
+			}
+
 			// If there are no new updates, return immediately. This is important for preventing UpdateSession calls from returning immediately.
 			if (updates.Count == 0)
 			{
@@ -554,6 +617,60 @@ namespace Horde.Server.Agents
 			return await TryUpdateAsync(agent, updateBuilder.Combine(updates), cancellationToken);
 		}
 
+		static List<PoolId> CreatePoolsList(IEnumerable<PoolId> pools)
+			=> pools.Distinct().OrderBy(x => x.Id.Text).ToList();
+
+		static List<PoolId> CreatePoolsList(IEnumerable<PoolId> dynamicPools, IEnumerable<PoolId> explicitPools, IEnumerable<string>? properties)
+		{
+			List<PoolId> pools = new List<PoolId>();
+			pools.AddRange(dynamicPools);
+			pools.AddRange(explicitPools);
+
+			if (properties != null)
+			{
+				foreach (string property in properties)
+				{
+					const string Key = KnownPropertyNames.RequestedPools + "=";
+					if (property.StartsWith(Key, StringComparison.Ordinal))
+					{
+						try
+						{
+							pools.AddRange(property[Key.Length..].Split(",").Select(x => new PoolId(x)));
+						}
+						catch
+						{
+							// Ignored
+						}
+					}
+				}
+			}
+
+			return CreatePoolsList(pools);
+		}
+
+		private static List<PoolId> GetRequestedPoolsFromProperties(IReadOnlyList<string> properties)
+		{
+			List<PoolId> poolIds = new();
+			foreach (string property in properties)
+			{
+				const string Key = KnownPropertyNames.RequestedPools + "=";
+				if (property.StartsWith(Key, StringComparison.InvariantCulture))
+				{
+					poolIds.AddRange(property[Key.Length..].Split(",").Select(x => new PoolId(x)));
+				}
+			}
+
+			return poolIds;
+		}
+
+		private static List<PoolId> CombineCurrentAndRequestedPools(IReadOnlyList<PoolId> pools, IReadOnlyList<string> properties)
+		{
+			HashSet<PoolId> uniquePools = new(pools);
+			uniquePools.UnionWith(GetRequestedPoolsFromProperties(properties));
+			return new List<PoolId>(uniquePools);
+		}
+
+
 		static bool ResourcesEqual(IReadOnlyDictionary<string, int>? dictA, IReadOnlyDictionary<string, int>? dictB)
 		{
 			dictA ??= ReadOnlyDictionary<string, int>.Empty;
@@ -600,7 +717,6 @@ namespace Horde.Server.Agents
 		{
 			List<string> newProperties = options.Properties.OrderBy(x => x, StringComparer.OrdinalIgnoreCase).ToList();
 			Dictionary<string, int> newResources = new(options.Resources);
-			List<PoolId> newPools = new(options.Pools);
 			List<PoolId> newDynamicPools = new(options.DynamicPools);
 
 			// Reset the agent to use the new session
@@ -614,7 +730,7 @@ namespace Horde.Server.Agents
 			updates.Add(updateBuilder.Unset(x => x.Deleted));
 			updates.Add(updateBuilder.Set(x => x.Properties, newProperties));
 			updates.Add(updateBuilder.Set(x => x.Resources, newResources));
-			updates.Add(updateBuilder.Set(x => x.Pools, newPools));
+			updates.Add(updateBuilder.Set(x => x.Pools, CreatePoolsList(agent.ExplicitPools, newDynamicPools, newProperties)));
 			updates.Add(updateBuilder.Set(x => x.DynamicPools, newDynamicPools));
 			updates.Add(updateBuilder.Set(x => x.Version, options.Version));
 			updates.Add(updateBuilder.Unset(x => x.RequestRestart));
diff --git a/Engine/Source/Programs/Horde/Horde.Server/Agents/AgentService.cs b/Engine/Source/Programs/Horde/Horde.Server/Agents/AgentService.cs
--- a/Engine/Source/Programs/Horde/Horde.Server/Agents/AgentService.cs
+++ b/Engine/Source/Programs/Horde/Horde.Server/Agents/AgentService.cs
@@ -360,28 +360,6 @@ namespace Horde.Server.Agents
 			return newDynamicPools;
 		}
 
-		private static List<PoolId> GetRequestedPoolsFromProperties(IReadOnlyList<string> properties)
-		{
-			List<PoolId> poolIds = new();
-			foreach (string property in properties)
-			{
-				const string Key = KnownPropertyNames.RequestedPools + "=";
-				if (property.StartsWith(Key, StringComparison.InvariantCulture))
-				{
-					poolIds.AddRange(property[Key.Length..].Split(",").Select(x => new PoolId(x)));
-				}
-			}
-
-			return poolIds;
-		}
-
-		private static List<PoolId> CombineCurrentAndRequestedPools(IReadOnlyList<PoolId> pools, IReadOnlyList<string> properties)
-		{
-			HashSet<PoolId> uniquePools = new(pools);
-			uniquePools.UnionWith(GetRequestedPoolsFromProperties(properties));
-			return new List<PoolId>(uniquePools);
-		}
-
 		/// <summary>
 		/// Callback for an agents 
 		/// </summary>
@@ -441,10 +419,9 @@ namespace Horde.Server.Agents
 
 					// Get the new pools for the agent
 					List<PoolId> dynamicPools = await GetDynamicPoolsAsync(agent, cancellationToken);
-					List<PoolId> pools = CombineCurrentAndRequestedPools(agent.ExplicitPools, properties);
 
 					// Reset the agent to use the new session
-					newAgent = await agent.TryCreateSessionAsync(new CreateSessionOptions(newSession.Id, sessionExpiresAt, status, properties, resources, pools, dynamicPools, lastStatusChange ?? utcNow, version), cancellationToken);
+					newAgent = await agent.TryCreateSessionAsync(new CreateSessionOptions(newSession.Id, sessionExpiresAt, status, properties, resources, dynamicPools, lastStatusChange ?? utcNow, version), cancellationToken);
 					if (newAgent != null)
 					{
 						LogPropertyChanges(agentLogger, agent.Properties, newAgent.Properties);
diff --git a/Engine/Source/Programs/Horde/Horde.Server/Agents/AgentsController.cs b/Engine/Source/Programs/Horde/Horde.Server/Agents/AgentsController.cs
--- a/Engine/Source/Programs/Horde/Horde.Server/Agents/AgentsController.cs
+++ b/Engine/Source/Programs/Horde/Horde.Server/Agents/AgentsController.cs
@@ -174,7 +174,7 @@ namespace Horde.Server.Agents
 				new Dictionary<string, int>(agent.Resources),
 				agent.UpdateTime,
 				agent.LastStatusChange,
-				agent.GetPools().Select(x => x.ToString()).ToList(),
+				agent.Pools.Select(x => x.ToString()).ToList(),
 				new { Devices = new[] { new { agent.Properties, agent.Resources } } },
 				leases,
 				agent.Workspaces.ConvertAll(x => CreateGetAgentWorkspaceResponse(x)),
diff --git a/Engine/Source/Programs/Horde/Horde.Server/Agents/Fleet/FleetService.cs b/Engine/Source/Programs/Horde/Horde.Server/Agents/Fleet/FleetService.cs
--- a/Engine/Source/Programs/Horde/Horde.Server/Agents/Fleet/FleetService.cs
+++ b/Engine/Source/Programs/Horde/Horde.Server/Agents/Fleet/FleetService.cs
@@ -203,7 +203,7 @@ namespace Horde.Server.Agents.Fleet
 		internal async Task<List<PoolWithAgents>> GetPoolsWithAgentsAsync(CancellationToken cancellationToken = default)
 		{
 			List<IAgent> agents = (await _agentCollection.FindAsync(status: AgentStatus.Ok, enabled: true, cancellationToken: cancellationToken)).Where(x => !x.RequestShutdown).ToList();
-			List<IAgent> GetAgentsInPool(PoolId poolId) => agents.FindAll(a => a.GetPools().Any(p => p == poolId));
+			List<IAgent> GetAgentsInPool(PoolId poolId) => agents.FindAll(a => a.Pools.Any(p => p == poolId));
 			IReadOnlyList<IPoolConfig> pools = await _poolCollection.GetConfigsAsync(cancellationToken);
 
 			return pools.Select(pool => new PoolWithAgents(pool, GetAgentsInPool(pool.Id))).ToList();
diff --git a/Engine/Source/Programs/Horde/Horde.Server/Agents/Fleet/LeaseUtilizationStrategy.cs b/Engine/Source/Programs/Horde/Horde.Server/Agents/Fleet/LeaseUtilizationStrategy.cs
--- a/Engine/Source/Programs/Horde/Horde.Server/Agents/Fleet/LeaseUtilizationStrategy.cs
+++ b/Engine/Source/Programs/Horde/Horde.Server/Agents/Fleet/LeaseUtilizationStrategy.cs
@@ -214,7 +214,7 @@ namespace Horde.Server.Agents.Fleet
 			// Find pool utilization over the query period
 			foreach (AgentData agentData in agentIdToData.Values)
 			{
-				foreach (PoolId poolId in agentData.Agent.GetPools())
+				foreach (PoolId poolId in agentData.Agent.Pools)
 				{
 					PoolData? poolData;
 					if (poolToData.TryGetValue(poolId, out poolData))
diff --git a/Engine/Source/Programs/Horde/Horde.Server/Agents/IAgent.cs b/Engine/Source/Programs/Horde/Horde.Server/Agents/IAgent.cs
--- a/Engine/Source/Programs/Horde/Horde.Server/Agents/IAgent.cs
+++ b/Engine/Source/Programs/Horde/Horde.Server/Agents/IAgent.cs
@@ -110,6 +110,11 @@ namespace Horde.Server.Agents
 		/// </summary>
 		public int? UpgradeAttemptCount { get; }
 
+		/// <summary>
+		/// All pools for this agent
+		/// </summary>
+		public IReadOnlyList<PoolId> Pools { get; }
+
 		/// <summary>
 		/// Dynamically applied pools
 		/// </summary>
@@ -267,9 +272,9 @@ namespace Horde.Server.Agents
 	/// <param name="RequestShutdown">Whether to request the machine be shut down</param>
 	/// <param name="RequestForceRestart">Request an immediate restart without waiting for leases to complete</param>
 	/// <param name="ShutdownReason">The reason for shutting down agent, ex. Autoscaler/Manual/Unexpected</param>
-	/// <param name="Pools">List of pools for the agent</param>
+	/// <param name="ExplicitPools">List of pools for the agent</param>
 	/// <param name="Comment">New comment</param>
-	public record class UpdateAgentOptions(bool? Enabled = null, bool? RequestConform = null, bool? RequestFullConform = null, bool? RequestRestart = null, bool? RequestShutdown = null, bool? RequestForceRestart = null, string? ShutdownReason = null, List<PoolId>? Pools = null, string? Comment = null);
+	public record class UpdateAgentOptions(bool? Enabled = null, bool? RequestConform = null, bool? RequestFullConform = null, bool? RequestRestart = null, bool? RequestShutdown = null, bool? RequestForceRestart = null, string? ShutdownReason = null, List<PoolId>? ExplicitPools = null, string? Comment = null);
 
 	/// <summary>
 	/// Options for starting a new agent session
@@ -279,11 +284,10 @@ namespace Horde.Server.Agents
 	/// <param name="Status">Status of the agent</param>
 	/// <param name="Properties">Properties for the current session</param>
 	/// <param name="Resources">Resources for the agent</param>
-	/// <param name="Pools">New list of pools for the agent</param>
 	/// <param name="DynamicPools">New list of dynamic pools for the agent</param>
 	/// <param name="LastStatusChange">Time to force status change timestamp to</param>
 	/// <param name="Version">Current version of the agent software</param>
-	public record class CreateSessionOptions(SessionId SessionId, DateTime SessionExpiresAt, AgentStatus Status, IReadOnlyList<string> Properties, IReadOnlyDictionary<string, int> Resources, IReadOnlyList<PoolId> Pools, IReadOnlyList<PoolId> DynamicPools, DateTime LastStatusChange, string? Version);
+	public record class CreateSessionOptions(SessionId SessionId, DateTime SessionExpiresAt, AgentStatus Status, IReadOnlyList<string> Properties, IReadOnlyDictionary<string, int> Resources, IReadOnlyList<PoolId> DynamicPools, DateTime LastStatusChange, string? Version);
 
 	/// <summary>
 	/// Options for updating a new agent session
@@ -373,24 +377,7 @@ namespace Horde.Server.Agents
 		/// <returns></returns>
 		public static bool IsInPool(this IAgent agent, PoolId poolId)
 		{
-			return agent.DynamicPools.Contains(poolId) || agent.ExplicitPools.Contains(poolId);
-		}
-
-		/// <summary>
-		/// Get all the pools for each agent
-		/// </summary>
-		/// <param name="agent">The agent to query</param>
-		/// <returns></returns>
-		public static IEnumerable<PoolId> GetPools(this IAgent agent)
-		{
-			foreach (PoolId poolId in agent.DynamicPools)
-			{
-				yield return poolId;
-			}
-			foreach (PoolId poolId in agent.ExplicitPools)
-			{
-				yield return poolId;
-			}
+			return agent.Pools.Contains(poolId);
 		}
 
 		/// <summary>
@@ -462,7 +449,7 @@ namespace Horde.Server.Agents
 			}
 			else if (name.Equals(KnownPropertyNames.Pool, StringComparison.OrdinalIgnoreCase))
 			{
-				foreach (PoolId poolId in agent.GetPools())
+				foreach (PoolId poolId in agent.Pools)
 				{
 					yield return poolId.ToString();
 				}
diff --git a/Engine/Source/Programs/Horde/Horde.Server/Agents/Pools/PoolService.cs b/Engine/Source/Programs/Horde/Horde.Server/Agents/Pools/PoolService.cs
--- a/Engine/Source/Programs/Horde/Horde.Server/Agents/Pools/PoolService.cs
+++ b/Engine/Source/Programs/Horde/Horde.Server/Agents/Pools/PoolService.cs
@@ -128,7 +128,7 @@ namespace Horde.Server.Agents.Pools
 			Dictionary<PoolId, IPoolConfig> poolMapping = await GetPoolLookupAsync(validAtTime, cancellationToken);
 
 			List<IPoolConfig> pools = new List<IPoolConfig>();
-			foreach (PoolId poolId in agent.GetPools())
+			foreach (PoolId poolId in agent.Pools)
 			{
 				if (poolMapping.TryGetValue(poolId, out IPoolConfig? pool))
 				{
@@ -220,7 +220,7 @@ namespace Horde.Server.Agents.Pools
 			AutoSdkConfig? autoSdkConfig = null;
 
 			Dictionary<PoolId, IPoolConfig> poolMapping = await GetPoolLookupAsync(validAtTime, cancellationToken);
-			foreach (PoolId poolId in agent.GetPools())
+			foreach (PoolId poolId in agent.Pools)
 			{
 				IPoolConfig? pool;
 				if (poolMapping.TryGetValue(poolId, out pool))
diff --git a/Engine/Source/Programs/Horde/Horde.Server/Agents/Pools/PoolsController.cs b/Engine/Source/Programs/Horde/Horde.Server/Agents/Pools/PoolsController.cs
--- a/Engine/Source/Programs/Horde/Horde.Server/Agents/Pools/PoolsController.cs
+++ b/Engine/Source/Programs/Horde/Horde.Server/Agents/Pools/PoolsController.cs
@@ -151,7 +151,7 @@ namespace Horde.Server.Agents.Pools
 					IReadOnlyList<IAgent> agents = await _agentCollection.FindAsync(cancellationToken: cancellationToken);
 					foreach (IAgent agent in agents)
 					{
-						foreach (PoolId poolId in agent.GetPools())
+						foreach (PoolId poolId in agent.Pools)
 						{
 							PoolInfo? poolInfo;
 							if (poolIdToInfo.TryGetValue(poolId, out poolInfo))
diff --git a/Engine/Source/Programs/Horde/Horde.Server/Agents/Utilization/UtilizationDataService.cs b/Engine/Source/Programs/Horde/Horde.Server/Agents/Utilization/UtilizationDataService.cs
--- a/Engine/Source/Programs/Horde/Horde.Server/Agents/Utilization/UtilizationDataService.cs
+++ b/Engine/Source/Programs/Horde/Horde.Server/Agents/Utilization/UtilizationDataService.cs
@@ -77,7 +77,7 @@ namespace Horde.Server.Agents.Utilization
 			agents = agents.Where(x => x.Enabled && x.IsSessionValid(currentTime)).ToList();
 
 			// Find all the agents
-			Dictionary<AgentId, List<PoolId>> agentToPoolIds = agents.ToDictionary(x => x.Id, x => x.GetPools().ToList());
+			Dictionary<AgentId, List<PoolId>> agentToPoolIds = agents.ToDictionary(x => x.Id, x => x.Pools.ToList());
 
 			// Generate all the telemetry data
 			DateTime bucketMinTime = minTime;
diff --git a/Engine/Source/Programs/Horde/Horde.Server/Jobs/JobTaskSource.cs b/Engine/Source/Programs/Horde/Horde.Server/Jobs/JobTaskSource.cs
--- a/Engine/Source/Programs/Horde/Horde.Server/Jobs/JobTaskSource.cs
+++ b/Engine/Source/Programs/Horde/Horde.Server/Jobs/JobTaskSource.cs
@@ -288,7 +288,7 @@ namespace Horde.Server.Jobs
 				List<object> outputWaiters = new List<object>();
 				foreach (QueueWaiter waiter in _waiters)
 				{
-					outputWaiters.Add(new { Id = waiter.Agent.Id.ToString(), Pools = waiter.Agent.GetPools().Select(x => x.ToString()).ToList(), waiter.Agent.Workspaces });
+					outputWaiters.Add(new { Id = waiter.Agent.Id.ToString(), Pools = waiter.Agent.Pools.Select(x => x.ToString()).ToList(), waiter.Agent.Workspaces });
 				}
 
 				return new { Items = outputItems, Waiters = outputWaiters };
diff --git a/Engine/Source/Programs/Horde/Horde.Server/Server/ServerController.cs b/Engine/Source/Programs/Horde/Horde.Server/Server/ServerController.cs
--- a/Engine/Source/Programs/Horde/Horde.Server/Server/ServerController.cs
+++ b/Engine/Source/Programs/Horde/Horde.Server/Server/ServerController.cs
@@ -230,7 +230,7 @@ namespace Horde.Server.Server
 				IReadOnlyList<IAgent> agents = await agentCollection.FindAsync(cancellationToken: cancellationToken);
 				foreach (IAgent agent in agents)
 				{
-					foreach (PoolId poolId in agent.GetPools())
+					foreach (PoolId poolId in agent.Pools)
 					{
 						int count;
 						if (!poolIdToCount.TryGetValue(poolId, out count))
