Horde: Allow registering multiple workspace materializer factories.

diff --git a/Engine/Source/Programs/Horde/Drivers/JobDriver/DriverApp.cs b/Engine/Source/Programs/Horde/Drivers/JobDriver/DriverApp.cs
--- a/Engine/Source/Programs/Horde/Drivers/JobDriver/DriverApp.cs
+++ b/Engine/Source/Programs/Horde/Drivers/JobDriver/DriverApp.cs
@@ -13,39 +13,28 @@ using Serilog.Formatting.Json;
 
 namespace JobDriver
 {
-	class DriverApp
+	/// <summary>
+	/// Application class for the job driver
+	/// </summary>
+	public static class DriverApp
 	{
+		/// <summary>
+		/// Main entry point
+		/// </summary>
 		public static async Task<int> Main(string[] args)
 		{
 			CommandLineArguments arguments = new CommandLineArguments(args);
 
-			// Read the driver config
-			IConfiguration configuration = new ConfigurationBuilder()
-				.AddJsonFile("appsettings.json", optional: false)
-				.AddEnvironmentVariables()
-				.Build();
-
 			// Create the services 
 			IServiceCollection services = new ServiceCollection();
-			services.AddOptions<DriverSettings>().Configure(options => configuration.GetSection("Driver").Bind(options)).ValidateDataAnnotations();
-			services.AddLogging(builder => builder.AddEpicDefault());
-			services.AddHorde(options => options.AllowAuthPrompt = false);
-
-			services.AddSingleton<IJobExecutorFactory, PerforceExecutorFactory>();
-			services.AddSingleton<IJobExecutorFactory, WorkspaceExecutorFactory>();
-			services.AddSingleton<IJobExecutorFactory, LocalExecutorFactory>();
-			services.AddSingleton<IJobExecutorFactory, TestExecutorFactory>();
-
-			services.AddSingleton<IWorkspaceMaterializerFactory, WorkspaceMaterializerFactory>();
-
-			services.AddCommandsFromAssembly(Assembly.GetExecutingAssembly());
+			RegisterServices(services);
 
 			// Run the host
 			await using ServiceProvider serviceProvider = services.BuildServiceProvider();
 			return await CommandHost.RunAsync(arguments, serviceProvider, null);
 		}
 
-		public static ILoggerProvider CreateLoggerProvider(IConfiguration configuration)
+		static ILoggerProvider CreateLoggerProvider(IConfiguration configuration)
 		{
 			//			ConsoleTheme theme;
 			//			if (RuntimeInformation.IsOSPlatform(OSPlatform.Windows) && Environment.OSVersion.Version < new Version(10, 0))
@@ -64,5 +53,32 @@ namespace JobDriver
 
 			return new SerilogLoggerProvider(loggerConfiguration.CreateLogger());
 		}
+
+		/// <summary>
+		/// Helper method to register services for this app
+		/// </summary>
+		/// <param name="services"></param>
+		public static void RegisterServices(IServiceCollection services)
+		{
+			// Read the driver config
+			IConfiguration configuration = new ConfigurationBuilder()
+				.AddJsonFile("appsettings.json", optional: false)
+				.AddEnvironmentVariables()
+				.Build();
+
+			// Register the services
+			services.AddOptions<DriverSettings>().Configure(options => configuration.GetSection("Driver").Bind(options)).ValidateDataAnnotations();
+			services.AddLogging(builder => builder.AddEpicDefault());
+			services.AddHorde(options => options.AllowAuthPrompt = false);
+
+			services.AddSingleton<IJobExecutorFactory, PerforceExecutorFactory>();
+			services.AddSingleton<IJobExecutorFactory, WorkspaceExecutorFactory>();
+			services.AddSingleton<IJobExecutorFactory, LocalExecutorFactory>();
+			services.AddSingleton<IJobExecutorFactory, TestExecutorFactory>();
+
+			services.AddSingleton<IWorkspaceMaterializerFactory, WorkspaceMaterializerFactory>();
+
+			services.AddCommandsFromAssembly(Assembly.GetExecutingAssembly());
+		}
 	}
 }
diff --git a/Engine/Source/Programs/Horde/Drivers/JobDriver/Execution/ManagedWorkspaceMaterializer.cs b/Engine/Source/Programs/Horde/Drivers/JobDriver/Execution/ManagedWorkspaceMaterializer.cs
--- a/Engine/Source/Programs/Horde/Drivers/JobDriver/Execution/ManagedWorkspaceMaterializer.cs
+++ b/Engine/Source/Programs/Horde/Drivers/JobDriver/Execution/ManagedWorkspaceMaterializer.cs
@@ -16,6 +16,11 @@ namespace JobDriver.Execution;
 /// </summary>
 public sealed class ManagedWorkspaceMaterializer : IWorkspaceMaterializer
 {
+	/// <summary>
+	/// Name of this materializer
+	/// </summary>
+	public const string Name = "ManagedWorkspace";
+
 	private readonly RpcAgentWorkspace _agentWorkspace;
 	private readonly bool _useCacheFile;
 	private readonly bool _cleanDuringFinalize;
diff --git a/Engine/Source/Programs/Horde/Drivers/JobDriver/Execution/WorkspaceExecutor.cs b/Engine/Source/Programs/Horde/Drivers/JobDriver/Execution/WorkspaceExecutor.cs
--- a/Engine/Source/Programs/Horde/Drivers/JobDriver/Execution/WorkspaceExecutor.cs
+++ b/Engine/Source/Programs/Horde/Drivers/JobDriver/Execution/WorkspaceExecutor.cs
@@ -152,14 +152,14 @@ namespace JobDriver.Execution
 
 	class WorkspaceExecutorFactory : IJobExecutorFactory
 	{
-		private readonly IWorkspaceMaterializerFactory _materializerFactory;
+		private readonly IEnumerable<IWorkspaceMaterializerFactory> _materializerFactories;
 		private readonly ILoggerFactory _loggerFactory;
 
 		public string Name => WorkspaceExecutor.Name;
 
-		public WorkspaceExecutorFactory(IWorkspaceMaterializerFactory materializerFactory, ILoggerFactory loggerFactory)
+		public WorkspaceExecutorFactory(IEnumerable<IWorkspaceMaterializerFactory> materializerFactories, ILoggerFactory loggerFactory)
 		{
-			_materializerFactory = materializerFactory;
+			_materializerFactories = materializerFactories;
 			_loggerFactory = loggerFactory;
 		}
 
@@ -169,12 +169,12 @@ namespace JobDriver.Execution
 			IWorkspaceMaterializer? autoSdkMaterializer = null;
 			try
 			{
-				WorkspaceMaterializerType type = GetMaterializerType(options.JobOptions.WorkspaceMaterializer, WorkspaceMaterializerType.ManagedWorkspace);
+				string name = String.IsNullOrEmpty(options.JobOptions.WorkspaceMaterializer) ? ManagedWorkspaceMaterializer.Name : options.JobOptions.WorkspaceMaterializer;
 
-				workspaceMaterializer = await _materializerFactory.CreateMaterializerAsync(type, workspaceInfo, options, forAutoSdk: false, cancellationToken);
+				workspaceMaterializer = await CreateMaterializerAsync(name, workspaceInfo, options.WorkingDir, forAutoSdk: false, cancellationToken);
 				if (autoSdkWorkspaceInfo != null)
 				{
-					autoSdkMaterializer = await _materializerFactory.CreateMaterializerAsync(type, autoSdkWorkspaceInfo, options, forAutoSdk: true, cancellationToken);
+					autoSdkMaterializer = await CreateMaterializerAsync(name, autoSdkWorkspaceInfo, options.WorkingDir, forAutoSdk: true, cancellationToken);
 				}
 
 				return new WorkspaceExecutor(options, workspaceMaterializer, autoSdkMaterializer, _loggerFactory.CreateLogger<WorkspaceExecutor>());
@@ -187,18 +187,16 @@ namespace JobDriver.Execution
 			}
 		}
 
-		private static WorkspaceMaterializerType GetMaterializerType(string name, WorkspaceMaterializerType defaultValue)
+		async Task<IWorkspaceMaterializer> CreateMaterializerAsync(string name, RpcAgentWorkspace workspaceInfo, DirectoryReference workspaceDir, bool forAutoSdk, CancellationToken cancellationToken)
 		{
-			if (String.IsNullOrEmpty(name))
+			foreach (IWorkspaceMaterializerFactory materializerFactory in _materializerFactories)
 			{
-				return defaultValue;
-			}
-
-			if (Enum.TryParse(name, true, out WorkspaceMaterializerType enumType))
-			{
-				return enumType;
+				IWorkspaceMaterializer? materializer = await materializerFactory.CreateMaterializerAsync(name, workspaceInfo, workspaceDir, forAutoSdk, cancellationToken);
+				if (materializer != null)
+				{
+					return materializer;
+				}
 			}
-
 			throw new ArgumentException($"Unable to find materializer type '{name}'");
 		}
 	}
diff --git a/Engine/Source/Programs/Horde/Drivers/JobDriver/Execution/WorkspaceMaterializer.cs b/Engine/Source/Programs/Horde/Drivers/JobDriver/Execution/WorkspaceMaterializer.cs
--- a/Engine/Source/Programs/Horde/Drivers/JobDriver/Execution/WorkspaceMaterializer.cs
+++ b/Engine/Source/Programs/Horde/Drivers/JobDriver/Execution/WorkspaceMaterializer.cs
@@ -101,26 +101,21 @@ public interface IWorkspaceMaterializer : IDisposable
 	Task FinalizeAsync(CancellationToken cancellationToken);
 }
 
-enum WorkspaceMaterializerType
-{
-	ManagedWorkspace,
-}
-
 /// <summary>
 /// Factory for creating new workspace materializers
 /// </summary>
-interface IWorkspaceMaterializerFactory
+public interface IWorkspaceMaterializerFactory
 {
 	/// <summary>
 	/// Creates a new workspace materializer instance
 	/// </summary>
-	/// <param name="type">Type of materializer to instantiate</param>
+	/// <param name="name">Name of the materializer to create</param>
 	/// <param name="workspaceInfo">Agent workspace</param>
-	/// <param name="options">Job options</param>
+	/// <param name="workspaceDir">Directory for the workspace</param>
 	/// <param name="forAutoSdk">Whether intended for AutoSDK materialization</param>
 	/// <param name="cancellationToken">Cancellation token for the operation</param>
 	/// <returns>A new workspace materializer instance</returns>
-	Task<IWorkspaceMaterializer> CreateMaterializerAsync(WorkspaceMaterializerType type, RpcAgentWorkspace workspaceInfo, JobExecutorOptions options, bool forAutoSdk = false, CancellationToken cancellationToken = default);
+	Task<IWorkspaceMaterializer?> CreateMaterializerAsync(string name, RpcAgentWorkspace workspaceInfo, DirectoryReference workspaceDir, bool forAutoSdk = false, CancellationToken cancellationToken = default);
 }
 
 class WorkspaceMaterializerFactory : IWorkspaceMaterializerFactory
@@ -131,17 +126,19 @@ class WorkspaceMaterializerFactory : IWorkspaceMaterializerFactory
 		=> _serviceProvider = serviceProvider;
 
 	/// <inheritdoc/>
-	public async Task<IWorkspaceMaterializer> CreateMaterializerAsync(WorkspaceMaterializerType type, RpcAgentWorkspace workspaceInfo, JobExecutorOptions options, bool forAutoSdk, CancellationToken cancellationToken)
+	public async Task<IWorkspaceMaterializer?> CreateMaterializerAsync(string name, RpcAgentWorkspace workspaceInfo, DirectoryReference workspaceDir, bool forAutoSdk, CancellationToken cancellationToken)
 	{
-		return type switch
+		if (name.Equals(ManagedWorkspaceMaterializer.Name, StringComparison.OrdinalIgnoreCase))
 		{
-			WorkspaceMaterializerType.ManagedWorkspace when forAutoSdk 
-				=> await ManagedWorkspaceMaterializer.CreateAsync(workspaceInfo, options.WorkingDir, true, false, _serviceProvider.GetRequiredService<ILogger<ManagedWorkspaceMaterializer>>(), cancellationToken),
-
-			WorkspaceMaterializerType.ManagedWorkspace 
-				=> await ManagedWorkspaceMaterializer.CreateAsync(workspaceInfo, options.WorkingDir, false, true, _serviceProvider.GetRequiredService<ILogger<ManagedWorkspaceMaterializer>>(), cancellationToken),
-
-			_ => throw new Exception("Unhandled materializer option: " + type)
-		};
+			if (forAutoSdk)
+			{
+				return await ManagedWorkspaceMaterializer.CreateAsync(workspaceInfo, workspaceDir, true, false, _serviceProvider.GetRequiredService<ILogger<ManagedWorkspaceMaterializer>>(), cancellationToken);
+			}
+			else
+			{
+				return await ManagedWorkspaceMaterializer.CreateAsync(workspaceInfo, workspaceDir, false, true, _serviceProvider.GetRequiredService<ILogger<ManagedWorkspaceMaterializer>>(), cancellationToken);
+			}
+		}
+		return null;
 	}
 }
