Horde: Fix static analysis warnings.

#rnx

diff --git a/Engine/Source/Programs/Horde/Horde.Server/Agents/AgentCollection.cs b/Engine/Source/Programs/Horde/Horde.Server/Agents/AgentCollection.cs
--- a/Engine/Source/Programs/Horde/Horde.Server/Agents/AgentCollection.cs
+++ b/Engine/Source/Programs/Horde/Horde.Server/Agents/AgentCollection.cs
@@ -8,7 +8,6 @@ using System.Linq;
 using System.Linq.Expressions;
 using System.Threading;
 using System.Threading.Tasks;
-using Amazon.EC2.Model;
 using EpicGames.Core;
 using EpicGames.Horde.Agents;
 using EpicGames.Horde.Agents.Leases;
@@ -19,7 +18,6 @@ using Google.Protobuf.WellKnownTypes;
 using Horde.Server.Auditing;
 using Horde.Server.Server;
 using HordeCommon.Rpc.Tasks;
-using Microsoft.Extensions.Azure;
 using MongoDB.Bson;
 using MongoDB.Bson.Serialization.Attributes;
 using MongoDB.Driver;
@@ -72,7 +70,7 @@ namespace Horde.Server.Agents
 			string IAgent.EnrollmentKey => _document.EnrollmentKey;
 			DateTime IAgent.UpdateTime => _document.UpdateTime;
 			uint IAgent.UpdateIndex => _document.UpdateIndex;
-			
+
 			public Agent(AgentCollection collection, AgentDocument document)
 			{
 				_collection = collection;
@@ -648,29 +646,6 @@ namespace Horde.Server.Agents
 			return CreatePoolsList(pools);
 		}
 
-		private static List<PoolId> GetRequestedPoolsFromProperties(IReadOnlyList<string> properties)
-		{
-			List<PoolId> poolIds = new();
-			foreach (string property in properties)
-			{
-				const string Key = KnownPropertyNames.RequestedPools + "=";
-				if (property.StartsWith(Key, StringComparison.InvariantCulture))
-				{
-					poolIds.AddRange(property[Key.Length..].Split(",").Select(x => new PoolId(x)));
-				}
-			}
-
-			return poolIds;
-		}
-
-		private static List<PoolId> CombineCurrentAndRequestedPools(IReadOnlyList<PoolId> pools, IReadOnlyList<string> properties)
-		{
-			HashSet<PoolId> uniquePools = new(pools);
-			uniquePools.UnionWith(GetRequestedPoolsFromProperties(properties));
-			return new List<PoolId>(uniquePools);
-		}
-
-
 		static bool ResourcesEqual(IReadOnlyDictionary<string, int>? dictA, IReadOnlyDictionary<string, int>? dictB)
 		{
 			dictA ??= ReadOnlyDictionary<string, int>.Empty;
