Horde: Prevent agent documents being updated if dynamic pools have not changed. Fixes UpdateSession() being called repeatedly.

#rnx

diff --git a/Engine/Source/Programs/Horde/Horde.Server/Agents/AgentCollection.cs b/Engine/Source/Programs/Horde/Horde.Server/Agents/AgentCollection.cs
--- a/Engine/Source/Programs/Horde/Horde.Server/Agents/AgentCollection.cs
+++ b/Engine/Source/Programs/Horde/Horde.Server/Agents/AgentCollection.cs
@@ -568,7 +568,10 @@ namespace Horde.Server.Agents
 			if (options.DynamicPools != null)
 			{
 				List<PoolId> dynamicPools = CreatePoolsList(options.DynamicPools).ToList();
-				updates.Add(updateBuilder.Set(x => x.DynamicPools, dynamicPools));
+				if (!Enumerable.SequenceEqual(dynamicPools, agent.DynamicPools))
+				{
+					updates.Add(updateBuilder.Set(x => x.DynamicPools, dynamicPools));
+				}
 			}
 			if (options.Leases != null)
 			{
