Horde: Remove unused method causing static analysis warning.

#rnx

diff --git a/Engine/Source/Programs/Horde/Drivers/JobDriver/DriverApp.cs b/Engine/Source/Programs/Horde/Drivers/JobDriver/DriverApp.cs
--- a/Engine/Source/Programs/Horde/Drivers/JobDriver/DriverApp.cs
+++ b/Engine/Source/Programs/Horde/Drivers/JobDriver/DriverApp.cs
@@ -34,26 +34,6 @@ namespace JobDriver
 			return await CommandHost.RunAsync(arguments, serviceProvider, null);
 		}
 
-		static ILoggerProvider CreateLoggerProvider(IConfiguration configuration)
-		{
-			//			ConsoleTheme theme;
-			//			if (RuntimeInformation.IsOSPlatform(OSPlatform.Windows) && Environment.OSVersion.Version < new Version(10, 0))
-			//			{
-			//				theme = SystemConsoleTheme.Literate;
-			//			}
-			//			else
-			//			{
-			//				theme = AnsiConsoleTheme.Code;
-			//			}
-
-			LoggerConfiguration loggerConfiguration = new LoggerConfiguration()
-				.WriteTo.Console(new JsonFormatter(renderMessage: true))
-				.ReadFrom.Configuration(configuration)
-				.Enrich.FromLogContext();
-
-			return new SerilogLoggerProvider(loggerConfiguration.CreateLogger());
-		}
-
 		/// <summary>
 		/// Helper method to register services for this app
 		/// </summary>
