Horde: Remove unnecessary using directives.

#rnx

diff --git a/Engine/Source/Programs/Horde/Drivers/JobDriver/DriverApp.cs b/Engine/Source/Programs/Horde/Drivers/JobDriver/DriverApp.cs
--- a/Engine/Source/Programs/Horde/Drivers/JobDriver/DriverApp.cs
+++ b/Engine/Source/Programs/Horde/Drivers/JobDriver/DriverApp.cs
@@ -6,10 +6,6 @@ using EpicGames.Horde;
 using JobDriver.Execution;
 using Microsoft.Extensions.Configuration;
 using Microsoft.Extensions.DependencyInjection;
-using Microsoft.Extensions.Logging;
-using Serilog;
-using Serilog.Extensions.Logging;
-using Serilog.Formatting.Json;
 
 namespace JobDriver
 {
