Horde: Fix conform commands using the incorrect access token to communicate with the server.

diff --git a/Engine/Source/Programs/Horde/Drivers/JobDriver/Commands/Execution/ExecuteConformCommand.cs b/Engine/Source/Programs/Horde/Drivers/JobDriver/Commands/Execution/ExecuteConformCommand.cs
--- a/Engine/Source/Programs/Horde/Drivers/JobDriver/Commands/Execution/ExecuteConformCommand.cs
+++ b/Engine/Source/Programs/Horde/Drivers/JobDriver/Commands/Execution/ExecuteConformCommand.cs
@@ -45,8 +45,7 @@ namespace JobDriver.Commands.Execution
 
 			ConformTask conformTask = ConformTask.Parser.ParseFrom(Convert.FromBase64String(Task));
 
-			string accessToken = null!;
-			IHordeClient hordeClient = _hordeClientFactory.Create(accessToken);
+			IHordeClient hordeClient = _hordeClientFactory.Create();
 
 			ConformExecutor conformExecutor = new ConformExecutor(hordeClient, WorkingDir, AgentId, LeaseId, conformTask, _driverSettings.Value, logger);
 			await conformExecutor.ExecuteAsync(CancellationToken.None);
