Horde: Rename WorkspaceMaterializerFactory to ManagedWorkspaceMaterializerFactory now that it only creates workspaces of a particular type.

#rnx

diff --git a/Engine/Source/Programs/Horde/Drivers/JobDriver/DriverApp.cs b/Engine/Source/Programs/Horde/Drivers/JobDriver/DriverApp.cs
--- a/Engine/Source/Programs/Horde/Drivers/JobDriver/DriverApp.cs
+++ b/Engine/Source/Programs/Horde/Drivers/JobDriver/DriverApp.cs
@@ -52,7 +52,7 @@ namespace JobDriver
 			services.AddSingleton<IJobExecutorFactory, LocalExecutorFactory>();
 			services.AddSingleton<IJobExecutorFactory, TestExecutorFactory>();
 
-			services.AddSingleton<IWorkspaceMaterializerFactory, WorkspaceMaterializerFactory>();
+			services.AddSingleton<IWorkspaceMaterializerFactory, ManagedWorkspaceMaterializerFactory>();
 
 			services.AddCommandsFromAssembly(Assembly.GetExecutingAssembly());
 		}
diff --git a/Engine/Source/Programs/Horde/Drivers/JobDriver/Execution/ManagedWorkspaceMaterializer.cs b/Engine/Source/Programs/Horde/Drivers/JobDriver/Execution/ManagedWorkspaceMaterializer.cs
--- a/Engine/Source/Programs/Horde/Drivers/JobDriver/Execution/ManagedWorkspaceMaterializer.cs
+++ b/Engine/Source/Programs/Horde/Drivers/JobDriver/Execution/ManagedWorkspaceMaterializer.cs
@@ -8,6 +8,7 @@ using Microsoft.Extensions.Logging;
 using OpenTracing;
 using OpenTracing.Util;
 using EpicGames.Perforce;
+using Microsoft.Extensions.DependencyInjection;
 
 namespace JobDriver.Execution;
 
@@ -154,4 +155,29 @@ public sealed class ManagedWorkspaceMaterializer : IWorkspaceMaterializer
 		scope.Span.SetTag("CleanDuringFinalize", _cleanDuringFinalize);
 		return scope;
 	}
-}
\ No newline at end of file
+}
+
+class ManagedWorkspaceMaterializerFactory : IWorkspaceMaterializerFactory
+{
+	readonly IServiceProvider _serviceProvider;
+
+	public ManagedWorkspaceMaterializerFactory(IServiceProvider serviceProvider)
+		=> _serviceProvider = serviceProvider;
+
+	/// <inheritdoc/>
+	public async Task<IWorkspaceMaterializer?> CreateMaterializerAsync(string name, RpcAgentWorkspace workspaceInfo, DirectoryReference workspaceDir, bool forAutoSdk, CancellationToken cancellationToken)
+	{
+		if (name.Equals(ManagedWorkspaceMaterializer.Name, StringComparison.OrdinalIgnoreCase))
+		{
+			if (forAutoSdk)
+			{
+				return await ManagedWorkspaceMaterializer.CreateAsync(workspaceInfo, workspaceDir, true, false, _serviceProvider.GetRequiredService<ILogger<ManagedWorkspaceMaterializer>>(), cancellationToken);
+			}
+			else
+			{
+				return await ManagedWorkspaceMaterializer.CreateAsync(workspaceInfo, workspaceDir, false, true, _serviceProvider.GetRequiredService<ILogger<ManagedWorkspaceMaterializer>>(), cancellationToken);
+			}
+		}
+		return null;
+	}
+}
diff --git a/Engine/Source/Programs/Horde/Drivers/JobDriver/Execution/WorkspaceMaterializer.cs b/Engine/Source/Programs/Horde/Drivers/JobDriver/Execution/WorkspaceMaterializer.cs
--- a/Engine/Source/Programs/Horde/Drivers/JobDriver/Execution/WorkspaceMaterializer.cs
+++ b/Engine/Source/Programs/Horde/Drivers/JobDriver/Execution/WorkspaceMaterializer.cs
@@ -2,8 +2,6 @@
 
 using EpicGames.Core;
 using HordeCommon.Rpc.Messages;
-using Microsoft.Extensions.DependencyInjection;
-using Microsoft.Extensions.Logging;
 
 namespace JobDriver.Execution;
 
@@ -15,7 +13,6 @@ public class WorkspaceMaterializationException : Exception
 	/// <summary>
 	/// Constructor
 	/// </summary>
-	/// <param name="message"></param>
 	public WorkspaceMaterializationException(string? message) : base(message)
 	{
 	}
@@ -23,8 +20,6 @@ public class WorkspaceMaterializationException : Exception
 	/// <summary>
 	/// Constructor
 	/// </summary>
-	/// <param name="message"></param>
-	/// <param name="innerException"></param>
 	public WorkspaceMaterializationException(string? message, Exception? innerException) : base(message, innerException)
 	{
 	}
@@ -111,34 +106,9 @@ public interface IWorkspaceMaterializerFactory
 	/// </summary>
 	/// <param name="name">Name of the materializer to create</param>
 	/// <param name="workspaceInfo">Agent workspace</param>
-	/// <param name="workspaceDir">Directory for the workspace</param>
+	/// <param name="workingDir">Working directory for the agent (ie. the root directory, not the workspace location)</param>
 	/// <param name="forAutoSdk">Whether intended for AutoSDK materialization</param>
 	/// <param name="cancellationToken">Cancellation token for the operation</param>
 	/// <returns>A new workspace materializer instance</returns>
-	Task<IWorkspaceMaterializer?> CreateMaterializerAsync(string name, RpcAgentWorkspace workspaceInfo, DirectoryReference workspaceDir, bool forAutoSdk = false, CancellationToken cancellationToken = default);
-}
-
-class WorkspaceMaterializerFactory : IWorkspaceMaterializerFactory
-{
-	readonly IServiceProvider _serviceProvider;
-
-	public WorkspaceMaterializerFactory(IServiceProvider serviceProvider)
-		=> _serviceProvider = serviceProvider;
-
-	/// <inheritdoc/>
-	public async Task<IWorkspaceMaterializer?> CreateMaterializerAsync(string name, RpcAgentWorkspace workspaceInfo, DirectoryReference workspaceDir, bool forAutoSdk, CancellationToken cancellationToken)
-	{
-		if (name.Equals(ManagedWorkspaceMaterializer.Name, StringComparison.OrdinalIgnoreCase))
-		{
-			if (forAutoSdk)
-			{
-				return await ManagedWorkspaceMaterializer.CreateAsync(workspaceInfo, workspaceDir, true, false, _serviceProvider.GetRequiredService<ILogger<ManagedWorkspaceMaterializer>>(), cancellationToken);
-			}
-			else
-			{
-				return await ManagedWorkspaceMaterializer.CreateAsync(workspaceInfo, workspaceDir, false, true, _serviceProvider.GetRequiredService<ILogger<ManagedWorkspaceMaterializer>>(), cancellationToken);
-			}
-		}
-		return null;
-	}
+	Task<IWorkspaceMaterializer?> CreateMaterializerAsync(string name, RpcAgentWorkspace workspaceInfo, DirectoryReference workingDir, bool forAutoSdk = false, CancellationToken cancellationToken = default);
 }
