Horde: Update cached agents using database read replica

This data is already assumed a few seconds stale so reading from a secondary MongoDB instance is fine.

#rnx

diff --git a/Engine/Source/Programs/Horde/Horde.Server/Agents/AgentService.cs b/Engine/Source/Programs/Horde/Horde.Server/Agents/AgentService.cs
--- a/Engine/Source/Programs/Horde/Horde.Server/Agents/AgentService.cs
+++ b/Engine/Source/Programs/Horde/Horde.Server/Agents/AgentService.cs
@@ -1166,7 +1166,7 @@ namespace Horde.Server.Agents
 		{
 			using TelemetrySpan span = _tracer.StartActiveSpan($"{nameof(AgentService)}.{nameof(RefreshCachedAgentsAsync)}");
 			Dictionary<AgentId, IAgent> agents = new();
-			IReadOnlyList<IAgent> agentList = await Agents.FindAsync(cancellationToken: cancellationToken);
+			IReadOnlyList<IAgent> agentList = await Agents.FindAsync(consistentRead: false, cancellationToken: cancellationToken);
 			foreach (IAgent agent in agentList)
 			{
 				agents[agent.Id] = agent;
