Horde: Use cached agent data for fleet and pool size handling

diff --git a/Engine/Source/Programs/Horde/Horde.Server.Tests/Fleet/FleetServiceTest.cs b/Engine/Source/Programs/Horde/Horde.Server.Tests/Fleet/FleetServiceTest.cs
--- a/Engine/Source/Programs/Horde/Horde.Server.Tests/Fleet/FleetServiceTest.cs
+++ b/Engine/Source/Programs/Horde/Horde.Server.Tests/Fleet/FleetServiceTest.cs
@@ -235,7 +235,7 @@ namespace Horde.Server.Tests.Fleet
 			serverSettingsOpt.Value.FleetManagerV2 = FleetManagerType.AwsReuse;
 
 			FleetService service = new(
-				AgentCollection, GraphCollection, JobCollection, LeaseCollection, PoolCollection, new DowntimeServiceStub(isDowntimeActive), StreamCollection, Meter,
+				AgentService, GraphCollection, JobCollection, LeaseCollection, PoolCollection, new DowntimeServiceStub(isDowntimeActive), StreamCollection, Meter,
 				new StubFleetManagerFactory(fleetManager), Clock, Cache, serverSettingsOpt, GlobalConfig, ServiceProvider, Tracer, loggerFactory.CreateLogger<FleetService>());
 
 			return service;
diff --git a/Engine/Source/Programs/Horde/Horde.Server/Agents/Fleet/FleetService.cs b/Engine/Source/Programs/Horde/Horde.Server/Agents/Fleet/FleetService.cs
--- a/Engine/Source/Programs/Horde/Horde.Server/Agents/Fleet/FleetService.cs
+++ b/Engine/Source/Programs/Horde/Horde.Server/Agents/Fleet/FleetService.cs
@@ -64,7 +64,7 @@ namespace Horde.Server.Agents.Fleet
 		/// </summary>
 		private const int MaxParallelTasks = 10;
 
-		private readonly IAgentCollection _agentCollection;
+		private readonly AgentService _agentService;
 		private readonly IGraphCollection _graphCollection;
 		private readonly IJobCollection _jobCollection;
 		private readonly ILeaseCollection _leaseCollection;
@@ -89,7 +89,7 @@ namespace Horde.Server.Agents.Fleet
 		/// Constructor
 		/// </summary>
 		public FleetService(
-			IAgentCollection agentCollection,
+			AgentService agentService,
 			IGraphCollection graphCollection,
 			IJobCollection jobCollection,
 			ILeaseCollection leaseCollection,
@@ -106,7 +106,7 @@ namespace Horde.Server.Agents.Fleet
 			Tracer tracer,
 			ILogger<FleetService> logger)
 		{
-			_agentCollection = agentCollection;
+			_agentService = agentService;
 			_graphCollection = graphCollection;
 			_jobCollection = jobCollection;
 			_leaseCollection = leaseCollection;
@@ -202,7 +202,9 @@ namespace Horde.Server.Agents.Fleet
 
 		internal async Task<List<PoolWithAgents>> GetPoolsWithAgentsAsync(CancellationToken cancellationToken = default)
 		{
-			List<IAgent> agents = (await _agentCollection.FindAsync(status: AgentStatus.Ok, enabled: true, cancellationToken: cancellationToken)).Where(x => !x.RequestShutdown).ToList();
+			List<IAgent> agents = (await _agentService.GetCachedAgentsAsync(cancellationToken))
+				.Where(x => x is { Status: AgentStatus.Ok, Enabled: true, RequestShutdown: false }).ToList();
+			
 			List<IAgent> GetAgentsInPool(PoolId poolId) => agents.FindAll(a => a.Pools.Any(p => p == poolId));
 			IReadOnlyList<IPoolConfig> pools = await _poolCollection.GetConfigsAsync(cancellationToken);
 
@@ -396,7 +398,7 @@ namespace Horde.Server.Agents.Fleet
 
 							case PoolSizeStrategy.LeaseUtilization:
 								LeaseUtilizationSettings luSettings = DeserializeConfig<LeaseUtilizationSettings>(info.Config);
-								LeaseUtilizationStrategy luStrategy = new(_agentCollection, _poolCollection, _leaseCollection, _clock, _cache, luSettings);
+								LeaseUtilizationStrategy luStrategy = new(_agentService, _poolCollection, _leaseCollection, _clock, _cache, luSettings);
 								return info.ExtraAgentCount != 0 ? new ExtraAgentCountStrategy(luStrategy, info.ExtraAgentCount) : luStrategy;
 
 							case PoolSizeStrategy.ComputeQueueAwsMetric:
@@ -434,7 +436,7 @@ namespace Horde.Server.Agents.Fleet
 					{
 						luSettings.NumReserveAgents = pool.NumReserveAgents.Value;
 					}
-					return new LeaseUtilizationStrategy(_agentCollection, _poolCollection, _leaseCollection, _clock, _cache, luSettings);
+					return new LeaseUtilizationStrategy(_agentService, _poolCollection, _leaseCollection, _clock, _cache, luSettings);
 				case PoolSizeStrategy.NoOp:
 					return new NoOpPoolSizeStrategy();
 				default:
@@ -455,7 +457,11 @@ namespace Horde.Server.Agents.Fleet
 		/// <returns>Number of pending shutdown cancelled</returns>
 		private async Task<int> CancelPendingShutdownsAsync(IPool pool, int count, CancellationToken cancellationToken)
 		{
-			IReadOnlyList<IAgent> agents = await _agentCollection.FindAsync(status: AgentStatus.Ok, enabled: true, poolId: pool.Id, cancellationToken: cancellationToken);
+			IReadOnlyList<IAgent> agents = (await _agentService.GetCachedAgentsAsync(cancellationToken))
+				.Where(x => x is { Status: AgentStatus.Ok, Enabled: true })
+				.Where(x => x.Pools.Contains(pool.Id))
+				.ToList();
+			
 			int numShutdownsCancelled = 0;
 			foreach (IAgent agent in agents)
 			{
diff --git a/Engine/Source/Programs/Horde/Horde.Server/Agents/Fleet/LeaseUtilizationStrategy.cs b/Engine/Source/Programs/Horde/Horde.Server/Agents/Fleet/LeaseUtilizationStrategy.cs
--- a/Engine/Source/Programs/Horde/Horde.Server/Agents/Fleet/LeaseUtilizationStrategy.cs
+++ b/Engine/Source/Programs/Horde/Horde.Server/Agents/Fleet/LeaseUtilizationStrategy.cs
@@ -127,7 +127,7 @@ namespace Horde.Server.Agents.Fleet
 
 		internal LeaseUtilizationSettings Settings { get; }
 
-		private readonly IAgentCollection _agentCollection;
+		private readonly AgentService _agentService;
 		private readonly IPoolCollection _poolCollection;
 		private readonly ILeaseCollection _leaseCollection;
 		private readonly IClock _clock;
@@ -136,15 +136,15 @@ namespace Horde.Server.Agents.Fleet
 		/// <summary>
 		/// Constructor
 		/// </summary>
-		/// <param name="agentCollection"></param>
+		/// <param name="agentService"></param>
 		/// <param name="poolCollection"></param>
 		/// <param name="leaseCollection"></param>
 		/// <param name="clock"></param>
 		/// <param name="cache"></param>
 		/// <param name="settings"></param>
-		public LeaseUtilizationStrategy(IAgentCollection agentCollection, IPoolCollection poolCollection, ILeaseCollection leaseCollection, IClock clock, IMemoryCache cache, LeaseUtilizationSettings settings)
+		public LeaseUtilizationStrategy(AgentService agentService, IPoolCollection poolCollection, ILeaseCollection leaseCollection, IClock clock, IMemoryCache cache, LeaseUtilizationSettings settings)
 		{
-			_agentCollection = agentCollection;
+			_agentService = agentService;
 			_poolCollection = poolCollection;
 			_leaseCollection = leaseCollection;
 			_clock = clock;
@@ -157,7 +157,7 @@ namespace Horde.Server.Agents.Fleet
 			using TelemetrySpan span = OpenTelemetryTracers.Horde.StartActiveSpan($"{nameof(LeaseUtilizationStrategy)}.{nameof(GetAgentDataAsync)}");
 
 			// Find all the current agents
-			IReadOnlyList<IAgent> agents = await _agentCollection.FindAsync(status: AgentStatus.Ok, cancellationToken: cancellationToken);
+			IReadOnlyList<IAgent> agents = (await _agentService.GetCachedAgentsAsync(cancellationToken)).Where(x => x.Status == AgentStatus.Ok).ToList();
 
 			// Query leases in last interval
 			DateTime maxTime = _clock.UtcNow;
