Horde: Improve logging for temp storage copies.

#rnx

diff --git a/Engine/Source/Programs/Horde/Drivers/JobDriver/Utility/TempStorage.cs b/Engine/Source/Programs/Horde/Drivers/JobDriver/Utility/TempStorage.cs
--- a/Engine/Source/Programs/Horde/Drivers/JobDriver/Utility/TempStorage.cs
+++ b/Engine/Source/Programs/Horde/Drivers/JobDriver/Utility/TempStorage.cs
@@ -593,7 +593,7 @@ namespace JobDriver.Utility
 			FileReference localFileListLocation = GetTagManifestLocation(manifestDir, nodeName, tagName);
 			if (FileReference.Exists(localFileListLocation))
 			{
-				logger.LogInformation("Reading local file list from {File}", localFileListLocation.FullName);
+				logger.LogInformation("Reading node \"{NodeName}\" tag \"{TagName}\" from {File}", nodeName, tagName, localFileListLocation.FullName);
 			}
 			else
 			{
@@ -727,7 +727,7 @@ namespace JobDriver.Utility
 			TempStorageBlockManifest? manifest;
 			if (local)
 			{
-				logger.LogInformation("Reading block manifest from {File}", localManifestFile.FullName);
+				logger.LogInformation("Reading node \"{NodeName}\" block \"{BlockName}\" from {File}", nodeName, blockName, localManifestFile);
 				manifest = TempStorageBlockManifest.Load(localManifestFile);
 			}
 			else
