Horde: Update release notes.

#rnx

diff --git a/Engine/Source/Programs/Horde/Docs/ReleaseNotes.md b/Engine/Source/Programs/Horde/Docs/ReleaseNotes.md
--- a/Engine/Source/Programs/Horde/Docs/ReleaseNotes.md
+++ b/Engine/Source/Programs/Horde/Docs/ReleaseNotes.md
@@ -1,5 +1,21 @@
 # Release Notes
 
+## 2024-05-31
+
+* Use cached agent data for fleet and pool size handling (34029489)
+* Allow for licenses to be in markdown files (.md) - resolves issue with cassandracsharpdriver package
+  Also deleted license file I submitted earlier as this was already submitted under a different name but identical content (34028142)
+* Fix conform commands using the incorrect access token to communicate with the server. (34012793)
+* [UBA]
+  * Added support for explicit root id for path roots (this is just c# stub code) (34011609)
+* Allow registering multiple workspace materializer factories. (34010953)
+* Change the code generation of UHT for BP implementable events to check for any script implementations of the event. 
+  If there are none, then just call directly into the native C++ implemenation of the function, saving us the cost of copying the function params to the script VM. (34009523)
+* Add an agent document property for the combined list of pools, and upgrade existing agent documents on read to include it. This allows indexed searches for dynamic pools. (34009396)
+* Use the standard version field for log events. This is automatically mapped to the Datadog version field by our pipeline. (34005245)
+* Set the service.version property for log events. (34004705)
+* Add a log enricher which adds the server version. (34003927)
+
 ## 2024-05-30
 
 * Update test projects to NET 8. (34002443)
