Horde: Use read replica for agent lookups in compute service

diff --git a/Engine/Source/Programs/Horde/Horde.Server/Compute/ComputeService.cs b/Engine/Source/Programs/Horde/Horde.Server/Compute/ComputeService.cs
--- a/Engine/Source/Programs/Horde/Horde.Server/Compute/ComputeService.cs
+++ b/Engine/Source/Programs/Horde/Horde.Server/Compute/ComputeService.cs
@@ -384,7 +384,7 @@ namespace Horde.Server.Compute
 			try
 			{
 				PoolId? poolId = arp.Requirements.Pool != null ? new PoolId(arp.Requirements.Pool) : null;
-				IReadOnlyList<IAgent> agents = await _agentCollection.FindAsync(poolId, cancellationToken: cancellationToken);
+				IReadOnlyList<IAgent> agents = await _agentCollection.FindAsync(poolId, consistentRead: false, cancellationToken: cancellationToken);
 				foreach (IAgent agent in agents)
 				{
 					Dictionary<string, int> assignedResources = new Dictionary<string, int>();
