Horde: Use cached agent data for agent and pool controller

diff --git a/Engine/Source/Programs/Horde/Horde.Server/Agents/AgentsController.cs b/Engine/Source/Programs/Horde/Horde.Server/Agents/AgentsController.cs
--- a/Engine/Source/Programs/Horde/Horde.Server/Agents/AgentsController.cs
+++ b/Engine/Source/Programs/Horde/Horde.Server/Agents/AgentsController.cs
@@ -71,10 +71,20 @@ namespace Horde.Server.Agents
 			{
 				return Forbid(AgentAclAction.ListAgents);
 			}
+			
+			IEnumerable<IAgent> agents = (await _agentService.GetCachedAgentsAsync(HttpContext.RequestAborted))
+				.Where(x => poolId == null || x.Pools.Contains(poolId.Value))
+				.Where(x => modifiedAfter == null || modifiedAfter.Value.UtcDateTime >= x.UpdateTime);
 
-			IReadOnlyList<IAgent> agents = await _agentService.FindAgentsAsync(poolId, modifiedAfter?.UtcDateTime, null, includeDeleted, index, count, HttpContext.RequestAborted);
+			if (!includeDeleted)
+			{
+				agents = agents.Where(x => !x.Deleted);
+			}
+			
+			agents = index != null ? agents.Skip(index.Value) : agents;
+			agents = count != null ? agents.Take(count.Value) : agents;
 
-			List<object> responses = new List<object>();
+			List<object> responses = [];
 			foreach (IAgent agent in agents)
 			{
 				if (condition == null || agent.SatisfiesCondition(condition))
diff --git a/Engine/Source/Programs/Horde/Horde.Server/Agents/Pools/PoolsController.cs b/Engine/Source/Programs/Horde/Horde.Server/Agents/Pools/PoolsController.cs
--- a/Engine/Source/Programs/Horde/Horde.Server/Agents/Pools/PoolsController.cs
+++ b/Engine/Source/Programs/Horde/Horde.Server/Agents/Pools/PoolsController.cs
@@ -27,8 +27,8 @@ namespace Horde.Server.Agents.Pools
 	[Route("[controller]")]
 	public class PoolsController : HordeControllerBase
 	{
+		readonly AgentService _agentService;
 		readonly IPoolCollection _poolCollection;
-		readonly IAgentCollection _agentCollection;
 		readonly IUtilizationDataCollection _utilizationDataCollection;
 		readonly IClock _clock;
 		readonly IOptionsSnapshot<GlobalConfig> _globalConfig;
@@ -36,10 +36,10 @@ namespace Horde.Server.Agents.Pools
 		/// <summary>
 		/// Constructor
 		/// </summary>
-		public PoolsController(IPoolCollection poolCollection, IAgentCollection agentCollection, IUtilizationDataCollection utilizationDataCollection, IClock clock, IOptionsSnapshot<GlobalConfig> globalConfig)
+		public PoolsController(AgentService agentService, IPoolCollection poolCollection, IUtilizationDataCollection utilizationDataCollection, IClock clock, IOptionsSnapshot<GlobalConfig> globalConfig)
 		{
+			_agentService = agentService;
 			_poolCollection = poolCollection;
-			_agentCollection = agentCollection;
 			_utilizationDataCollection = utilizationDataCollection;
 			_clock = clock;
 			_globalConfig = globalConfig;
@@ -148,7 +148,7 @@ namespace Horde.Server.Agents.Pools
 
 				if (stats || numAgents > 0)
 				{
-					IReadOnlyList<IAgent> agents = await _agentCollection.FindAsync(cancellationToken: cancellationToken);
+					IReadOnlyList<IAgent> agents = await _agentService.GetCachedAgentsAsync(cancellationToken);
 					foreach (IAgent agent in agents)
 					{
 						foreach (PoolId poolId in agent.Pools)
