Horde: Add artifact metadata to job responses.

diff --git a/Engine/Source/Programs/Horde/Horde.Server/Jobs/JobsController.cs b/Engine/Source/Programs/Horde/Horde.Server/Jobs/JobsController.cs
--- a/Engine/Source/Programs/Horde/Horde.Server/Jobs/JobsController.cs
+++ b/Engine/Source/Programs/Horde/Horde.Server/Jobs/JobsController.cs
@@ -551,7 +551,7 @@ namespace Horde.Server.Jobs
 						string? stepKey = artifact.Keys.FirstOrDefault(x => x.StartsWith(artifactStepKeyPrefix, StringComparison.Ordinal));
 						if (stepKey != null && JobStepId.TryParse(stepKey.Substring(artifactStepKeyPrefix.Length), out JobStepId jobStepId))
 						{
-							response.Artifacts.Add(new GetJobArtifactResponse(artifact.Id, artifact.Name, artifact.Type, artifact.Description, artifact.Keys.ToList(), jobStepId));
+							response.Artifacts.Add(new GetJobArtifactResponse(artifact.Id, artifact.Name, artifact.Type, artifact.Description, artifact.Keys.ToList(), artifact.Metadata.ToList(), jobStepId));
 							addedArtifacts.Add((artifact.Name, jobStepId));
 						}
 					}
@@ -574,7 +574,7 @@ namespace Horde.Server.Jobs
 							IGraphArtifact? graphArtifact;
 							if (outputNameToArtifact.TryGetValue(outputName, out graphArtifact) && !addedArtifacts.Contains((graphArtifact.Name, step.Id)))
 							{
-								response.Artifacts.Add(new GetJobArtifactResponse(null, graphArtifact.Name, graphArtifact.Type, graphArtifact.Description, graphArtifact.Keys.ToList(), step.Id));
+								response.Artifacts.Add(new GetJobArtifactResponse(null, graphArtifact.Name, graphArtifact.Type, graphArtifact.Description, graphArtifact.Keys.ToList(), graphArtifact.Metadata.ToList(), step.Id));
 							}
 						}
 					}
diff --git a/Engine/Source/Programs/Shared/EpicGames.Horde/Jobs/JobMessages.cs b/Engine/Source/Programs/Shared/EpicGames.Horde/Jobs/JobMessages.cs
--- a/Engine/Source/Programs/Shared/EpicGames.Horde/Jobs/JobMessages.cs
+++ b/Engine/Source/Programs/Shared/EpicGames.Horde/Jobs/JobMessages.cs
@@ -574,6 +574,7 @@ namespace EpicGames.Horde.Jobs
 	/// <param name="Type">Artifact type</param>
 	/// <param name="Description">Description to display for the artifact on the dashboard</param>
 	/// <param name="Keys">Keys for the artifact</param>
+	/// <param name="Metadata">Metadata for the artifact</param>
 	/// <param name="StepId">Step producing the artifact</param>
 	public record class GetJobArtifactResponse
 	(
@@ -582,6 +583,7 @@ namespace EpicGames.Horde.Jobs
 		ArtifactType Type,
 		string? Description,
 		List<string> Keys,
+		List<string> Metadata,
 		JobStepId StepId
 	);
 
