Horde: Make view log on issue view go to first issue event line
#rnx

diff --git a/Engine/Source/Programs/Horde/HordeDashboard/src/components/IssueViewV2.tsx b/Engine/Source/Programs/Horde/HordeDashboard/src/components/IssueViewV2.tsx
--- a/Engine/Source/Programs/Horde/HordeDashboard/src/components/IssueViewV2.tsx
+++ b/Engine/Source/Programs/Horde/HordeDashboard/src/components/IssueViewV2.tsx
@@ -1536,6 +1536,12 @@ const StepPanel: React.FC<{ streamId: string, hstep: GetIssueStepResponse }> = o
 
       const textSize = "small";
 
+
+      let viewLogURL = `/log/${hstep.logId}`;
+      if (logEvents?.length) {
+         viewLogURL += `?lineindex=${logEvents[0].lineIndex}`;
+      }
+
       return <div style={{ paddingTop: 8, height: "100%" }}>
          <Stack style={{ flexBasis: "52px", flexShrink: 0 }}>
             <Stack horizontal verticalAlign="center" style={{ backgroundColor: backgroundColor, width: "100%", paddingLeft: 8, padding: 12 }}>
@@ -1567,7 +1573,7 @@ const StepPanel: React.FC<{ streamId: string, hstep: GetIssueStepResponse }> = o
                            </Stack>
                         </Stack>
                         <Stack style={{ paddingRight: 18 }}>
-                           <Link className={"view-log-link"} to={`/log/${hstep.logId}`} target="_blank">
+                           <Link className={"view-log-link"} to={viewLogURL} target="_blank">
                               <Stack horizontal verticalAlign="center" tokens={{ childrenGap: 8 }}>
                                  <Icon styles={{ root: { margin: '0px', padding: '0px', paddingTop: "2px" } }} iconName="AlignLeft" className={hordeClasses.iconBlue} />
                                  <Text variant={textSize} styles={{ root: { margin: '0px', padding: '0px' } }} className="view-log-link">View Log</Text>
