Horde: Add an Epic-specific appsettings.Epic.json file for the Job driver.

#rnx

diff --git a/Engine/Source/Programs/Horde/Drivers/JobDriver/DriverApp.cs b/Engine/Source/Programs/Horde/Drivers/JobDriver/DriverApp.cs
--- a/Engine/Source/Programs/Horde/Drivers/JobDriver/DriverApp.cs
+++ b/Engine/Source/Programs/Horde/Drivers/JobDriver/DriverApp.cs
@@ -39,6 +39,7 @@ namespace JobDriver
 			// Read the driver config
 			IConfiguration configuration = new ConfigurationBuilder()
 				.AddJsonFile("appsettings.json", optional: false)
+				.AddJsonFile("appsettings.Epic.json", optional: true)
 				.AddEnvironmentVariables()
 				.Build();
 
diff --git a/Engine/Source/Programs/Horde/Drivers/JobDriver/JobDriver.csproj b/Engine/Source/Programs/Horde/Drivers/JobDriver/JobDriver.csproj
--- a/Engine/Source/Programs/Horde/Drivers/JobDriver/JobDriver.csproj
+++ b/Engine/Source/Programs/Horde/Drivers/JobDriver/JobDriver.csproj
@@ -112,6 +112,13 @@
       <CopyToOutputDirectory>Always</CopyToOutputDirectory>
       <CopyToPublishDirectory>Always</CopyToPublishDirectory>
     </Content>
+    <Content Include="..\..\..\..\..\Restricted\NotForLicensees\Source\Programs\Horde\Drivers\JobDriver\appsettings.Epic.json" Condition="Exists('..\..\..\..\..\Restricted\NotForLicensees\Source\Programs\Horde\Drivers\JobDriver\appsettings.Epic.json')">
+      <Link>appsettings.Epic.json</Link>
+      <DependentUpon>appsettings.json</DependentUpon>
+      <Generator></Generator>
+      <CopyToOutputDirectory>Always</CopyToOutputDirectory>
+      <CopyToPublishDirectory>Always</CopyToPublishDirectory>
+    </Content>
   </ItemGroup>
 
 </Project>
