Horde: Improve trends view default zoom and add trends line
#rnx

diff --git a/Engine/Source/Programs/Horde/HordeDashboard/src/components/jobDetailsV2/JobDetailStepTrendsV2.tsx b/Engine/Source/Programs/Horde/HordeDashboard/src/components/jobDetailsV2/JobDetailStepTrendsV2.tsx
--- a/Engine/Source/Programs/Horde/HordeDashboard/src/components/jobDetailsV2/JobDetailStepTrendsV2.tsx
+++ b/Engine/Source/Programs/Horde/HordeDashboard/src/components/jobDetailsV2/JobDetailStepTrendsV2.tsx
@@ -71,6 +71,8 @@ class StepTrendsDataView extends JobDataView {
       this.minTime = this.maxTime = undefined;
       this.maxMinutes = 0;
 
+      const alltimes: number[] = [];
+
       history = history.filter(r => {
 
          if (!r.startTime) {
@@ -121,6 +123,8 @@ class StepTrendsDataView extends JobDataView {
             this.maxMinutes = Math.max(minutes, this.maxMinutes);
          }
 
+         alltimes.push(minutes);
+
          durations.set(h.jobId, minutes);
       });
 
@@ -130,6 +134,30 @@ class StepTrendsDataView extends JobDataView {
          return (durations.get(h.jobId) ?? (24 * 60)) <= (23 * 60)
       })
 
+      const median = (arr: number[]): number | undefined => {
+         if (!arr.length) return undefined;
+         const s = [...arr].sort((a, b) => a - b);
+         const mid = Math.floor(s.length / 2);
+         return s.length % 2 ? s[mid] : ((s[mid - 1] + s[mid]) / 2);
+      };
+
+      this.median = undefined;
+
+      this.median = median(alltimes);
+      if (this.median) {
+
+         this.median *= 2;
+
+         if (this.median < 1) {
+            this.median = 1;
+         }
+
+
+         if (this.median > this.maxMinutes) {
+            this.median = undefined;
+         }
+      }
+
       this.history = history;
    }
 
@@ -212,6 +240,7 @@ class StepTrendsDataView extends JobDataView {
 
    durations = new Map<string, number>();
    maxMinutes = 0;
+   median?: number;
 
    minTime?: Date;
    maxTime?: Date;
@@ -271,7 +300,7 @@ class StepTrendsRenderer {
          .range([margin.left, width - margin.right])
 
       const y = this.scaleY = d3.scaleLinear()
-         .domain([0, dataView.maxMinutes]).nice()
+         .domain([0, dataView.median ?? dataView.maxMinutes]).nice()
          .range([height - margin.bottom, margin.top])
 
 
@@ -365,22 +394,61 @@ class StepTrendsRenderer {
 
       const lineI = d3.range(data.length);
 
-      const curve = d3.curveLinear;
-      const line:any = d3.line()
-         .defined(i => true)
+      const plotTrendY = (i: any, scaleY?: any) => {
+         const idx = i as any as number;
+
+         let sum = 0;
+         let count = 0;
+         const range = 50;
+         for (let j = idx - range; j < idx + range; j++) {
+
+            if (j < 0 || j >= data.length) {
+               continue;
+            }
+
+            const v = dataView.durations.get(data[j].jobId)!;
+
+            if (dataView.median && v > dataView.median) {
+               continue;
+            }
+
+            sum += v;
+            count++;
+         }
+
+         const f = scaleY ?? y;
+
+         if (!count) {
+            return f(dataView.durations.get(data[idx].jobId)!) - 2
+         }
+
+         return f(sum / count) - 2
+
+      }
+
+      const curve = d3.curveMonotoneX;
+      const line: any = d3.line()
+         .defined(i => {
+            if (dataView.durations.size < 10) {
+               return false;
+            }
+            return true;
+         })
          .curve(curve)
          .x(i => { return x(new Date(data[i as any].startTime!).getTime() / 1000) })
-         .y(i => { return y(dataView.durations.get(data[i as any].jobId)!) - 2 })
+         .y(i => {
+            return plotTrendY(i);
+         })
 
       svg.append("path")
          .attr("clip-path", `url(#${clipId})`)
          .attr("class", "linechart")
          .attr("fill", "none")
-         .attr("stroke", dashboard.darktheme ? "#6D6C6B" : "#6D6C6B")
-         .attr("stroke-width", 1.0)
+         .attr("stroke", dashboard.darktheme ? "#EEEEEE" : "#035CA1")
+         .attr("stroke-width", 3.0)
          .attr("stroke-linecap", "round")
          .attr("stroke-linejoin", "round")
-         .attr("stroke-opacity", 0.5)
+         .attr("stroke-opacity", 1)
          .attr("d", line(lineI as any));
 
       const xAxis = (g: SelectionType) => {
@@ -443,17 +511,14 @@ class StepTrendsRenderer {
             .attr("x", (d: any) => x(new Date(d.startTime!).getTime() / 1000) - (barWidth / 2))
             .attr("width", barWidth);
 
-         const strokeWidth = zoomLevel > 4 ? 2 : 1;
-
-         const scaledLine:any = d3.line()
+         const scaledLine: any = d3.line()
             .defined(i => true)
             .curve(curve)
             .x(i => { return x(new Date(data[i as any].startTime!).getTime() / 1000) })
-            .y(i => { return renderer.scaleY!(dataView.durations.get(data[i as any].jobId)!) - 2 })
+            .y(i => plotTrendY(i, renderer.scaleY!))
 
 
          svg!.selectAll(".linechart")
-            .attr("stroke-width", strokeWidth)
             .attr("d", scaledLine(lineI as any));
 
 
@@ -472,11 +537,11 @@ class StepTrendsRenderer {
             .attr("height", (d: any) => scaleY(0) - scaleY(dataView.durations.get(d.jobId)!))
 
 
-         const scaledLine:any = d3.line()
+         const scaledLine: any = d3.line()
             .defined(i => true)
             .curve(curve)
             .x(i => { return x(new Date(data[i as any].startTime!).getTime() / 1000) })
-            .y(i => { return scaleY(dataView.durations.get(data[i as any].jobId)!) - 2 })
+            .y(i => { return plotTrendY(i, scaleY) })
 
 
          svg!.selectAll(".linechart")
@@ -545,7 +610,7 @@ class StepTrendsRenderer {
       }
 
       // events
-      svg.on("wheel", (event:any) => { event.preventDefault(); })
+      svg.on("wheel", (event: any) => { event.preventDefault(); })
 
       svg.on("mousemove", (event) => { this.showToolTip(true); handleMouseMove(event); });
       svg.on("mouseleave", (event) => { if (!dataView.tooltip.frozen) dataView.tooltip.update(undefined); })
@@ -589,7 +654,7 @@ class StepTrendsRenderer {
 const GraphTooltip: React.FC<{ dataView: StepTrendsDataView }> = observer(({ dataView }) => {
 
    const { modeColors } = getHordeStyling();
-   
+
    // subscribe
    if (dataView.tooltip.updated) { }
 
@@ -756,7 +821,7 @@ export const StepTrendsPanelV2: React.FC<{ jobDetails: JobDetailsV2; stepId: str
                      min={1}
                      max={dataView.maxMinutes}
                      step={dataView.maxMinutes / 20}
-                     defaultValue={dataView.maxMinutes}
+                     defaultValue={dataView.median ?? dataView.maxMinutes}
                      showValue
                      valueFormat={(value) => {
                         if (!value) return ""
