Horde: Fix step preflights
#rnx
#skipci

diff --git a/Engine/Source/Programs/Horde/HordeDashboard/src/components/jobDetailsV2/StepRetryModal.tsx b/Engine/Source/Programs/Horde/HordeDashboard/src/components/jobDetailsV2/StepRetryModal.tsx
--- a/Engine/Source/Programs/Horde/HordeDashboard/src/components/jobDetailsV2/StepRetryModal.tsx
+++ b/Engine/Source/Programs/Horde/HordeDashboard/src/components/jobDetailsV2/StepRetryModal.tsx
@@ -259,24 +259,12 @@ export const StepRetryModal: React.FC<{ stepId: string; jobDetails: JobDetailsV2
 
       const job = jobData;
 
-      const args: string[] = [];
-      args.push(`-Target=Setup Build`);
-      args.push(`-Target=${step!.name}`);
-
-      job.arguments?.forEach(arg => {
-
-         if (arg.toLowerCase().indexOf("-target=") !== -1) {
-            return;
-         }
-
-         args.push(arg);
-      });
-
       const data: CreateJobRequest = {
          streamId: job.streamId,
          templateId: job.templateId!,
-         arguments: args,
+         parameters: job.parameters,
          change: job.change,
+         targets: [step!.name],
          preflightChange: fixCL
       };
 
