Horde: Fix null reference exception in TestExecutor.

#rnx

diff --git a/Engine/Source/Programs/Horde/Drivers/JobDriver/Execution/TestExecutor.cs b/Engine/Source/Programs/Horde/Drivers/JobDriver/Execution/TestExecutor.cs
--- a/Engine/Source/Programs/Horde/Drivers/JobDriver/Execution/TestExecutor.cs
+++ b/Engine/Source/Programs/Horde/Drivers/JobDriver/Execution/TestExecutor.cs
@@ -43,8 +43,10 @@ namespace JobDriver.Execution
 			return result;
 		}
 
-		public override Task InitializeAsync(RpcBeginBatchResponse batch, ILogger logger, CancellationToken cancellationToken)
+		public override async Task InitializeAsync(RpcBeginBatchResponse batch, ILogger logger, CancellationToken cancellationToken)
 		{
+			await base.InitializeAsync(batch, logger, cancellationToken);
+
 			logger.LogInformation("Initializing");
 
 			foreach (string argument in batch.Arguments)
@@ -63,8 +65,6 @@ namespace JobDriver.Execution
 					}
 				}
 			}
-
-			return Task.CompletedTask;
 		}
 
 		protected override async Task<bool> SetupAsync(JobStepInfo step, ILogger logger, CancellationToken cancellationToken)
