Horde: Instrument agent find endpoint

#rnx

diff --git a/Engine/Source/Programs/Horde/Horde.Server/Agents/AgentsController.cs b/Engine/Source/Programs/Horde/Horde.Server/Agents/AgentsController.cs
--- a/Engine/Source/Programs/Horde/Horde.Server/Agents/AgentsController.cs
+++ b/Engine/Source/Programs/Horde/Horde.Server/Agents/AgentsController.cs
@@ -22,6 +22,7 @@ using Microsoft.AspNetCore.Authorization;
 using Microsoft.AspNetCore.Mvc;
 using Microsoft.Extensions.Logging;
 using Microsoft.Extensions.Options;
+using OpenTelemetry.Trace;
 
 namespace Horde.Server.Agents
 {
@@ -37,17 +38,19 @@ namespace Horde.Server.Agents
 		readonly IAgentTelemetryCollection _agentTelemetryCollection;
 		readonly IUserCollection _userCollection;
 		readonly IOptionsSnapshot<GlobalConfig> _globalConfig;
+		readonly Tracer _tracer;
 		readonly ILogger<AgentsController> _logger;
 
 		/// <summary>
 		/// Constructor
 		/// </summary>
-		public AgentsController(AgentService agentService, IAgentTelemetryCollection agentTelemetryCollection, IUserCollection userCollection, IOptionsSnapshot<GlobalConfig> globalConfig, ILogger<AgentsController> logger)
+		public AgentsController(AgentService agentService, IAgentTelemetryCollection agentTelemetryCollection, IUserCollection userCollection, IOptionsSnapshot<GlobalConfig> globalConfig, Tracer tracer, ILogger<AgentsController> logger)
 		{
 			_agentService = agentService;
 			_agentTelemetryCollection = agentTelemetryCollection;
 			_userCollection = userCollection;
 			_globalConfig = globalConfig;
+			_tracer = tracer;
 			_logger = logger;
 		}
 
@@ -67,32 +70,40 @@ namespace Horde.Server.Agents
 		[ProducesResponseType(typeof(List<GetAgentResponse>), 200)]
 		public async Task<ActionResult<List<object>>> FindAgentsAsync([FromQuery] PoolId? poolId = null, [FromQuery] Condition? condition = null, [FromQuery] bool includeDeleted = false, [FromQuery] int? index = null, [FromQuery] int? count = null, [FromQuery] DateTimeOffset? modifiedAfter = null, [FromQuery] PropertyFilter? filter = null)
 		{
+			using TelemetrySpan span = _tracer.StartActiveSpan($"{nameof(AgentsController)}.{nameof(FindAgentsAsync)}");
+			
 			if (!_globalConfig.Value.Authorize(AgentAclAction.ListAgents, User))
 			{
 				return Forbid(AgentAclAction.ListAgents);
 			}
 			
-			IEnumerable<IAgent> agents = (await _agentService.GetCachedAgentsAsync(HttpContext.RequestAborted))
+			IEnumerable<IAgent> agentsEnumerable = (await _agentService.GetCachedAgentsAsync(HttpContext.RequestAborted))
 				.Where(x => poolId == null || x.Pools.Contains(poolId.Value))
 				.Where(x => modifiedAfter == null || modifiedAfter.Value.UtcDateTime >= x.UpdateTime);
 
 			if (!includeDeleted)
 			{
-				agents = agents.Where(x => !x.Deleted);
+				agentsEnumerable = agentsEnumerable.Where(x => !x.Deleted);
 			}
 			
-			agents = index != null ? agents.Skip(index.Value) : agents;
-			agents = count != null ? agents.Take(count.Value) : agents;
-
+			agentsEnumerable = index != null ? agentsEnumerable.Skip(index.Value) : agentsEnumerable;
+			agentsEnumerable = count != null ? agentsEnumerable.Take(count.Value) : agentsEnumerable;
+			
+			List<IAgent> agents = agentsEnumerable.ToList();
 			List<object> responses = [];
-			foreach (IAgent agent in agents)
+			using TelemetrySpan filterSpan = _tracer.StartActiveSpan($"FilterAgentResponses");
 			{
-				if (condition == null || agent.SatisfiesCondition(condition))
+				foreach (IAgent agent in agents)
 				{
-					responses.Add(await GetAgentResponseAsync(agent, filter));
+					if (condition == null || agent.SatisfiesCondition(condition))
+					{
+						responses.Add(await GetAgentResponseAsync(agent, filter));
+					}
 				}
 			}
-
+			
+			span.SetAttribute("NumAgents", agents.Count);
+			span.SetAttribute("NumResponses", responses.Count);
 			return responses;
 		}
 
