Horde: Better artifacts panel UX
#rnx
#sklipci

diff --git a/Engine/Source/Programs/Horde/HordeDashboard/src/backend/Api.ts b/Engine/Source/Programs/Horde/HordeDashboard/src/backend/Api.ts
--- a/Engine/Source/Programs/Horde/HordeDashboard/src/backend/Api.ts
+++ b/Engine/Source/Programs/Horde/HordeDashboard/src/backend/Api.ts
@@ -1211,6 +1211,10 @@ export type GetJobArtifactResponse = {
 	description?: string;
 	/// Step producing the artifact
 	stepId: string;
+	/// Keys for the artifact	
+	keys: string[];
+	/// Metadata for the artifact
+	metadata: string[];
 }
 
 /**Information about a job */
diff --git a/Engine/Source/Programs/Horde/HordeDashboard/src/components/jobDetailsV2/JobDetailArtifacts.tsx b/Engine/Source/Programs/Horde/HordeDashboard/src/components/jobDetailsV2/JobDetailArtifacts.tsx
--- a/Engine/Source/Programs/Horde/HordeDashboard/src/components/jobDetailsV2/JobDetailArtifacts.tsx
+++ b/Engine/Source/Programs/Horde/HordeDashboard/src/components/jobDetailsV2/JobDetailArtifacts.tsx
@@ -1,6 +1,6 @@
 // Copyright Epic Games, Inc. All Rights Reserved.
 
-import { DetailsList, DetailsListLayoutMode, DirectionalHint, FontIcon, IColumn, IContextualMenuProps, PrimaryButton, SelectionMode, Stack, Text } from '@fluentui/react';
+import { DetailsList, DetailsListLayoutMode, DirectionalHint, FontIcon, IColumn, IContextualMenuProps, IGroup, PrimaryButton, SelectionMode, Stack, Text } from '@fluentui/react';
 import { observer } from 'mobx-react-lite';
 import React, { useEffect, useState } from 'react';
 import { GetJobArtifactResponse, JobStepState } from '../../backend/Api';
@@ -12,6 +12,7 @@ import { getStepStatusColor } from '../../styles/colors';
 import { JobDataView, JobDetailsV2 } from "./JobDetailsViewCommon";
 import { JobArtifactsModal } from '../artifacts/ArtifactsModal';
 
+import { useConst } from '@fluentui/react-hooks';
 const sideRail: ISideRailLink = { text: "Artifacts", url: "rail_artifacts" };
 
 class JobArtifactsDataView extends JobDataView {
@@ -71,8 +72,9 @@ JobDetailsV2.registerDataView("JobArtifactsDataView", (details: JobDetailsV2) =>
 
 export const JobArtifactsPanel: React.FC<{ jobDetails: JobDetailsV2 }> = observer(({ jobDetails }) => {
 
+   const constGroups = useConst<IGroup[]>([]);
    const [selected, setSelected] = useState<GetJobArtifactResponse | undefined>(undefined);
-
+   
    const artifactView = jobDetails.getDataView<JobArtifactsDataView>("JobArtifactsDataView");
 
    jobDetails.subscribe();
@@ -107,13 +109,85 @@ export const JobArtifactsPanel: React.FC<{ jobDetails: JobDetailsV2 }> = observe
 
    let artifacts = [...jobData.artifacts!];
 
+   // artifact name => category (id may be unpopulated)
+   const clookup: Map<string, string> = new Map();
+
+   artifacts.forEach(a => {
+
+      let category = "Uncategorized";
+
+      const m = a.metadata.find(m => m.toLowerCase().startsWith("dashboard-category"));
+
+      if (m) {
+         const elements = m.split("=");
+         if (elements.length > 1) {
+            if (elements[1].trim()) {
+               category = elements[1].trim();
+            }
+         }
+      }
+
+      clookup.set(a.name, category);
+   })
+
    artifacts = artifacts.sort((a, b) => {
-      const aname = a.description ?? a.name;
-      const bname = b.description ?? b.name;
-      return aname.localeCompare(bname)
+
+      const acat = clookup.get(a.name)!;
+      const bcat = clookup.get(b.name)!;
+
+      if (acat !== bcat) {
+
+         if (acat !== "Uncategorized" && bcat === "Uncategorized") {
+            return -1;
+         }
+
+         if (acat === "Uncategorized" && bcat !== "Uncategorized") {
+            return 1;
+         }
+
+         return acat.localeCompare(bcat);
+      }
+
+      return a.name.localeCompare(b.name)
    });
 
 
+   let groups: IGroup[] | undefined;
+
+   groups = [];
+   let cgroup: string = "";
+
+   // emit groups
+   for (let i = 0; i < artifacts.length; i++) {
+
+      const a = artifacts[i];
+      const cat = clookup.get(a.name)!;
+
+      if (cat != cgroup) {
+
+         if (cgroup) {
+            groups[groups.length - 1].count = i - groups[groups.length - 1].startIndex;
+         }
+
+         cgroup = cat;
+
+         if (cgroup) {
+            const key = `group_key_${cgroup}`;
+            groups.push({ startIndex: i, name: cgroup, key:key, count: 0, isCollapsed: constGroups.find(g => g.key === key)?.isCollapsed ?? true });
+         }
+      }
+   }
+
+   if (cgroup) {
+      groups[groups.length - 1].count = artifacts.length - groups[groups.length - 1].startIndex;
+   }
+
+   while (constGroups.length > 0) {
+      constGroups.pop();
+   }
+   
+   constGroups.push(...groups);
+
    const renderItem = (item: GetJobArtifactResponse, index?: number, column?: IColumn) => {
 
       if (!column) {
@@ -140,14 +214,14 @@ export const JobArtifactsPanel: React.FC<{ jobDetails: JobDetailsV2 }> = observe
       if (column.key === 'column_download') {
          return <Stack horizontalAlign="end" verticalAlign="center" verticalFill={true} style={{ paddingRight: 8 }}>
             <Stack horizontal tokens={{ childrenGap: 18 }}>
-            <PrimaryButton text="Browse" 
+               <PrimaryButton text="Browse"
                   disabled={!stepFinished}
                   style={{ fontFamily: "Horde Open Sans SemiBold" }}
                   onClick={() => {
                      setSelected(item)
                   }}
                />
-               
+
                <PrimaryButton split text="Download as Zip" menuProps={downloadProps}
                   disabled={!stepFinished || !item.id}
                   style={{ fontFamily: "Horde Open Sans SemiBold" }}
@@ -221,10 +295,11 @@ export const JobArtifactsPanel: React.FC<{ jobDetails: JobDetailsV2 }> = observe
             </Stack>
             <Stack >
                <Stack style={{ paddingTop: 8 }} tokens={{ childrenGap: 12 }}>
-                  <DetailsList
+                  <DetailsList                     
                      styles={{ root: { overflowX: "hidden" } }}
                      isHeaderVisible={false}
                      items={artifacts}
+                     groups={groups}
                      columns={columns}
                      selectionMode={SelectionMode.none}
                      layoutMode={DetailsListLayoutMode.justified}
