Horde: Fix issue with trends line
#rnx

diff --git a/Engine/Source/Programs/Horde/HordeDashboard/src/components/jobDetailsV2/JobDetailStepTrendsV2.tsx b/Engine/Source/Programs/Horde/HordeDashboard/src/components/jobDetailsV2/JobDetailStepTrendsV2.tsx
--- a/Engine/Source/Programs/Horde/HordeDashboard/src/components/jobDetailsV2/JobDetailStepTrendsV2.tsx
+++ b/Engine/Source/Programs/Horde/HordeDashboard/src/components/jobDetailsV2/JobDetailStepTrendsV2.tsx
@@ -394,6 +394,9 @@ class StepTrendsRenderer {
 
       const lineI = d3.range(data.length);
 
+      const showTrendLine = dataView.durations.size >= 10;
+
+
       const plotTrendY = (i: any, scaleY?: any) => {
          const idx = i as any as number;
 
@@ -428,12 +431,7 @@ class StepTrendsRenderer {
 
       const curve = d3.curveMonotoneX;
       const line: any = d3.line()
-         .defined(i => {
-            if (dataView.durations.size < 10) {
-               return false;
-            }
-            return true;
-         })
+         .defined(i => showTrendLine)
          .curve(curve)
          .x(i => { return x(new Date(data[i as any].startTime!).getTime() / 1000) })
          .y(i => {
@@ -512,7 +510,7 @@ class StepTrendsRenderer {
             .attr("width", barWidth);
 
          const scaledLine: any = d3.line()
-            .defined(i => true)
+            .defined(i => showTrendLine)
             .curve(curve)
             .x(i => { return x(new Date(data[i as any].startTime!).getTime() / 1000) })
             .y(i => plotTrendY(i, renderer.scaleY!))
@@ -538,7 +536,7 @@ class StepTrendsRenderer {
 
 
          const scaledLine: any = d3.line()
-            .defined(i => true)
+            .defined(i => showTrendLine)
             .curve(curve)
             .x(i => { return x(new Date(data[i as any].startTime!).getTime() / 1000) })
             .y(i => { return plotTrendY(i, scaleY) })
