Horde: Add a couple of additional artifact commands.

#rnx

diff --git a/Engine/Source/Programs/Horde/Horde/Commands/Artifacts/ArtifactDownload.cs b/Engine/Source/Programs/Horde/Horde/Commands/Artifacts/ArtifactDownload.cs
new file mode 100644
--- /dev/null
+++ b/Engine/Source/Programs/Horde/Horde/Commands/Artifacts/ArtifactDownload.cs
@@ -0,0 +1,72 @@
+// Copyright Epic Games, Inc. All Rights Reserved.
+
+using System.ComponentModel;
+using System.Diagnostics;
+using EpicGames.Core;
+using EpicGames.Horde;
+using EpicGames.Horde.Artifacts;
+using EpicGames.Horde.Storage;
+using EpicGames.Horde.Storage.Nodes;
+using Microsoft.Extensions.Logging;
+
+namespace Horde.Commands.Artifacts
+{
+	[Command("artifact", "download", "Downloads an artifact by id")]
+	class ArtifactDownload : Command
+	{
+		[CommandLine("-Id=")]
+		[Description("Unique identifier for the artifact")]
+		public ArtifactId Id { get; set; }
+
+		[CommandLine("-OutputDir=", Required = true)]
+		[Description("Directory to write extracted files.")]
+		public DirectoryReference OutputDir { get; set; } = null!;
+
+		[CommandLine("-Stats")]
+		[Description("Outputs stats about the extraction process.")]
+		public bool Stats { get; set; }
+
+		[CommandLine("-CleanOutput")]
+		[Description("If set, deletes the contents of the output directory before extraction.")]
+		public bool CleanOutput { get; set; }
+
+		readonly IHordeClient _hordeClient;
+
+		public ArtifactDownload(IHordeClient hordeClient)
+		{
+			_hordeClient = hordeClient;
+		}
+
+		public override async Task<int> ExecuteAsync(ILogger logger)
+		{
+			HordeHttpClient httpClient = _hordeClient.CreateHttpClient();
+
+			GetArtifactResponse artifact = await httpClient.GetArtifactAsync(Id);
+			logger.LogInformation("Downloading artifact {Id}: {Description}", Id, artifact.Description);
+
+			if (CleanOutput)
+			{
+				logger.LogInformation("Deleting contents of {OutputDir}...", OutputDir);
+				FileUtils.ForceDeleteDirectoryContents(OutputDir);
+			}
+
+			using IStorageClient store = _hordeClient.CreateStorageClient(artifact.Id);
+			IBlobHandle handle = await store.ReadRefAsync(new RefName("default"));
+
+			Stopwatch timer = Stopwatch.StartNew();
+
+			DirectoryNode node = await handle.ReadBlobAsync<DirectoryNode>();
+			await node.CopyToDirectoryAsync(OutputDir.ToDirectoryInfo(), new ExtractStatsLogger(logger), logger, CancellationToken.None);
+
+			logger.LogInformation("Elapsed: {Time}s", timer.Elapsed.TotalSeconds);
+
+			if (Stats)
+			{
+				StorageStats stats = store.GetStats();
+				stats.Print(logger);
+			}
+
+			return 0;
+		}
+	}
+}
diff --git a/Engine/Source/Programs/Horde/Horde/Commands/Artifacts/ArtifactInfo.cs b/Engine/Source/Programs/Horde/Horde/Commands/Artifacts/ArtifactInfo.cs
new file mode 100644
--- /dev/null
+++ b/Engine/Source/Programs/Horde/Horde/Commands/Artifacts/ArtifactInfo.cs
@@ -0,0 +1,57 @@
+// Copyright Epic Games, Inc. All Rights Reserved.
+
+using System.ComponentModel;
+using EpicGames.Core;
+using EpicGames.Horde;
+using EpicGames.Horde.Artifacts;
+using Microsoft.Extensions.Logging;
+
+namespace Horde.Commands.Artifacts
+{
+	[Command("artifact", "info", "Gets information about an artifact")]
+	class ArtifactInfo : Command
+	{
+		[CommandLine("-Id=")]
+		[Description("Unique identifier for the artifact")]
+		public ArtifactId Id { get; set; }
+
+		readonly IHordeClient _hordeClient;
+
+		public ArtifactInfo(IHordeClient hordeClient)
+		{
+			_hordeClient = hordeClient;
+		}
+
+		public override async Task<int> ExecuteAsync(ILogger logger)
+		{
+			HordeHttpClient httpClient = _hordeClient.CreateHttpClient();
+
+			GetArtifactResponse artifact = await httpClient.GetArtifactAsync(Id);
+			logger.LogInformation("Id: {Id}", Id);
+			logger.LogInformation("Name: {Name}", artifact.Name);
+			logger.LogInformation("Type: {Type}", artifact.Type);
+			logger.LogInformation("Stream: {Stream}", artifact.StreamId);
+			logger.LogInformation("Change: {Change}", artifact.Change);
+			logger.LogInformation("Description: {Desc}", artifact.Description);
+
+			if (artifact.Keys.Count > 0)
+			{
+				logger.LogInformation("Keys:");
+				foreach (string key in artifact.Keys)
+				{
+					logger.LogInformation("  {Key}", key);
+				}
+			}
+			if (artifact.Metadata.Count > 0)
+			{
+				logger.LogInformation("Metadata:");
+				foreach (string meta in artifact.Metadata)
+				{
+					logger.LogInformation("  {Meta}", meta);
+				}
+			}
+
+			return 0;
+		}
+	}
+}
\ No newline at end of file
