Horde: Use cached agent data for JobTaskSource

JobTaskSource ticks every 5 seconds refreshing agent and pool data. The cached agent data is refreshed at the same rate.

diff --git a/Engine/Source/Programs/Horde/Horde.Server.Tests/TestSetup.cs b/Engine/Source/Programs/Horde/Horde.Server.Tests/TestSetup.cs
--- a/Engine/Source/Programs/Horde/Horde.Server.Tests/TestSetup.cs
+++ b/Engine/Source/Programs/Horde/Horde.Server.Tests/TestSetup.cs
@@ -253,6 +253,7 @@ namespace Horde.Server.Tests
 
 			services.AddSingleton<AclService>();
 			services.AddSingleton<AgentService>();
+			services.AddSingleton(provider => new Lazy<AgentService>(provider.GetRequiredService<AgentService>));
 			services.AddSingleton<AgentRelayService>();
 			services.AddSingleton<AwsAutoScalingLifecycleService>();
 			services.AddSingleton<ArtifactExpirationService>();
diff --git a/Engine/Source/Programs/Horde/Horde.Server/Jobs/JobTaskSource.cs b/Engine/Source/Programs/Horde/Horde.Server/Jobs/JobTaskSource.cs
--- a/Engine/Source/Programs/Horde/Horde.Server/Jobs/JobTaskSource.cs
+++ b/Engine/Source/Programs/Horde/Horde.Server/Jobs/JobTaskSource.cs
@@ -187,11 +187,15 @@ namespace Horde.Server.Jobs
 				Agent = agent;
 			}
 		}
-
+		
+		/// <summary>
+		/// Use lazy initialization of agent service to prevent circular dependency as AgentService depend on this class
+		/// </summary>
+		readonly Lazy<AgentService> _agentService;
+		
 		readonly AclService _aclService;
 		readonly IStreamCollection _streamCollection;
 		readonly ILogCollection _logCollection;
-		readonly IAgentCollection _agentsCollection;
 		readonly IJobCollection _jobs;
 		readonly IJobStepRefCollection _jobStepRefs;
 		readonly IGraphCollection _graphs;
@@ -238,10 +242,10 @@ namespace Horde.Server.Jobs
 		/// <summary>
 		/// Constructor
 		/// </summary>
-		public JobTaskSource(AclService aclService, IAgentCollection agents, IJobCollection jobs, IJobStepRefCollection jobStepRefs, IBisectTaskCollection bisectTasks, IGraphCollection graphs, IPoolCollection pools, PoolService poolService, IUgsMetadataCollection ugsMetadataCollection, IStreamCollection streamCollection, ILogCollection logCollection, PerforceLoadBalancer perforceLoadBalancer, IClock clock, IOptionsMonitor<ServerSettings> settings, IOptionsMonitor<GlobalConfig> globalConfig, Tracer tracer, ILogger<JobTaskSource> logger)
+		public JobTaskSource(AclService aclService, Lazy<AgentService> agentService, IJobCollection jobs, IJobStepRefCollection jobStepRefs, IBisectTaskCollection bisectTasks, IGraphCollection graphs, IPoolCollection pools, PoolService poolService, IUgsMetadataCollection ugsMetadataCollection, IStreamCollection streamCollection, ILogCollection logCollection, PerforceLoadBalancer perforceLoadBalancer, IClock clock, IOptionsMonitor<ServerSettings> settings, IOptionsMonitor<GlobalConfig> globalConfig, Tracer tracer, ILogger<JobTaskSource> logger)
 		{
 			_aclService = aclService;
-			_agentsCollection = agents;
+			_agentService = agentService;
 			_jobs = jobs;
 			_jobStepRefs = jobStepRefs;
 			_bisectTasks = bisectTasks;
@@ -378,7 +382,7 @@ namespace Horde.Server.Jobs
 			Dictionary<StreamId, IStream> streams = streamsList.ToDictionary(x => x.Id, x => x);
 
 			// Find all the pools which are valid (ie. have at least one online agent)
-			IReadOnlyList<IAgent> agents = await _agentsCollection.FindAsync(cancellationToken: cancellationToken);
+			IReadOnlyList<IAgent> agents = await _agentService.Value.GetCachedAgentsAsync(cancellationToken: cancellationToken);
 			IReadOnlyList<IPoolConfig> pools = await _poolCollection.GetConfigsAsync(cancellationToken);
 			Dictionary<PoolId, PoolStatus> poolStatus = GetPoolStatus(_clock.UtcNow, pools, agents);
 
diff --git a/Engine/Source/Programs/Horde/Horde.Server/Startup.cs b/Engine/Source/Programs/Horde/Horde.Server/Startup.cs
--- a/Engine/Source/Programs/Horde/Horde.Server/Startup.cs
+++ b/Engine/Source/Programs/Horde/Horde.Server/Startup.cs
@@ -555,6 +555,7 @@ namespace Horde.Server
 
 			services.AddSingleton<AclService>();
 			services.AddSingleton<AgentService>();
+			services.AddSingleton(provider => new Lazy<AgentService>(provider.GetRequiredService<AgentService>));
 			services.AddHostedService<ArtifactExpirationService>();
 			services.AddSingleton<ConsistencyService>();
 			services.AddSingleton<RequestTrackerService>();
