Horde: Record endpoint address in MongoDB command tracer

Helps differentiate between primary and secondary use (read-only ops)

diff --git a/Engine/Source/Programs/Horde/Horde.Server/Utilities/MongoCommandTracer.cs b/Engine/Source/Programs/Horde/Horde.Server/Utilities/MongoCommandTracer.cs
--- a/Engine/Source/Programs/Horde/Horde.Server/Utilities/MongoCommandTracer.cs
+++ b/Engine/Source/Programs/Horde/Horde.Server/Utilities/MongoCommandTracer.cs
@@ -150,6 +150,7 @@ public class MongoCommandTracer
 			sa.Add("db.operationId", ev.OperationId ?? -1);
 			sa.Add("db.requestId", ev.RequestId);
 			sa.Add("db.serviceId", ev.ServiceId?.ToString());
+			sa.Add("db.endpoint", ev.ConnectionId?.ServerId.EndPoint.ToString());
 			
 			TelemetrySpan span = _tracer.StartActiveSpan(name, SpanKind.Client, parentContext: parentSpanEntry?.Span.Context ?? Tracer.CurrentSpan.Context, sa);
 			
