Horde: Add more trace attributes to compute resource allocation

diff --git a/Engine/Source/Programs/Horde/Horde.Server/Compute/ComputeService.cs b/Engine/Source/Programs/Horde/Horde.Server/Compute/ComputeService.cs
--- a/Engine/Source/Programs/Horde/Horde.Server/Compute/ComputeService.cs
+++ b/Engine/Source/Programs/Horde/Horde.Server/Compute/ComputeService.cs
@@ -19,6 +19,7 @@ using EpicGames.Horde.Compute.Clients;
 using EpicGames.Horde.Compute.Transports;
 using EpicGames.Horde.Jobs;
 using EpicGames.Horde.Logs;
+using EpicGames.Horde.Users;
 using Google.Protobuf;
 using Google.Protobuf.WellKnownTypes;
 using Horde.Common.Rpc;
@@ -79,6 +80,11 @@ namespace Horde.Server.Compute
 		/// Cluster ID
 		/// </summary>
 		public ClusterId ClusterId { get; }
+		
+		/// <summary>
+		/// User allocating the resource
+		/// </summary>
+		public UserId? UserId { get; init; }
 
 		/// <summary>
 		/// Desired protocol version
@@ -350,8 +356,13 @@ namespace Horde.Server.Compute
 		public async Task<ComputeResource?> TryAllocateResourceAsync(AllocateResourceParams arp, CancellationToken cancellationToken)
 		{
 			using TelemetrySpan span = _tracer.StartActiveSpan($"{nameof(ComputeService)}.{nameof(TryAllocateResourceAsync)}");
+			span.SetAttribute("clusterId", arp.ClusterId.ToString());
+			span.SetAttribute("userId", arp.UserId?.ToString());
 			span.SetAttribute("requestId", arp.RequestId);
 			span.SetAttribute("requestIp", arp.RequesterIp?.ToString());
+			span.SetAttribute("connectionMode", arp.ConnectionMode?.ToString());
+			span.SetAttribute("usePublicIp", arp.UsePublicIp.ToString());
+			span.SetAttribute("encryption", arp.Encryption.ToString());
 			span.SetAttribute("parentLeaseId", arp.ParentLeaseId?.ToString());
 			span.SetAttribute("req.pool", arp.Requirements.Pool);
 			span.SetAttribute("req.condition", arp.Requirements.Condition?.ToString());
@@ -428,6 +439,8 @@ namespace Horde.Server.Compute
 								await _agentService.CreateLeaseAsync(newAgent, lease, cancellationToken);
 								span.SetAttribute("allocatedLeaseId", leaseId.ToString());
 								span.SetAttribute("allocatedAgentId", newAgent.Id.ToString());
+								span.SetAttribute("allocatedConnectionAddress", resource.ConnectionAddress);
+								span.SetAttribute("allocatedIp", resource.Ip.ToString());
 
 								await LogRequestAsync(AllocationOutcome.Accepted, arp.RequestId, arp.Requirements, arp.ParentLeaseId, span, cancellationToken);
 								return resource;
