Horde: Grouped and ungrouped artifacts
#rnx
#skipci

diff --git a/Engine/Source/Programs/Horde/HordeDashboard/src/components/jobDetailsV2/JobDetailArtifacts.tsx b/Engine/Source/Programs/Horde/HordeDashboard/src/components/jobDetailsV2/JobDetailArtifacts.tsx
--- a/Engine/Source/Programs/Horde/HordeDashboard/src/components/jobDetailsV2/JobDetailArtifacts.tsx
+++ b/Engine/Source/Programs/Horde/HordeDashboard/src/components/jobDetailsV2/JobDetailArtifacts.tsx
@@ -74,7 +74,7 @@ export const JobArtifactsPanel: React.FC<{ jobDetails: JobDetailsV2 }> = observe
 
    const constGroups = useConst<IGroup[]>([]);
    const [selected, setSelected] = useState<GetJobArtifactResponse | undefined>(undefined);
-   
+
    const artifactView = jobDetails.getDataView<JobArtifactsDataView>("JobArtifactsDataView");
 
    jobDetails.subscribe();
@@ -112,9 +112,12 @@ export const JobArtifactsPanel: React.FC<{ jobDetails: JobDetailsV2 }> = observe
    // artifact name => category (id may be unpopulated)
    const clookup: Map<string, string> = new Map();
 
+   let grouped: GetJobArtifactResponse[] = [];
+   const ungrouped: GetJobArtifactResponse[] = [];
+
    artifacts.forEach(a => {
 
-      let category = "Uncategorized";
+      let category = "";
 
       const m = a.metadata.find(m => m.toLowerCase().startsWith("dashboard-category"));
 
@@ -127,24 +130,24 @@ export const JobArtifactsPanel: React.FC<{ jobDetails: JobDetailsV2 }> = observe
          }
       }
 
-      clookup.set(a.name, category);
+      if (category) {
+         clookup.set(a.name, category);
+         grouped.push(a);
+      }
+      else {
+         ungrouped.push(a);
+      }
+
+
    })
 
-   artifacts = artifacts.sort((a, b) => {
+   grouped = grouped.sort((a, b) => {
 
       const acat = clookup.get(a.name)!;
       const bcat = clookup.get(b.name)!;
 
       if (acat !== bcat) {
 
-         if (acat !== "Uncategorized" && bcat === "Uncategorized") {
-            return -1;
-         }
-
-         if (acat === "Uncategorized" && bcat !== "Uncategorized") {
-            return 1;
-         }
-
          return acat.localeCompare(bcat);
       }
 
@@ -152,15 +155,14 @@ export const JobArtifactsPanel: React.FC<{ jobDetails: JobDetailsV2 }> = observe
    });
 
 
-   let groups: IGroup[] | undefined;
+   let groups: IGroup[] = [];
 
-   groups = [];
    let cgroup: string = "";
 
    // emit groups
-   for (let i = 0; i < artifacts.length; i++) {
+   for (let i = 0; i < grouped.length; i++) {
 
-      const a = artifacts[i];
+      const a = grouped[i];
       const cat = clookup.get(a.name)!;
 
       if (cat != cgroup) {
@@ -173,20 +175,23 @@ export const JobArtifactsPanel: React.FC<{ jobDetails: JobDetailsV2 }> = observe
 
          if (cgroup) {
             const key = `group_key_${cgroup}`;
-            groups.push({ startIndex: i, name: cgroup, key:key, count: 0, isCollapsed: constGroups.find(g => g.key === key)?.isCollapsed ?? true });
+            groups.push({ startIndex: i, name: cgroup, key: key, count: 0, isCollapsed: constGroups.find(g => g.key === key)?.isCollapsed ?? true });
          }
       }
    }
 
    if (cgroup) {
-      groups[groups.length - 1].count = artifacts.length - groups[groups.length - 1].startIndex;
+      groups[groups.length - 1].count = grouped.length - groups[groups.length - 1].startIndex;
    }
 
-   while (constGroups.length > 0) {
-      constGroups.pop();
+   if (groups?.length) {
+      while (constGroups.length > 0) {
+         constGroups.pop();
+      }
+
+      constGroups.push(...groups);
+
    }
-   
-   constGroups.push(...groups);
 
    const renderItem = (item: GetJobArtifactResponse, index?: number, column?: IColumn) => {
 
@@ -295,17 +300,28 @@ export const JobArtifactsPanel: React.FC<{ jobDetails: JobDetailsV2 }> = observe
             </Stack>
             <Stack >
                <Stack style={{ paddingTop: 8 }} tokens={{ childrenGap: 12 }}>
-                  <DetailsList                     
+                  {!!ungrouped.length && <DetailsList
+                     styles={{ root: { overflowX: "hidden" } }}
+                     isHeaderVisible={false}
+                     items={ungrouped}
+                     columns={columns}
+                     selectionMode={SelectionMode.none}
+                     layoutMode={DetailsListLayoutMode.justified}
+                     compact
+                     onRenderItemColumn={renderItem}
+                  />}
+
+                  {!!grouped.length && <DetailsList
                      styles={{ root: { overflowX: "hidden" } }}
                      isHeaderVisible={false}
-                     items={artifacts}
+                     items={grouped}
                      groups={groups}
                      columns={columns}
                      selectionMode={SelectionMode.none}
                      layoutMode={DetailsListLayoutMode.justified}
                      compact
                      onRenderItemColumn={renderItem}
-                  />
+                  />}
                </Stack>
             </Stack>
          </Stack>
