Horde: Re-encode json log events which are downgraded from warning to information level.

diff --git a/Engine/Source/Programs/Horde/Drivers/JobDriver/Utility/JobStepLogger.cs b/Engine/Source/Programs/Horde/Drivers/JobDriver/Utility/JobStepLogger.cs
--- a/Engine/Source/Programs/Horde/Drivers/JobDriver/Utility/JobStepLogger.cs
+++ b/Engine/Source/Programs/Horde/Drivers/JobDriver/Utility/JobStepLogger.cs
@@ -1,5 +1,6 @@
 // Copyright Epic Games, Inc. All Rights Reserved.
 
+using EpicGames.Core;
 using EpicGames.Horde;
 using EpicGames.Horde.Jobs;
 using EpicGames.Horde.Logs;
@@ -77,6 +78,13 @@ namespace JobDriver.Utility
 			if (logLevel == LogLevel.Warning && !_warnings)
 			{
 				logLevel = LogLevel.Information;
+
+				// If this is a json log event, we need to re-encode it to pick up the new log level.
+				if (state is JsonLogEvent jsonEvent)
+				{
+					LogEvent logEvent = LogEvent.FromState(logLevel, eventId, state, exception, formatter);
+					state = (TState)(object)new JsonLogEvent(logEvent);
+				}
 			}
 
 			// Write to the local logger
