Horde: Use compact log when greater 500k lines
#rnx
#skipci

diff --git a/Engine/Source/Programs/Horde/HordeDashboard/src/components/LogStyle.tsx b/Engine/Source/Programs/Horde/HordeDashboard/src/components/LogStyle.tsx
--- a/Engine/Source/Programs/Horde/HordeDashboard/src/components/LogStyle.tsx
+++ b/Engine/Source/Programs/Horde/HordeDashboard/src/components/LogStyle.tsx
@@ -13,7 +13,7 @@ export const logMetricNormal = {
    fontSize: 11
 }
 
-const adjustForDisplayScale = (window as any).safari === undefined && window.devicePixelRatio > 1.25;
+const adjustForDisplayScale = /*(window as any).safari === undefined &&*/ window.devicePixelRatio > 1.25;
 
 export const logMetricSmall = {
    // note, making this any larger limits log line range for overflow
diff --git a/Engine/Source/Programs/Horde/HordeDashboard/src/components/LogView.tsx b/Engine/Source/Programs/Horde/HordeDashboard/src/components/LogView.tsx
--- a/Engine/Source/Programs/Horde/HordeDashboard/src/components/LogView.tsx
+++ b/Engine/Source/Programs/Horde/HordeDashboard/src/components/LogView.tsx
@@ -578,7 +578,7 @@ export const LogList: React.FC<{ logId: string }> = observer(({ logId }) => {
    const errors = logSource.errors.sort((a, b) => a.lineIndex - b.lineIndex);
    const issues = logSource.issues;
 
-   handler.compact = !!(logSource.logItems?.length > 1000000);
+   handler.compact = !!(logSource.logItems?.length > 500000);
 
    const onRenderCell = (item?: LogItem, index?: number, isScrolling?: boolean): JSX.Element => {
 
