Horde: Fix P4 server health not being updated when HTTP health check endpoint returns degraded.

diff --git a/Engine/Source/Programs/Horde/Horde.Server/Perforce/PerforceLoadBalancer.cs b/Engine/Source/Programs/Horde/Horde.Server/Perforce/PerforceLoadBalancer.cs
--- a/Engine/Source/Programs/Horde/Horde.Server/Perforce/PerforceLoadBalancer.cs
+++ b/Engine/Source/Programs/Horde/Horde.Server/Perforce/PerforceLoadBalancer.cs
@@ -648,7 +648,7 @@ namespace Horde.Server.Perforce
 			if (entry.HealthCheckUrl != null)
 			{
 				ServerHealth http = await GetServerHealthViaHttpAsync(entry.HealthCheckUrl, cancellationToken);
-				if (http.Status == PerforceServerStatus.Unhealthy)
+				if (http.Status != PerforceServerStatus.Unknown && http.Status < health.Status)
 				{
 					// Unhealthy HTTP status will override any p4 info-based status
 					health = new ServerHealth(http.Status, http.Detail);
