Horde: Add a setting ("ReportWarnings") to exclude issues which are only warnings from summary reports.

diff --git a/Engine/Source/Programs/Horde/Horde.Server/Issues/IssueReportService.cs b/Engine/Source/Programs/Horde/Horde.Server/Issues/IssueReportService.cs
--- a/Engine/Source/Programs/Horde/Horde.Server/Issues/IssueReportService.cs
+++ b/Engine/Source/Programs/Horde/Horde.Server/Issues/IssueReportService.cs
@@ -185,6 +185,11 @@ namespace Horde.Server.Issues
 							HashSet<int> issueIds = new HashSet<int>(report.IssueSpans.Select(x => x.IssueId));
 							report.Issues.AddRange(issues.Where(x => issueIds.Contains(x.Id)));
 
+							if (!workflowConfig.ReportWarnings)
+							{
+								report.Issues.RemoveAll(x => x.Severity == IssueSeverity.Warning);
+							}
+
 							DateTime reportTime = lastScheduledReportTime;
 
 							IssueReportGroup? group = groups.FirstOrDefault(x => x.Channel == workflowConfig.ReportChannel && x.Time == reportTime);
diff --git a/Engine/Source/Programs/Horde/Horde.Server/Issues/WorkflowConfig.cs b/Engine/Source/Programs/Horde/Horde.Server/Issues/WorkflowConfig.cs
--- a/Engine/Source/Programs/Horde/Horde.Server/Issues/WorkflowConfig.cs
+++ b/Engine/Source/Programs/Horde/Horde.Server/Issues/WorkflowConfig.cs
@@ -33,6 +33,11 @@ namespace Horde.Server.Issues
 		/// </summary>
 		public string? ReportChannel { get; set; }
 
+		/// <summary>
+		/// Whether to include issues with a warning status in the summary
+		/// </summary>
+		public bool ReportWarnings { get; set; } = true;
+
 		/// <summary>
 		/// Whether to group issues by template in the report
 		/// </summary>
