Horde: Print a message to the batch log whenever a lease is cancelled or fails due to an unhandled exception.

diff --git a/Engine/Source/Programs/Horde/Horde.Agent/Leases/Handlers/JobHandler.cs b/Engine/Source/Programs/Horde/Horde.Agent/Leases/Handlers/JobHandler.cs
--- a/Engine/Source/Programs/Horde/Horde.Agent/Leases/Handlers/JobHandler.cs
+++ b/Engine/Source/Programs/Horde/Horde.Agent/Leases/Handlers/JobHandler.cs
@@ -43,30 +43,42 @@ namespace Horde.Agent.Leases.Handlers
 				GlobalTracer.Instance.ActiveSpan?.SetTag("batchId", executeTask.BatchId.ToString());
 
 				await using IServerLogger logger = session.HordeClient.CreateServerLogger(LogId.Parse(executeTask.LogId)).WithLocalLogger(localLogger);
+				try
+				{
+					executeTask.JobOptions ??= new RpcJobOptions();
 
-				executeTask.JobOptions ??= new RpcJobOptions();
+					List<string> arguments = new List<string>();
+					arguments.Add("execute");
+					arguments.Add("job");
+					arguments.Add($"-AgentId={session.AgentId}");
+					arguments.Add($"-SessionId={session.SessionId}");
+					arguments.Add($"-LeaseId={leaseId}");
+					arguments.Add($"-WorkingDir={session.WorkingDir}");
+					arguments.Add($"-Task={Convert.ToBase64String(executeTask.ToByteArray())}");
 
-				List<string> arguments = new List<string>();
-				arguments.Add("execute");
-				arguments.Add("job");
-				arguments.Add($"-AgentId={session.AgentId}");
-				arguments.Add($"-SessionId={session.SessionId}");
-				arguments.Add($"-LeaseId={leaseId}");
-				arguments.Add($"-WorkingDir={session.WorkingDir}");
-				arguments.Add($"-Task={Convert.ToBase64String(executeTask.ToByteArray())}");
+					string driverName = String.IsNullOrEmpty(executeTask.JobOptions.Driver) ? "JobDriver" : executeTask.JobOptions.Driver;
+					FileReference driverAssembly = FileReference.Combine(new DirectoryReference(AppContext.BaseDirectory), driverName, $"{driverName}.dll");
 
-				string driverName = String.IsNullOrEmpty(executeTask.JobOptions.Driver) ? "JobDriver" : executeTask.JobOptions.Driver;
-				FileReference driverAssembly = FileReference.Combine(new DirectoryReference(AppContext.BaseDirectory), driverName, $"{driverName}.dll");
+					Dictionary<string, string> environment = ManagedProcess.GetCurrentEnvVars();
+					environment[HordeHttpClient.HordeUrlEnvVarName] = session.HordeClient.ServerUrl.ToString();
+					environment[HordeHttpClient.HordeTokenEnvVarName] = executeTask.Token;
+					environment["UE_LOG_JSON_TO_STDOUT"] = "1";
 
-				Dictionary<string, string> environment = ManagedProcess.GetCurrentEnvVars();
-				environment[HordeHttpClient.HordeUrlEnvVarName] = session.HordeClient.ServerUrl.ToString();
-				environment[HordeHttpClient.HordeTokenEnvVarName] = executeTask.Token;
-				environment["UE_LOG_JSON_TO_STDOUT"] = "1";
+					int exitCode = await RunDotNetProcessAsync(driverAssembly, arguments, environment, false, logger, cancellationToken);
+					logger.LogInformation("Driver finished with exit code {ExitCode}", exitCode);
 
-				int exitCode = await RunDotNetProcessAsync(driverAssembly, arguments, environment, false, logger, cancellationToken);
-				logger.LogInformation("Driver finished with exit code {ExitCode}", exitCode);
-
-				return (exitCode == 0) ? LeaseResult.Success : LeaseResult.Failed;
+					return (exitCode == 0) ? LeaseResult.Success : LeaseResult.Failed;
+				}
+				catch (OperationCanceledException ex)
+				{
+					logger.LogError(ex, "Lease was cancelled");
+					throw;
+				}
+				catch (Exception ex)
+				{
+					logger.LogError(ex, "Unhandled exception: {Message}", ex.Message);
+					throw;
+				}
 			}
 			finally
 			{
