Horde: Updating schema documentation

#rnx

diff --git a/Engine/Source/Programs/Horde/Docs/Config/Schema/Streams.md b/Engine/Source/Programs/Horde/Docs/Config/Schema/Streams.md
--- a/Engine/Source/Programs/Horde/Docs/Config/Schema/Streams.md
+++ b/Engine/Source/Programs/Horde/Docs/Config/Schema/Streams.md
@@ -479,6 +479,7 @@ Name | Description
 `reportTimes` | `string[]`<br>Times of day at which to send a report
 `summaryTab` | `string`<br>Name of the tab to post summary data to
 `reportChannel` | `string`<br>Channel to post summary information for these templates.
+`reportWarnings` | `boolean`<br>Whether to include issues with a warning status in the summary
 `groupIssuesByTemplate` | `boolean`<br>Whether to group issues by template in the report
 `triageChannel` | `string`<br>Channel to post threads for triaging new issues
 `triagePrefix` | `string`<br>Prefix for all triage messages
