Horde: Allow user configured backends and namespaces to override the defaults in defaults.global.json.

diff --git a/Engine/Source/Programs/Horde/Horde.Server/Storage/StorageConfig.cs b/Engine/Source/Programs/Horde/Horde.Server/Storage/StorageConfig.cs
--- a/Engine/Source/Programs/Horde/Horde.Server/Storage/StorageConfig.cs
+++ b/Engine/Source/Programs/Horde/Horde.Server/Storage/StorageConfig.cs
@@ -73,11 +73,11 @@ namespace Horde.Server.Storage
 		/// </summary>
 		internal void PostLoad(GlobalConfig globalConfig)
 		{
-			// Create a lookup for backend configs
+			// Create a lookup for backend configs, allowing later entries to override previous ones
 			_backendLookup.Clear();
 			foreach (BackendConfig backendConfig in Backends)
 			{
-				_backendLookup.Add(backendConfig.Id, backendConfig);
+				_backendLookup[backendConfig.Id] = backendConfig;
 			}
 
 			// Fixup all the inherited properties
@@ -112,7 +112,7 @@ namespace Horde.Server.Storage
 				}
 
 				namespaceConfig.PostLoad(globalConfig, backendConfig);
-				_namespaceLookup.Add(namespaceConfig.Id, namespaceConfig);
+				_namespaceLookup[namespaceConfig.Id] = namespaceConfig;
 			}
 		}
 
