Horde: Fix incorrect path in agent build script.

diff --git a/Engine/Source/Programs/Horde/BuildHorde.xml b/Engine/Source/Programs/Horde/BuildHorde.xml
--- a/Engine/Source/Programs/Horde/BuildHorde.xml
+++ b/Engine/Source/Programs/Horde/BuildHorde.xml
@@ -139,7 +139,7 @@
 			<Sign Files="$(ServerToolsStagingDir_HordeAgentWin64)/HordeAgent.exe;$(ServerToolsStagingDir_HordeAgentWin64)/HordeAgent.dll;$(ServerToolsStagingDir_HordeAgentWin64)/EpicGames.*.dll" If="'$(SignOutputs)' == 'true'"/>
 
 			<!-- Build the driver -->
-			<DotNet Arguments="publish &quot;$(RootDir)/Engine/Source/Programs/Horde/JobDriver/JobDriver.csproj&quot; --output &quot;$(ServerToolsStagingDir_HordeAgentWin64)/JobDriver&quot; --runtime win7-x64 --self-contained $(VersionArguments)"/>
+			<DotNet Arguments="publish &quot;$(RootDir)/Engine/Source/Programs/Horde/Drivers/JobDriver/JobDriver.csproj&quot; --output &quot;$(ServerToolsStagingDir_HordeAgentWin64)/JobDriver&quot; --runtime win7-x64 --self-contained $(VersionArguments)"/>
 			<Sign Files="$(ServerToolsStagingDir_HordeAgentWin64)/JobDriver/JobDriver.exe;$(ServerToolsStagingDir_HordeAgentWin64)/JobDriver/JobDriver.dll;$(ServerToolsStagingDir_HordeAgentWin64)/JobDriver/EpicGames.*.dll" If="'$(SignOutputs)' == 'true'"/>
 
 			<!-- Build the tray app -->
