Horde: Add a link to the producing job to .uartifact files, as well as other metadata and keys on the artifact object.

diff --git a/Engine/Source/Programs/Horde/Horde.Server/Artifacts/ArtifactsController.cs b/Engine/Source/Programs/Horde/Horde.Server/Artifacts/ArtifactsController.cs
--- a/Engine/Source/Programs/Horde/Horde.Server/Artifacts/ArtifactsController.cs
+++ b/Engine/Source/Programs/Horde/Horde.Server/Artifacts/ArtifactsController.cs
@@ -5,6 +5,7 @@ using System.Collections.Generic;
 using System.IO;
 using System.Linq;
 using System.Net.Mime;
+using System.Text.RegularExpressions;
 using System.Threading;
 using System.Threading.Tasks;
 using EpicGames.Core;
@@ -549,6 +550,18 @@ namespace Horde.Server.Artifacts
 			Uri baseUri = new Uri(_globalConfig.ServerSettings.ServerUrl, $"api/v2/artifacts/{artifact.Id}");
 
 			ArtifactDescriptor descriptor = new ArtifactDescriptor(baseUri, new RefName("default"), fileFilter);
+			descriptor.Keys = new List<string>(artifact.Keys);
+			descriptor.Metadata = new List<string>(artifact.Metadata);
+
+			foreach (string key in artifact.Keys)
+			{
+				Match match = Regex.Match(key, @"^job:([0-9a-zA-Z]{12})$");
+				if (match.Success)
+				{
+					descriptor.JobUrl = new Uri(_globalConfig.ServerSettings.DashboardUrl, $"job/{match.Value}");
+					break;
+				}
+			}
 
 			byte[] data = descriptor.Serialize();
 			return new FileStreamResult(new MemoryStream(data), "application/x-horde-artifact") { FileDownloadName = $"{artifact.Name}.uartifact" };
diff --git a/Engine/Source/Programs/Shared/EpicGames.Horde/Artifacts/ArtifactDescriptor.cs b/Engine/Source/Programs/Shared/EpicGames.Horde/Artifacts/ArtifactDescriptor.cs
--- a/Engine/Source/Programs/Shared/EpicGames.Horde/Artifacts/ArtifactDescriptor.cs
+++ b/Engine/Source/Programs/Shared/EpicGames.Horde/Artifacts/ArtifactDescriptor.cs
@@ -26,6 +26,21 @@ namespace EpicGames.Horde.Storage
 		/// </summary>
 		public RefName RefName { get; set; }
 
+		/// <summary>
+		/// Optional URL to the job that produced this artifact
+		/// </summary>
+		public Uri? JobUrl { get; set; }
+
+		/// <summary>
+		/// Keys associated with the artifact
+		/// </summary>
+		public List<string>? Keys { get; set; }
+
+		/// <summary>
+		/// Metadata associated with the artifact
+		/// </summary>
+		public List<string>? Metadata { get; set; }
+
 		/// <summary>
 		/// Filter for the files selected for download
 		/// </summary>
