Horde: Fix job id matching for .uartifact files.

#rnx

diff --git a/Engine/Source/Programs/Horde/Horde.Server/Artifacts/ArtifactsController.cs b/Engine/Source/Programs/Horde/Horde.Server/Artifacts/ArtifactsController.cs
--- a/Engine/Source/Programs/Horde/Horde.Server/Artifacts/ArtifactsController.cs
+++ b/Engine/Source/Programs/Horde/Horde.Server/Artifacts/ArtifactsController.cs
@@ -555,7 +555,7 @@ namespace Horde.Server.Artifacts
 
 			foreach (string key in artifact.Keys)
 			{
-				Match match = Regex.Match(key, @"^job:([0-9a-zA-Z]{12})$");
+				Match match = Regex.Match(key, @"^job:([0-9a-zA-Z]{24})$");
 				if (match.Success)
 				{
 					descriptor.JobUrl = new Uri(_globalConfig.ServerSettings.DashboardUrl, $"job/{match.Value}");
