Horde: Fix bad initial agent load
#rnx

diff --git a/Engine/Source/Programs/Horde/HordeDashboard/src/backend/AgentStore.ts b/Engine/Source/Programs/Horde/HordeDashboard/src/backend/AgentStore.ts
--- a/Engine/Source/Programs/Horde/HordeDashboard/src/backend/AgentStore.ts
+++ b/Engine/Source/Programs/Horde/HordeDashboard/src/backend/AgentStore.ts
@@ -33,7 +33,7 @@ export class AgentStore {
     private _pools: PoolData[] = [];
 
     // don't need this for now
-    modifiedAfterDate: Date = new Date(0);
+    modifiedAfterDate?: Date;
 
     inflight?: boolean;
 
