Horde: Fix URL in exported .uartifact files.

#rnx

diff --git a/Engine/Source/Programs/Horde/Horde.Server/Artifacts/ArtifactsController.cs b/Engine/Source/Programs/Horde/Horde.Server/Artifacts/ArtifactsController.cs
--- a/Engine/Source/Programs/Horde/Horde.Server/Artifacts/ArtifactsController.cs
+++ b/Engine/Source/Programs/Horde/Horde.Server/Artifacts/ArtifactsController.cs
@@ -558,7 +558,7 @@ namespace Horde.Server.Artifacts
 				Match match = Regex.Match(key, @"^job:([0-9a-zA-Z]{24})$");
 				if (match.Success)
 				{
-					descriptor.JobUrl = new Uri(_globalConfig.ServerSettings.DashboardUrl, $"job/{match.Value}");
+					descriptor.JobUrl = new Uri(_globalConfig.ServerSettings.DashboardUrl, $"job/{match.Groups[1].Value}");
 					break;
 				}
 			}
