Horde: Log reasons when a lease is terminated due to the cancellation token being set.

diff --git a/Engine/Source/Programs/Horde/Horde.Agent.Tests/Services/TelemetryServiceTest.cs b/Engine/Source/Programs/Horde/Horde.Agent.Tests/Services/TelemetryServiceTest.cs
--- a/Engine/Source/Programs/Horde/Horde.Agent.Tests/Services/TelemetryServiceTest.cs
+++ b/Engine/Source/Programs/Horde/Horde.Agent.Tests/Services/TelemetryServiceTest.cs
@@ -4,7 +4,6 @@ using System;
 using System.Collections.Generic;
 using System.Threading;
 using System.Threading.Tasks;
-using Horde.Agent.Leases.Handlers;
 using Horde.Agent.Services;
 using Microsoft.Extensions.Logging.Abstractions;
 using Microsoft.Extensions.Options;
@@ -24,10 +23,9 @@ public sealed class TelemetryServiceTest : IDisposable
 		OptionsWrapper<AgentSettings> settingsOpt = new(settings);
 
 		using WorkerService workerService = new WorkerService(null!, null!, null!, null!);
-		JobHandler jobHandler = new JobHandler();
 		GrpcService grpcService = new(settingsOpt, NullLogger<GrpcService>.Instance, _loggerFactory);
 
-		_telemetryService = new TelemetryService(workerService, jobHandler, grpcService, null!, settingsOpt, NullLogger<TelemetryService>.Instance);
+		_telemetryService = new TelemetryService(workerService, grpcService, null!, settingsOpt, NullLogger<TelemetryService>.Instance);
 	}
 
 	[TestMethod]
diff --git a/Engine/Source/Programs/Horde/Horde.Agent/AgentApp.cs b/Engine/Source/Programs/Horde/Horde.Agent/AgentApp.cs
--- a/Engine/Source/Programs/Horde/Horde.Agent/AgentApp.cs
+++ b/Engine/Source/Programs/Horde/Horde.Agent/AgentApp.cs
@@ -193,16 +193,16 @@ namespace Horde.Agent
 			services.AddSingleton<TelemetryService>();
 			services.AddHostedService(sp => sp.GetRequiredService<TelemetryService>());
 
-			services.AddSingleton<JobHandler>();
+			services.AddSingleton<JobHandlerFactory>();
 			services.AddSingleton<StatusService>();
 			services.AddHostedService<StatusService>(sp => sp.GetRequiredService<StatusService>());
 
-			services.AddSingleton<LeaseHandler, ComputeHandler>();
-			services.AddSingleton<LeaseHandler, ConformHandler>();
-			services.AddSingleton<LeaseHandler, JobHandler>(x => x.GetRequiredService<JobHandler>());
-			services.AddSingleton<LeaseHandler, RestartHandler>();
-			services.AddSingleton<LeaseHandler, ShutdownHandler>();
-			services.AddSingleton<LeaseHandler, UpgradeHandler>();
+			services.AddSingleton<LeaseHandlerFactory, ComputeHandlerFactory>();
+			services.AddSingleton<LeaseHandlerFactory, ConformHandlerFactory>();
+			services.AddSingleton<LeaseHandlerFactory, JobHandlerFactory>(x => x.GetRequiredService<JobHandlerFactory>());
+			services.AddSingleton<LeaseHandlerFactory, RestartHandlerFactory>();
+			services.AddSingleton<LeaseHandlerFactory, ShutdownHandlerFactory>();
+			services.AddSingleton<LeaseHandlerFactory, UpgradeHandlerFactory>();
 
 			services.AddSingleton<CapabilitiesService>();
 			services.AddSingleton<ISessionFactory, SessionFactory>();
diff --git a/Engine/Source/Programs/Horde/Horde.Agent/Leases/Handlers/ComputeHandler.cs b/Engine/Source/Programs/Horde/Horde.Agent/Leases/Handlers/ComputeHandler.cs
--- a/Engine/Source/Programs/Horde/Horde.Agent/Leases/Handlers/ComputeHandler.cs
+++ b/Engine/Source/Programs/Horde/Horde.Agent/Leases/Handlers/ComputeHandler.cs
@@ -7,6 +7,7 @@ using EpicGames.Horde.Compute;
 using EpicGames.Horde.Compute.Transports;
 using EpicGames.Horde.Logs;
 using Horde.Agent.Services;
+using HordeCommon.Rpc.Messages;
 using HordeCommon.Rpc.Tasks;
 using Microsoft.Extensions.Logging;
 using Microsoft.Extensions.Options;
@@ -24,14 +25,15 @@ namespace Horde.Agent.Leases.Handlers
 		/// <summary>
 		/// Constructor
 		/// </summary>
-		public ComputeHandler(ComputeListenerService listenerService, IOptions<AgentSettings> settings)
+		public ComputeHandler(RpcLease lease, ComputeListenerService listenerService, IOptions<AgentSettings> settings)
+			: base(lease)
 		{
 			_listenerService = listenerService;
 			_settings = settings.Value;
 		}
 
 		/// <inheritdoc/>
-		public override async Task<LeaseResult> ExecuteAsync(ISession session, LeaseId leaseId, ComputeTask computeTask, ILogger localLogger, CancellationToken cancellationToken)
+		protected override async Task<LeaseResult> ExecuteAsync(ISession session, LeaseId leaseId, ComputeTask computeTask, ILogger localLogger, CancellationToken cancellationToken)
 		{
 			await using IServerLogger? serverLogger = (computeTask.LogId != null) ? session.HordeClient.CreateServerLogger(LogId.Parse(computeTask.LogId), LogLevel.Trace).WithLocalLogger(localLogger) : null;
 			ILogger logger = serverLogger ?? localLogger;
@@ -153,5 +155,20 @@ namespace Horde.Agent.Leases.Handlers
 			}
 		}
 	}
+
+	class ComputeHandlerFactory : LeaseHandlerFactory<ComputeTask>
+	{
+		readonly ComputeListenerService _listenerService;
+		readonly IOptions<AgentSettings> _settings;
+
+		public ComputeHandlerFactory(ComputeListenerService listenerService, IOptions<AgentSettings> settings)
+		{
+			_listenerService = listenerService;
+			_settings = settings;
+		}
+
+		public override LeaseHandler<ComputeTask> CreateHandler(RpcLease lease)
+			=> new ComputeHandler(lease, _listenerService, _settings);
+	}
 }
 
diff --git a/Engine/Source/Programs/Horde/Horde.Agent/Leases/Handlers/ConformHandler.cs b/Engine/Source/Programs/Horde/Horde.Agent/Leases/Handlers/ConformHandler.cs
--- a/Engine/Source/Programs/Horde/Horde.Agent/Leases/Handlers/ConformHandler.cs
+++ b/Engine/Source/Programs/Horde/Horde.Agent/Leases/Handlers/ConformHandler.cs
@@ -6,6 +6,7 @@ using EpicGames.Horde.Agents.Leases;
 using EpicGames.Horde.Logs;
 using Google.Protobuf;
 using Horde.Agent.Services;
+using HordeCommon.Rpc.Messages;
 using HordeCommon.Rpc.Tasks;
 using Microsoft.Extensions.Logging;
 
@@ -13,8 +14,12 @@ namespace Horde.Agent.Leases.Handlers
 {
 	class ConformHandler : LeaseHandler<ConformTask>
 	{
+		public ConformHandler(RpcLease lease)
+			: base(lease)
+		{ }
+
 		/// <inheritdoc/>
-		public override async Task<LeaseResult> ExecuteAsync(ISession session, LeaseId leaseId, ConformTask conformTask, ILogger localLogger, CancellationToken cancellationToken)
+		protected override async Task<LeaseResult> ExecuteAsync(ISession session, LeaseId leaseId, ConformTask conformTask, ILogger localLogger, CancellationToken cancellationToken)
 		{
 			await using IServerLogger serverLogger = session.HordeClient.CreateServerLogger(LogId.Parse(conformTask.LogId)).WithLocalLogger(localLogger);
 
@@ -40,5 +45,11 @@ namespace Horde.Agent.Leases.Handlers
 			return (exitCode == 0) ? LeaseResult.Success : LeaseResult.Failed;
 		}
 	}
+
+	class ConformHandlerFactory : LeaseHandlerFactory<ConformTask>
+	{
+		public override LeaseHandler<ConformTask> CreateHandler(RpcLease lease)
+			=> new ConformHandler(lease);
+	}
 }
 
diff --git a/Engine/Source/Programs/Horde/Horde.Agent/Leases/Handlers/JobHandler.cs b/Engine/Source/Programs/Horde/Horde.Agent/Leases/Handlers/JobHandler.cs
--- a/Engine/Source/Programs/Horde/Horde.Agent/Leases/Handlers/JobHandler.cs
+++ b/Engine/Source/Programs/Horde/Horde.Agent/Leases/Handlers/JobHandler.cs
@@ -6,6 +6,7 @@ using EpicGames.Horde.Agents.Leases;
 using EpicGames.Horde.Logs;
 using Google.Protobuf;
 using Horde.Agent.Services;
+using HordeCommon.Rpc.Messages;
 using HordeCommon.Rpc.Tasks;
 using Microsoft.Extensions.Logging;
 using OpenTracing.Util;
@@ -14,79 +15,61 @@ namespace Horde.Agent.Leases.Handlers
 {
 	class JobHandler : LeaseHandler<ExecuteJobTask>
 	{
-		/// <summary>
-		/// Current lease ID being executed
-		/// </summary>
-		public LeaseId? CurrentLeaseId { get; private set; } = null;
-
-		/// <summary>
-		/// Current job ID being executed
-		/// </summary>
-		public string? CurrentJobId { get; private set; } = null;
-
-		/// <summary>
-		/// Current job batch ID being executed
-		/// </summary>
-		public string? CurrentBatchId { get; private set; } = null;
+		public JobHandler(RpcLease lease)
+			: base(lease)
+		{ }
 
 		/// <inheritdoc/>
-		public override async Task<LeaseResult> ExecuteAsync(ISession session, LeaseId leaseId, ExecuteJobTask executeTask, ILogger localLogger, CancellationToken cancellationToken)
+		protected override async Task<LeaseResult> ExecuteAsync(ISession session, LeaseId leaseId, ExecuteJobTask executeTask, ILogger localLogger, CancellationToken cancellationToken)
 		{
+			GlobalTracer.Instance.ActiveSpan?.SetTag("jobId", executeTask.JobId.ToString());
+			GlobalTracer.Instance.ActiveSpan?.SetTag("jobName", executeTask.JobName.ToString());
+			GlobalTracer.Instance.ActiveSpan?.SetTag("batchId", executeTask.BatchId.ToString());
+
+			await using IServerLogger logger = session.HordeClient.CreateServerLogger(LogId.Parse(executeTask.LogId)).WithLocalLogger(localLogger);
 			try
 			{
-				CurrentLeaseId = leaseId;
-				CurrentJobId = executeTask.JobId;
-				CurrentBatchId = executeTask.BatchId;
-
-				GlobalTracer.Instance.ActiveSpan?.SetTag("jobId", executeTask.JobId.ToString());
-				GlobalTracer.Instance.ActiveSpan?.SetTag("jobName", executeTask.JobName.ToString());
-				GlobalTracer.Instance.ActiveSpan?.SetTag("batchId", executeTask.BatchId.ToString());
-
-				await using IServerLogger logger = session.HordeClient.CreateServerLogger(LogId.Parse(executeTask.LogId)).WithLocalLogger(localLogger);
-				try
-				{
-					executeTask.JobOptions ??= new RpcJobOptions();
+				executeTask.JobOptions ??= new RpcJobOptions();
 
-					List<string> arguments = new List<string>();
-					arguments.Add("execute");
-					arguments.Add("job");
-					arguments.Add($"-AgentId={session.AgentId}");
-					arguments.Add($"-SessionId={session.SessionId}");
-					arguments.Add($"-LeaseId={leaseId}");
-					arguments.Add($"-WorkingDir={session.WorkingDir}");
-					arguments.Add($"-Task={Convert.ToBase64String(executeTask.ToByteArray())}");
+				List<string> arguments = new List<string>();
+				arguments.Add("execute");
+				arguments.Add("job");
+				arguments.Add($"-AgentId={session.AgentId}");
+				arguments.Add($"-SessionId={session.SessionId}");
+				arguments.Add($"-LeaseId={leaseId}");
+				arguments.Add($"-WorkingDir={session.WorkingDir}");
+				arguments.Add($"-Task={Convert.ToBase64String(executeTask.ToByteArray())}");
 
-					string driverName = String.IsNullOrEmpty(executeTask.JobOptions.Driver) ? "JobDriver" : executeTask.JobOptions.Driver;
-					FileReference driverAssembly = FileReference.Combine(new DirectoryReference(AppContext.BaseDirectory), driverName, $"{driverName}.dll");
+				string driverName = String.IsNullOrEmpty(executeTask.JobOptions.Driver) ? "JobDriver" : executeTask.JobOptions.Driver;
+				FileReference driverAssembly = FileReference.Combine(new DirectoryReference(AppContext.BaseDirectory), driverName, $"{driverName}.dll");
 
-					Dictionary<string, string> environment = ManagedProcess.GetCurrentEnvVars();
-					environment[HordeHttpClient.HordeUrlEnvVarName] = session.HordeClient.ServerUrl.ToString();
-					environment[HordeHttpClient.HordeTokenEnvVarName] = executeTask.Token;
-					environment["UE_LOG_JSON_TO_STDOUT"] = "1";
+				Dictionary<string, string> environment = ManagedProcess.GetCurrentEnvVars();
+				environment[HordeHttpClient.HordeUrlEnvVarName] = session.HordeClient.ServerUrl.ToString();
+				environment[HordeHttpClient.HordeTokenEnvVarName] = executeTask.Token;
+				environment["UE_LOG_JSON_TO_STDOUT"] = "1";
 
-					int exitCode = await RunDotNetProcessAsync(driverAssembly, arguments, environment, false, logger, cancellationToken);
-					logger.LogInformation("Driver finished with exit code {ExitCode}", exitCode);
+				int exitCode = await RunDotNetProcessAsync(driverAssembly, arguments, environment, false, logger, cancellationToken);
+				logger.LogInformation("Driver finished with exit code {ExitCode}", exitCode);
 
-					return (exitCode == 0) ? LeaseResult.Success : LeaseResult.Failed;
-				}
-				catch (OperationCanceledException ex)
-				{
-					logger.LogError(ex, "Lease was cancelled");
-					throw;
-				}
-				catch (Exception ex)
-				{
-					logger.LogError(ex, "Unhandled exception: {Message}", ex.Message);
-					throw;
-				}
+				return (exitCode == 0) ? LeaseResult.Success : LeaseResult.Failed;
 			}
-			finally
+			catch (OperationCanceledException ex)
 			{
-				CurrentLeaseId = null;
-				CurrentJobId = null;
-				CurrentBatchId = null;
+				logger.LogError(ex, "Lease was cancelled ({Reason})", CancellationReason);
+				throw;
+			}
+			catch (Exception ex)
+			{
+				logger.LogError(ex, "Unhandled exception: {Message}", ex.Message);
+				throw;
 			}
 		}
 	}
+
+	class JobHandlerFactory : LeaseHandlerFactory<ExecuteJobTask>
+	{
+		public override LeaseHandler<ExecuteJobTask> CreateHandler(RpcLease lease)
+			=> new JobHandler(lease);
+	}
 }
 
diff --git a/Engine/Source/Programs/Horde/Horde.Agent/Leases/Handlers/RestartHandler.cs b/Engine/Source/Programs/Horde/Horde.Agent/Leases/Handlers/RestartHandler.cs
--- a/Engine/Source/Programs/Horde/Horde.Agent/Leases/Handlers/RestartHandler.cs
+++ b/Engine/Source/Programs/Horde/Horde.Agent/Leases/Handlers/RestartHandler.cs
@@ -3,6 +3,7 @@
 using EpicGames.Horde.Agents.Leases;
 using Horde.Agent.Services;
 using Horde.Agent.Utility;
+using HordeCommon.Rpc.Messages;
 using HordeCommon.Rpc.Tasks;
 using Microsoft.Extensions.Logging;
 
@@ -10,12 +11,22 @@ namespace Horde.Agent.Leases.Handlers
 {
 	class RestartHandler : LeaseHandler<RestartTask>
 	{
+		public RestartHandler(RpcLease lease)
+			: base(lease)
+		{ }
+
 		/// <inheritdoc/>
-		public override Task<LeaseResult> ExecuteAsync(ISession session, LeaseId leaseId, RestartTask task, ILogger logger, CancellationToken cancellationToken)
+		protected override Task<LeaseResult> ExecuteAsync(ISession session, LeaseId leaseId, RestartTask task, ILogger logger, CancellationToken cancellationToken)
 		{
 			logger.LogInformation("Scheduling restart task for agent {AgentId}", session.AgentId);
 			SessionResult result = new SessionResult((logger, ctx) => Shutdown.ExecuteAsync(true, logger, ctx));
 			return Task.FromResult(new LeaseResult(result));
 		}
 	}
+
+	class RestartHandlerFactory : LeaseHandlerFactory<RestartTask>
+	{
+		public override LeaseHandler<RestartTask> CreateHandler(RpcLease lease)
+			=> new RestartHandler(lease);
+	}
 }
diff --git a/Engine/Source/Programs/Horde/Horde.Agent/Leases/Handlers/ShutdownHandler.cs b/Engine/Source/Programs/Horde/Horde.Agent/Leases/Handlers/ShutdownHandler.cs
--- a/Engine/Source/Programs/Horde/Horde.Agent/Leases/Handlers/ShutdownHandler.cs
+++ b/Engine/Source/Programs/Horde/Horde.Agent/Leases/Handlers/ShutdownHandler.cs
@@ -3,6 +3,7 @@
 using EpicGames.Horde.Agents.Leases;
 using Horde.Agent.Services;
 using Horde.Agent.Utility;
+using HordeCommon.Rpc.Messages;
 using HordeCommon.Rpc.Tasks;
 using Microsoft.Extensions.Logging;
 
@@ -10,13 +11,23 @@ namespace Horde.Agent.Leases.Handlers
 {
 	class ShutdownHandler : LeaseHandler<ShutdownTask>
 	{
+		public ShutdownHandler(RpcLease lease)
+			: base(lease)
+		{ }
+
 		/// <inheritdoc/>
-		public override Task<LeaseResult> ExecuteAsync(ISession session, LeaseId leaseId, ShutdownTask task, ILogger logger, CancellationToken cancellationToken)
+		protected override Task<LeaseResult> ExecuteAsync(ISession session, LeaseId leaseId, ShutdownTask task, ILogger logger, CancellationToken cancellationToken)
 		{
 			logger.LogInformation("Scheduling shutdown task for agent {AgentId}", session.AgentId);
 			SessionResult result = new SessionResult((logger, ctx) => Shutdown.ExecuteAsync(false, logger, ctx));
 			return Task.FromResult(new LeaseResult(result));
 		}
 	}
+
+	class ShutdownHandlerFactory : LeaseHandlerFactory<ShutdownTask>
+	{
+		public override LeaseHandler<ShutdownTask> CreateHandler(RpcLease lease)
+			=> new ShutdownHandler(lease);
+	}
 }
 
diff --git a/Engine/Source/Programs/Horde/Horde.Agent/Leases/Handlers/UpgradeHandler.cs b/Engine/Source/Programs/Horde/Horde.Agent/Leases/Handlers/UpgradeHandler.cs
--- a/Engine/Source/Programs/Horde/Horde.Agent/Leases/Handlers/UpgradeHandler.cs
+++ b/Engine/Source/Programs/Horde/Horde.Agent/Leases/Handlers/UpgradeHandler.cs
@@ -10,6 +10,7 @@ using EpicGames.Horde.Logs;
 using Grpc.Core;
 using Horde.Agent.Services;
 using HordeCommon.Rpc;
+using HordeCommon.Rpc.Messages;
 using HordeCommon.Rpc.Tasks;
 using Microsoft.Extensions.Logging;
 
@@ -17,12 +18,12 @@ namespace Horde.Agent.Leases.Handlers
 {
 	class UpgradeHandler : LeaseHandler<UpgradeTask>
 	{
-		public UpgradeHandler()
-		{
-		}
+		public UpgradeHandler(RpcLease lease)
+			: base(lease)
+		{ }
 
 		/// <inheritdoc/>
-		public override async Task<LeaseResult> ExecuteAsync(ISession session, LeaseId leaseId, UpgradeTask task, ILogger localLogger, CancellationToken cancellationToken)
+		protected override async Task<LeaseResult> ExecuteAsync(ISession session, LeaseId leaseId, UpgradeTask task, ILogger localLogger, CancellationToken cancellationToken)
 		{
 			await using IServerLogger logger = session.HordeClient.CreateServerLogger(LogId.Parse(task.LogId)).WithLocalLogger(localLogger);
 
@@ -192,5 +193,11 @@ namespace Horde.Agent.Leases.Handlers
 			baseDir.Delete();
 		}
 	}
+
+	class UpgradeHandlerFactory : LeaseHandlerFactory<UpgradeTask>
+	{
+		public override LeaseHandler<UpgradeTask> CreateHandler(RpcLease lease)
+			=> new UpgradeHandler(lease);
+	}
 }
 
diff --git a/Engine/Source/Programs/Horde/Horde.Agent/Leases/LeaseHandler.cs b/Engine/Source/Programs/Horde/Horde.Agent/Leases/LeaseHandler.cs
--- a/Engine/Source/Programs/Horde/Horde.Agent/Leases/LeaseHandler.cs
+++ b/Engine/Source/Programs/Horde/Horde.Agent/Leases/LeaseHandler.cs
@@ -3,29 +3,173 @@
 using System.Diagnostics;
 using EpicGames.Core;
 using EpicGames.Horde.Agents.Leases;
+using EpicGames.Perforce.Managed;
 using Google.Protobuf;
-using Google.Protobuf.Reflection;
 using Google.Protobuf.WellKnownTypes;
 using Horde.Agent.Services;
+using HordeCommon.Rpc;
+using HordeCommon.Rpc.Messages;
 using Microsoft.Extensions.Logging;
+using OpenTracing;
+using OpenTracing.Util;
+using ByteString = Google.Protobuf.ByteString;
 
 namespace Horde.Agent.Leases
 {
 	/// <summary>
 	/// Handles execution of a specific lease type
 	/// </summary>
-	abstract class LeaseHandler
+	abstract class LeaseHandler : IDisposable
 	{
 		/// <summary>
-		/// Returns protobuf type urls for the handled message types
+		/// Identifier for this lease
 		/// </summary>
-		public abstract string LeaseType { get; }
+		public LeaseId Id { get; }
+
+		/// <summary>
+		/// The RPC lease state
+		/// </summary>
+		public RpcLease RpcLease { get; set; }
+
+		/// <summary>
+		/// Payload from the lease
+		/// </summary>
+		public Any RpcPayload { get; }
+
+		/// <summary>
+		/// Result from executing the lease
+		/// </summary>
+		public Task<LeaseResult> Result { get; private set; }
+
+		/// <summary>
+		/// Whether the lease has been cancelled
+		/// </summary>
+		public bool IsCancelled => _cancellationSource.IsCancellationRequested;
+
+		/// <summary>
+		/// Reason for cancellation
+		/// </summary>
+		public string CancellationReason => _cancellationReason ?? "unknown";
+
+		readonly CancellationTokenSource _cancellationSource = new CancellationTokenSource();
+		string? _cancellationReason;
+
+		/// <summary>
+		/// Constructor
+		/// </summary>
+		protected LeaseHandler(RpcLease rpcLease)
+		{
+			Id = LeaseId.Parse(rpcLease.Id);
+			RpcLease = rpcLease;
+			RpcPayload = rpcLease.Payload;
+			Result = Task.FromException<LeaseResult>(new InvalidOperationException("Lease has not been started"));
+		}
+
+		/// <inheritdoc/>
+		public void Dispose()
+		{
+			Dispose(disposing: true);
+			GC.SuppressFinalize(this);
+		}
+
+		/// <summary>
+		/// Overridable dispose method
+		/// </summary>
+		protected virtual void Dispose(bool disposing)
+		{
+			if (disposing)
+			{
+				_cancellationSource.Dispose();
+			}
+		}
+
+		/// <summary>
+		/// Starts executing the lease
+		/// </summary>
+		public void Start(ISession session, ILogger logger, LeaseLoggerFactory leaseLoggerFactory)
+		{
+			Result = Task.Run(() => HandleLeaseAsync(session, logger, leaseLoggerFactory));
+		}
+
+		/// <summary>
+		/// Cancels the lease
+		/// </summary>
+		public void Cancel(string reason)
+		{
+			Interlocked.CompareExchange(ref _cancellationReason, reason, null);
+			_cancellationSource.Cancel();
+		}
+
+		/// <summary>
+		/// Handle a lease request
+		/// </summary>
+		async Task<LeaseResult> HandleLeaseAsync(ISession session, ILogger logger, LeaseLoggerFactory leaseLoggerFactory)
+		{
+			using IScope scope = GlobalTracer.Instance.BuildSpan("HandleLease").WithResourceName(RpcLease.Id).StartActive();
+			scope.Span.SetTag("LeaseId", RpcLease.Id);
+			scope.Span.SetTag("AgentId", session.AgentId.ToString());
+			//			using IDisposable TraceProperty = LogContext.PushProperty("dd.trace_id", CorrelationIdentifier.TraceId.ToString());
+			//			using IDisposable SpanProperty = LogContext.PushProperty("dd.span_id", CorrelationIdentifier.SpanId.ToString());
+
+			logger.LogInformation("Handling lease {LeaseId}", Id);
+
+			// Get the lease outcome
+			LeaseResult result = LeaseResult.Failed;
+			try
+			{
+				result = await HandleLeasePayloadAsync(session, leaseLoggerFactory);
+			}
+			catch (OperationCanceledException) when (_cancellationSource.IsCancellationRequested)
+			{
+				logger.LogInformation("Lease {LeaseId} cancelled", Id);
+			}
+			catch (InsufficientSpaceException ex)
+			{
+				logger.LogError(ex, "{Message}", ex.Message);
+			}
+			catch (Exception ex)
+			{
+				logger.LogError(ex, "Unhandled exception while executing lease {LeaseId}: {Message}", Id, ex.Message);
+			}
+
+			// Update the state of the lease
+			RpcLease newRpcLease = new RpcLease(RpcLease);
+			if (_cancellationSource.IsCancellationRequested)
+			{
+				newRpcLease.State = RpcLeaseState.Cancelled;
+				newRpcLease.Outcome = RpcLeaseOutcome.Failed;
+				newRpcLease.Output = Google.Protobuf.ByteString.Empty;
+			}
+			else
+			{
+				newRpcLease.State = (result.Outcome == LeaseOutcome.Cancelled) ? RpcLeaseState.Cancelled : RpcLeaseState.Completed;
+				newRpcLease.Outcome = (RpcLeaseOutcome)result.Outcome;
+				newRpcLease.Output = (result.Output != null) ? ByteString.CopyFrom(result.Output) : ByteString.Empty;
+			}
+			RpcLease = newRpcLease;
+
+			logger.LogInformation("Transitioning lease {LeaseId} to {State}, outcome={Outcome}", Id, RpcLease.State, RpcLease.Outcome);
+
+			return result;
+		}
+
+		/// <summary>
+		/// Dispatch a lease payload to the appropriate handler
+		/// </summary>
+		internal async Task<LeaseResult> HandleLeasePayloadAsync(ISession session, LeaseLoggerFactory leaseLoggerFactory)
+		{
+			using ILoggerFactory loggerFactory = leaseLoggerFactory.CreateLoggerFactory(Id);
+			ILogger leaseLogger = loggerFactory.CreateLogger(GetType());
+
+			GlobalTracer.Instance.ActiveSpan?.SetTag("task", RpcPayload.TypeUrl);
+			return await ExecuteAsync(session, leaseLogger, _cancellationSource.Token);
+		}
 
 		/// <summary>
 		/// Executes a lease
 		/// </summary>
 		/// <returns>Result for the lease</returns>
-		public abstract Task<LeaseResult> ExecuteAsync(ISession session, LeaseId leaseId, Any message, ILogger logger, CancellationToken cancellationToken);
+		protected abstract Task<LeaseResult> ExecuteAsync(ISession session, ILogger logger, CancellationToken cancellationToken);
 
 		/// <summary>
 		/// Runs a child process, piping the output to the given logger
@@ -99,21 +243,30 @@ namespace Horde.Agent.Leases
 	/// <typeparam name="T">Type of the lease message</typeparam>
 	abstract class LeaseHandler<T> : LeaseHandler where T : IMessage<T>, new()
 	{
-		/// <summary>
-		/// Static for the message type descriptor
-		/// </summary>
-		public static MessageDescriptor Descriptor { get; } = new T().Descriptor;
-
-		/// <inheritdoc/>
-		public override string LeaseType { get; } = $"type.googleapis.com/{Descriptor.Name}";
+		protected LeaseHandler(RpcLease rpcLease)
+			: base(rpcLease)
+		{
+		}
 
 		/// <inheritdoc/>
-		public override Task<LeaseResult> ExecuteAsync(ISession session, LeaseId leaseId, Any message, ILogger logger, CancellationToken cancellationToken)
+		protected override Task<LeaseResult> ExecuteAsync(ISession session, ILogger logger, CancellationToken cancellationToken)
 		{
-			return ExecuteAsync(session, leaseId, message.Unpack<T>(), logger, cancellationToken);
+			return ExecuteAsync(session, Id, RpcLease.Payload.Unpack<T>(), logger, cancellationToken);
 		}
 
 		/// <inheritdoc/>
-		public abstract Task<LeaseResult> ExecuteAsync(ISession session, LeaseId leaseId, T message, ILogger logger, CancellationToken cancellationToken);
+		protected abstract Task<LeaseResult> ExecuteAsync(ISession session, LeaseId leaseId, T message, ILogger logger, CancellationToken cancellationToken);
+	}
+
+	class DefaultLeaseHandler : LeaseHandler
+	{
+		readonly LeaseResult _result;
+
+		public DefaultLeaseHandler(RpcLease lease, LeaseResult result) : base(lease)
+			=> _result = result;
+
+		/// <inheritdoc/>
+		protected override Task<LeaseResult> ExecuteAsync(ISession session, ILogger logger, CancellationToken cancellationToken)
+			=> Task.FromResult(_result);
 	}
 }
diff --git a/Engine/Source/Programs/Horde/Horde.Agent/Leases/LeaseHandlerFactory.cs b/Engine/Source/Programs/Horde/Horde.Agent/Leases/LeaseHandlerFactory.cs
new file mode 100644
--- /dev/null
+++ b/Engine/Source/Programs/Horde/Horde.Agent/Leases/LeaseHandlerFactory.cs
@@ -0,0 +1,43 @@
+// Copyright Epic Games, Inc. All Rights Reserved.
+
+using Google.Protobuf;
+using Google.Protobuf.Reflection;
+using HordeCommon.Rpc.Messages;
+
+namespace Horde.Agent.Leases
+{
+	abstract class LeaseHandlerFactory
+	{
+		/// <summary>
+		/// Returns protobuf type urls for the handled message types
+		/// </summary>
+		public abstract string LeaseType { get; }
+
+		/// <summary>
+		/// Creates a new lease handler for the given lease definition
+		/// </summary>
+		/// <param name="lease">The lease to create a handler for</param>
+		public abstract LeaseHandler CreateHandler(RpcLease lease);
+	}
+
+	/// <summary>
+	/// Implementation of <see cref="LeaseHandler"/> for a specific lease type
+	/// </summary>
+	/// <typeparam name="T">Type of the lease message</typeparam>
+	abstract class LeaseHandlerFactory<T> : LeaseHandlerFactory where T : IMessage<T>, new()
+	{
+		/// <summary>
+		/// Static for the message type descriptor
+		/// </summary>
+		public static MessageDescriptor Descriptor { get; } = new T().Descriptor;
+
+		/// <inheritdoc/>
+		public override string LeaseType { get; } = $"type.googleapis.com/{Descriptor.Name}";
+
+		/// <summary>
+		/// Creates a new lease handler for the given lease definition
+		/// </summary>
+		/// <param name="lease">The lease to create a handler for</param>
+		public override abstract LeaseHandler<T> CreateHandler(RpcLease lease);
+	}
+}
diff --git a/Engine/Source/Programs/Horde/Horde.Agent/Leases/LeaseManager.cs b/Engine/Source/Programs/Horde/Horde.Agent/Leases/LeaseManager.cs
--- a/Engine/Source/Programs/Horde/Horde.Agent/Leases/LeaseManager.cs
+++ b/Engine/Source/Programs/Horde/Horde.Agent/Leases/LeaseManager.cs
@@ -2,8 +2,6 @@
 
 using System.Diagnostics;
 using EpicGames.Core;
-using EpicGames.Horde.Agents.Leases;
-using EpicGames.Perforce.Managed;
 using Google.Protobuf.WellKnownTypes;
 using Grpc.Core;
 using Horde.Agent.Services;
@@ -11,55 +9,14 @@ using HordeCommon.Rpc;
 using HordeCommon.Rpc.Messages;
 using Microsoft.Extensions.DependencyInjection;
 using Microsoft.Extensions.Logging;
-using OpenTracing;
-using OpenTracing.Util;
 
 namespace Horde.Agent.Leases
 {
-	using ByteString = Google.Protobuf.ByteString;
-
 	/// <summary>
 	/// Implements the message handling loop for an agent. Runs asynchronously until disposed.
 	/// </summary>
 	class LeaseManager
 	{
-		/// <summary>
-		/// Stores information about an active session
-		/// </summary>
-		internal class LeaseInfo
-		{
-			/// <summary>
-			/// The worker lease state
-			/// </summary>
-			public RpcLease Lease { get; set; }
-
-			/// <summary>
-			/// Identifier for this lease
-			/// </summary>
-			public LeaseId LeaseId { get; set; }
-
-			/// <summary>
-			/// The task being executed for this lease
-			/// </summary>
-			public Task? Task { get; set; }
-
-			/// <summary>
-			/// Source for cancellation tokens for this session.
-			/// </summary>
-			public CancellationTokenSource CancellationTokenSource { get; }
-
-			/// <summary>
-			/// Constructor
-			/// </summary>
-			/// <param name="lease">The worker lease state</param>
-			public LeaseInfo(RpcLease lease)
-			{
-				Lease = lease;
-				LeaseId = LeaseId.Parse(lease.Id);
-				CancellationTokenSource = new CancellationTokenSource();
-			}
-		}
-
 		/// <summary>
 		/// Object used for controlling access to the access tokens and active sessions list
 		/// </summary>
@@ -68,7 +25,7 @@ namespace Horde.Agent.Leases
 		/// <summary>
 		/// The list of active leases.
 		/// </summary>
-		readonly List<LeaseInfo> _activeLeases = new List<LeaseInfo>();
+		readonly List<LeaseHandler> _activeLeases = new List<LeaseHandler>();
 
 		/// <summary>
 		/// Number of leases completed
@@ -120,28 +77,36 @@ namespace Horde.Agent.Leases
 		readonly CapabilitiesService _capabilitiesService;
 		readonly StatusService _statusService;
 		readonly ISystemMetrics _systemMetrics;
-		readonly Dictionary<string, LeaseHandler> _typeUrlToLeaseHandler;
+		readonly Dictionary<string, LeaseHandlerFactory> _typeUrlToLeaseHandler;
 		readonly LeaseLoggerFactory _leaseLoggerFactory;
 		readonly ILogger _logger;
 
 		RpcAgentCapabilities? _capabilities;
 
-		public LeaseManager(ISession session, CapabilitiesService capabilitiesService, StatusService statusService, ISystemMetrics systemMetrics, IEnumerable<LeaseHandler> leaseHandlers, LeaseLoggerFactory leaseLoggerFactory, ILogger logger)
+		public LeaseManager(ISession session, CapabilitiesService capabilitiesService, StatusService statusService, ISystemMetrics systemMetrics, IEnumerable<LeaseHandlerFactory> leaseHandlerFactories, LeaseLoggerFactory leaseLoggerFactory, ILogger logger)
 		{
 			_session = session;
 			_capabilitiesService = capabilitiesService;
 			_statusService = statusService;
 			_systemMetrics = systemMetrics;
-			_typeUrlToLeaseHandler = leaseHandlers.ToDictionary(x => x.LeaseType, x => x);
+			_typeUrlToLeaseHandler = leaseHandlerFactories.ToDictionary(x => x.LeaseType, x => x);
 			_leaseLoggerFactory = leaseLoggerFactory;
 			_logger = logger;
 		}
 
 		public LeaseManager(ISession session, IServiceProvider serviceProvider)
-			: this(session, serviceProvider.GetRequiredService<CapabilitiesService>(), serviceProvider.GetRequiredService<StatusService>(), serviceProvider.GetRequiredService<ISystemMetrics>(), serviceProvider.GetRequiredService<IEnumerable<LeaseHandler>>(), serviceProvider.GetRequiredService<LeaseLoggerFactory>(), serviceProvider.GetRequiredService<ILogger<LeaseManager>>())
+			: this(session, serviceProvider.GetRequiredService<CapabilitiesService>(), serviceProvider.GetRequiredService<StatusService>(), serviceProvider.GetRequiredService<ISystemMetrics>(), serviceProvider.GetRequiredService<IEnumerable<LeaseHandlerFactory>>(), serviceProvider.GetRequiredService<LeaseLoggerFactory>(), serviceProvider.GetRequiredService<ILogger<LeaseManager>>())
 		{
 		}
 
+		public List<RpcLease> GetActiveLeases()
+		{
+			lock (_lockObject)
+			{
+				return _activeLeases.Select(x => x.RpcLease).ToList();
+			}
+		}
+
 		public async Task<SessionResult> RunAsync(bool shutdownAfterFinishedLease, CancellationToken stoppingToken)
 		{
 			SessionResult result;
@@ -165,7 +130,7 @@ namespace Horde.Agent.Leases
 				try
 				{
 					_logger.LogInformation("Draining leases... ({NumLeases} remaining)", _activeLeases.Count);
-					await DrainLeasesAsync();
+					await DrainLeasesAsync("session terminating");
 				}
 				catch (Exception ex)
 				{
@@ -176,51 +141,52 @@ namespace Horde.Agent.Leases
 			return result;
 		}
 
-		async Task DrainLeasesAsync()
+		async Task DrainLeasesAsync(string reason)
 		{
 			for (int idx = 0; idx < _activeLeases.Count; idx++)
 			{
-				LeaseInfo activeLease = _activeLeases[idx];
-				if (activeLease.Task == null)
+				LeaseHandler activeLease = _activeLeases[idx];
+				if (activeLease.Result.IsCompleted)
 				{
 					_activeLeases.RemoveAt(idx--);
-					_logger.LogInformation("Removed lease {LeaseId}", activeLease.Lease.Id);
+					_logger.LogInformation("Removed lease {LeaseId}", activeLease.Id);
 				}
 				else
 				{
-					_logger.LogInformation("Cancelling active lease {LeaseId}", activeLease.Lease.Id);
-					activeLease.CancellationTokenSource.Cancel();
+					_logger.LogInformation("Cancelling active lease {LeaseId}", activeLease.Id);
+					activeLease.Cancel(reason);
 				}
 			}
 
 			while (_activeLeases.Count > 0)
 			{
-				List<Task> tasks = _activeLeases.Select(x => x.Task!).ToList();
+				List<Task> tasks = _activeLeases.Select(x => (Task)x.Result).ToList();
 				tasks.Add(Task.Delay(TimeSpan.FromMinutes(1.0)));
 				await Task.WhenAny(tasks);
 
 				for (int idx = 0; idx < _activeLeases.Count; idx++)
 				{
-					LeaseInfo activeLease = _activeLeases[idx];
-					if (activeLease.Task!.IsCompleted)
+					LeaseHandler activeLease = _activeLeases[idx];
+					if (activeLease.Result.IsCompleted)
 					{
 						_activeLeases.RemoveAt(idx--);
 						try
 						{
-							await activeLease.Task;
+							await activeLease.Result;
 						}
 						catch (OperationCanceledException)
 						{
 						}
 						catch (Exception ex)
 						{
-							_logger.LogError(ex, "Lease {LeaseId} threw an exception while terminating", activeLease.Lease.Id);
+							_logger.LogError(ex, "Lease {LeaseId} threw an exception while terminating", activeLease.Id);
 						}
-						_logger.LogInformation("Lease {LeaseId} has completed", activeLease.Lease.Id);
+						_logger.LogInformation("Lease {LeaseId} has completed", activeLease.Id);
+						activeLease.Dispose();
 					}
 					else
 					{
-						_logger.LogInformation("Still waiting for lease {LeaseId} to terminate...", activeLease.Lease.Id);
+						_logger.LogInformation("Still waiting for lease {LeaseId} to terminate...", activeLease.Id);
 					}
 				}
 			}
@@ -271,9 +237,9 @@ namespace Horde.Agent.Leases
 				// Get the new the lease states. If a restart is requested and we have no active leases, signal to the server that we're stopping.
 				lock (_lockObject)
 				{
-					foreach (LeaseInfo leaseInfo in _activeLeases)
+					foreach (LeaseHandler activeLease in _activeLeases)
 					{
-						updateSessionRequest.Leases.Add(new RpcLease(leaseInfo.Lease));
+						updateSessionRequest.Leases.Add(new RpcLease(activeLease.RpcLease));
 					}
 					if (_sessionResult != null && _activeLeases.Count == 0)
 					{
@@ -314,7 +280,7 @@ namespace Horde.Agent.Leases
 						// Now reconcile the local state to match what the server reports
 						if (updateSessionResponse != null)
 						{
-							bool atLeastOneLeaseFinished = _activeLeases.Any(x => x.Lease.State is RpcLeaseState.Completed or RpcLeaseState.Cancelled);
+							bool atLeastOneLeaseFinished = _activeLeases.Any(x => x.RpcLease.State is RpcLeaseState.Completed or RpcLeaseState.Cancelled);
 							if (atLeastOneLeaseFinished && shutdownAfterFinishedLease)
 							{
 								_logger.LogInformation("At least one lease executed. Requesting shutdown.");
@@ -324,7 +290,7 @@ namespace Horde.Agent.Leases
 							PoolIds = updateSessionResponse.PoolIds;
 
 							// Remove any leases which have completed
-							int numRemoved = _activeLeases.RemoveAll(x => (x.Lease.State == RpcLeaseState.Completed || x.Lease.State == RpcLeaseState.Cancelled) && !updateSessionResponse.Leases.Any(y => y.Id == x.Lease.Id && y.State != RpcLeaseState.Cancelled));
+							int numRemoved = _activeLeases.RemoveAll(x => (x.RpcLease.State == RpcLeaseState.Completed || x.RpcLease.State == RpcLeaseState.Cancelled) && !updateSessionResponse.Leases.Any(y => y.Id == x.RpcLease.Id && y.State != RpcLeaseState.Cancelled));
 							NumLeasesCompleted += numRemoved;
 
 							// Create any new leases and cancel any running leases
@@ -332,21 +298,33 @@ namespace Horde.Agent.Leases
 							{
 								if (serverLease.State == RpcLeaseState.Cancelled)
 								{
-									LeaseInfo? info = _activeLeases.FirstOrDefault(x => x.Lease.Id == serverLease.Id);
-									if (info != null)
+									LeaseHandler? handler = _activeLeases.FirstOrDefault(x => x.RpcLease.Id == serverLease.Id);
+									if (handler != null)
 									{
 										_logger.LogInformation("Cancelling lease {LeaseId}", serverLease.Id);
-										info.CancellationTokenSource.Cancel();
+										handler.Cancel("cancelled by server");
 									}
 								}
-								if (serverLease.State == RpcLeaseState.Pending && !_activeLeases.Any(x => x.Lease.Id == serverLease.Id))
+								if (serverLease.State == RpcLeaseState.Pending && !_activeLeases.Any(x => x.RpcLease.Id == serverLease.Id))
 								{
 									serverLease.State = RpcLeaseState.Active;
 
 									_logger.LogInformation("Adding lease {LeaseId}", serverLease.Id);
-									LeaseInfo info = new LeaseInfo(serverLease);
-									info.Task = Task.Run(() => HandleLeaseAsync(_session, info), CancellationToken.None);
-									_activeLeases.Add(info);
+
+									LeaseHandler leaseHandler = CreateLeaseHandler(serverLease);
+									leaseHandler.Start(_session, _logger, _leaseLoggerFactory);
+									leaseHandler.Result.ContinueWith((Task<LeaseResult> task) =>
+									{
+										LeaseResult result = task.Result;
+										if (result.SessionResult != null && _sessionResult == null)
+										{
+											_logger.LogInformation("Lease {LeaseId} is setting session result to {Result}", leaseHandler.Id, result.SessionResult.Outcome);
+											_sessionResult = result.SessionResult;
+										}
+										_updateLeasesEvent.Set();
+									}, TaskScheduler.Default);
+
+									_activeLeases.Add(leaseHandler);
 									OnLeaseActive?.Invoke(serverLease);
 								}
 							}
@@ -379,7 +357,7 @@ namespace Horde.Agent.Leases
 						if (_activeLeases.Count > 0)
 						{
 							_logger.LogInformation("Agent marked itself as busy. Draining any active leases to prevent them from using up local resources...");
-							await DrainLeasesAsync();
+							await DrainLeasesAsync("user is active");
 						}
 					}
 					else if (_activeLeases.Count == 0)
@@ -409,6 +387,17 @@ namespace Horde.Agent.Leases
 			}
 		}
 
+		LeaseHandler CreateLeaseHandler(RpcLease lease)
+		{
+			Any payload = lease.Payload;
+			if (!_typeUrlToLeaseHandler.TryGetValue(payload.TypeUrl, out LeaseHandlerFactory? leaseHandlerFactory))
+			{
+				_logger.LogError("Invalid lease payload type ({PayloadType})", payload.TypeUrl);
+				return new DefaultLeaseHandler(lease, LeaseResult.Failed);
+			}
+			return leaseHandlerFactory.CreateHandler(lease);
+		}
+
 		/// <summary>
 		/// Wrapper for <see cref="UpdateSessionInternalAsync"/> which filters/logs exceptions
 		/// </summary>
@@ -485,89 +474,6 @@ namespace Horde.Agent.Leases
 			return response;
 		}
 
-		/// <summary>
-		/// Handle a lease request
-		/// </summary>
-		/// <param name="session">The current session</param>
-		/// <param name="leaseInfo">Information about the lease</param>
-		/// <returns>Async task</returns>
-		async Task HandleLeaseAsync(ISession session, LeaseInfo leaseInfo)
-		{
-			using IScope scope = GlobalTracer.Instance.BuildSpan("HandleLease").WithResourceName(leaseInfo.Lease.Id).StartActive();
-			scope.Span.SetTag("LeaseId", leaseInfo.Lease.Id);
-			scope.Span.SetTag("AgentId", session.AgentId.ToString());
-			//			using IDisposable TraceProperty = LogContext.PushProperty("dd.trace_id", CorrelationIdentifier.TraceId.ToString());
-			//			using IDisposable SpanProperty = LogContext.PushProperty("dd.span_id", CorrelationIdentifier.SpanId.ToString());
-
-			_logger.LogInformation("Handling lease {LeaseId}", leaseInfo.Lease.Id);
-
-			// Get the lease outcome
-			LeaseResult result = LeaseResult.Failed;
-			try
-			{
-				result = await HandleLeasePayloadAsync(session, leaseInfo);
-			}
-			catch (OperationCanceledException) when (leaseInfo.CancellationTokenSource.IsCancellationRequested)
-			{
-				_logger.LogInformation("Lease {LeaseId} cancelled", leaseInfo.Lease.Id);
-			}
-			catch (InsufficientSpaceException ex)
-			{
-				_logger.LogError(ex, "{Message}", ex.Message);
-			}
-			catch (Exception ex)
-			{
-				_logger.LogError(ex, "Unhandled exception while executing lease {LeaseId}: {Message}", leaseInfo.Lease.Id, ex.Message);
-			}
-
-			// Update the state of the lease
-			lock (_lockObject)
-			{
-				if (leaseInfo.CancellationTokenSource.IsCancellationRequested)
-				{
-					leaseInfo.Lease.State = RpcLeaseState.Cancelled;
-					leaseInfo.Lease.Outcome = RpcLeaseOutcome.Failed;
-					leaseInfo.Lease.Output = ByteString.Empty;
-				}
-				else
-				{
-					leaseInfo.Lease.State = (result.Outcome == LeaseOutcome.Cancelled) ? RpcLeaseState.Cancelled : RpcLeaseState.Completed;
-					leaseInfo.Lease.Outcome = (RpcLeaseOutcome)result.Outcome;
-					leaseInfo.Lease.Output = (result.Output != null) ? ByteString.CopyFrom(result.Output) : ByteString.Empty;
-				}
-				_logger.LogInformation("Transitioning lease {LeaseId} to {State}, outcome={Outcome}", leaseInfo.Lease.Id, leaseInfo.Lease.State, leaseInfo.Lease.Outcome);
-
-				if (result.SessionResult != null && _sessionResult == null)
-				{
-					_logger.LogInformation("Lease {LeaseId} is setting session result to {Result}", leaseInfo.Lease.Id, result.SessionResult.Outcome);
-					_sessionResult = result.SessionResult;
-				}
-			}
-			_updateLeasesEvent.Set();
-		}
-
-		/// <summary>
-		/// Dispatch a lease payload to the appropriate handler
-		/// </summary>
-		/// <param name="session">The current session</param>
-		/// <param name="leaseInfo">Information about the lease</param>
-		/// <returns>Outcome from the lease</returns>
-		internal async Task<LeaseResult> HandleLeasePayloadAsync(ISession session, LeaseInfo leaseInfo)
-		{
-			Any payload = leaseInfo.Lease.Payload;
-			if (!_typeUrlToLeaseHandler.TryGetValue(payload.TypeUrl, out LeaseHandler? leaseHandler))
-			{
-				_logger.LogError("Invalid lease payload type ({PayloadType})", payload.TypeUrl);
-				return LeaseResult.Failed;
-			}
-
-			using ILoggerFactory leaseLoggerFactory = _leaseLoggerFactory.CreateLoggerFactory(leaseInfo.LeaseId);
-			ILogger leaseLogger = leaseLoggerFactory.CreateLogger(leaseHandler.GetType());
-
-			GlobalTracer.Instance.ActiveSpan?.SetTag("task", payload.TypeUrl);
-			return await leaseHandler.ExecuteAsync(session, LeaseId.Parse(leaseInfo.Lease.Id), payload, leaseLogger, leaseInfo.CancellationTokenSource.Token);
-		}
-
 		/// <summary>
 		/// Background task that updates the capabilities of this agent
 		/// </summary>
diff --git a/Engine/Source/Programs/Horde/Horde.Agent/Services/TelemetryService.cs b/Engine/Source/Programs/Horde/Horde.Agent/Services/TelemetryService.cs
--- a/Engine/Source/Programs/Horde/Horde.Agent/Services/TelemetryService.cs
+++ b/Engine/Source/Programs/Horde/Horde.Agent/Services/TelemetryService.cs
@@ -9,9 +9,10 @@ using EpicGames.Core;
 using Google.Protobuf.WellKnownTypes;
 using Grpc.Core;
 using Grpc.Net.Client;
-using Horde.Agent.Leases.Handlers;
 using HordeCommon.Rpc;
+using HordeCommon.Rpc.Messages;
 using HordeCommon.Rpc.Messages.Telemetry;
+using HordeCommon.Rpc.Tasks;
 using Microsoft.Extensions.Hosting;
 using Microsoft.Extensions.Logging;
 using Microsoft.Extensions.Options;
@@ -281,7 +282,6 @@ class TelemetryService : BackgroundService
 	private readonly TimeSpan _heartbeatMaxAllowedDiff = TimeSpan.FromSeconds(5);
 
 	private readonly WorkerService _workerService;
-	private readonly JobHandler _jobHandler;
 	private readonly GrpcService _grpcService;
 	private readonly AgentSettings _agentSettings;
 	private readonly ILogger<TelemetryService> _logger;
@@ -298,10 +298,9 @@ class TelemetryService : BackgroundService
 	/// <summary>
 	/// Constructor
 	/// </summary>
-	public TelemetryService(WorkerService workerService, JobHandler jobHandler, GrpcService grpcService, ISystemMetrics systemMetrics, IOptions<AgentSettings> settings, ILogger<TelemetryService> logger)
+	public TelemetryService(WorkerService workerService, GrpcService grpcService, ISystemMetrics systemMetrics, IOptions<AgentSettings> settings, ILogger<TelemetryService> logger)
 	{
 		_workerService = workerService;
-		_jobHandler = jobHandler;
 		_grpcService = grpcService;
 		_agentSettings = settings.Value;
 		_logger = logger;
@@ -567,12 +566,18 @@ class TelemetryService : BackgroundService
 
 			RpcSendTelemetryEventsRequest request = new();
 			Timestamp utcNow = Timestamp.FromDateTime(DateTime.UtcNow);
-			RpcExecutionMetadata em = new()
+
+			RpcExecutionMetadata em = new();
+			foreach (RpcLease lease in _workerService.GetActiveLeases())
 			{
-				LeaseId = _jobHandler.CurrentLeaseId.ToString(),
-				JobId = _jobHandler.CurrentJobId,
-				JobBatchId = _jobHandler.CurrentBatchId,
-			};
+				if (lease.Payload.TryUnpack<ExecuteJobTask>(out ExecuteJobTask task))
+				{
+					em.LeaseId = lease.Id;
+					em.JobId = task.JobId;
+					em.JobBatchId = task.BatchId;
+					break;
+				}
+			}
 
 			if (cpuMetrics != null)
 			{
diff --git a/Engine/Source/Programs/Horde/Horde.Agent/Services/WorkerService.cs b/Engine/Source/Programs/Horde/Horde.Agent/Services/WorkerService.cs
--- a/Engine/Source/Programs/Horde/Horde.Agent/Services/WorkerService.cs
+++ b/Engine/Source/Programs/Horde/Horde.Agent/Services/WorkerService.cs
@@ -3,6 +3,7 @@
 using System.Diagnostics;
 using EpicGames.Core;
 using Horde.Agent.Leases;
+using HordeCommon.Rpc.Messages;
 using Microsoft.Extensions.DependencyInjection;
 using Microsoft.Extensions.Hosting;
 using Microsoft.Extensions.Logging;
@@ -41,6 +42,14 @@ namespace Horde.Agent.Services
 			_serviceProvider = serviceProvider;
 		}
 
+		/// <summary>
+		/// Gets all the current leases held by the agent
+		/// </summary>
+		public List<RpcLease> GetActiveLeases()
+		{
+			return _currentLeaseManager?.GetActiveLeases() ?? new List<RpcLease>();
+		}
+
 		/// <summary>
 		/// Executes the ServerTaskAsync method and swallows the exception for the task being cancelled. This allows waiting for it to terminate.
 		/// </summary>
diff --git a/Engine/Source/Programs/Horde/Horde.Server/Jobs/JobRpcService.cs b/Engine/Source/Programs/Horde/Horde.Server/Jobs/JobRpcService.cs
--- a/Engine/Source/Programs/Horde/Horde.Server/Jobs/JobRpcService.cs
+++ b/Engine/Source/Programs/Horde/Horde.Server/Jobs/JobRpcService.cs
@@ -352,6 +352,7 @@ namespace Horde.Server.Jobs
 				{
 					if (otherGroup != graph.Groups[batch.GroupIdx])
 					{
+						//						AnomalyDetectorConfiguration INPUTDEPENDENCIES
 						foreach (NodeOutputRef outputRef in otherGroup.Nodes.SelectMany(x => x.Inputs))
 						{
 							if (outputRef.NodeRef.GroupIdx == batch.GroupIdx && outputRef.NodeRef.NodeIdx == step.NodeIdx && !response.PublishOutputs.Contains(outputRef.OutputIdx))
