Horde: Forward all claims for access tokens minted for Horde accounts

The role claim was missing which is used for resolving many permissions.

diff --git a/Engine/Source/Programs/Horde/Horde.Server/Accounts/OAuthController.cs b/Engine/Source/Programs/Horde/Horde.Server/Accounts/OAuthController.cs
--- a/Engine/Source/Programs/Horde/Horde.Server/Accounts/OAuthController.cs
+++ b/Engine/Source/Programs/Horde/Horde.Server/Accounts/OAuthController.cs
@@ -442,7 +442,12 @@ namespace Horde.Server.Accounts
 			payload[ClaimTypes.Name] = subject.Name;
 			payload[ClaimTypes.Email] = subject.Email;
 			payload[HordeClaimTypes.AccountId] = subject.Id.ToString();
-
+			
+			foreach (IUserClaim claim in subject.Claims)
+			{
+				payload[claim.Type] = claim.Value;
+			}
+			
 			return CreateAndSignJwt(globals, payload);
 		}
 
