Horde: Guard against exceptions from querying EC2 instance metadata on Azure.

diff --git a/Engine/Source/Programs/Horde/Drivers/JobDriver/Execution/JobExecutor.cs b/Engine/Source/Programs/Horde/Drivers/JobDriver/Execution/JobExecutor.cs
--- a/Engine/Source/Programs/Horde/Drivers/JobDriver/Execution/JobExecutor.cs
+++ b/Engine/Source/Programs/Horde/Drivers/JobDriver/Execution/JobExecutor.cs
@@ -1787,9 +1787,15 @@ namespace JobDriver.Execution
 			FileUtils.ForceDeleteFile(graphUpdateFile);
 			newEnvVars["UE_HORDE_GRAPH_UPDATE"] = graphUpdateFile.FullName;
 
-			// TODO: These are AWS specific, this should be extended to handle more clouds or for licensees to be able to set these
-			newEnvVars["UE_HORDE_AVAILABILITY_ZONE"] = Amazon.Util.EC2InstanceMetadata.AvailabilityZone ?? "";
-			newEnvVars["UE_HORDE_REGION"] = Amazon.Util.EC2InstanceMetadata.Region?.DisplayName ?? "";
+			try
+			{
+				// TODO: These are AWS specific, this should be extended to handle more clouds or for licensees to be able to set these
+				newEnvVars["UE_HORDE_AVAILABILITY_ZONE"] = Amazon.Util.EC2InstanceMetadata.AvailabilityZone ?? "";
+				newEnvVars["UE_HORDE_REGION"] = Amazon.Util.EC2InstanceMetadata.Region?.DisplayName ?? "";
+			}
+			catch
+			{
+			}
 
 			newEnvVars["UE_HORDE_JOBID"] = JobId.ToString();
 			newEnvVars["UE_HORDE_BATCHID"] = BatchId.ToString();
