Horde: Add an endpoint to query an unsync manifest for an artifact in JSON format.

#rnx

diff --git a/Engine/Source/Programs/Horde/Horde.Server/Artifacts/ArtifactsController.cs b/Engine/Source/Programs/Horde/Horde.Server/Artifacts/ArtifactsController.cs
--- a/Engine/Source/Programs/Horde/Horde.Server/Artifacts/ArtifactsController.cs
+++ b/Engine/Source/Programs/Horde/Horde.Server/Artifacts/ArtifactsController.cs
@@ -5,6 +5,7 @@ using System.Collections.Generic;
 using System.IO;
 using System.Linq;
 using System.Net.Mime;
+using System.Text.Json;
 using System.Text.RegularExpressions;
 using System.Threading;
 using System.Threading.Tasks;
@@ -567,6 +568,102 @@ namespace Horde.Server.Artifacts
 			return new FileStreamResult(new MemoryStream(data), "application/x-horde-artifact") { FileDownloadName = $"{artifact.Name}.uartifact" };
 		}
 
+		/// <summary>
+		/// Creates an Unsync manifest for an artifact
+		/// </summary>
+		/// <param name="id">The artifact id</param>
+		/// <param name="cancellationToken">Cancellation token for the request</param>
+		/// <returns>An unsync manifest stream</returns>
+		[HttpGet]
+		[Route("/api/v2/artifacts/{id}/unsync")]
+		public async Task<ActionResult> GetUnsyncManifestAsync(ArtifactId id, CancellationToken cancellationToken = default)
+		{
+			IArtifact? artifact = await _artifactCollection.GetAsync(id, cancellationToken);
+			if (artifact == null)
+			{
+				return NotFound(id);
+			}
+			if (!_globalConfig.Authorize(artifact.AclScope, ArtifactAclAction.ReadArtifact, User))
+			{
+				return Forbid(ArtifactAclAction.ReadArtifact, artifact.AclScope);
+			}
+
+			IStorageClient storageClient = _storageService.CreateClient(artifact.NamespaceId);
+
+			IBlobRef<DirectoryNode>? target = await storageClient.TryReadRefAsync<DirectoryNode>(artifact.RefName, cancellationToken: cancellationToken);
+			if (target == null)
+			{
+				return NotFound(id);
+			}
+
+			using MemoryStream stream = new MemoryStream();
+
+			Utf8JsonWriter writer = new Utf8JsonWriter(stream, new JsonWriterOptions { Indented = true });
+			writer.WriteStartObject();
+			writer.WriteString("type", "unsync_manifest");
+			writer.WriteString("hash_strong", "Blake3.160");
+			writer.WriteString("chunking", "Variable");
+			writer.WriteStartArray("files");
+			await WriteUnsyncFileManifestAsync(new Utf8StringBuilder(), target, writer, cancellationToken);
+			writer.WriteEndArray();
+			writer.WriteEndObject();
+
+			return Content(stream.ToString() ?? String.Empty);
+		}
+
+		async Task WriteUnsyncFileManifestAsync(Utf8StringBuilder path, IBlobRef<DirectoryNode> blobRef, Utf8JsonWriter writer, CancellationToken cancellationToken)
+		{
+			int initialPathLen = path.Length;
+			DirectoryNode directoryNode = await blobRef.ReadBlobAsync(cancellationToken);
+
+			foreach (FileEntry fileEntry in directoryNode.Files)
+			{
+				path.Append(fileEntry.Name);
+
+				writer.WriteStartObject();
+				writer.WriteString("name", path.WrittenSpan);
+				writer.WriteBoolean("read_only", (fileEntry.Flags & FileEntryFlags.ReadOnly) != 0);
+				writer.WriteNumber("size", fileEntry.Length);
+				writer.WriteStartArray("blocks");
+				await WriteUnsyncBlockManifestAsync(fileEntry.Target, 0, writer, cancellationToken);
+				writer.WriteEndArray();
+				writer.WriteEndObject();
+
+				path.Length = initialPathLen;
+			}
+
+			foreach (DirectoryEntry directoryEntry in directoryNode.Directories)
+			{
+				path.Append(directoryEntry.Name);
+				path.Append('/');
+				await WriteUnsyncFileManifestAsync(path, directoryEntry.Handle, writer, cancellationToken);
+				path.Length = initialPathLen;
+			}
+		}
+
+		async Task WriteUnsyncBlockManifestAsync(ChunkedDataNodeRef chunkedDataRef, long offset, Utf8JsonWriter writer, CancellationToken cancellationToken)
+		{
+			if (chunkedDataRef.Type == ChunkedDataNodeType.Leaf)
+			{
+				writer.WriteStartObject();
+				writer.WriteNumber("offset", offset);
+				writer.WriteNumber("size", chunkedDataRef.Length);
+				writer.WriteString("hash_strong", chunkedDataRef.Handle.Hash.ToString());
+				writer.WriteEndObject();
+			}
+			else
+			{
+				InteriorChunkedDataNode node = await chunkedDataRef.Handle.ReadBlobAsync<InteriorChunkedDataNode>(cancellationToken: cancellationToken);
+
+				long childOffset = offset;
+				foreach (ChunkedDataNodeRef childChunkedDataRef in node.Children)
+				{
+					await WriteUnsyncBlockManifestAsync(childChunkedDataRef, childOffset, writer, cancellationToken);
+					childOffset += childChunkedDataRef.Length;
+				}
+			}
+		}
+
 		/// <summary>
 		/// Downloads an individual file from an artifact
 		/// </summary>
diff --git a/Engine/Source/Programs/Shared/EpicGames.Core/Utf8StringBuilder.cs b/Engine/Source/Programs/Shared/EpicGames.Core/Utf8StringBuilder.cs
--- a/Engine/Source/Programs/Shared/EpicGames.Core/Utf8StringBuilder.cs
+++ b/Engine/Source/Programs/Shared/EpicGames.Core/Utf8StringBuilder.cs
@@ -15,7 +15,20 @@ namespace EpicGames.Core
 		/// <summary>
 		/// Returns the length of this string
 		/// </summary>
-		public int Length => _writer.Length;
+		public int Length
+		{
+			get => _writer.Length;
+			set
+			{
+				int prevLength = _writer.Length;
+				_writer.Resize(value);
+
+				if (value > prevLength)
+				{
+					_writer.WrittenSpan.Slice(prevLength).Fill(0);
+				}
+			}
+		}
 
 		/// <summary>
 		/// Accessor for the written span
