Horde: Stream the unsync manifest response.

#rnx

diff --git a/Engine/Source/Programs/Horde/Horde.Server/Artifacts/ArtifactsController.cs b/Engine/Source/Programs/Horde/Horde.Server/Artifacts/ArtifactsController.cs
--- a/Engine/Source/Programs/Horde/Horde.Server/Artifacts/ArtifactsController.cs
+++ b/Engine/Source/Programs/Horde/Horde.Server/Artifacts/ArtifactsController.cs
@@ -4,6 +4,7 @@ using System;
 using System.Collections.Generic;
 using System.IO;
 using System.Linq;
+using System.Net;
 using System.Net.Mime;
 using System.Text.Json;
 using System.Text.RegularExpressions;
@@ -27,6 +28,7 @@ using Horde.Server.Utilities;
 using HordeCommon.Rpc.Tasks;
 using Microsoft.AspNetCore.Authorization;
 using Microsoft.AspNetCore.Http;
+using Microsoft.AspNetCore.Http.Features;
 using Microsoft.AspNetCore.Mvc;
 using Microsoft.AspNetCore.StaticFiles;
 using Microsoft.Extensions.Logging;
@@ -588,7 +590,7 @@ namespace Horde.Server.Artifacts
 				return Forbid(ArtifactAclAction.ReadArtifact, artifact.AclScope);
 			}
 
-			IStorageClient storageClient = _storageService.CreateClient(artifact.NamespaceId);
+			using IStorageClient storageClient = _storageService.CreateClient(artifact.NamespaceId);
 
 			IBlobRef<DirectoryNode>? target = await storageClient.TryReadRefAsync<DirectoryNode>(artifact.RefName, cancellationToken: cancellationToken);
 			if (target == null)
@@ -596,22 +598,34 @@ namespace Horde.Server.Artifacts
 				return NotFound(id);
 			}
 
-			using MemoryStream stream = new MemoryStream();
+			// Disable buffering for the response
+			IHttpResponseBodyFeature? responseBodyFeature = HttpContext.Features.Get<IHttpResponseBodyFeature>();
+			responseBodyFeature?.DisableBuffering();
+
+			// Write the response directly to the writer
+			HttpResponse response = HttpContext.Response;
+			response.ContentType = "application/json";
+			response.StatusCode = (int)HttpStatusCode.OK;
+			await response.StartAsync(cancellationToken);
+
+			using (Utf8JsonWriter writer = new Utf8JsonWriter(HttpContext.Response.BodyWriter, new JsonWriterOptions { Indented = true }))
+			{
+				writer.WriteStartObject();
+				writer.WriteString("type", "unsync_manifest");
+				writer.WriteString("hash_strong", "Blake3.160");
+				writer.WriteString("chunking", "Variable");
+				writer.WriteStartArray("files");
+				await WriteUnsyncFileManifestAsync(new Utf8StringBuilder(), target, writer, cancellationToken);
+				writer.WriteEndArray();
+				writer.WriteEndObject();
+			}
 
-			Utf8JsonWriter writer = new Utf8JsonWriter(stream, new JsonWriterOptions { Indented = true });
-			writer.WriteStartObject();
-			writer.WriteString("type", "unsync_manifest");
-			writer.WriteString("hash_strong", "Blake3.160");
-			writer.WriteString("chunking", "Variable");
-			writer.WriteStartArray("files");
-			await WriteUnsyncFileManifestAsync(new Utf8StringBuilder(), target, writer, cancellationToken);
-			writer.WriteEndArray();
-			writer.WriteEndObject();
+			await response.CompleteAsync();
 
-			return Content(stream.ToString() ?? String.Empty);
+			return Ok();
 		}
 
-		async Task WriteUnsyncFileManifestAsync(Utf8StringBuilder path, IBlobRef<DirectoryNode> blobRef, Utf8JsonWriter writer, CancellationToken cancellationToken)
+		static async Task WriteUnsyncFileManifestAsync(Utf8StringBuilder path, IBlobRef<DirectoryNode> blobRef, Utf8JsonWriter writer, CancellationToken cancellationToken)
 		{
 			int initialPathLen = path.Length;
 			DirectoryNode directoryNode = await blobRef.ReadBlobAsync(cancellationToken);
@@ -641,7 +655,7 @@ namespace Horde.Server.Artifacts
 			}
 		}
 
-		async Task WriteUnsyncBlockManifestAsync(ChunkedDataNodeRef chunkedDataRef, long offset, Utf8JsonWriter writer, CancellationToken cancellationToken)
+		static async Task WriteUnsyncBlockManifestAsync(ChunkedDataNodeRef chunkedDataRef, long offset, Utf8JsonWriter writer, CancellationToken cancellationToken)
 		{
 			if (chunkedDataRef.Type == ChunkedDataNodeType.Leaf)
 			{
