Horde: Ensure the redis config file argument is always first on the command line.

diff --git a/Engine/Source/Programs/Horde/Horde.Server/Server/RedisProcess.cs b/Engine/Source/Programs/Horde/Horde.Server/Server/RedisProcess.cs
--- a/Engine/Source/Programs/Horde/Horde.Server/Server/RedisProcess.cs
+++ b/Engine/Source/Programs/Horde/Horde.Server/Server/RedisProcess.cs
@@ -81,7 +81,7 @@ namespace Horde.Server.Server
 
 			// Launch the child process
 			_processGroup = new ManagedProcessGroup();
-			_process = new ManagedProcess(_processGroup, tempRedisExe.FullName, $"--bind 127.0.0.1 --port {_port} --save \"\" --appendonly no {arguments}", redisExe.Directory.FullName, null, ProcessPriorityClass.Normal);
+			_process = new ManagedProcess(_processGroup, tempRedisExe.FullName, $"{arguments} --bind 127.0.0.1 --port {_port} --save \"\" --appendonly no", redisExe.Directory.FullName, null, ProcessPriorityClass.Normal);
 			_process.StdIn.Close();
 
 			_relayTask = BackgroundTask.StartNew(RelayOutputAsync);
