Horde: Adding support for optioinal cancellation reasons for jobs and steps

diff --git a/Engine/Source/Programs/Horde/Horde.Server/Jobs/IJob.cs b/Engine/Source/Programs/Horde/Horde.Server/Jobs/IJob.cs
--- a/Engine/Source/Programs/Horde/Horde.Server/Jobs/IJob.cs
+++ b/Engine/Source/Programs/Horde/Horde.Server/Jobs/IJob.cs
@@ -86,6 +86,11 @@ namespace Horde.Server.Jobs
 		/// </summary>
 		public UserId? AbortedByUserId { get; }
 
+		/// <summary>
+		/// Optional reason for why the job was canceled
+		/// </summary>
+		public string? CancellationReason { get; }
+
 		/// <summary>
 		/// Identifier of the bisect task that started this job
 		/// </summary>
@@ -301,8 +306,9 @@ namespace Horde.Server.Jobs
 		/// <param name="arguments">New arguments for the job</param>
 		/// <param name="labelIdxToTriggerId">New trigger ID for a label in the job</param>
 		/// <param name="jobTrigger">New downstream job id</param>
+		/// <param name="cancellationReason">Optional reason why the job was canceled</param>
 		/// <param name="cancellationToken">Cancellation token for the operation</param>
-		Task<IJob?> TryUpdateJobAsync(string? name = null, Priority? priority = null, bool? autoSubmit = null, int? autoSubmitChange = null, string? autoSubmitMessage = null, UserId? abortedByUserId = null, ObjectId? notificationTriggerId = null, List<JobReport>? reports = null, List<string>? arguments = null, KeyValuePair<int, ObjectId>? labelIdxToTriggerId = null, KeyValuePair<TemplateId, JobId>? jobTrigger = null, CancellationToken cancellationToken = default);
+		Task<IJob?> TryUpdateJobAsync(string? name = null, Priority? priority = null, bool? autoSubmit = null, int? autoSubmitChange = null, string? autoSubmitMessage = null, UserId? abortedByUserId = null, ObjectId? notificationTriggerId = null, List<JobReport>? reports = null, List<string>? arguments = null, KeyValuePair<int, ObjectId>? labelIdxToTriggerId = null, KeyValuePair<TemplateId, JobId>? jobTrigger = null, string? cancellationReason = null, CancellationToken cancellationToken = default);
 
 		/// <summary>
 		/// Updates the state of a batch
@@ -331,9 +337,10 @@ namespace Horde.Server.Jobs
 		/// <param name="newPriority">New priority for this step</param>
 		/// <param name="newReports">New report documents</param>
 		/// <param name="newProperties">Property changes. Any properties with a null value will be removed.</param>
+		/// <param name="newCancellationReason">The reason the job step was canceled</param>
 		/// <param name="cancellationToken">Cancellation token for the operation</param>
 		/// <returns>True if the job was updated, false if it was deleted in the meantime</returns>
-		Task<IJob?> TryUpdateStepAsync(JobStepBatchId batchId, JobStepId stepId, JobStepState newState = default, JobStepOutcome newOutcome = default, JobStepError? newError = null, bool? newAbortRequested = null, UserId? newAbortByUserId = null, LogId? newLogId = null, ObjectId? newNotificationTriggerId = null, UserId? newRetryByUserId = null, Priority? newPriority = null, List<JobReport>? newReports = null, Dictionary<string, string?>? newProperties = null, CancellationToken cancellationToken = default);
+		Task<IJob?> TryUpdateStepAsync(JobStepBatchId batchId, JobStepId stepId, JobStepState newState = default, JobStepOutcome newOutcome = default, JobStepError? newError = null, bool? newAbortRequested = null, UserId? newAbortByUserId = null, LogId? newLogId = null, ObjectId? newNotificationTriggerId = null, UserId? newRetryByUserId = null, Priority? newPriority = null, List<JobReport>? newReports = null, Dictionary<string, string?>? newProperties = null, string? newCancellationReason = null, CancellationToken cancellationToken = default);
 
 		/// <summary>
 		/// Attempts to update the node groups to be executed for a job. Fails if another write happens in the meantime.
@@ -1575,6 +1582,11 @@ namespace Horde.Server.Jobs
 		/// </summary>
 		public UserId? AbortedByUserId { get; }
 
+		/// <summary>
+		/// Optional reason for why the job step was canceled
+		/// </summary>
+		public string? CancellationReason { get; }
+
 		/// <summary>
 		/// List of reports for this step
 		/// </summary>
diff --git a/Engine/Source/Programs/Horde/Horde.Server/Jobs/JobCollection.cs b/Engine/Source/Programs/Horde/Horde.Server/Jobs/JobCollection.cs
--- a/Engine/Source/Programs/Horde/Horde.Server/Jobs/JobCollection.cs
+++ b/Engine/Source/Programs/Horde/Horde.Server/Jobs/JobCollection.cs
@@ -63,6 +63,8 @@ namespace Horde.Server.Jobs
 			ContentHash IJob.GraphHash => _document.GraphHash;
 			UserId? IJob.StartedByUserId => _document.StartedByUserId;
 			UserId? IJob.AbortedByUserId => _document.AbortedByUserId;
+			string? IJob.CancellationReason => _document.CancellationReason;
+
 			BisectTaskId? IJob.StartedByBisectTaskId => _document.StartedByBisectTaskId;
 			string IJob.Name => _document.Name;
 			int IJob.Change => _document.Change;
@@ -222,15 +224,15 @@ namespace Horde.Server.Jobs
 				return _collection.CreateJobObject(newDocument, newGraph);
 			}
 
-			public async Task<IJob?> TryUpdateJobAsync(string? name = null, Priority? priority = null, bool? autoSubmit = null, int? autoSubmitChange = null, string? autoSubmitMessage = null, UserId? abortedByUserId = null, ObjectId? notificationTriggerId = null, List<JobReport>? reports = null, List<string>? arguments = null, KeyValuePair<int, ObjectId>? labelIdxToTriggerId = null, KeyValuePair<TemplateId, JobId>? jobTrigger = null, CancellationToken cancellationToken = default)
+			public async Task<IJob?> TryUpdateJobAsync(string? name = null, Priority? priority = null, bool? autoSubmit = null, int? autoSubmitChange = null, string? autoSubmitMessage = null, UserId? abortedByUserId = null, ObjectId? notificationTriggerId = null, List<JobReport>? reports = null, List<string>? arguments = null, KeyValuePair<int, ObjectId>? labelIdxToTriggerId = null, KeyValuePair<TemplateId, JobId>? jobTrigger = null, string? cancellationReason = null, CancellationToken cancellationToken = default)
 			{
-				JobDocument? newDocument = await _collection.TryUpdateJobAsync(Document, Graph, name, priority, autoSubmit, autoSubmitChange, autoSubmitMessage, abortedByUserId, notificationTriggerId, reports, arguments, labelIdxToTriggerId, jobTrigger, cancellationToken);
+				JobDocument? newDocument = await _collection.TryUpdateJobAsync(Document, Graph, name, priority, autoSubmit, autoSubmitChange, autoSubmitMessage, abortedByUserId, notificationTriggerId, reports, arguments, labelIdxToTriggerId, jobTrigger, cancellationReason, cancellationToken);
 				return _collection.CreateJobObject(newDocument, Graph);
 			}
 
-			public async Task<IJob?> TryUpdateStepAsync(JobStepBatchId batchId, JobStepId stepId, JobStepState newState = JobStepState.Unspecified, JobStepOutcome newOutcome = JobStepOutcome.Unspecified, JobStepError? newError = null, bool? newAbortRequested = null, UserId? newAbortByUserId = null, LogId? newLogId = null, ObjectId? newNotificationTriggerId = null, UserId? newRetryByUserId = null, Priority? newPriority = null, List<JobReport>? newReports = null, Dictionary<string, string?>? newProperties = null, CancellationToken cancellationToken = default)
+			public async Task<IJob?> TryUpdateStepAsync(JobStepBatchId batchId, JobStepId stepId, JobStepState newState = JobStepState.Unspecified, JobStepOutcome newOutcome = JobStepOutcome.Unspecified, JobStepError? newError = null, bool? newAbortRequested = null, UserId? newAbortByUserId = null, LogId? newLogId = null, ObjectId? newNotificationTriggerId = null, UserId? newRetryByUserId = null, Priority? newPriority = null, List<JobReport>? newReports = null, Dictionary<string, string?>? newProperties = null, string? newCancellationReason = null, CancellationToken cancellationToken = default)
 			{
-				JobDocument? newDocument = await _collection.TryUpdateStepAsync(_document, _graph, batchId, stepId, newState, newOutcome, newError, newAbortRequested, newAbortByUserId, newLogId, newNotificationTriggerId, newRetryByUserId, newPriority, newReports, newProperties, cancellationToken);
+				JobDocument? newDocument = await _collection.TryUpdateStepAsync(_document, _graph, batchId, stepId, newState, newOutcome, newError, newAbortRequested, newAbortByUserId, newLogId, newNotificationTriggerId, newRetryByUserId, newPriority, newReports, newProperties, newCancellationReason, cancellationToken);
 				return _collection.CreateJobObject(newDocument, Graph);
 			}
 		}
@@ -320,6 +322,7 @@ namespace Horde.Server.Jobs
 			UserId? IJobStep.RetriedByUserId => Document.RetriedByUserId;
 			bool IJobStep.AbortRequested => Document.AbortRequested;
 			UserId? IJobStep.AbortedByUserId => Document.AbortedByUserId;
+			string? IJobStep.CancellationReason => Document.CancellationReason;
 			IReadOnlyList<IJobReport>? IJobStep.Reports => Document.Reports;
 			IReadOnlyDictionary<string, string>? IJobStep.Properties => Document.Properties;
 
@@ -412,6 +415,9 @@ namespace Horde.Server.Jobs
 			[BsonIgnoreIfNull, BsonElement("AbortedByUser")]
 			public string? AbortedByUserDeprecated { get; set; }
 
+			[BsonIgnoreIfNull]
+			public string? CancellationReason { get; set; }
+
 			[BsonRequired]
 			public string Name { get; set; }
 
@@ -649,6 +655,9 @@ namespace Horde.Server.Jobs
 			[BsonElement("AbortByUser")]
 			public string? AbortedByUserDeprecated { get; set; }
 
+			[BsonIgnoreIfNull]
+			public string? CancellationReason { get; set; }
+
 			[BsonIgnoreIfNull]
 			public List<JobReport>? Reports { get; set; }
 
@@ -989,7 +998,7 @@ namespace Horde.Server.Jobs
 		}
 
 		/// <inheritdoc/>
-		async Task<JobDocument?> TryUpdateJobAsync(JobDocument jobDocument, IGraph graph, string? name, Priority? priority, bool? autoSubmit, int? autoSubmitChange, string? autoSubmitMessage, UserId? abortedByUserId, ObjectId? notificationTriggerId, List<JobReport>? reports, List<string>? arguments, KeyValuePair<int, ObjectId>? labelIdxToTriggerId, KeyValuePair<TemplateId, JobId>? jobTrigger, CancellationToken cancellationToken)
+		async Task<JobDocument?> TryUpdateJobAsync(JobDocument jobDocument, IGraph graph, string? name, Priority? priority, bool? autoSubmit, int? autoSubmitChange, string? autoSubmitMessage, UserId? abortedByUserId, ObjectId? notificationTriggerId, List<JobReport>? reports, List<string>? arguments, KeyValuePair<int, ObjectId>? labelIdxToTriggerId, KeyValuePair<TemplateId, JobId>? jobTrigger, string? cancellationReason, CancellationToken cancellationToken)
 		{
 			// Create the update 
 			UpdateDefinitionBuilder<JobDocument> updateBuilder = Builders<JobDocument>.Update;
@@ -1031,6 +1040,13 @@ namespace Horde.Server.Jobs
 				updates.Add(updateBuilder.Set(x => x.AbortedByUserId, jobDocument.AbortedByUserId));
 				updateBatches = true;
 			}
+
+			if (!String.IsNullOrEmpty(cancellationReason))
+			{
+				jobDocument.CancellationReason = cancellationReason;
+				updates.Add(updateBuilder.Set(x => x.CancellationReason, jobDocument.CancellationReason));				
+			}
+
 			if (notificationTriggerId != null)
 			{
 				jobDocument.NotificationTriggerId = notificationTriggerId.Value;
@@ -1227,7 +1243,7 @@ namespace Horde.Server.Jobs
 		}
 
 		/// <inheritdoc/>
-		async Task<JobDocument?> TryUpdateStepAsync(JobDocument jobDocument, IGraph graph, JobStepBatchId batchId, JobStepId stepId, JobStepState newState, JobStepOutcome newOutcome, JobStepError? newError, bool? newAbortRequested, UserId? newAbortByUserId, LogId? newLogId, ObjectId? newNotificationTriggerId, UserId? newRetryByUserId, Priority? newPriority, List<JobReport>? newReports, Dictionary<string, string?>? newProperties, CancellationToken cancellationToken)
+		async Task<JobDocument?> TryUpdateStepAsync(JobDocument jobDocument, IGraph graph, JobStepBatchId batchId, JobStepId stepId, JobStepState newState, JobStepOutcome newOutcome, JobStepError? newError, bool? newAbortRequested, UserId? newAbortByUserId, LogId? newLogId, ObjectId? newNotificationTriggerId, UserId? newRetryByUserId, Priority? newPriority, List<JobReport>? newReports, Dictionary<string, string?>? newProperties, string? newCancellationReason, CancellationToken cancellationToken)
 		{
 			jobDocument = Clone(jobDocument);
 
@@ -1253,7 +1269,7 @@ namespace Horde.Server.Jobs
 							// Update the request abort status
 							if (newAbortRequested != null && step.AbortRequested == false)
 							{
-								step.AbortRequested = newAbortRequested.Value;
+								step.AbortRequested = newAbortRequested.Value;								
 								updates.Add(updateBuilder.Set(x => x.Batches[batchIdx].Steps[stepIdx].AbortRequested, step.AbortRequested));
 
 								// If the step is pending, and not running on an agent, set to aborted
@@ -1275,6 +1291,15 @@ namespace Horde.Server.Jobs
 								refreshDependentJobSteps = true;
 							}
 
+							// Update the reason the step was canceled
+							if (!String.IsNullOrEmpty(newCancellationReason))
+							{
+								step.CancellationReason = newCancellationReason;
+								updates.Add(updateBuilder.Set(x => x.Batches[batchIdx].Steps[stepIdx].CancellationReason, step.CancellationReason));
+
+								refreshDependentJobSteps = true;
+							}
+
 							// Update the state
 							if (newState != JobStepState.Unspecified && step.State != newState)
 							{
diff --git a/Engine/Source/Programs/Horde/Horde.Server/Jobs/JobService.cs b/Engine/Source/Programs/Horde/Horde.Server/Jobs/JobService.cs
--- a/Engine/Source/Programs/Horde/Horde.Server/Jobs/JobService.cs
+++ b/Engine/Source/Programs/Horde/Horde.Server/Jobs/JobService.cs
@@ -358,8 +358,9 @@ namespace Horde.Server.Jobs
 		/// <param name="reports">New reports to add</param>
 		/// <param name="arguments">New arguments for the job</param>
 		/// <param name="labelIdxToTriggerId">New trigger ID for a label in the job</param>
+		/// <param name="cancellationReason">Reason the job was canceled</param>
 		/// <param name="cancellationToken">Cancellation token for the operation</param>
-		public async Task<IJob?> UpdateJobAsync(IJob job, string? name = null, Priority? priority = null, bool? autoSubmit = null, UserId? abortedByUserId = null, ObjectId? onCompleteTriggerId = null, List<JobReport>? reports = null, List<string>? arguments = null, KeyValuePair<int, ObjectId>? labelIdxToTriggerId = null, CancellationToken cancellationToken = default)
+		public async Task<IJob?> UpdateJobAsync(IJob job, string? name = null, Priority? priority = null, bool? autoSubmit = null, UserId? abortedByUserId = null, ObjectId? onCompleteTriggerId = null, List<JobReport>? reports = null, List<string>? arguments = null, KeyValuePair<int, ObjectId>? labelIdxToTriggerId = null, string? cancellationReason = null, CancellationToken cancellationToken = default)
 		{
 			using TelemetrySpan span = _tracer.StartActiveSpan($"{nameof(JobService)}.{nameof(UpdateJobAsync)}");
 			span.SetAttribute("JobId", job.Id.ToString());
@@ -374,7 +375,7 @@ namespace Horde.Server.Jobs
 				IReadOnlyList<(LabelState, LabelOutcome)> oldLabelStates = newJob.GetLabelStates(graph);
 
 				// Update the new list of job steps
-				newJob = await newJob.TryUpdateJobAsync(name, priority, autoSubmit, null, null, abortedByUserId, onCompleteTriggerId, reports, arguments, labelIdxToTriggerId, cancellationToken: cancellationToken);
+				newJob = await newJob.TryUpdateJobAsync(name, priority, autoSubmit, null, null, abortedByUserId, onCompleteTriggerId, reports, arguments, labelIdxToTriggerId, cancellationReason: cancellationReason, cancellationToken: cancellationToken);
 				if (newJob != null)
 				{
 					// Update any badges that have been modified
@@ -973,9 +974,10 @@ namespace Horde.Server.Jobs
 		/// <param name="newPriority">New priority for this step</param>
 		/// <param name="newReports">New list of reports</param>
 		/// <param name="newProperties">Property changes. Any properties with a null value will be removed.</param>
+		/// <param name="newCancellationReason">The reason the step was canceled</param>
 		/// <param name="cancellationToken">Cancellation token for the operation</param>
 		/// <returns>True if the job was updated, false if it was deleted in the meantime</returns>
-		public async Task<IJob?> UpdateStepAsync(IJob job, JobStepBatchId batchId, JobStepId stepId, StreamConfig streamConfig, JobStepState newState = JobStepState.Unspecified, JobStepOutcome newOutcome = JobStepOutcome.Unspecified, JobStepError? newError = null, bool? newAbortRequested = null, UserId? newAbortByUserId = null, LogId? newLogId = null, ObjectId? newNotificationTriggerId = null, UserId? newRetryByUserId = null, Priority? newPriority = null, List<JobReport>? newReports = null, Dictionary<string, string?>? newProperties = null, CancellationToken cancellationToken = default)
+		public async Task<IJob?> UpdateStepAsync(IJob job, JobStepBatchId batchId, JobStepId stepId, StreamConfig streamConfig, JobStepState newState = JobStepState.Unspecified, JobStepOutcome newOutcome = JobStepOutcome.Unspecified, JobStepError? newError = null, bool? newAbortRequested = null, UserId? newAbortByUserId = null, LogId? newLogId = null, ObjectId? newNotificationTriggerId = null, UserId? newRetryByUserId = null, Priority? newPriority = null, List<JobReport>? newReports = null, Dictionary<string, string?>? newProperties = null, string? newCancellationReason = null, CancellationToken cancellationToken = default)
 		{
 			using TelemetrySpan span = _tracer.StartActiveSpan($"{nameof(JobService)}.{nameof(UpdateStepAsync)}");
 			span.SetAttribute("Job", job.Id.ToString());
@@ -985,7 +987,7 @@ namespace Horde.Server.Jobs
 			using IDisposable? scope = _logger.BeginScope("UpdateStepAsync({JobId}:{BatchId}:{StepId})", job.Id, batchId, stepId);
 			for (; ; )
 			{
-				IJob? newJob = await TryUpdateStepAsync(job, batchId, stepId, streamConfig, newState, newOutcome, newError, newAbortRequested, newAbortByUserId, newLogId, newNotificationTriggerId, newRetryByUserId, newPriority, newReports, newProperties, cancellationToken);
+				IJob? newJob = await TryUpdateStepAsync(job, batchId, stepId, streamConfig, newState, newOutcome, newError, newAbortRequested, newAbortByUserId, newLogId, newNotificationTriggerId, newRetryByUserId, newPriority, newReports, newProperties, newCancellationReason, cancellationToken);
 				if (newJob != null)
 				{
 					return newJob;
@@ -1019,9 +1021,10 @@ namespace Horde.Server.Jobs
 		/// <param name="newPriority">New priority for this step</param>
 		/// <param name="newReports">New reports</param>
 		/// <param name="newProperties">Property changes. Any properties with a null value will be removed.</param>
+		/// <param name="newCancellationReason">Property changes. Any properties with a null value will be removed.</param>
 		/// <param name="cancellationToken">Cancellation token for the operation</param>
 		/// <returns>True if the job was updated, false if it was deleted in the meantime</returns>
-		public async Task<IJob?> TryUpdateStepAsync(IJob job, JobStepBatchId batchId, JobStepId stepId, StreamConfig streamConfig, JobStepState newState = JobStepState.Unspecified, JobStepOutcome newOutcome = JobStepOutcome.Unspecified, JobStepError? newError = null, bool? newAbortRequested = null, UserId? newAbortByUserId = null, LogId? newLogId = null, ObjectId? newTriggerId = null, UserId? newRetryByUserId = null, Priority? newPriority = null, List<JobReport>? newReports = null, Dictionary<string, string?>? newProperties = null, CancellationToken cancellationToken = default)
+		public async Task<IJob?> TryUpdateStepAsync(IJob job, JobStepBatchId batchId, JobStepId stepId, StreamConfig streamConfig, JobStepState newState = JobStepState.Unspecified, JobStepOutcome newOutcome = JobStepOutcome.Unspecified, JobStepError? newError = null, bool? newAbortRequested = null, UserId? newAbortByUserId = null, LogId? newLogId = null, ObjectId? newTriggerId = null, UserId? newRetryByUserId = null, Priority? newPriority = null, List<JobReport>? newReports = null, Dictionary<string, string?>? newProperties = null, string? newCancellationReason = null, CancellationToken cancellationToken = default)
 		{
 			using TelemetrySpan span = _tracer.StartActiveSpan($"{nameof(JobService)}.{nameof(TryUpdateStepAsync)}");
 			span.SetAttribute("Job", job.Id.ToString());
@@ -1053,7 +1056,7 @@ namespace Horde.Server.Jobs
 			}
 
 			// Update the step
-			IJob? newJob = await job.TryUpdateStepAsync(batchId, stepId, newState, newOutcome, newError, newAbortRequested, newAbortByUserId, newLogId, newTriggerId, newRetryByUserId, newPriority, newReports, newProperties, cancellationToken);
+			IJob? newJob = await job.TryUpdateStepAsync(batchId, stepId, newState, newOutcome, newError, newAbortRequested, newAbortByUserId, newLogId, newTriggerId, newRetryByUserId, newPriority, newReports, newProperties, newCancellationReason, cancellationToken);
 			if (newJob != null)
 			{
 				job = newJob;
diff --git a/Engine/Source/Programs/Horde/Horde.Server/Jobs/JobsController.cs b/Engine/Source/Programs/Horde/Horde.Server/Jobs/JobsController.cs
--- a/Engine/Source/Programs/Horde/Horde.Server/Jobs/JobsController.cs
+++ b/Engine/Source/Programs/Horde/Horde.Server/Jobs/JobsController.cs
@@ -350,7 +350,7 @@ namespace Horde.Server.Jobs
 				abortedByUserId = User.GetUserId();
 			}
 
-			IJob? newJob = await _jobService.UpdateJobAsync(job, name: request.Name, priority: request.Priority, autoSubmit: request.AutoSubmit, abortedByUserId: abortedByUserId, arguments: request.Arguments);
+			IJob? newJob = await _jobService.UpdateJobAsync(job, name: request.Name, priority: request.Priority, autoSubmit: request.AutoSubmit, abortedByUserId: abortedByUserId, arguments: request.Arguments, cancellationReason: request.CancellationReason);
 			if (newJob == null)
 			{
 				return NotFound(jobId);
@@ -607,6 +607,7 @@ namespace Horde.Server.Jobs
 			response.StartedByBisectTaskId = job.StartedByBisectTaskId;
 			response.AbortedByUser = abortedByUserInfo?.Login;
 			response.AbortedByUserInfo = abortedByUserInfo;
+			response.CancellationReason = job.CancellationReason;
 			response.CreateTime = new DateTimeOffset(job.CreateTimeUtc);
 			response.State = job.GetState();
 			response.Priority = job.Priority;
@@ -750,6 +751,7 @@ namespace Horde.Server.Jobs
 			response.AbortRequested = step.AbortRequested;
 			response.AbortByUser = abortedByUserInfo?.Login;
 			response.AbortedByUserInfo = abortedByUserInfo;
+			response.CancellationReason = step.CancellationReason;
 			response.RetryByUser = retriedByUserInfo?.Login;
 			response.RetriedByUserInfo = retriedByUserInfo;
 			response.LogId = step.LogId?.ToString();
@@ -1587,7 +1589,7 @@ namespace Horde.Server.Jobs
 					retryNodeRef = new NodeRef(step.Batch.GroupIdx, step.NodeIdx);
 				}
 
-				IJob? newJob = await _jobService.UpdateStepAsync(job, batchId, stepId, streamConfig, request.State, request.Outcome, null, request.AbortRequested, abortByUser, (request.LogId == null) ? null : LogId.Parse(request.LogId), null, retryByUser, request.Priority, null, request.Properties);
+				IJob? newJob = await _jobService.UpdateStepAsync(job, batchId, stepId, streamConfig, request.State, request.Outcome, null, request.AbortRequested, abortByUser, (request.LogId == null) ? null : LogId.Parse(request.LogId), null, retryByUser, request.Priority, null, request.Properties, request.CancellationReason);
 				if (newJob == null)
 				{
 					return NotFound(jobId);
diff --git a/Engine/Source/Programs/Horde/HordeDashboard/src/backend/Api.ts b/Engine/Source/Programs/Horde/HordeDashboard/src/backend/Api.ts
--- a/Engine/Source/Programs/Horde/HordeDashboard/src/backend/Api.ts
+++ b/Engine/Source/Programs/Horde/HordeDashboard/src/backend/Api.ts
@@ -1169,6 +1169,9 @@ export type UpdateJobRequest = {
 	/** Mark this job as aborted */
 	aborted?: boolean;
 
+	/** Reason job was canceled */
+	cancellationReason?: string;
+
 	/** New list of arguments for the job. Only -Target= arguments can be modified after the job has started.  */
 	arguments?: string[];
 
@@ -1253,6 +1256,9 @@ export type GetJobResponse = {
 	/** The user that started this job */
 	abortedByUserInfo?: GetThinUserInfoResponse;
 
+	/** The reason the job was canceled */
+	cancellationReason?: string;
+
 	/** Whether job was created by a bisect task */
 	startedByBisectTaskId?: string;
 
@@ -1328,6 +1334,9 @@ export type UpdateStepRequest = {
 	/**If the step has been requested to abort */
 	abortRequested?: boolean;
 
+	/** The reason the step was canceled */	
+	cancellationReason?: string;
+
 	/**Specifies the log file id for this step */
 	logId?: string;
 
@@ -1379,6 +1388,9 @@ export type GetStepResponse = {
 	/* The user that retried this step */
 	retriedByUserInfo?: GetThinUserInfoResponse;
 
+	/** The reason the job step canceled */
+	cancellationReason?: string;
+
 	/**The log id for this step */
 	logId?: string;
 
diff --git a/Engine/Source/Programs/Horde/HordeDashboard/src/components/jobDetailsV2/AbortJobModal.tsx b/Engine/Source/Programs/Horde/HordeDashboard/src/components/jobDetailsV2/AbortJobModal.tsx
--- a/Engine/Source/Programs/Horde/HordeDashboard/src/components/jobDetailsV2/AbortJobModal.tsx
+++ b/Engine/Source/Programs/Horde/HordeDashboard/src/components/jobDetailsV2/AbortJobModal.tsx
@@ -1,5 +1,5 @@
 // Copyright Epic Games, Inc. All Rights Reserved.
-import { DefaultButton, IconButton, Modal, PrimaryButton, Stack, Text } from '@fluentui/react';
+import { DefaultButton, ITextField, IconButton, Modal, PrimaryButton, Stack, Text, TextField } from '@fluentui/react';
 import React from 'react';
 import backend from '../../backend';
 import { GetJobResponse } from '../../backend/Api';
@@ -9,6 +9,8 @@ import { getHordeStyling } from '../../styles/Styles';
 
 export const AbortJobModal: React.FC<{ jobDetails?: JobDetailsV2; jobDataIn?: GetJobResponse; stepId?: string, show: boolean; onClose: () => void }> = ({ jobDetails, jobDataIn, stepId, show, onClose }) => {
 
+   const reasonRef = React.useRef<ITextField>(null);
+
    const { hordeClasses } = getHordeStyling();
 
    const jobData = jobDetails?.jobData ?? jobDataIn;
@@ -29,8 +31,10 @@ export const AbortJobModal: React.FC<{ jobDetails?: JobDetailsV2; jobDataIn?: Ge
 
    const onAbort = async () => {
 
+      const reason = reasonRef.current?.value?.trim();
+
       if (!stepId) {
-         await backend.updateJob(jobData.id, { aborted: true }).then((response) => {
+         await backend.updateJob(jobData.id, { aborted: true, cancellationReason: reason ? reason : undefined }).then((response) => {
             console.log("Job canceled", response);
          }).catch((reason) => {
             // @todo: error ui
@@ -47,7 +51,7 @@ export const AbortJobModal: React.FC<{ jobDetails?: JobDetailsV2; jobDataIn?: Ge
             return;
          }
 
-         await backend.updateJobStep(jobData.id, batch.id, stepId, { abortRequested: true }).then((response) => {
+         await backend.updateJobStep(jobData.id, batch.id, stepId, { abortRequested: true, cancellationReason: reason ? reason : undefined }).then((response) => {
             console.log("Job step canceled", response);
          }).catch((reason) => {
             // @todo: error ui
@@ -60,7 +64,7 @@ export const AbortJobModal: React.FC<{ jobDetails?: JobDetailsV2; jobDataIn?: Ge
 
    };
 
-   const height = 140;
+   const height = 300;
 
    return <Modal isOpen={show} className={hordeClasses.modal} styles={{ main: { padding: 8, width: 540, height: height, minHeight: height } }} onDismiss={() => { if (show) close() }}>
       <Stack horizontal styles={{ root: { padding: 8 } }}>
@@ -77,6 +81,9 @@ export const AbortJobModal: React.FC<{ jobDetails?: JobDetailsV2; jobDataIn?: Ge
       </Stack>
 
       <Stack styles={{ root: { padding: 8 } }}>
+         <Stack style={{ paddingBottom: 24 }}>
+            <TextField style={{height: 120}} componentRef={reasonRef} label="Cancellation Reason" multiline resizable={false} />
+         </Stack>
          <Stack horizontal tokens={{ childrenGap: 16 }} styles={{ root: { paddingTop: 12, paddingLeft: 8, paddingBottom: 8 } }}>
             <Stack grow />
             <PrimaryButton text={(jobDetails && stepId) ? "Cancel Step" : "Cancel Job"} disabled={false} onClick={() => { onAbort(); }} />
diff --git a/Engine/Source/Programs/Horde/HordeDashboard/src/components/jobDetailsV2/JobDetailViewSummary.tsx b/Engine/Source/Programs/Horde/HordeDashboard/src/components/jobDetailsV2/JobDetailViewSummary.tsx
--- a/Engine/Source/Programs/Horde/HordeDashboard/src/components/jobDetailsV2/JobDetailViewSummary.tsx
+++ b/Engine/Source/Programs/Horde/HordeDashboard/src/components/jobDetailsV2/JobDetailViewSummary.tsx
@@ -1,6 +1,6 @@
 // Copyright Epic Games, Inc. All Rights Reserved.
 
-import { Stack, Text } from '@fluentui/react';
+import { Label, Stack, Text } from '@fluentui/react';
 import { observer } from 'mobx-react-lite';
 import React, { useEffect } from 'react';
 import { JobState, ReportPlacement } from '../../backend/Api';
@@ -52,7 +52,7 @@ JobDetailsV2.registerDataView("SummaryDataView", (details: JobDetailsV2) => new
 
 export const SummaryPanel: React.FC<{ jobDetails: JobDetailsV2 }> = observer(({ jobDetails }) => {
 
-   const { hordeClasses } = getHordeStyling();
+   const { hordeClasses, modeColors } = getHordeStyling();
 
    if (jobDetails.updated) { }
 
@@ -83,7 +83,7 @@ export const SummaryPanel: React.FC<{ jobDetails: JobDetailsV2 }> = observer(({
 
    if (!jobData) {
       return null;
-   }   
+   }
 
    const price = jobDetails.jobPrice();
 
@@ -106,7 +106,7 @@ export const SummaryPanel: React.FC<{ jobDetails: JobDetailsV2 }> = observer(({
 
    const reportData = jobDetails.getReportData(ReportPlacement.Summary);
 
-   return (<Stack id={ sideRail.url} styles={{ root: { paddingTop: 0, paddingRight: 12 } }}>
+   return (<Stack id={sideRail.url} styles={{ root: { paddingTop: 0, paddingRight: 12 } }}>
       <Stack className={hordeClasses.raised} >
          <Stack tokens={{ childrenGap: 12 }} grow>
             <Stack horizontal>
@@ -115,9 +115,15 @@ export const SummaryPanel: React.FC<{ jobDetails: JobDetailsV2 }> = observer(({
                </Stack>
             </Stack>
             <Stack >
-               <Stack tokens={{childrenGap: 12}}>
+               <Stack tokens={{ childrenGap: 12 }}>
                   <Text styles={{ root: { whiteSpace: "pre" } }}>{"" + summary}</Text>
-                  {!!reportData && <Stack> <Markdown>{reportData}</Markdown></Stack>}                  
+                  {!!jobData.cancellationReason && <Stack style={{ color: modeColors.text }} tokens={{ childrenGap: 12 }}>
+                     <Label>Cancellation Reason</Label>
+                     <Markdown>{jobData.cancellationReason}</Markdown>
+                  </Stack>
+                  }
+
+                  {!!reportData && <Stack> <Markdown>{reportData}</Markdown></Stack>}
                   {!!price && <Stack>
                      <Text>{`Estimated cost: $${price.toFixed(2)}`}</Text>
                   </Stack>}
diff --git a/Engine/Source/Programs/Horde/HordeDashboard/src/components/jobDetailsV2/StepDetailSummary.tsx b/Engine/Source/Programs/Horde/HordeDashboard/src/components/jobDetailsV2/StepDetailSummary.tsx
--- a/Engine/Source/Programs/Horde/HordeDashboard/src/components/jobDetailsV2/StepDetailSummary.tsx
+++ b/Engine/Source/Programs/Horde/HordeDashboard/src/components/jobDetailsV2/StepDetailSummary.tsx
@@ -1,5 +1,5 @@
 
-import { Stack, Text } from '@fluentui/react';
+import { Label, Stack, Text } from '@fluentui/react';
 import { observer } from 'mobx-react-lite';
 import React, { useEffect } from 'react';
 import { JobStepError, JobStepOutcome, JobStepState, ReportPlacement } from '../../backend/Api';
@@ -227,6 +227,8 @@ export const StepSummaryPanel: React.FC<{ jobDetails: JobDetailsV2; stepId: stri
       }
    }
 
+   const step = jobDetails.stepById(stepId)!
+
    const reportData = jobDetails.getReportData(ReportPlacement.Summary, stepId);
 
    return (<Stack id={sideRail.url} styles={{ root: { paddingTop: 0, paddingRight: 12 } }}>
@@ -241,6 +243,11 @@ export const StepSummaryPanel: React.FC<{ jobDetails: JobDetailsV2; stepId: stri
                <Stack style={{ color: modeColors.text }}>
                   <Markdown>{getStepSummaryMarkdown(jobDetails, stepId)}</Markdown>
                </Stack>
+               {!!step?.cancellationReason && <Stack style={{ color: modeColors.text }} tokens={{ childrenGap: 12 }}>
+                  <Label>Cancellation Reason</Label>
+                  <Markdown>{step.cancellationReason}</Markdown>
+               </Stack>
+               }
                {!!reportData && <Stack style={{ paddingTop: 8 }}> <Markdown>{reportData}</Markdown> </Stack>}
                {!!priceText && <Stack style={{ paddingTop: 8 }}>
                   <Text>{priceText}</Text>
diff --git a/Engine/Source/Programs/Shared/EpicGames.Horde/Jobs/JobMessages.cs b/Engine/Source/Programs/Shared/EpicGames.Horde/Jobs/JobMessages.cs
--- a/Engine/Source/Programs/Shared/EpicGames.Horde/Jobs/JobMessages.cs
+++ b/Engine/Source/Programs/Shared/EpicGames.Horde/Jobs/JobMessages.cs
@@ -308,6 +308,11 @@ namespace EpicGames.Horde.Jobs
 		/// </summary>
 		public bool? Aborted { get; set; }
 
+		/// <summary>
+		/// Optional reason the job was canceled
+		/// </summary>
+		public string? CancellationReason { get; set; }
+
 		/// <summary>
 		/// New list of arguments for the job. Only -Target= arguments can be modified after the job has started.
 		/// </summary>
@@ -455,6 +460,11 @@ namespace EpicGames.Horde.Jobs
 		/// </summary>
 		public GetThinUserInfoResponse? AbortedByUserInfo { get; set; }
 
+		/// <summary>
+		/// Optional reason the job was canceled
+		/// </summary>
+		public string? CancellationReason { get; set; }
+
 		/// <summary>
 		/// Priority of the job
 		/// </summary>
@@ -661,6 +671,11 @@ namespace EpicGames.Horde.Jobs
 		/// </summary>
 		public bool? AbortRequested { get; set; }
 
+		/// <summary>
+		/// Optional reason the job step was canceled
+		/// </summary>
+		public string? CancellationReason { get; set; }
+
 		/// <summary>
 		/// Specifies the log file id for this step
 		/// </summary>
@@ -800,6 +815,11 @@ namespace EpicGames.Horde.Jobs
 		/// </summary>
 		public GetThinUserInfoResponse? AbortedByUserInfo { get; set; }
 
+		/// <summary>
+		/// Optional reason the job step was canceled
+		/// </summary>
+		public string? CancellationReason { get; set; }
+
 		/// <summary>
 		/// Name of the user that requested this step be run again [DEPRECATED]
 		/// </summary>
