Horde: Fix issue where files will not be uploaded to temp storage if they were previously tagged as outputs in another step. Now allows output files to be produced by multiple steps, as long as their attributes match.

diff --git a/Engine/Source/Programs/AutomationTool/BuildGraph/BuildGraph.cs b/Engine/Source/Programs/AutomationTool/BuildGraph/BuildGraph.cs
--- a/Engine/Source/Programs/AutomationTool/BuildGraph/BuildGraph.cs
+++ b/Engine/Source/Programs/AutomationTool/BuildGraph/BuildGraph.cs
@@ -736,7 +736,7 @@ namespace AutomationTool
 				}
 			}
 
-			// ExecuteAsync the command
+			// Execute the command
 			if (!listOnly)
 			{
 				if (singleNode != null)
@@ -1100,18 +1100,20 @@ namespace AutomationTool
 			}
 
 			// Read all the input storage blocks, keeping track of which block each file came from
-			Dictionary<FileReference, TempStorageBlock> fileToStorageBlock = new Dictionary<FileReference, TempStorageBlock>();
+			Dictionary<string, TempStorageFile> inputFiles = new Dictionary<string, TempStorageFile>(FileReference.Comparer);
 			foreach (KeyValuePair<TempStorageBlock, TempStorageManifest> pair in inputManifests)
 			{
-				TempStorageBlock inputStorageBlock = pair.Key;
-				foreach (FileReference file in pair.Value.Files.Select(x => x.ToFileReference(rootDir)))
+				foreach (TempStorageFile newFile in pair.Value.Files)
 				{
-					TempStorageBlock? currentStorageBlock;
-					if (fileToStorageBlock.TryGetValue(file, out currentStorageBlock) && !TempStorage.IsDuplicateBuildProduct(file))
+					TempStorageFile? existingFile;
+					if (inputFiles.TryGetValue(newFile.RelativePath, out existingFile) && !TempStorage.IsDuplicateBuildProduct(newFile.ToFileReference(rootDir)))
 					{
-						Logger.LogError("File '{File}' was produced by {InputStorageBlock} and {CurrentStorageBlock}", file, inputStorageBlock, currentStorageBlock);
+						if (existingFile.LastWriteTimeUtcTicks != newFile.LastWriteTimeUtcTicks)
+						{
+							Logger.LogError("File '{File}' was produced by multiple nodes", newFile.RelativePath);// {InputStorageBlock} and {CurrentStorageBlock}", file, inputStorageBlock, currentStorageBlock);
+						}
 					}
-					fileToStorageBlock[file] = inputStorageBlock;
+					inputFiles[newFile.RelativePath] = newFile;
 				}
 			}
 
@@ -1195,7 +1197,7 @@ namespace AutomationTool
 				HashSet<FileReference> files = tagNameToFileSet[output.TagName];
 				foreach (FileReference file in files)
 				{
-					if (!fileToStorageBlock.ContainsKey(file) && file.IsUnderDirectory(rootDir))
+					if (file.IsUnderDirectory(rootDir))
 					{
 						if (output == node.DefaultOutput)
 						{
@@ -1236,12 +1238,13 @@ namespace AutomationTool
 			// Write all the storage blocks, and update the mapping from file to storage block
 			using (GlobalTracer.Instance.BuildSpan("TempStorage").WithTag("resource", "Write").StartActive())
 			{
+				Dictionary<FileReference, TempStorageBlock> outputFileToStorageBlock = new Dictionary<FileReference, TempStorageBlock>();
 				foreach (KeyValuePair<string, HashSet<FileReference>> pair in outputStorageBlockToFiles)
 				{
 					TempStorageBlock outputBlock = new TempStorageBlock(node.Name, pair.Key);
 					foreach (FileReference file in pair.Value)
 					{
-						fileToStorageBlock.Add(file, outputBlock);
+						outputFileToStorageBlock.Add(file, outputBlock);
 					}
 					storage.Archive(node.Name, pair.Key, pair.Value.ToArray(), pair.Value.Any(x => referencedOutputFiles.Contains(x)));
 				}
@@ -1258,7 +1261,7 @@ namespace AutomationTool
 					foreach (FileReference file in files)
 					{
 						TempStorageBlock? storageBlock;
-						if (fileToStorageBlock.TryGetValue(file, out storageBlock))
+						if (outputFileToStorageBlock.TryGetValue(file, out storageBlock))
 						{
 							storageBlocks.Add(storageBlock);
 						}
diff --git a/Engine/Source/Programs/Horde/Drivers/JobDriver/Execution/JobExecutor.cs b/Engine/Source/Programs/Horde/Drivers/JobDriver/Execution/JobExecutor.cs
--- a/Engine/Source/Programs/Horde/Drivers/JobDriver/Execution/JobExecutor.cs
+++ b/Engine/Source/Programs/Horde/Drivers/JobDriver/Execution/JobExecutor.cs
@@ -1119,7 +1119,7 @@ namespace JobDriver.Execution
 			bool result = await ExecuteWithTempStorageAsync(step, workspaceDir, arguments.ToString(), useP4, logger, cancellationToken);
 			return result;
 		}
-
+		
 		protected async Task CreateArtifactsAsync(JobStepId stepId, ArtifactName name, ArtifactType type, DirectoryReference baseDir, IEnumerable<(string, FileReference)> files, ILogger logger, CancellationToken cancellationToken)
 		{
 			await CreateArtifactAsync(stepId, name, type, baseDir, files.Select(x => x.Item2), logger, cancellationToken);
@@ -1208,18 +1208,20 @@ namespace JobDriver.Execution
 			}
 
 			// Read all the input storage blocks, keeping track of which block each file came from
-			Dictionary<FileReference, TempStorageBlockRef> fileToStorageBlock = new Dictionary<FileReference, TempStorageBlockRef>();
-			foreach (KeyValuePair<TempStorageBlockRef, TempStorageBlockManifest> pair in inputManifests)
+			Dictionary<string, (TempStorageFile, TempStorageBlockRef)> inputPathToSource = new Dictionary<string, (TempStorageFile, TempStorageBlockRef)>(FileReference.Comparer);
+			foreach ((TempStorageBlockRef inputStorageBlock, TempStorageBlockManifest inputManifest) in inputManifests)
 			{
-				TempStorageBlockRef inputStorageBlock = pair.Key;
-				foreach (FileReference file in pair.Value.Files.Select(x => x.ToFileReference(workspaceDir)))
+				foreach (TempStorageFile inputFile in inputManifest.Files)
 				{
-					TempStorageBlockRef? currentStorageBlock;
-					if (fileToStorageBlock.TryGetValue(file, out currentStorageBlock) && !TempStorage.IsDuplicateBuildProduct(file))
+					(TempStorageFile File, TempStorageBlockRef BlockRef) source;
+					if (!inputPathToSource.TryGetValue(inputFile.RelativePath, out source))
+					{
+						inputPathToSource.Add(inputFile.RelativePath, (inputFile, inputStorageBlock));
+					}
+					else if (!inputFile.Equals(source.File) && !TempStorage.IsDuplicateBuildProduct(inputFile.ToFileReference(workspaceDir)))
 					{
-						logger.LogError("File '{File}' was produced by {InputBlock} and {CurrentBlock}", file, inputStorageBlock.ToString(), currentStorageBlock.ToString());
+						logger.LogError("File '{File}' was produced by {InputBlock} and {CurrentBlock}", inputFile.RelativePath, inputStorageBlock.ToString(), source.BlockRef.ToString());
 					}
-					fileToStorageBlock[file] = inputStorageBlock;
 				}
 			}
 
@@ -1276,7 +1278,7 @@ namespace JobDriver.Execution
 			}
 
 			// Find a block name for all new outputs
-			Dictionary<FileReference, string> fileToBlockName = new Dictionary<FileReference, string>();
+			Dictionary<FileReference, string> outputFileToBlockName = new Dictionary<FileReference, string>();
 			for (int idx = 0; idx < step.OutputNames.Count; idx++)
 			{
 				string tagName = step.OutputNames[idx];
@@ -1287,25 +1289,25 @@ namespace JobDriver.Execution
 				HashSet<FileReference> files = tagNameToFileSet[tagName];
 				foreach (FileReference file in files)
 				{
-					if (!fileToStorageBlock.ContainsKey(file) && file.IsUnderDirectory(workspaceDir))
+					if (file.IsUnderDirectory(workspaceDir))
 					{
 						if (isDefaultOutput)
 						{
-							if (!fileToBlockName.ContainsKey(file))
+							if (!outputFileToBlockName.ContainsKey(file))
 							{
-								fileToBlockName[file] = "";
+								outputFileToBlockName[file] = "";
 							}
 						}
 						else
 						{
 							string? blockName;
-							if (fileToBlockName.TryGetValue(file, out blockName) && blockName.Length > 0)
+							if (outputFileToBlockName.TryGetValue(file, out blockName) && blockName.Length > 0)
 							{
-								fileToBlockName[file] = $"{blockName}+{outputNameWithoutHash}";
+								outputFileToBlockName[file] = $"{blockName}+{outputNameWithoutHash}";
 							}
 							else
 							{
-								fileToBlockName[file] = outputNameWithoutHash;
+								outputFileToBlockName[file] = outputNameWithoutHash;
 							}
 						}
 					}
@@ -1314,7 +1316,7 @@ namespace JobDriver.Execution
 
 			// Invert the dictionary to make a mapping of storage block to the files each contains
 			Dictionary<string, HashSet<FileReference>> outputStorageBlockToFiles = new Dictionary<string, HashSet<FileReference>>();
-			foreach (KeyValuePair<FileReference, string> pair in fileToBlockName)
+			foreach (KeyValuePair<FileReference, string> pair in outputFileToBlockName)
 			{
 				HashSet<FileReference>? files;
 				if (!outputStorageBlockToFiles.TryGetValue(pair.Value, out files))
@@ -1350,12 +1352,13 @@ namespace JobDriver.Execution
 					DirectoryNode outputNode = new DirectoryNode();
 
 					// Create all the output blocks
+					Dictionary<FileReference, TempStorageBlockRef> outputFileToStorageBlock = new Dictionary<FileReference, TempStorageBlockRef>();
 					foreach (KeyValuePair<string, HashSet<FileReference>> pair in outputStorageBlockToFiles)
 					{
 						TempStorageBlockRef outputBlock = new TempStorageBlockRef(step.Name, pair.Key);
 						foreach (FileReference file in pair.Value)
 						{
-							fileToStorageBlock.Add(file, outputBlock);
+							outputFileToStorageBlock.Add(file, outputBlock);
 						}
 						if (pair.Value.Any(x => referencedOutputFiles.Contains(x)))
 						{
@@ -1372,7 +1375,7 @@ namespace JobDriver.Execution
 						foreach (FileReference file in files)
 						{
 							TempStorageBlockRef? storageBlock;
-							if (fileToStorageBlock.TryGetValue(file, out storageBlock))
+							if (outputFileToStorageBlock.TryGetValue(file, out storageBlock))
 							{
 								storageBlocks.Add(storageBlock);
 							}
diff --git a/Engine/Source/Programs/Horde/Drivers/JobDriver/Utility/TempStorage.cs b/Engine/Source/Programs/Horde/Drivers/JobDriver/Utility/TempStorage.cs
--- a/Engine/Source/Programs/Horde/Drivers/JobDriver/Utility/TempStorage.cs
+++ b/Engine/Source/Programs/Horde/Drivers/JobDriver/Utility/TempStorage.cs
@@ -93,7 +93,7 @@ namespace JobDriver.Utility
 	/// Information about a single file in temp storage
 	/// </summary>
 	[DebuggerDisplay("{RelativePath}")]
-	public class TempStorageFile
+	public class TempStorageFile : IEquatable<TempStorageFile>
 	{
 		/// <summary>
 		/// The path of the file, relative to the engine root. Stored using forward slashes.
@@ -301,6 +301,24 @@ namespace JobDriver.Utility
 		{
 			return FileReference.Combine(rootDir, RelativePath.Replace('/', Path.DirectorySeparatorChar));
 		}
+
+		/// <inheritdoc/>
+		public override bool Equals(object? other)
+			=> Equals(other as TempStorageFile);
+
+		/// <inheritdoc/>
+		public bool Equals(TempStorageFile? other)
+		{
+			return other != null
+				&& String.Equals(RelativePath, other.RelativePath, StringComparison.Ordinal)
+				&& LastWriteTimeUtcTicks == other.LastWriteTimeUtcTicks
+				&& Length == other.Length
+				&& (Digest == null || other.Digest == null || String.Equals(Digest, other.Digest, StringComparison.Ordinal));
+		}
+
+		/// <inheritdoc/>
+		public override int GetHashCode()
+			=> HashCode.Combine(RelativePath, LastWriteTimeUtcTicks, Length);
 	}
 
 	/// <summary>
