Horde: Tweak cancellation reason summary
#rnx

diff --git a/Engine/Source/Programs/Horde/HordeDashboard/src/components/jobDetailsV2/JobDetailViewSummary.tsx b/Engine/Source/Programs/Horde/HordeDashboard/src/components/jobDetailsV2/JobDetailViewSummary.tsx
--- a/Engine/Source/Programs/Horde/HordeDashboard/src/components/jobDetailsV2/JobDetailViewSummary.tsx
+++ b/Engine/Source/Programs/Horde/HordeDashboard/src/components/jobDetailsV2/JobDetailViewSummary.tsx
@@ -117,7 +117,7 @@ export const SummaryPanel: React.FC<{ jobDetails: JobDetailsV2 }> = observer(({
             <Stack >
                <Stack tokens={{ childrenGap: 12 }}>
                   <Text styles={{ root: { whiteSpace: "pre" } }}>{"" + summary}</Text>
-                  {!!jobData.cancellationReason && <Stack style={{ color: modeColors.text }} tokens={{ childrenGap: 12 }}>
+                  {!!jobData.cancellationReason && <Stack style={{ color: modeColors.text }} tokens={{ childrenGap: 8 }}>
                      <Label>Cancellation Reason</Label>
                      <Markdown>{jobData.cancellationReason}</Markdown>
                   </Stack>
diff --git a/Engine/Source/Programs/Horde/HordeDashboard/src/components/jobDetailsV2/StepDetailSummary.tsx b/Engine/Source/Programs/Horde/HordeDashboard/src/components/jobDetailsV2/StepDetailSummary.tsx
--- a/Engine/Source/Programs/Horde/HordeDashboard/src/components/jobDetailsV2/StepDetailSummary.tsx
+++ b/Engine/Source/Programs/Horde/HordeDashboard/src/components/jobDetailsV2/StepDetailSummary.tsx
@@ -243,7 +243,7 @@ export const StepSummaryPanel: React.FC<{ jobDetails: JobDetailsV2; stepId: stri
                <Stack style={{ color: modeColors.text }}>
                   <Markdown>{getStepSummaryMarkdown(jobDetails, stepId)}</Markdown>
                </Stack>
-               {!!step?.cancellationReason && <Stack style={{ color: modeColors.text }} tokens={{ childrenGap: 12 }}>
+               {!!step?.cancellationReason && <Stack style={{ color: modeColors.text, paddingTop: 12, paddingBottom: 12 }} tokens={{ childrenGap: 8 }}>
                   <Label>Cancellation Reason</Label>
                   <Markdown>{step.cancellationReason}</Markdown>
                </Stack>
