Horde: Fix batch ready bug with job timeline on retried steps
#rnx

diff --git a/Engine/Source/Programs/Horde/HordeDashboard/src/components/jobDetailsV2/JobDetailTimeline.tsx b/Engine/Source/Programs/Horde/HordeDashboard/src/components/jobDetailsV2/JobDetailTimeline.tsx
--- a/Engine/Source/Programs/Horde/HordeDashboard/src/components/jobDetailsV2/JobDetailTimeline.tsx
+++ b/Engine/Source/Programs/Horde/HordeDashboard/src/components/jobDetailsV2/JobDetailTimeline.tsx
@@ -213,9 +213,13 @@ class TimelineDataView extends JobDataView {
 
          const utcBatchStart = new Date(b.startTime as string);
          const utcBatchFinish = new Date(b.finishTime as string);
-         const utcBatchReady = new Date(b.readyTime as string)
+         let utcBatchReady = new Date(b.readyTime as string)
          const utcStepStart = new Date(b.steps[0].startTime as string)
 
+         if (filterStep && jobDetails.getStepRetryNumber(filterStep.id) > 0) {
+            utcBatchReady = utcBatchStart;
+         }
+
          if (!filterStep || !maxTime) {
             if (utcBatchFinish.getTime() > maxTime) {
                maxTime = utcBatchFinish.getTime();
