Horde: Add a GC verification mode, which sets a flag in the DB rather than deleting an object from the store.

diff --git a/Engine/Source/Programs/Horde/Horde.Server/Storage/StorageConfig.cs b/Engine/Source/Programs/Horde/Horde.Server/Storage/StorageConfig.cs
--- a/Engine/Source/Programs/Horde/Horde.Server/Storage/StorageConfig.cs
+++ b/Engine/Source/Programs/Horde/Horde.Server/Storage/StorageConfig.cs
@@ -53,7 +53,12 @@ namespace Horde.Server.Storage
 		/// <summary>
 		/// Whether to enable garbage collection
 		/// </summary>
-		public bool EnableGC { get; set; } = true;
+		public bool EnableGc { get; set; } = true;
+
+		/// <summary>
+		/// Whether to enable garbage collection in verification mode (nothing deleted, just logging on access to deleted blobs)
+		/// </summary>
+		public bool EnableGcVerification { get; set; }
 
 		/// <summary>
 		/// List of storage backends
diff --git a/Engine/Source/Programs/Horde/Horde.Server/Storage/StorageService.cs b/Engine/Source/Programs/Horde/Horde.Server/Storage/StorageService.cs
--- a/Engine/Source/Programs/Horde/Horde.Server/Storage/StorageService.cs
+++ b/Engine/Source/Programs/Horde/Horde.Server/Storage/StorageService.cs
@@ -68,12 +68,18 @@ namespace Horde.Server.Storage
 
 			#region Blobs
 
-			public Task<Stream> OpenBlobAsync(BlobLocator locator, int offset, int? length, CancellationToken cancellationToken = default)
-				=> _store.OpenAsync(GetObjectKey(locator), offset, length, cancellationToken);
+			public async Task<Stream> OpenBlobAsync(BlobLocator locator, int offset, int? length, CancellationToken cancellationToken = default)
+			{
+				await _outer.CheckBlobExists(NamespaceId, locator);
+				return await _store.OpenAsync(GetObjectKey(locator), offset, length, cancellationToken);
+			}
 
 			/// <inheritdoc/>
-			public Task<IReadOnlyMemoryOwner<byte>> ReadBlobAsync(BlobLocator locator, int offset, int? length, CancellationToken cancellationToken = default)
-				=> _store.ReadAsync(GetObjectKey(locator), offset, length, cancellationToken);
+			public async Task<IReadOnlyMemoryOwner<byte>> ReadBlobAsync(BlobLocator locator, int offset, int? length, CancellationToken cancellationToken = default)
+			{
+				await _outer.CheckBlobExists(NamespaceId, locator);
+				return await _store.ReadAsync(GetObjectKey(locator), offset, length, cancellationToken);
+			}
 
 			/// <inheritdoc/>
 			public async Task<BlobLocator> WriteBlobAsync(Stream stream, IReadOnlyCollection<BlobLocator>? imports, string? basePath = null, CancellationToken cancellationToken = default)
@@ -87,9 +93,10 @@ namespace Horde.Server.Storage
 			}
 
 			/// <inheritdoc/>
-			public ValueTask<Uri?> TryGetBlobReadRedirectAsync(BlobLocator locator, CancellationToken cancellationToken = default)
+			public async ValueTask<Uri?> TryGetBlobReadRedirectAsync(BlobLocator locator, CancellationToken cancellationToken = default)
 			{
-				return _store.TryGetReadRedirectAsync(GetObjectKey(locator), cancellationToken);
+				await _outer.CheckBlobExists(NamespaceId, locator);
+				return await _store.TryGetReadRedirectAsync(GetObjectKey(locator), cancellationToken);
 			}
 
 			/// <inheritdoc/>
@@ -231,6 +238,8 @@ namespace Horde.Server.Storage
 			}
 		}
 
+		const int CurrentGcVersion = 1;
+
 		internal class BlobInfo
 		{
 			public ObjectId Id { get; set; }
@@ -247,6 +256,9 @@ namespace Horde.Server.Storage
 			[BsonElement("ali"), BsonIgnoreIfNull]
 			public List<AliasInfo>? Aliases { get; set; }
 
+			[BsonElement("del"), BsonIgnoreIfDefault]
+			public int GcVersion { get; set; }
+
 			[BsonIgnore]
 			public BlobLocator Locator => new BlobLocator(Path);
 
@@ -574,7 +586,17 @@ namespace Horde.Server.Storage
 			await _blobCollection.InsertOneAsync(blobInfo, new InsertOneOptions { }, cancellationToken);
 		}
 
-		/// <inheritdoc/>
+		async ValueTask CheckBlobExists(NamespaceId namespaceId, BlobLocator locator)
+		{
+			if (_globalConfig.CurrentValue.Storage.EnableGcVerification)
+			{
+				if (await _blobCollection.Find(x => x.Locator == locator && x.GcVersion >= CurrentGcVersion).AnyAsync())
+				{
+					_logger.LogWarning("Blob {Locator} accessed after being garbage collected", locator);
+				}
+			}
+		}
+
 		async Task<bool> IsBlobReferencedAsync(ObjectId blobInfoId, CancellationToken cancellationToken = default)
 		{
 			FilterDefinition<BlobInfo> blobFilter = Builders<BlobInfo>.Filter.AnyEq(x => x.Imports, blobInfoId);
@@ -960,7 +982,7 @@ namespace Horde.Server.Storage
 				State state = CreateState(globalConfig);
 
 				StorageConfig storageConfig = state.Config.Storage;
-				if (!storageConfig.EnableGC)
+				if (!storageConfig.EnableGc && !storageConfig.EnableGcVerification)
 				{
 					break;
 				}
@@ -1011,7 +1033,7 @@ namespace Horde.Server.Storage
 							{
 								try
 								{
-									await TickGcForNamespaceAsync(state.Namespaces[namespaceId], gcState.LastImportBlobInfoId, utcNow, cancellationToken);
+									await TickGcForNamespaceAsync(state.Namespaces[namespaceId], gcState.LastImportBlobInfoId, utcNow, storageConfig.EnableGc, cancellationToken);
 								}
 								catch (Exception ex)
 								{
@@ -1026,7 +1048,7 @@ namespace Horde.Server.Storage
 			}
 		}
 
-		async Task TickGcForNamespaceAsync(NamespaceInfo namespaceInfo, ObjectId lastImportBlobInfoId, DateTime utcNow, CancellationToken cancellationToken)
+		async Task TickGcForNamespaceAsync(NamespaceInfo namespaceInfo, ObjectId lastImportBlobInfoId, DateTime utcNow, bool deleteObjects, CancellationToken cancellationToken)
 		{
 			using IStorageClient client = this.CreateClient(namespaceInfo.Id);
 
@@ -1037,7 +1059,7 @@ namespace Horde.Server.Storage
 			double score = GetGcTimestamp(utcNow);
 
 			RedisSortedSetKey<RedisValue> checkSet = GetGcCheckSet(namespaceInfo.Id);
-			while (_globalConfig.CurrentValue.Storage.EnableGC)
+			while (_globalConfig.CurrentValue.Storage.EnableGc || _globalConfig.CurrentValue.Storage.EnableGcVerification)
 			{
 				long length = await _redisService.GetDatabase().SortedSetLengthAsync(checkSet);
 				int batchSize = (int)Math.Min(length, 1024);
@@ -1060,7 +1082,16 @@ namespace Horde.Server.Storage
 
 					if (blobInfoId < lastImportBlobInfoId && !await IsBlobReferencedAsync(blobInfoId, cancellationToken))
 					{
-						BlobInfo? info = await _blobCollection.FindOneAndDeleteAsync(x => x.Id == blobInfoId, cancellationToken: cancellationToken);
+						BlobInfo? info;
+						if (deleteObjects)
+						{
+							info = await _blobCollection.FindOneAndDeleteAsync(x => x.Id == blobInfoId, cancellationToken: cancellationToken);
+						}
+						else
+						{
+							info = await _blobCollection.FindOneAndUpdateAsync(x => x.Id == blobInfoId, Builders<BlobInfo>.Update.Set(x => x.GcVersion, CurrentGcVersion), cancellationToken: cancellationToken);
+						}
+
 						if (info != null)
 						{
 							if (info.Imports != null)
@@ -1072,7 +1103,12 @@ namespace Horde.Server.Storage
 
 							ObjectKey objectKey = GetObjectKey(new BlobLocator(info.Path));
 							_logger.LogDebug("Deleting {NamespaceId} blob {BlobId}, key: {ObjectKey} ({ImportCount} imports)", namespaceInfo.Id, blobInfoId, objectKey, info.Imports?.Count ?? 0);
-							await namespaceInfo.Store.DeleteAsync(objectKey, cancellationToken);
+
+							if (deleteObjects)
+							{
+								await namespaceInfo.Store.DeleteAsync(objectKey, cancellationToken);
+							}
+
 							numItemsRemoved++;
 						}
 					}
