Horde: Another stab at fixing timeline views for retried steps
#rnx

diff --git a/Engine/Source/Programs/Horde/HordeDashboard/src/components/jobDetailsV2/JobDetailTimeline.tsx b/Engine/Source/Programs/Horde/HordeDashboard/src/components/jobDetailsV2/JobDetailTimeline.tsx
--- a/Engine/Source/Programs/Horde/HordeDashboard/src/components/jobDetailsV2/JobDetailTimeline.tsx
+++ b/Engine/Source/Programs/Horde/HordeDashboard/src/components/jobDetailsV2/JobDetailTimeline.tsx
@@ -142,7 +142,7 @@ class TimelineDataView extends JobDataView {
 
       // patch batches with latest finished step time, this renders the timeline more dynamically than waiting for 
       // complete batches, note: should present batch stopping time, which can be significant
-      
+
       const nbatches = job.batches?.map(b => {
 
          const newBatch = {
@@ -158,7 +158,7 @@ class TimelineDataView extends JobDataView {
                   break;
                }
             }
-            
+
             if (lastStep) {
                newBatch.state = JobStepBatchState.Complete
                newBatch.finishTime = lastStep.finishTime;
@@ -177,7 +177,7 @@ class TimelineDataView extends JobDataView {
                b.steps = b.steps?.filter(s => {
                   if (!s.startTime || !s.finishTime) {
                      return false;
-                  }                  
+                  }
                   return label.steps.indexOf(s?.id ?? "") !== -1;
                });
             });
@@ -216,7 +216,7 @@ class TimelineDataView extends JobDataView {
          let utcBatchReady = new Date(b.readyTime as string)
          const utcStepStart = new Date(b.steps[0].startTime as string)
 
-         if (filterStep && jobDetails.getStepRetryNumber(filterStep.id) > 0) {
+         if (!!b.steps.find( s => jobDetails.getStepRetryNumber(s.id) > 0)) {
             utcBatchReady = utcBatchStart;
          }
 
