Horde: Improve top axis for trendsview
#rnx

diff --git a/Engine/Source/Programs/Horde/HordeDashboard/src/components/jobDetailsV2/JobDetailStepTrendsV2.tsx b/Engine/Source/Programs/Horde/HordeDashboard/src/components/jobDetailsV2/JobDetailStepTrendsV2.tsx
--- a/Engine/Source/Programs/Horde/HordeDashboard/src/components/jobDetailsV2/JobDetailStepTrendsV2.tsx
+++ b/Engine/Source/Programs/Horde/HordeDashboard/src/components/jobDetailsV2/JobDetailStepTrendsV2.tsx
@@ -451,10 +451,40 @@ class StepTrendsRenderer {
 
       const xAxis = (g: SelectionType) => {
 
+         const dateMin = dataView.minTime!
+         const dateMax = dataView.maxTime!
+
+         let ticks: number[] = [];
+         for (const date of d3.timeDays(dateMin, dateMax, 1).reverse()) {
+            ticks.push(date.getTime() / 1000);
+         }
+
+         if (ticks.length > 14) {
+            let nticks = [...ticks];
+            // remove first and last, will be readded 
+            const first = nticks.shift()!;
+            const last = nticks.pop()!;
+
+            const n = Math.floor(nticks.length / 12);
+
+            const rticks: number[] = [];
+            for (let i = 0; i < nticks.length; i = i + n) {
+               rticks.push(nticks[i]);
+            }
+
+            rticks.unshift(first);
+            rticks.push(last);
+            ticks = rticks;
+
+         }
+
+
          g.attr("transform", `translate(0,24)`)
             .style("font-family", "Horde Open Sans SemiBold")
             .style("font-size", "12px")
             .call(d3.axisTop(x)
+               .tickValues(ticks)
+               //.ticks(d3.timeDays(dateMin, dateMax))
                .tickFormat(d => {
                   return getHumanTime(new Date((d as number) * 1000));
                })
