Horde: Add a retry policy to the Epic telemetry sink, to stop 502 errors causing server log entries.

diff --git a/Engine/Source/Programs/Horde/Horde.Server/Telemetry/Sinks/EpicTelemetrySink.cs b/Engine/Source/Programs/Horde/Horde.Server/Telemetry/Sinks/EpicTelemetrySink.cs
--- a/Engine/Source/Programs/Horde/Horde.Server/Telemetry/Sinks/EpicTelemetrySink.cs
+++ b/Engine/Source/Programs/Horde/Horde.Server/Telemetry/Sinks/EpicTelemetrySink.cs
@@ -12,7 +12,10 @@ using System.Threading.Tasks;
 using EpicGames.Core;
 using EpicGames.Horde.Telemetry;
 using Microsoft.AspNetCore.WebUtilities;
+using Microsoft.Extensions.Http;
 using Microsoft.Extensions.Logging;
+using Polly;
+using Polly.Extensions.Http;
 
 namespace Horde.Server.Telemetry.Sinks
 {
@@ -63,7 +66,7 @@ namespace Horde.Server.Telemetry.Sinks
 
 		readonly object _lockObject = new object();
 
-		readonly IHttpClientFactory _httpClientFactory;
+		readonly HttpClient _httpClient;
 		readonly Uri? _baseUrl;
 		readonly JsonSerializerOptions _jsonOptions;
 
@@ -132,9 +135,9 @@ namespace Horde.Server.Telemetry.Sinks
 		/// <summary>
 		/// Constructor
 		/// </summary>
-		public EpicTelemetrySink(IEpicTelemetrySinkConfig config, IHttpClientFactory httpClientFactory, ILogger<EpicTelemetrySink> logger)
+		public EpicTelemetrySink(IEpicTelemetrySinkConfig config, ILogger<EpicTelemetrySink> logger)
 		{
-			_httpClientFactory = httpClientFactory;
+			_httpClient = CreateHttpClient();
 
 			_jsonOptions = new JsonSerializerOptions();
 			Startup.ConfigureJsonSerializer(_jsonOptions);
@@ -145,6 +148,17 @@ namespace Horde.Server.Telemetry.Sinks
 			_baseUrl = config.Url;
 		}
 
+		static HttpClient CreateHttpClient()
+		{
+			HttpMessageHandler handler = new SocketsHttpHandler();
+
+			TimeSpan[] retryTimes = new[] { TimeSpan.FromSeconds(1.0), TimeSpan.FromSeconds(5.0), TimeSpan.FromSeconds(10.0) };
+			IAsyncPolicy<HttpResponseMessage> handleTransientErrorPolicy = Policy<HttpResponseMessage>.Handle<HttpRequestException>().OrTransientHttpStatusCode().WaitAndRetryAsync(retryTimes);
+			handler = new PolicyHttpMessageHandler(handleTransientErrorPolicy) { InnerHandler = handler };
+
+			return new HttpClient(handler);
+		}
+
 		/// <inheritdoc/>
 		public ValueTask DisposeAsync()
 		{
@@ -152,6 +166,7 @@ namespace Horde.Server.Telemetry.Sinks
 			{
 				writer.Dispose();
 			}
+			_httpClient.Dispose();
 			return new ValueTask();
 		}
 
@@ -195,7 +210,6 @@ namespace Horde.Server.Telemetry.Sinks
 			// Post the event data
 			foreach ((Uri uri, byte[] packet) in packets)
 			{
-				HttpClient httpClient = _httpClientFactory.CreateClient(HttpClientName);
 				using (HttpRequestMessage request = new HttpRequestMessage())
 				{
 					request.RequestUri = uri;
@@ -204,7 +218,7 @@ namespace Horde.Server.Telemetry.Sinks
 					request.Content = new ByteArrayContent(packet);
 					request.Content.Headers.ContentType = new MediaTypeHeaderValue("application/json");
 
-					using (HttpResponseMessage response = await httpClient.SendAsync(request, cancellationToken))
+					using (HttpResponseMessage response = await _httpClient.SendAsync(request, cancellationToken))
 					{
 						if (response.IsSuccessStatusCode)
 						{
diff --git a/Engine/Source/Programs/Horde/Horde.Server/Telemetry/TelemetryManager.cs b/Engine/Source/Programs/Horde/Horde.Server/Telemetry/TelemetryManager.cs
--- a/Engine/Source/Programs/Horde/Horde.Server/Telemetry/TelemetryManager.cs
+++ b/Engine/Source/Programs/Horde/Horde.Server/Telemetry/TelemetryManager.cs
@@ -44,7 +44,7 @@ namespace Horde.Server.Telemetry
 				switch (config)
 				{
 					case EpicTelemetryConfig epicConfig:
-						_telemetrySinks.Add(new EpicTelemetrySink(epicConfig, httpClientFactory, loggerFactory.CreateLogger<EpicTelemetrySink>()));
+						_telemetrySinks.Add(new EpicTelemetrySink(epicConfig, loggerFactory.CreateLogger<EpicTelemetrySink>()));
 						break;
 					case ClickHouseTelemetryConfig chConfig:
 						_telemetrySinks.Add(new ClickHouseTelemetrySink(chConfig, httpClientFactory, loggerFactory.CreateLogger<ClickHouseTelemetrySink>()));
