Horde: Fix some static analysis warnings.

#rnx

diff --git a/Engine/Source/Programs/Horde/Horde.Server/Storage/StorageService.cs b/Engine/Source/Programs/Horde/Horde.Server/Storage/StorageService.cs
--- a/Engine/Source/Programs/Horde/Horde.Server/Storage/StorageService.cs
+++ b/Engine/Source/Programs/Horde/Horde.Server/Storage/StorageService.cs
@@ -70,14 +70,14 @@ namespace Horde.Server.Storage
 
 			public async Task<Stream> OpenBlobAsync(BlobLocator locator, int offset, int? length, CancellationToken cancellationToken = default)
 			{
-				await _outer.CheckBlobExists(NamespaceId, locator);
+				await _outer.CheckBlobExistsAsync(NamespaceId, locator, cancellationToken);
 				return await _store.OpenAsync(GetObjectKey(locator), offset, length, cancellationToken);
 			}
 
 			/// <inheritdoc/>
 			public async Task<IReadOnlyMemoryOwner<byte>> ReadBlobAsync(BlobLocator locator, int offset, int? length, CancellationToken cancellationToken = default)
 			{
-				await _outer.CheckBlobExists(NamespaceId, locator);
+				await _outer.CheckBlobExistsAsync(NamespaceId, locator, cancellationToken);
 				return await _store.ReadAsync(GetObjectKey(locator), offset, length, cancellationToken);
 			}
 
@@ -95,7 +95,7 @@ namespace Horde.Server.Storage
 			/// <inheritdoc/>
 			public async ValueTask<Uri?> TryGetBlobReadRedirectAsync(BlobLocator locator, CancellationToken cancellationToken = default)
 			{
-				await _outer.CheckBlobExists(NamespaceId, locator);
+				await _outer.CheckBlobExistsAsync(NamespaceId, locator, cancellationToken);
 				return await _store.TryGetReadRedirectAsync(GetObjectKey(locator), cancellationToken);
 			}
 
@@ -586,11 +586,11 @@ namespace Horde.Server.Storage
 			await _blobCollection.InsertOneAsync(blobInfo, new InsertOneOptions { }, cancellationToken);
 		}
 
-		async ValueTask CheckBlobExists(NamespaceId namespaceId, BlobLocator locator)
+		async ValueTask CheckBlobExistsAsync(NamespaceId namespaceId, BlobLocator locator, CancellationToken cancellationToken)
 		{
 			if (_globalConfig.CurrentValue.Storage.EnableGcVerification)
 			{
-				if (await _blobCollection.Find(x => x.Locator == locator && x.GcVersion >= CurrentGcVersion).AnyAsync())
+				if (await _blobCollection.Find(x => x.NamespaceId == namespaceId && x.Locator == locator && x.GcVersion >= CurrentGcVersion).AnyAsync(cancellationToken))
 				{
 					_logger.LogWarning("Blob {Locator} accessed after being garbage collected", locator);
 				}
diff --git a/Engine/Source/Programs/Horde/Horde.Server/Telemetry/Sinks/EpicTelemetrySink.cs b/Engine/Source/Programs/Horde/Horde.Server/Telemetry/Sinks/EpicTelemetrySink.cs
--- a/Engine/Source/Programs/Horde/Horde.Server/Telemetry/Sinks/EpicTelemetrySink.cs
+++ b/Engine/Source/Programs/Horde/Horde.Server/Telemetry/Sinks/EpicTelemetrySink.cs
@@ -152,11 +152,13 @@ namespace Horde.Server.Telemetry.Sinks
 		{
 			HttpMessageHandler handler = new SocketsHttpHandler();
 
+#pragma warning disable CA2000
 			TimeSpan[] retryTimes = new[] { TimeSpan.FromSeconds(1.0), TimeSpan.FromSeconds(5.0), TimeSpan.FromSeconds(10.0) };
 			IAsyncPolicy<HttpResponseMessage> handleTransientErrorPolicy = Policy<HttpResponseMessage>.Handle<HttpRequestException>().OrTransientHttpStatusCode().WaitAndRetryAsync(retryTimes);
 			handler = new PolicyHttpMessageHandler(handleTransientErrorPolicy) { InnerHandler = handler };
+#pragma warning restore CA2000
 
-			return new HttpClient(handler);
+			return new HttpClient(handler, true);
 		}
 
 		/// <inheritdoc/>
