Horde: Updating schema documentation

#rnx

diff --git a/Engine/Source/Programs/Horde/Docs/Config/Schema/Globals.md b/Engine/Source/Programs/Horde/Docs/Config/Schema/Globals.md
--- a/Engine/Source/Programs/Horde/Docs/Config/Schema/Globals.md
+++ b/Engine/Source/Programs/Horde/Docs/Config/Schema/Globals.md
@@ -401,7 +401,8 @@ Configuration for storage
 
 Name | Description
 ---- | -----------
-`enableGC` | `boolean`<br>Whether to enable garbage collection
+`enableGc` | `boolean`<br>Whether to enable garbage collection
+`enableGcVerification` | `boolean`<br>Whether to enable garbage collection in verification mode (nothing deleted, just logging on access to deleted blobs)
 `backends` | [`BackendConfig`](#backendconfig)`[]`<br>List of storage backends
 `namespaces` | [`NamespaceConfig`](#namespaceconfig)`[]`<br>List of namespaces for storage
 
