Horde: Use self-contained job driver instances if the agent app is self contained.

#rnx

diff --git a/Engine/Source/Programs/Horde/Horde.Agent/Leases/Handlers/ConformHandler.cs b/Engine/Source/Programs/Horde/Horde.Agent/Leases/Handlers/ConformHandler.cs
--- a/Engine/Source/Programs/Horde/Horde.Agent/Leases/Handlers/ConformHandler.cs
+++ b/Engine/Source/Programs/Horde/Horde.Agent/Leases/Handlers/ConformHandler.cs
@@ -39,7 +39,7 @@ namespace Horde.Agent.Leases.Handlers
 			environment[HordeHttpClient.HordeTokenEnvVarName] = await session.HordeClient.GetAccessTokenAsync(false, cancellationToken) ?? String.Empty;
 			environment["UE_LOG_JSON_TO_STDOUT"] = "1";
 
-			int exitCode = await RunDotNetProcessAsync(driverAssembly, arguments, environment, false, serverLogger, cancellationToken);
+			int exitCode = await RunDotNetProcessAsync(driverAssembly, arguments, environment, AgentApp.IsSelfContained, serverLogger, cancellationToken);
 			serverLogger.LogInformation("Driver finished with exit code {ExitCode}", exitCode);
 
 			return (exitCode == 0) ? LeaseResult.Success : LeaseResult.Failed;
diff --git a/Engine/Source/Programs/Horde/Horde.Agent/Leases/Handlers/JobHandler.cs b/Engine/Source/Programs/Horde/Horde.Agent/Leases/Handlers/JobHandler.cs
--- a/Engine/Source/Programs/Horde/Horde.Agent/Leases/Handlers/JobHandler.cs
+++ b/Engine/Source/Programs/Horde/Horde.Agent/Leases/Handlers/JobHandler.cs
@@ -48,7 +48,7 @@ namespace Horde.Agent.Leases.Handlers
 				environment[HordeHttpClient.HordeTokenEnvVarName] = executeTask.Token;
 				environment["UE_LOG_JSON_TO_STDOUT"] = "1";
 
-				int exitCode = await RunDotNetProcessAsync(driverAssembly, arguments, environment, false, logger, cancellationToken);
+				int exitCode = await RunDotNetProcessAsync(driverAssembly, arguments, environment, AgentApp.IsSelfContained, logger, cancellationToken);
 				logger.LogInformation("Driver finished with exit code {ExitCode}", exitCode);
 
 				return (exitCode == 0) ? LeaseResult.Success : LeaseResult.Failed;
