Horde: Update release notes
#rnx

diff --git a/Engine/Source/Programs/Horde/Docs/ReleaseNotes.md b/Engine/Source/Programs/Horde/Docs/ReleaseNotes.md
--- a/Engine/Source/Programs/Horde/Docs/ReleaseNotes.md
+++ b/Engine/Source/Programs/Horde/Docs/ReleaseNotes.md
@@ -1,5 +1,23 @@
 # Release Notes
 
+## 2024-06-06
+
+* Add a retry policy to the Epic telemetry sink, to stop 502 errors causing server log entries. (34141477)
+* Add a GC verification mode, which sets a flag in the DB rather than deleting an object from the store. (34135143)
+* Fix issue where files will not be uploaded to temp storage if they were previously tagged as outputs in another step. Now allows output files to be produced by multiple steps, as long as their attributes match. (34132683)
+* Adding support for optional cancellation reasons for jobs and steps (34131480)
+* Forward all claims for access tokens minted for Horde accounts The role claim was missing which is used for resolving many permissions. (34119849)
+* Log reasons when a lease is terminated due to the cancellation token being set. (34111700)
+* Add a link to the producing job to .uartifact files, as well as other metadata and keys on the artifact object. (34095063)
+* Allow user configured backends and namespaces to override the defaults in defaults.global.json. (34092465)
+* Print a message to the batch log whenever a lease is cancelled or fails due to an unhandled exception. (34084201)
+* Add a setting ("ReportWarnings") to exclude issues which are only warnings from summary reports. (34078316)
+* Fix P4 server health not being updated when HTTP health check endpoint returns degraded. (34077493)
+* Re-encode json log events which are downgraded from warning to information level. (34066649)
+* Add more trace attributes to compute resource allocation (34064620)
+* Record endpoint address in MongoDB command tracer, jelps differentiate between primary and secondary use (read-only ops) (34063808)
+* Use cached agent data for JobTaskSource, this ticks every 5 seconds refreshing agent and pool data. The cached agent data is refreshed at the same rate. (34062377)
+
 ## 2024-05-31
 
 * Use cached agent data for fleet and pool size handling (34029489)
