Horde: Fix browse not being disabled on invalid artifact
#rnx

diff --git a/Engine/Source/Programs/Horde/HordeDashboard/src/components/jobDetailsV2/JobDetailArtifacts.tsx b/Engine/Source/Programs/Horde/HordeDashboard/src/components/jobDetailsV2/JobDetailArtifacts.tsx
--- a/Engine/Source/Programs/Horde/HordeDashboard/src/components/jobDetailsV2/JobDetailArtifacts.tsx
+++ b/Engine/Source/Programs/Horde/HordeDashboard/src/components/jobDetailsV2/JobDetailArtifacts.tsx
@@ -220,7 +220,7 @@ export const JobArtifactsPanel: React.FC<{ jobDetails: JobDetailsV2 }> = observe
          return <Stack horizontalAlign="end" verticalAlign="center" verticalFill={true} style={{ paddingRight: 8 }}>
             <Stack horizontal tokens={{ childrenGap: 18 }}>
                <PrimaryButton text="Browse"
-                  disabled={!stepFinished}
+                  disabled={!stepFinished || !item.id}
                   style={{ fontFamily: "Horde Open Sans SemiBold" }}
                   onClick={() => {
                      setSelected(item)
