Horde: Disable config updates when running locally in Redis read-only mode, preventing an exception on startup.

diff --git a/Engine/Source/Programs/Horde/Horde.Server/Configuration/ConfigService.cs b/Engine/Source/Programs/Horde/Horde.Server/Configuration/ConfigService.cs
--- a/Engine/Source/Programs/Horde/Horde.Server/Configuration/ConfigService.cs
+++ b/Engine/Source/Programs/Horde/Horde.Server/Configuration/ConfigService.cs
@@ -354,8 +354,15 @@ namespace Horde.Server.Configuration
 		{
 			try
 			{
+				bool forceConfigUpdate = _serverSettings.ForceConfigUpdateOnStartup;
+				if (forceConfigUpdate && _redisService.ReadOnlyMode)
+				{
+					_logger.LogInformation("Ignoring flag to force config update on startup due to Redis configured in read-only mode.");
+					forceConfigUpdate = false;
+				}
+
 				ReadOnlyMemory<byte> data = ReadOnlyMemory<byte>.Empty;
-				if (!_serverSettings.ForceConfigUpdateOnStartup)
+				if (!forceConfigUpdate)
 				{
 					data = await ReadSnapshotDataAsync();
 				}
