Horde: Fix compiler warning

diff --git a/Engine/Source/Programs/Horde/Horde.Agent/Leases/LeaseHandlerFactory.cs b/Engine/Source/Programs/Horde/Horde.Agent/Leases/LeaseHandlerFactory.cs
--- a/Engine/Source/Programs/Horde/Horde.Agent/Leases/LeaseHandlerFactory.cs
+++ b/Engine/Source/Programs/Horde/Horde.Agent/Leases/LeaseHandlerFactory.cs
@@ -38,6 +38,6 @@ namespace Horde.Agent.Leases
 		/// Creates a new lease handler for the given lease definition
 		/// </summary>
 		/// <param name="lease">The lease to create a handler for</param>
-		public override abstract LeaseHandler<T> CreateHandler(RpcLease lease);
+		public abstract override LeaseHandler<T> CreateHandler(RpcLease lease);
 	}
 }
