Horde: Set user ID in compute assignment request track usage by user

diff --git a/Engine/Source/Programs/Horde/Horde.Server/Compute/ComputeController.cs b/Engine/Source/Programs/Horde/Horde.Server/Compute/ComputeController.cs
--- a/Engine/Source/Programs/Horde/Horde.Server/Compute/ComputeController.cs
+++ b/Engine/Source/Programs/Horde/Horde.Server/Compute/ComputeController.cs
@@ -68,7 +68,8 @@ namespace Horde.Server.Compute
 				ConnectionMode = request.Connection?.ModePreference,
 				RequesterPublicIp = request.Connection?.ClientPublicIp,
 				UsePublicIp = request.Connection?.PreferPublicIp,
-				Encryption = ComputeService.ConvertEncryptionToProto(request.Connection?.Encryption)
+				Encryption = ComputeService.ConvertEncryptionToProto(request.Connection?.Encryption),
+				UserId = User.GetUserId(),
 			};
 
 			ComputeResource? computeResource;
