Horde: Check agent's compute cluster membership on assignment

Previously it was only checked by the task source.

diff --git a/Engine/Source/Programs/Horde/Horde.Server.Tests/Compute/ComputeServiceTest.cs b/Engine/Source/Programs/Horde/Horde.Server.Tests/Compute/ComputeServiceTest.cs
--- a/Engine/Source/Programs/Horde/Horde.Server.Tests/Compute/ComputeServiceTest.cs
+++ b/Engine/Source/Programs/Horde/Horde.Server.Tests/Compute/ComputeServiceTest.cs
@@ -245,6 +245,42 @@ namespace Horde.Server.Tests.Compute
 			
 			Assert.IsNull(await ComputeService.TryAllocateResourceAsync(CreateParams("does-not-exist"), CancellationToken.None));
 		}
+		
+		[TestMethod]
+		public async Task Assignment_ComputeClusterMemberConditionAsync()
+		{
+			ClusterId cluster1 = new("cluster1");
+			ClusterId cluster2 = new("cluster2");
+			PoolId poolFoo = new ("foo");
+			PoolId poolBar = new ("bar");
+
+			GlobalConfig.CurrentValue.Compute = new List<ComputeClusterConfig>
+			{
+				new() { Id = cluster1, Condition = $"pool == '{poolFoo.ToString()}'" },
+				new() { Id = cluster2, Condition = $"pool == '{poolBar.ToString()}'" }
+			};
+			GlobalConfig.CurrentValue.PostLoad(new ServerSettings());
+
+			List<string> props = ["ComputeIp=11.0.0.1", "ComputePort=5000"];
+			IAgent agent1 = await CreateAgentAsync(new PoolId("foo"), properties: props);
+			IAgent agent2 = await CreateAgentAsync(new PoolId("bar"), properties: props);
+			
+			{
+				AllocateResourceParams arpFoo = new(cluster1, ComputeProtocol.Latest, new Requirements { Pool = poolFoo.ToString() });
+				Assert.AreEqual(agent1.Id, (await ComputeService.TryAllocateResourceAsync(arpFoo, CancellationToken.None))!.AgentId);
+				
+				AllocateResourceParams arpBar = new(cluster1, ComputeProtocol.Latest, new Requirements { Pool = poolBar.ToString() });
+				Assert.IsNull(await ComputeService.TryAllocateResourceAsync(arpBar, CancellationToken.None));
+			}
+			
+			{
+				AllocateResourceParams arpFoo = new(cluster2, ComputeProtocol.Latest, new Requirements { Pool = poolFoo.ToString() });
+				Assert.IsNull(await ComputeService.TryAllocateResourceAsync(arpFoo, CancellationToken.None));
+				
+				AllocateResourceParams arpBar = new(cluster2, ComputeProtocol.Latest, new Requirements { Pool = poolBar.ToString() });
+				Assert.AreEqual(agent2.Id, (await ComputeService.TryAllocateResourceAsync(arpBar, CancellationToken.None))!.AgentId);
+			}
+		}
 
 		[TestMethod]
 		public async Task Connection_Direct_IpConnection_Async()
@@ -304,26 +340,27 @@ namespace Horde.Server.Tests.Compute
 			Assert.AreEqual(new ComputeResourcePort(10004, 13000), cr.Ports["myOtherPort"]);
 		}
 
-		private async Task<ComputeService> CreateComputeServiceAsync(string? tunnelAddress, string pool)
+		private async Task<ComputeService> CreateComputeServiceAsync(string? tunnelAddress, string pool, ComputeClusterConfig ccc)
 		{
 			ServerSettings ss = new() { ComputeTunnelAddress = tunnelAddress };
 			ComputeService cs = new(AgentCollection, LogCollection, AgentService, AgentRelayService, GetRedisServiceSingleton(),
 				new TestOptionsMonitor<ServerSettings>(ss), GlobalConfig, Clock, Tracer, Meter,
 				NullLogger<ComputeService>.Instance);
-			List<string> props = new() { "ComputeIp=11.0.0.1", "ComputePort=5000" };
-			await CreateAgentAsync(new PoolId(pool), properties: props);
+			await CreateAgentAsync(new PoolId(pool), properties: ["ComputeIp=11.0.0.1", "ComputePort=5000"]);
+			GlobalConfig.CurrentValue.Compute = [new ComputeClusterConfig { Id = _cluster1 }];
+			GlobalConfig.CurrentValue.PostLoad(ss);
 			return cs;
 		}
 
 		private async Task<ComputeResource?> AllocateAsync(
-			ConnectionMode connectionMode,
+			ConnectionMode connectionMode = ConnectionMode.Direct,
 			Dictionary<string, int>? ports = null,
 			bool usePublicIp = false,
 			string[]? relayIps = null,
 			string? tunnelAddress = null,
 			string pool = "myPool")
 		{
-			await using ComputeService cs = await CreateComputeServiceAsync(tunnelAddress, pool);
+			await using ComputeService cs = await CreateComputeServiceAsync(tunnelAddress, pool, new ComputeClusterConfig { Id = _cluster1});
 			AllocateResourceParams arp = new(_cluster1, ComputeProtocol.Latest, new Requirements() { Pool = pool})
 			{
 				ConnectionMode = connectionMode,
diff --git a/Engine/Source/Programs/Horde/Horde.Server.Tests/TestSetup.cs b/Engine/Source/Programs/Horde/Horde.Server.Tests/TestSetup.cs
--- a/Engine/Source/Programs/Horde/Horde.Server.Tests/TestSetup.cs
+++ b/Engine/Source/Programs/Horde/Horde.Server.Tests/TestSetup.cs
@@ -387,7 +387,7 @@ namespace Horde.Server.Tests
 		/// <param name="ephemeral">Whether the agent is ephemeral</param>
 		/// <returns>A new agent</returns>
 		public async Task<IAgent> CreateAgentAsync(
-			PoolId poolId,
+			PoolId? poolId,
 			bool enabled = true,
 			bool requestShutdown = false,
 			List<string>? properties = null,
@@ -413,7 +413,7 @@ namespace Horde.Server.Tests
 			IAgent? agent = await AgentService.CreateAgentAsync("TestAgent" + s_agentIdCounter++, ephemeral, "");
 			Assert.IsNotNull(agent);
 
-			agent = await agent.TryUpdateAsync(new UpdateAgentOptions { Enabled = enabled, ExplicitPools = new List<PoolId> { poolId } });
+			agent = await agent.TryUpdateAsync(new UpdateAgentOptions { Enabled = enabled, ExplicitPools = poolId != null ? [poolId.Value] : [] });
 			Assert.IsNotNull(agent);
 
 			agent = await AgentService.CreateSessionAsync(agent, AgentStatus.Ok, tempProps, resources, null);
diff --git a/Engine/Source/Programs/Horde/Horde.Server/Compute/ComputeService.cs b/Engine/Source/Programs/Horde/Horde.Server/Compute/ComputeService.cs
--- a/Engine/Source/Programs/Horde/Horde.Server/Compute/ComputeService.cs
+++ b/Engine/Source/Programs/Horde/Horde.Server/Compute/ComputeService.cs
@@ -376,11 +376,16 @@ namespace Horde.Server.Compute
 				span.SetAttribute($"req.res.{name}.min", resReq.Min);
 				span.SetAttribute($"req.res.{name}.max", resReq.Max);
 			}
+			
+			if (!_globalConfig.CurrentValue.TryGetComputeCluster(arp.ClusterId, out ComputeClusterConfig? clusterConfig))
+			{
+				throw new ArgumentException($"Cluster '{arp.ClusterId}' not found");
+			}
 
 			byte[] certificate;
 			using (TelemetrySpan _ = _tracer.StartActiveSpan("Generating certificate"))
 			{
-				certificate = GenerateCert(arp.Encryption); // A no-op if certificate is not required	
+				certificate = GenerateCert(arp.Encryption); // A no-op if certificate is not required
 			}
 
 			List<Condition> conditions = [];
@@ -400,7 +405,8 @@ namespace Horde.Server.Compute
 				{
 					Dictionary<string, int> assignedResources = new Dictionary<string, int>();
 
-					bool match = agent.MeetsRequirements(arp.Requirements, assignedResources, conditions);
+					bool isMemberOfCluster = clusterConfig.Condition == null || agent.SatisfiesCondition(clusterConfig.Condition);
+					bool match = isMemberOfCluster && agent.MeetsRequirements(arp.Requirements, assignedResources, conditions);
 					if (match)
 					{
 						using TelemetrySpan matchSpan = _tracer.StartActiveSpan("Found match");
