Horde: Add an endpoint for downloading blocks from an unsync manifest.

#rnx

diff --git a/Engine/Source/Programs/Horde/Horde.Server.Tests/Artifacts/UnsyncCacheTests.cs b/Engine/Source/Programs/Horde/Horde.Server.Tests/Artifacts/UnsyncCacheTests.cs
new file mode 100644
--- /dev/null
+++ b/Engine/Source/Programs/Horde/Horde.Server.Tests/Artifacts/UnsyncCacheTests.cs
@@ -0,0 +1,76 @@
+﻿// Copyright Epic Games, Inc. All Rights Reserved.
+
+using System;
+using System.Collections.Generic;
+using System.Linq;
+using System.Text;
+using System.Threading;
+using System.Threading.Tasks;
+using EpicGames.Core;
+using EpicGames.Horde.Storage;
+using EpicGames.Horde.Storage.Bundles;
+using EpicGames.Horde.Storage.Nodes;
+using Horde.Server.Artifacts;
+using Microsoft.Extensions.Logging.Abstractions;
+using Microsoft.VisualStudio.TestTools.UnitTesting;
+using Moq;
+
+namespace Horde.Server.Tests.Artifacts
+{
+	[TestClass]
+	public class UnsyncCacheTests
+	{
+		[TestMethod]
+		public async Task TestBlobsAsync()
+		{
+			NamespaceId ns = new NamespaceId("test");
+			byte[] source = Enumerable.Range(0, 1024 * 1024).Select(x => (byte)(x % 257)).ToArray();
+
+			using BundleStorageClient storageClient = BundleStorageClient.CreateInMemory(NullLogger.Instance);
+
+			Dictionary<IoHash, ReadOnlyMemory<byte>> chunks = new Dictionary<IoHash, ReadOnlyMemory<byte>>();
+
+			IBlobRef<DirectoryNode> directoryRef;
+			await using (IBlobWriter writer = storageClient.CreateBlobWriter())
+			{
+				using ChunkedDataWriter chunkedWriter = new ChunkedDataWriter(writer, new ChunkingOptions());
+				await chunkedWriter.AppendAsync(source, CancellationToken.None);
+
+				ChunkedData data = await chunkedWriter.FlushAsync();
+
+				DirectoryNode directory = new DirectoryNode();
+				directory.AddFile("hello.txt", FileEntryFlags.None, source.Length, data);
+
+				directoryRef = await writer.WriteBlobAsync(directory);
+			}
+
+			RefName refName = new RefName("test");
+			await storageClient.WriteRefAsync(refName, directoryRef);
+
+			Mock<IStorageClientFactory> factory = new Mock<IStorageClientFactory>();
+			factory.Setup(x => x.TryCreateClient(ns)).Returns(storageClient);
+
+			Mock<IArtifact> artifact = new Mock<IArtifact>();
+			artifact.SetupGet(x => x.NamespaceId).Returns(ns);
+			artifact.SetupGet(x => x.RefName).Returns(refName);
+
+			using UnsyncCache cache = new UnsyncCache(factory.Object);
+
+			UnsyncManifest? manifest = await cache.GetManifestAsync(artifact.Object);
+			Assert.IsNotNull(manifest);
+
+			UnsyncFile file = manifest.Files[0];
+			Assert.AreEqual(file.Name.ToString(), "hello.txt");
+
+			int offset = 0;
+			foreach (UnsyncBlock block in file.Blocks)
+			{
+				using BlobData? blobData = await cache.ReadBlobAsync(artifact.Object, block.Blob.Hash);
+				Assert.IsNotNull(blobData);
+
+				Assert.IsTrue(blobData.Data.Span.SequenceEqual(source.AsSpan(offset, blobData.Data.Length)));
+				offset += (int)block.Length;
+			}
+		}
+	}
+}
diff --git a/Engine/Source/Programs/Horde/Horde.Server/Artifacts/ArtifactsController.cs b/Engine/Source/Programs/Horde/Horde.Server/Artifacts/ArtifactsController.cs
--- a/Engine/Source/Programs/Horde/Horde.Server/Artifacts/ArtifactsController.cs
+++ b/Engine/Source/Programs/Horde/Horde.Server/Artifacts/ArtifactsController.cs
@@ -48,19 +48,21 @@ namespace Horde.Server.Artifacts
 		readonly ILeaseCollection _leaseCollection;
 		readonly IJobCollection _jobCollection;
 		readonly AclService _aclService;
+		readonly UnsyncCache _unsyncCache;
 		readonly GlobalConfig _globalConfig;
 		readonly ILogger _logger;
 
 		/// <summary>
 		/// Constructor
 		/// </summary>
-		public ArtifactsController(IArtifactCollection artifactCollection, StorageService storageService, ILeaseCollection leaseCollection, IJobCollection jobCollection, AclService aclService, IOptionsSnapshot<GlobalConfig> globalConfig, ILogger<ArtifactsController> logger)
+		public ArtifactsController(IArtifactCollection artifactCollection, StorageService storageService, ILeaseCollection leaseCollection, IJobCollection jobCollection, AclService aclService, UnsyncCache unsyncCache, IOptionsSnapshot<GlobalConfig> globalConfig, ILogger<ArtifactsController> logger)
 		{
 			_artifactCollection = artifactCollection;
 			_storageService = storageService;
 			_leaseCollection = leaseCollection;
 			_jobCollection = jobCollection;
 			_aclService = aclService;
+			_unsyncCache = unsyncCache;
 			_globalConfig = globalConfig.Value;
 			_logger = logger;
 		}
@@ -590,10 +592,8 @@ namespace Horde.Server.Artifacts
 				return Forbid(ArtifactAclAction.ReadArtifact, artifact.AclScope);
 			}
 
-			using IStorageClient storageClient = _storageService.CreateClient(artifact.NamespaceId);
-
-			IBlobRef<DirectoryNode>? target = await storageClient.TryReadRefAsync<DirectoryNode>(artifact.RefName, cancellationToken: cancellationToken);
-			if (target == null)
+			UnsyncManifest? manifest = await _unsyncCache.GetManifestAsync(artifact, cancellationToken);
+			if (manifest == null)
 			{
 				return NotFound(id);
 			}
@@ -614,68 +614,98 @@ namespace Horde.Server.Artifacts
 				writer.WriteString("type", "unsync_manifest");
 				writer.WriteString("hash_strong", "Blake3.160");
 				writer.WriteString("chunking", "Variable");
+
 				writer.WriteStartArray("files");
-				await WriteUnsyncFileManifestAsync(new Utf8StringBuilder(), target, writer, cancellationToken);
+				foreach (UnsyncFile file in manifest.Files)
+				{
+					writer.WriteStartObject();
+					writer.WriteString("name", file.Name);
+					if (file.Executable)
+					{
+						writer.WriteBoolean("exec", file.Executable);
+					}
+					if (file.ReadOnly)
+					{
+						writer.WriteBoolean("read_only", file.ReadOnly);
+					}
+					writer.WriteNumber("size", file.Length);
+					writer.WriteStartArray("blocks");
+					foreach (UnsyncBlock block in file.Blocks)
+					{
+						writer.WriteStartObject();
+						writer.WriteNumber("offset", block.Offset);
+						writer.WriteNumber("size", block.Length);
+						writer.WriteString("hash_strong", block.Blob.Hash.ToString());
+						writer.WriteEndObject();
+					}
+					writer.WriteEndArray();
+					writer.WriteEndObject();
+				}
 				writer.WriteEndArray();
+
 				writer.WriteEndObject();
 			}
 
 			await response.CompleteAsync();
-
 			return Ok();
 		}
 
-		static async Task WriteUnsyncFileManifestAsync(Utf8StringBuilder path, IBlobRef<DirectoryNode> blobRef, Utf8JsonWriter writer, CancellationToken cancellationToken)
+		/// <summary>
+		/// Creates an Unsync manifest for an artifact
+		/// </summary>
+		/// <param name="id">The artifact id</param>
+		/// <param name="request">Information about the blobs to return</param>
+		/// <param name="cancellationToken">Cancellation token for the request</param>
+		/// <returns>An unsync manifest stream</returns>
+		[HttpPost]
+		[Route("/api/v2/artifacts/{id}/unsync-blobs")]
+		public async Task<ActionResult> GetUnsyncBlobsAsync(ArtifactId id, [FromBody] GetUnsyncDataRequest request, CancellationToken cancellationToken = default)
 		{
-			int initialPathLen = path.Length;
-			DirectoryNode directoryNode = await blobRef.ReadBlobAsync(cancellationToken);
-
-			foreach (FileEntry fileEntry in directoryNode.Files)
+			IArtifact? artifact = await _artifactCollection.GetAsync(id, cancellationToken);
+			if (artifact == null)
 			{
-				path.Append(fileEntry.Name);
-
-				writer.WriteStartObject();
-				writer.WriteString("name", path.WrittenSpan);
-				writer.WriteBoolean("read_only", (fileEntry.Flags & FileEntryFlags.ReadOnly) != 0);
-				writer.WriteNumber("size", fileEntry.Length);
-				writer.WriteStartArray("blocks");
-				await WriteUnsyncBlockManifestAsync(fileEntry.Target, 0, writer, cancellationToken);
-				writer.WriteEndArray();
-				writer.WriteEndObject();
-
-				path.Length = initialPathLen;
+				return NotFound(id);
 			}
-
-			foreach (DirectoryEntry directoryEntry in directoryNode.Directories)
+			if (!_globalConfig.Authorize(artifact.AclScope, ArtifactAclAction.ReadArtifact, User))
 			{
-				path.Append(directoryEntry.Name);
-				path.Append('/');
-				await WriteUnsyncFileManifestAsync(path, directoryEntry.Handle, writer, cancellationToken);
-				path.Length = initialPathLen;
+				return Forbid(ArtifactAclAction.ReadArtifact, artifact.AclScope);
 			}
-		}
 
-		static async Task WriteUnsyncBlockManifestAsync(ChunkedDataNodeRef chunkedDataRef, long offset, Utf8JsonWriter writer, CancellationToken cancellationToken)
-		{
-			if (chunkedDataRef.Type == ChunkedDataNodeType.Leaf)
+			if (!String.Equals(request.HashStrong, "Blake3.160", StringComparison.OrdinalIgnoreCase))
 			{
-				writer.WriteStartObject();
-				writer.WriteNumber("offset", offset);
-				writer.WriteNumber("size", chunkedDataRef.Length);
-				writer.WriteString("hash_strong", chunkedDataRef.Handle.Hash.ToString());
-				writer.WriteEndObject();
+				return BadRequest($"Unsupported hash algorithm: {request.HashStrong}");
 			}
-			else
+
+			// Disable buffering for the response
+			IHttpResponseBodyFeature? responseBodyFeature = HttpContext.Features.Get<IHttpResponseBodyFeature>();
+			responseBodyFeature?.DisableBuffering();
+
+			// Send the response headers
+			HttpResponse response = HttpContext.Response;
+
+			await response.StartAsync(cancellationToken);
+			foreach (GetUnsyncBlockRequest block in request.Files.SelectMany(x => x.Blocks))
 			{
-				InteriorChunkedDataNode node = await chunkedDataRef.Handle.ReadBlobAsync<InteriorChunkedDataNode>(cancellationToken: cancellationToken);
+				if (block.Hash == null || !IoHash.TryParse(block.Hash, out IoHash hash))
+				{
+					throw new InvalidOperationException();
+				}
 
-				long childOffset = offset;
-				foreach (ChunkedDataNodeRef childChunkedDataRef in node.Children)
+				BlobData? blobData = await _unsyncCache.ReadBlobAsync(artifact, hash, cancellationToken);
+				if (blobData == null)
 				{
-					await WriteUnsyncBlockManifestAsync(childChunkedDataRef, childOffset, writer, cancellationToken);
-					childOffset += childChunkedDataRef.Length;
+					throw new InvalidOperationException();
 				}
+
+				Memory<byte> data = response.BodyWriter.GetMemory(blobData.Data.Length);
+				blobData.Data.CopyTo(data);
+				response.BodyWriter.Advance(blobData.Data.Length);
+
+				await response.BodyWriter.FlushAsync(cancellationToken);
 			}
+			await response.CompleteAsync();
+
+			return Ok();
 		}
 
 		/// <summary>
diff --git a/Engine/Source/Programs/Horde/Horde.Server/Artifacts/UnsyncCache.cs b/Engine/Source/Programs/Horde/Horde.Server/Artifacts/UnsyncCache.cs
new file mode 100644
--- /dev/null
+++ b/Engine/Source/Programs/Horde/Horde.Server/Artifacts/UnsyncCache.cs
@@ -0,0 +1,219 @@
+// Copyright Epic Games, Inc. All Rights Reserved.
+
+using System;
+using System.Collections.Frozen;
+using System.Collections.Generic;
+using System.Linq;
+using System.Threading;
+using System.Threading.Tasks;
+using EpicGames.Core;
+using EpicGames.Horde.Storage;
+using EpicGames.Horde.Storage.Nodes;
+using Microsoft.Extensions.Caching.Memory;
+
+namespace Horde.Server.Artifacts
+{
+	/// <summary>
+	/// Manifest for Unsync
+	/// </summary>
+	/// <param name="Files">Files in the manifest</param>
+	public record class UnsyncManifest(IReadOnlyList<UnsyncFile> Files);
+
+	/// <summary>
+	/// Output file for an unsync manifest
+	/// </summary>
+	/// <param name="Name">Full path to the file</param>
+	/// <param name="ReadOnly">Whether the file is read only</param>
+	/// <param name="Executable">Whether to set the executable bit on the file</param>
+	/// <param name="Length">Length of the file</param>
+	/// <param name="ModTime">Last modified time for the file</param>
+	/// <param name="Blocks">List of blocks making up the file</param>
+	public record class UnsyncFile(Utf8String Name, bool ReadOnly, bool Executable, long Length, DateTime ModTime, IReadOnlyList<UnsyncBlock> Blocks);
+
+	/// <summary>
+	/// Describes a block of data in an unsync manifest
+	/// </summary>
+	/// <param name="Offset">Offset within the file</param>
+	/// <param name="Length">Length of the block</param>
+	/// <param name="Blob">Handle to the corresponding blob</param>
+	public record class UnsyncBlock(long Offset, long Length, IBlobRef<LeafChunkedDataNode> Blob);
+
+	/// <summary>
+	/// Implements a cache for downloading Unsync blobs
+	/// </summary>
+	public sealed class UnsyncCache : IDisposable
+	{
+		class ArtifactInfo : IDisposable
+		{
+			public IStorageClient StorageClient { get; }
+			public UnsyncManifest Manifest { get; }
+			public FrozenDictionary<IoHash, IBlobRef<LeafChunkedDataNode>> Blobs { get; }
+
+			public ArtifactInfo(IStorageClient storageClient, UnsyncManifest manifest, FrozenDictionary<IoHash, IBlobRef<LeafChunkedDataNode>> blobs)
+			{
+				StorageClient = storageClient;
+				Manifest = manifest;
+				Blobs = blobs;
+			}
+
+			public void Dispose()
+				=> StorageClient.Dispose();
+		}
+
+		readonly IStorageClientFactory _storageClientFactory;
+		readonly MemoryCache _cache;
+		readonly object _lockObject = new object();
+
+		/// <summary>
+		/// Constructor
+		/// </summary>
+		public UnsyncCache(IStorageClientFactory storageClientFactory)
+		{
+			_storageClientFactory = storageClientFactory;
+			_cache = new MemoryCache(new MemoryCacheOptions());
+		}
+
+		/// <inheritdoc/>
+		public void Dispose()
+			=> _cache.Dispose();
+
+		/// <summary>
+		/// 
+		/// </summary>
+		/// <param name="artifact"></param>
+		/// <param name="cancellationToken"></param>
+		/// <returns></returns>
+		public async ValueTask<UnsyncManifest?> GetManifestAsync(IArtifact artifact, CancellationToken cancellationToken = default)
+		{
+			ArtifactInfo? artifactInfo = await GetArtifactInfoAsync(artifact, cancellationToken);
+			return artifactInfo?.Manifest;
+		}
+
+		async Task<ArtifactInfo?> GetArtifactInfoAsync(IArtifact artifact, CancellationToken cancellationToken = default)
+		{
+			// Get the artifact info
+			BackgroundTask<ArtifactInfo?>? artifactInfoTask;
+			if (!_cache.TryGetValue(artifact.Id, out artifactInfoTask) || artifactInfoTask == null)
+			{
+				lock (_lockObject)
+				{
+					if (!_cache.TryGetValue(artifact.Id, out artifactInfoTask) || artifactInfoTask == null)
+					{
+						artifactInfoTask = BackgroundTask.StartNew(ctx => ReadArtifactAsync(artifact, ctx));
+						using (ICacheEntry entry = _cache.CreateEntry(artifact.Id))
+						{
+							entry.SetSize(1);
+							entry.SetSlidingExpiration(TimeSpan.FromHours(1.0));
+							entry.SetValue(artifactInfoTask);
+						}
+					}
+				}
+			}
+
+			// Wait for the read to finish
+			Task<ArtifactInfo?> task = artifactInfoTask.Task ?? Task.FromResult<ArtifactInfo?>(null);
+			return await task.WaitAsync(cancellationToken);
+		}
+
+		/// <summary>
+		/// Reads a blob from an artifact
+		/// </summary>
+		/// <param name="artifact"></param>
+		/// <param name="blobHash"></param>
+		/// <param name="cancellationToken"></param>
+		/// <returns></returns>
+		public async ValueTask<BlobData?> ReadBlobAsync(IArtifact artifact, IoHash blobHash, CancellationToken cancellationToken = default)
+		{
+			ArtifactInfo? artifactInfo = await GetArtifactInfoAsync(artifact, cancellationToken);
+			if (artifactInfo == null)
+			{
+				return null;
+			}
+
+			IBlobRef<LeafChunkedDataNode>? blobRef;
+			if (!artifactInfo.Blobs.TryGetValue(blobHash, out blobRef))
+			{
+				return null;
+			}
+
+			return await blobRef.ReadBlobDataAsync(cancellationToken);
+		}
+
+		async Task<ArtifactInfo?> ReadArtifactAsync(IArtifact artifact, CancellationToken cancellationToken)
+		{
+			IStorageClient? storageClient = null;
+			try
+			{
+				storageClient = _storageClientFactory.CreateClient(artifact.NamespaceId);
+
+				IBlobRef<DirectoryNode>? target = await storageClient.TryReadRefAsync<DirectoryNode>(artifact.RefName, cancellationToken: cancellationToken);
+				if (target == null)
+				{
+					return null;
+				}
+
+				List<UnsyncFile> files = new List<UnsyncFile>();
+				await FindFilesAsync(new Utf8StringBuilder(), target, files, cancellationToken);
+
+				Dictionary<IoHash, IBlobRef<LeafChunkedDataNode>> blocks = new Dictionary<IoHash, IBlobRef<LeafChunkedDataNode>>();
+				foreach (UnsyncBlock block in files.SelectMany(x => x.Blocks))
+				{
+					blocks[block.Blob.Hash] = block.Blob;
+				}
+
+				ArtifactInfo artifactInfo = new ArtifactInfo(storageClient, new UnsyncManifest(files), blocks.ToFrozenDictionary());
+				storageClient = null;
+				return artifactInfo;
+			}
+			finally
+			{
+				storageClient?.Dispose();
+			}
+		}
+
+		static async Task FindFilesAsync(Utf8StringBuilder path, IBlobRef<DirectoryNode> directoryNodeRef, List<UnsyncFile> files, CancellationToken cancellationToken)
+		{
+			DirectoryNode directoryNode = await directoryNodeRef.ReadBlobAsync(cancellationToken);
+
+			int initialPathLength = path.Length;
+
+			foreach (FileEntry fileEntry in directoryNode.Files)
+			{
+				List<UnsyncBlock> blocks = new List<UnsyncBlock>();
+				await FindBlocksAsync(0, fileEntry.Target, blocks, cancellationToken);
+
+				bool readOnly = (fileEntry.Flags & FileEntryFlags.ReadOnly) != 0;
+				bool executable = (fileEntry.Flags & FileEntryFlags.Executable) != 0;
+
+				path.Append(fileEntry.Name);
+				files.Add(new UnsyncFile(path.ToUtf8String(), readOnly, executable, fileEntry.Length, fileEntry.ModTime, blocks));
+				path.Length = initialPathLength;
+			}
+
+			foreach (DirectoryEntry directoryEntry in directoryNode.Directories)
+			{
+				path.Append(directoryEntry.Name);
+				path.Append('/');
+				await FindFilesAsync(path, directoryEntry.Handle, files, cancellationToken);
+				path.Length = initialPathLength;
+			}
+		}
+
+		static async Task FindBlocksAsync(long offset, ChunkedDataNodeRef nodeRef, List<UnsyncBlock> blocks, CancellationToken cancellationToken)
+		{
+			if (nodeRef.Type == ChunkedDataNodeType.Leaf)
+			{
+				blocks.Add(new UnsyncBlock(offset, nodeRef.Length, nodeRef.GetLeafHandle()));
+			}
+			else
+			{
+				InteriorChunkedDataNode interiorNode = await nodeRef.Handle.ReadBlobAsync<InteriorChunkedDataNode>(cancellationToken: cancellationToken);
+				foreach (ChunkedDataNodeRef childRef in interiorNode.Children)
+				{
+					await FindBlocksAsync(offset, childRef, blocks, cancellationToken);
+					offset += childRef.Length;
+				}
+			}
+		}
+	}
+}
diff --git a/Engine/Source/Programs/Horde/Horde.Server/Startup.cs b/Engine/Source/Programs/Horde/Horde.Server/Startup.cs
--- a/Engine/Source/Programs/Horde/Horde.Server/Startup.cs
+++ b/Engine/Source/Programs/Horde/Horde.Server/Startup.cs
@@ -577,6 +577,7 @@ namespace Horde.Server
 			services.AddSingleton<LogTailService>();
 			services.AddHostedService(provider => provider.GetRequiredService<LogTailService>()); // NB: Runs even on workers, to receive tail notifications.
 			services.AddSingleton<INotificationService, NotificationService>();
+			services.AddSingleton<UnsyncCache>();
 
 			if (settings.Commits.ReplicateMetadata)
 			{
diff --git a/Engine/Source/Programs/Shared/EpicGames.Horde/Artifacts/UnsyncMessages.cs b/Engine/Source/Programs/Shared/EpicGames.Horde/Artifacts/UnsyncMessages.cs
new file mode 100644
--- /dev/null
+++ b/Engine/Source/Programs/Shared/EpicGames.Horde/Artifacts/UnsyncMessages.cs
@@ -0,0 +1,56 @@
+// Copyright Epic Games, Inc.All Rights Reserved.
+
+using System;
+using System.Collections.Generic;
+using System.Text.Json.Serialization;
+
+#pragma warning disable CA2227
+
+namespace EpicGames.Horde.Artifacts
+{
+	/// <summary>
+	/// Request to return a set of blobs for Unsync
+	/// </summary>
+	public class GetUnsyncDataRequest
+	{
+		/// <summary>
+		/// The strong hash algorithm
+		/// </summary>
+		[JsonPropertyName("hash_strong")]
+		public string? HashStrong { get; set; }
+
+		/// <summary>
+		/// Files to retrieve
+		/// </summary>
+		[JsonPropertyName("files")]
+		public List<GetUnsyncFileRequest> Files { get; set; } = new List<GetUnsyncFileRequest>();
+	}
+
+	/// <summary>
+	/// Requests a set of blobs from a particular file
+	/// </summary>
+	public class GetUnsyncFileRequest
+	{
+		/// <summary>
+		/// Path to the file
+		/// </summary>
+		public string Name { get; set; } = String.Empty;
+
+		/// <summary>
+		/// Blobs to return
+		/// </summary>
+		public List<GetUnsyncBlockRequest> Blocks { get; set; } = new List<GetUnsyncBlockRequest>();
+	}
+
+	/// <summary>
+	/// Requests a block of data
+	/// </summary>
+	public class GetUnsyncBlockRequest
+	{
+		/// <summary>
+		/// Hash of the block
+		/// </summary>
+		[JsonPropertyName("hash_strong")]
+		public string? Hash { get; set; }
+	}
+}
diff --git a/Engine/Source/Programs/Shared/EpicGames.Horde/Storage/Nodes/ChunkedDataNode.cs b/Engine/Source/Programs/Shared/EpicGames.Horde/Storage/Nodes/ChunkedDataNode.cs
--- a/Engine/Source/Programs/Shared/EpicGames.Horde/Storage/Nodes/ChunkedDataNode.cs
+++ b/Engine/Source/Programs/Shared/EpicGames.Horde/Storage/Nodes/ChunkedDataNode.cs
@@ -168,6 +168,30 @@ namespace EpicGames.Horde.Storage.Nodes
 		{
 		}
 
+		/// <summary>
+		/// Gets the target interior node handle
+		/// </summary>
+		public IBlobRef<InteriorChunkedDataNode> GetInteriorHandle()
+		{
+			if (Type != ChunkedDataNodeType.Interior)
+			{
+				throw new InvalidOperationException("Node is not an interior node");
+			}
+			return BlobRef.Create<InteriorChunkedDataNode>(Handle.Hash, Handle, Handle.SerializerOptions);
+		}
+
+		/// <summary>
+		/// Gets the target leaf node handle
+		/// </summary>
+		public IBlobRef<LeafChunkedDataNode> GetLeafHandle()
+		{
+			if (Type != ChunkedDataNodeType.Leaf)
+			{
+				throw new InvalidOperationException("Node is not a leaf node");
+			}
+			return BlobRef.Create<LeafChunkedDataNode>(Handle.Hash, Handle, Handle.SerializerOptions);
+		}
+
 		/// <summary>
 		/// Read the node which is the target of this ref
 		/// </summary>
