Horde: Update logging to clarify whether tags or blocks are being read from temp storage.

#rnx

diff --git a/Engine/Source/Programs/Horde/Drivers/JobDriver/Utility/TempStorage.cs b/Engine/Source/Programs/Horde/Drivers/JobDriver/Utility/TempStorage.cs
--- a/Engine/Source/Programs/Horde/Drivers/JobDriver/Utility/TempStorage.cs
+++ b/Engine/Source/Programs/Horde/Drivers/JobDriver/Utility/TempStorage.cs
@@ -611,7 +611,7 @@ namespace JobDriver.Utility
 			FileReference localFileListLocation = GetTagManifestLocation(manifestDir, nodeName, tagName);
 			if (FileReference.Exists(localFileListLocation))
 			{
-				logger.LogInformation("Reading node \"{NodeName}\" tag \"{TagName}\" from {File}", nodeName, tagName, localFileListLocation.FullName);
+				logger.LogInformation("Reading tag \"{NodeName}\":\"{TagName}\" from {File}", nodeName, tagName, localFileListLocation.FullName);
 			}
 			else
 			{
@@ -630,7 +630,7 @@ namespace JobDriver.Utility
 				NamespaceId namespaceId = new NamespaceId(artifact.NamespaceId);
 				RefName refName = new RefName(artifact.RefName);
 
-				logger.LogInformation("Reading node \"{NodeName}\" tag \"{TagName}\" from temp storage (artifact: {ArtifactId} '{ArtifactName}' ({ArtifactType}), ns: {NamespaceId}, ref: {RefName}, localFile: {LocalFile})", nodeName, tagName, artifact.Id, artifactName, artifactType, namespaceId, refName, localFileListLocation);
+				logger.LogInformation("Reading tag \"{NodeName}\":\"{TagName}\" from temp storage (artifact: {ArtifactId} '{ArtifactName}' ({ArtifactType}), ns: {NamespaceId}, ref: {RefName}, localFile: {LocalFile})", nodeName, tagName, artifact.Id, artifactName, artifactType, namespaceId, refName, localFileListLocation);
 
 				using IStorageClient storageClient = hordeClient.CreateStorageClient(namespaceId, artifact.Token);
 				DirectoryNode node = await storageClient.ReadRefTargetAsync<DirectoryNode>(artifact.RefName, cancellationToken: cancellationToken);
@@ -745,7 +745,7 @@ namespace JobDriver.Utility
 			TempStorageBlockManifest? manifest;
 			if (local)
 			{
-				logger.LogInformation("Reading node \"{NodeName}\" block \"{BlockName}\" from {File}", nodeName, blockName, localManifestFile);
+				logger.LogInformation("Reading block \"{NodeName}\":\"{BlockName}\" from {File}", nodeName, blockName, localManifestFile);
 				manifest = TempStorageBlockManifest.Load(localManifestFile);
 			}
 			else
@@ -768,7 +768,7 @@ namespace JobDriver.Utility
 				NamespaceId namespaceId = new NamespaceId(artifact.NamespaceId);
 				RefName refName = new RefName(artifact.RefName);
 
-				logger.LogInformation("Reading node \"{NodeName}\" block \"{BlockName}\" from temp storage (artifact: {ArtifactId} '{ArtifactName}' ({ArtifactType}), ns: {NamespaceId}, ref: {RefName}, local: {LocalFile}, blockdir: {BlockDir})", nodeName, blockName, artifact.Id, artifactName, artifactType, namespaceId, refName, localManifestFile, blockDirectoryName);
+				logger.LogInformation("Reading block \"{NodeName}\":\"{BlockName}\" from temp storage (artifact: {ArtifactId} '{ArtifactName}' ({ArtifactType}), ns: {NamespaceId}, ref: {RefName}, local: {LocalFile}, blockdir: {BlockDir})", nodeName, blockName, artifact.Id, artifactName, artifactType, namespaceId, refName, localManifestFile, blockDirectoryName);
 
 				using IStorageClient storageClient = hordeClient.CreateStorageClient(namespaceId, artifact.Token);
 				DirectoryNode node = await storageClient.ReadRefTargetAsync<DirectoryNode>(refName, cancellationToken: cancellationToken);
