Horde: Include the target block size in unsync manifests.

#rnx

diff --git a/Engine/Source/Programs/Horde/Horde.Server/Artifacts/ArtifactsController.cs b/Engine/Source/Programs/Horde/Horde.Server/Artifacts/ArtifactsController.cs
--- a/Engine/Source/Programs/Horde/Horde.Server/Artifacts/ArtifactsController.cs
+++ b/Engine/Source/Programs/Horde/Horde.Server/Artifacts/ArtifactsController.cs
@@ -629,6 +629,7 @@ namespace Horde.Server.Artifacts
 						writer.WriteBoolean("read_only", file.ReadOnly);
 					}
 					writer.WriteNumber("size", file.Length);
+					writer.WriteNumber("block_size", LeafChunkedDataNodeOptions.Default.TargetSize);
 					writer.WriteStartArray("blocks");
 					foreach (UnsyncBlock block in file.Blocks)
 					{
