Horde: Fix some artifacts not showing up in job responses, due to default cap of 100 artifacts returned.

diff --git a/Engine/Source/Programs/Horde/Horde.Server/Jobs/JobsController.cs b/Engine/Source/Programs/Horde/Horde.Server/Jobs/JobsController.cs
--- a/Engine/Source/Programs/Horde/Horde.Server/Jobs/JobsController.cs
+++ b/Engine/Source/Programs/Horde/Horde.Server/Jobs/JobsController.cs
@@ -544,7 +544,7 @@ namespace Horde.Server.Jobs
 
 				string artifactKey = $"job:{job.Id}";
 				string artifactStepKeyPrefix = $"job:{job.Id}/step:";
-				await foreach (IArtifact artifact in _artifactCollection.FindAsync(keys: new[] { artifactKey }, cancellationToken: cancellationToken))
+				await foreach (IArtifact artifact in _artifactCollection.FindAsync(keys: new[] { artifactKey }, maxResults: 5000, cancellationToken: cancellationToken))
 				{
 					if (IncludeArtifactInResponse(artifact))
 					{
