Fix static analysis warnings.

#rnx

diff --git a/Engine/Source/Programs/Horde/Horde.Server.Tests/Artifacts/UnsyncCacheTests.cs b/Engine/Source/Programs/Horde/Horde.Server.Tests/Artifacts/UnsyncCacheTests.cs
--- a/Engine/Source/Programs/Horde/Horde.Server.Tests/Artifacts/UnsyncCacheTests.cs
+++ b/Engine/Source/Programs/Horde/Horde.Server.Tests/Artifacts/UnsyncCacheTests.cs
@@ -3,7 +3,6 @@
 using System;
 using System.Collections.Generic;
 using System.Linq;
-using System.Text;
 using System.Threading;
 using System.Threading.Tasks;
 using EpicGames.Core;
