Horde: Add compute endpoint that automatically finds the best cluster

Currently uses IP of requester together with networks from global config.

diff --git a/Engine/Source/Programs/Horde/Horde.Server.Tests/Compute/ComputeServiceTest.cs b/Engine/Source/Programs/Horde/Horde.Server.Tests/Compute/ComputeServiceTest.cs
--- a/Engine/Source/Programs/Horde/Horde.Server.Tests/Compute/ComputeServiceTest.cs
+++ b/Engine/Source/Programs/Horde/Horde.Server.Tests/Compute/ComputeServiceTest.cs
@@ -281,6 +281,37 @@ namespace Horde.Server.Tests.Compute
 				Assert.AreEqual(agent2.Id, (await ComputeService.TryAllocateResourceAsync(arpBar, CancellationToken.None))!.AgentId);
 			}
 		}
+		
+		[TestMethod]
+		[DataRow("11.0.0.1", "11.0.0.1", null, false)]
+		[DataRow("11.0.0.1", "11.0.0.1", "200.20.20.20", false)]
+		[DataRow("200.20.20.20", "11.0.0.1", "200.20.20.20", true)]
+		[DataRow("11.0.0.1", "11.0.0.1", "bad-ip", true)]
+		[DataRow("200.20.20.20", null, "200.20.20.20", true)]
+		public void Cluster_ResolveRequesterIp(string expectIp, string? internalIp, string? publicIp, bool? usePublicIp)
+		{
+			Assert.AreEqual(IPAddress.Parse(expectIp), ComputeService.ResolveRequesterIp(internalIp == null ? null : IPAddress.Parse(internalIp), usePublicIp, publicIp));
+		}
+		
+		[TestMethod]
+		public void Cluster_FindBest()
+		{
+			GlobalConfig globalConfig = new () { Networks = [
+				new NetworkConfig { CidrBlock = "11.0.0.0/16", Id = "network1", ComputeId = "compute1" },
+				new NetworkConfig { CidrBlock = "12.0.0.0/16", Id = "network2", ComputeId = "compute2" },
+			] };
+			globalConfig.PostLoad(new ServerSettings());
+			Assert.AreEqual(new ClusterId("compute1"), ComputeService.FindBestComputeClusterId(globalConfig, IPAddress.Parse("11.0.0.1")));
+			Assert.AreEqual(new ClusterId("compute2"), ComputeService.FindBestComputeClusterId(globalConfig, IPAddress.Parse("12.0.1.1")));
+			Assert.ThrowsException<ComputeServiceException>(() => ComputeService.FindBestComputeClusterId(globalConfig, IPAddress.Parse("123.123.123.123")));
+			
+			GlobalConfig globalConfigCatchAll = new ()
+			{
+				Networks = [new NetworkConfig { CidrBlock = "0.0.0.0/0", Id = "catchAll", ComputeId = "catchAll" }]
+			};
+			globalConfigCatchAll.PostLoad(new ServerSettings());
+			Assert.AreEqual(new ClusterId("catchAll"), ComputeService.FindBestComputeClusterId(globalConfigCatchAll, IPAddress.Parse("123.123.123.123")));
+		}
 
 		[TestMethod]
 		public async Task Connection_Direct_IpConnection_Async()
diff --git a/Engine/Source/Programs/Horde/Horde.Server/Compute/ComputeController.cs b/Engine/Source/Programs/Horde/Horde.Server/Compute/ComputeController.cs
--- a/Engine/Source/Programs/Horde/Horde.Server/Compute/ComputeController.cs
+++ b/Engine/Source/Programs/Horde/Horde.Server/Compute/ComputeController.cs
@@ -37,7 +37,23 @@ namespace Horde.Server.Compute
 			_computeService = computeService;
 			_globalConfig = globalConfig;
 		}
-
+		
+		/// <summary>
+		/// Add tasks to be executed remotely and auto-select appropriate compute cluster to use
+		/// </summary>
+		/// <param name="request">The request parameters</param>
+		/// <param name="cancellationToken">Cancellation token for the operation</param>
+		/// <returns></returns>
+		[HttpPost]
+		[Authorize]
+		[Route("/api/v2/compute")]
+		public async Task<ActionResult<AssignComputeResponse>> AssignComputeResourceAsync([FromBody] AssignComputeRequest request, CancellationToken cancellationToken)
+		{
+			IPAddress requesterIp = ComputeService.ResolveRequesterIp(HttpContext.Connection.RemoteIpAddress, request.Connection?.PreferPublicIp, request.Connection?.ClientPublicIp);
+			ClusterId clusterId = ComputeService.FindBestComputeClusterId(_globalConfig.Value, requesterIp);
+			return await AssignComputeResourceInClusterAsync(clusterId, request, cancellationToken);
+		}
+		
 		/// <summary>
 		/// Add tasks to be executed remotely
 		/// </summary>
@@ -48,7 +64,7 @@ namespace Horde.Server.Compute
 		[HttpPost]
 		[Authorize]
 		[Route("/api/v2/compute/{clusterId}")]
-		public async Task<ActionResult<AssignComputeResponse>> AssignComputeResourceAsync(ClusterId clusterId, [FromBody] AssignComputeRequest request, CancellationToken cancellationToken)
+		public async Task<ActionResult<AssignComputeResponse>> AssignComputeResourceInClusterAsync(ClusterId clusterId, [FromBody] AssignComputeRequest request, CancellationToken cancellationToken)
 		{
 			if (!_globalConfig.Value.TryGetComputeCluster(clusterId, out ComputeClusterConfig? clusterConfig))
 			{
diff --git a/Engine/Source/Programs/Horde/Horde.Server/Compute/ComputeService.cs b/Engine/Source/Programs/Horde/Horde.Server/Compute/ComputeService.cs
--- a/Engine/Source/Programs/Horde/Horde.Server/Compute/ComputeService.cs
+++ b/Engine/Source/Programs/Horde/Horde.Server/Compute/ComputeService.cs
@@ -50,7 +50,7 @@ namespace Horde.Server.Compute
 		public bool ShowToUser { get; init; } = false;
 
 		/// <inheritdoc/>
-		public ComputeServiceException(string? message, Exception? innerException) : base(message, innerException)
+		public ComputeServiceException(string? message, Exception? innerException = null) : base(message, innerException)
 		{
 		}
 	}
@@ -776,6 +776,54 @@ namespace Horde.Server.Compute
 
 			throw new Exception("Unable to resolve a suitable connection mode for compute task");
 		}
+		
+		/// <summary>
+		/// Resolve IP of the requester
+		/// </summary>
+		/// <param name="requesterIp"></param>
+		/// <param name="usePublicIp"></param>
+		/// <param name="requesterPublicIp"></param>
+		/// <returns>IP address of request</returns>
+		/// <exception cref="ComputeServiceException"></exception>
+		public static IPAddress ResolveRequesterIp(IPAddress? requesterIp, bool? usePublicIp, string? requesterPublicIp)
+		{
+			if (usePublicIp is true && requesterPublicIp != null)
+			{
+				if (IPAddress.TryParse(requesterPublicIp, out IPAddress? publicIp))
+				{
+					requesterIp = publicIp;
+				}
+			}
+			
+			if (requesterIp == null)
+			{
+				throw new ComputeServiceException("Unable to determine IP of requester");
+			}
+			
+			return requesterIp;
+		}
+		
+		/// <summary>
+		/// Find best compute cluster for given parameters
+		/// </summary>
+		/// <param name="globalConfig">Global config</param>
+		/// <param name="requesterIp">IP of client (initiator) requesting a compute resource</param>
+		/// <returns>Compute cluster ID</returns>
+		/// <exception cref="ComputeServiceException"></exception>
+		public static ClusterId FindBestComputeClusterId(GlobalConfig globalConfig, IPAddress? requesterIp)
+		{
+			if (requesterIp == null || !globalConfig.TryGetNetworkConfig(requesterIp, out NetworkConfig? networkConfig))
+			{
+				throw new ComputeServiceException("Unable to find a matching network config");
+			}
+			
+			if (networkConfig.ComputeId == null)
+			{
+				throw new ComputeServiceException($"Network config '{networkConfig.Id}' has no compute ID set");
+			}
+			
+			return new ClusterId(networkConfig.ComputeId);
+		}
 
 		private static IPAddress FindBestRelayIp(IPAddress? clientIp, IPAddress? publicClientIp, IEnumerable<string> relayIps)
 		{
