Horde: Fix DI exception due to IStorageClientFactory having scoped lifetime.

#rnx

diff --git a/Engine/Source/Programs/Horde/Horde.Server.Tests/Artifacts/UnsyncCacheTests.cs b/Engine/Source/Programs/Horde/Horde.Server.Tests/Artifacts/UnsyncCacheTests.cs
--- a/Engine/Source/Programs/Horde/Horde.Server.Tests/Artifacts/UnsyncCacheTests.cs
+++ b/Engine/Source/Programs/Horde/Horde.Server.Tests/Artifacts/UnsyncCacheTests.cs
@@ -10,6 +10,7 @@ using EpicGames.Horde.Storage;
 using EpicGames.Horde.Storage.Bundles;
 using EpicGames.Horde.Storage.Nodes;
 using Horde.Server.Artifacts;
+using Horde.Server.Storage;
 using Microsoft.Extensions.Logging.Abstractions;
 using Microsoft.VisualStudio.TestTools.UnitTesting;
 using Moq;
@@ -46,7 +47,7 @@ namespace Horde.Server.Tests.Artifacts
 			RefName refName = new RefName("test");
 			await storageClient.WriteRefAsync(refName, directoryRef);
 
-			Mock<IStorageClientFactory> factory = new Mock<IStorageClientFactory>();
+			Mock<IStorageService> factory = new Mock<IStorageService>();
 			factory.Setup(x => x.TryCreateClient(ns)).Returns(storageClient);
 
 			Mock<IArtifact> artifact = new Mock<IArtifact>();
diff --git a/Engine/Source/Programs/Horde/Horde.Server/Artifacts/UnsyncCache.cs b/Engine/Source/Programs/Horde/Horde.Server/Artifacts/UnsyncCache.cs
--- a/Engine/Source/Programs/Horde/Horde.Server/Artifacts/UnsyncCache.cs
+++ b/Engine/Source/Programs/Horde/Horde.Server/Artifacts/UnsyncCache.cs
@@ -9,6 +9,7 @@ using System.Threading.Tasks;
 using EpicGames.Core;
 using EpicGames.Horde.Storage;
 using EpicGames.Horde.Storage.Nodes;
+using Horde.Server.Storage;
 using Microsoft.Extensions.Caching.Memory;
 
 namespace Horde.Server.Artifacts
@@ -60,16 +61,16 @@ namespace Horde.Server.Artifacts
 				=> StorageClient.Dispose();
 		}
 
-		readonly IStorageClientFactory _storageClientFactory;
+		readonly IStorageService _storageService;
 		readonly MemoryCache _cache;
 		readonly object _lockObject = new object();
 
 		/// <summary>
 		/// Constructor
 		/// </summary>
-		public UnsyncCache(IStorageClientFactory storageClientFactory)
+		public UnsyncCache(IStorageService storageService)
 		{
-			_storageClientFactory = storageClientFactory;
+			_storageService = storageService;
 			_cache = new MemoryCache(new MemoryCacheOptions());
 		}
 
@@ -144,7 +145,7 @@ namespace Horde.Server.Artifacts
 			IStorageClient? storageClient = null;
 			try
 			{
-				storageClient = _storageClientFactory.CreateClient(artifact.NamespaceId);
+				storageClient = _storageService.CreateClient(artifact.NamespaceId);
 
 				IBlobRef<DirectoryNode>? target = await storageClient.TryReadRefAsync<DirectoryNode>(artifact.RefName, cancellationToken: cancellationToken);
 				if (target == null)
diff --git a/Engine/Source/Programs/Horde/Horde.Server/Startup.cs b/Engine/Source/Programs/Horde/Horde.Server/Startup.cs
--- a/Engine/Source/Programs/Horde/Horde.Server/Startup.cs
+++ b/Engine/Source/Programs/Horde/Horde.Server/Startup.cs
@@ -612,6 +612,7 @@ namespace Horde.Server
 			services.AddSingleton<DeviceService>();
 			services.AddSingleton<NoticeService>();
 			services.AddSingleton<StorageService>();
+			services.AddSingleton<IStorageService>(sp => sp.GetRequiredService<StorageService>());
 			services.AddScoped(sp => sp.GetRequiredService<StorageService>().CreateStorageClientFactory(sp.GetRequiredService<IOptionsSnapshot<GlobalConfig>>().Value));
 			services.AddSingleton<TestDataService>();
 			services.AddSingleton<BundleCache>();
diff --git a/Engine/Source/Programs/Horde/Horde.Server/Storage/StorageService.cs b/Engine/Source/Programs/Horde/Horde.Server/Storage/StorageService.cs
--- a/Engine/Source/Programs/Horde/Horde.Server/Storage/StorageService.cs
+++ b/Engine/Source/Programs/Horde/Horde.Server/Storage/StorageService.cs
@@ -31,10 +31,17 @@ using StackExchange.Redis;
 
 namespace Horde.Server.Storage
 {
+	/// <summary>
+	/// Interface for the storage service
+	/// </summary>
+	public interface IStorageService : IStorageClientFactory
+	{
+	}
+
 	/// <summary>
 	/// Functionality related to the storage service
 	/// </summary>
-	public sealed class StorageService : IHostedService, IStorageClientFactory, IAsyncDisposable
+	public sealed class StorageService : IHostedService, IStorageService, IAsyncDisposable
 	{
 		class RefCount
 		{
