Horde: Fix static analysis warning.

#rnx

diff --git a/Engine/Source/Programs/Horde/Horde.Server.Tests/Compute/ComputeServiceTest.cs b/Engine/Source/Programs/Horde/Horde.Server.Tests/Compute/ComputeServiceTest.cs
--- a/Engine/Source/Programs/Horde/Horde.Server.Tests/Compute/ComputeServiceTest.cs
+++ b/Engine/Source/Programs/Horde/Horde.Server.Tests/Compute/ComputeServiceTest.cs
@@ -378,7 +378,7 @@ namespace Horde.Server.Tests.Compute
 				new TestOptionsMonitor<ServerSettings>(ss), GlobalConfig, Clock, Tracer, Meter,
 				NullLogger<ComputeService>.Instance);
 			await CreateAgentAsync(new PoolId(pool), properties: ["ComputeIp=11.0.0.1", "ComputePort=5000"]);
-			GlobalConfig.CurrentValue.Compute = [new ComputeClusterConfig { Id = _cluster1 }];
+			GlobalConfig.CurrentValue.Compute = [ccc];
 			GlobalConfig.CurrentValue.PostLoad(ss);
 			return cs;
 		}
