Horde: Fix truncated paths in unsync manifests.

#rnx

diff --git a/Engine/Source/Programs/Horde/Horde.Server/Artifacts/UnsyncCache.cs b/Engine/Source/Programs/Horde/Horde.Server/Artifacts/UnsyncCache.cs
--- a/Engine/Source/Programs/Horde/Horde.Server/Artifacts/UnsyncCache.cs
+++ b/Engine/Source/Programs/Horde/Horde.Server/Artifacts/UnsyncCache.cs
@@ -187,7 +187,7 @@ namespace Horde.Server.Artifacts
 				bool executable = (fileEntry.Flags & FileEntryFlags.Executable) != 0;
 
 				path.Append(fileEntry.Name);
-				files.Add(new UnsyncFile(path.ToUtf8String(), readOnly, executable, fileEntry.Length, fileEntry.ModTime, blocks));
+				files.Add(new UnsyncFile(path.ToUtf8String().Clone(), readOnly, executable, fileEntry.Length, fileEntry.ModTime, blocks));
 				path.Length = initialPathLength;
 			}
 
