Horde: Ignore cancellations in AWS instance lifecycle

diff --git a/Engine/Source/Programs/Horde/Horde.Agent/Services/AwsInstanceLifecycleService.cs b/Engine/Source/Programs/Horde/Horde.Agent/Services/AwsInstanceLifecycleService.cs
--- a/Engine/Source/Programs/Horde/Horde.Agent/Services/AwsInstanceLifecycleService.cs
+++ b/Engine/Source/Programs/Horde/Horde.Agent/Services/AwsInstanceLifecycleService.cs
@@ -257,7 +257,14 @@ class AwsInstanceLifecycleService : BackgroundService
 	/// <inheritdoc/>
 	protected override async Task ExecuteAsync(CancellationToken stoppingToken)
 	{
-		await MonitorInstanceLifecycleAsync(stoppingToken);
+		try
+		{
+			await MonitorInstanceLifecycleAsync(stoppingToken);
+		}
+		catch (OperationCanceledException) when (stoppingToken.IsCancellationRequested)
+		{
+			// Ignore any exceptions if cancellation has been requested
+		}
 	}
 }
 
