Horde: Artifact Info box for artifact modal
#rnx

diff --git a/Engine/Source/Programs/Horde/HordeDashboard/src/components/artifacts/ArtifactsModal.tsx b/Engine/Source/Programs/Horde/HordeDashboard/src/components/artifacts/ArtifactsModal.tsx
--- a/Engine/Source/Programs/Horde/HordeDashboard/src/components/artifacts/ArtifactsModal.tsx
+++ b/Engine/Source/Programs/Horde/HordeDashboard/src/components/artifacts/ArtifactsModal.tsx
@@ -1,6 +1,6 @@
 // Copyright Epic Games, Inc. All Rights Reserved.
 
-import { DetailsList, DetailsListLayoutMode, FontIcon, IColumn, IconButton, Modal, PrimaryButton, ScrollablePane, Selection, SelectionMode, SelectionZone, Spinner, SpinnerSize, Stack, Text, mergeStyleSets } from "@fluentui/react";
+import { Callout, DetailsList, DetailsListLayoutMode, DirectionalHint, FontIcon, IColumn, IconButton, Modal, PrimaryButton, ScrollablePane, Selection, SelectionMode, SelectionZone, Spinner, SpinnerSize, Stack, Text, mergeStyleSets } from "@fluentui/react";
 import { action, makeObservable, observable } from "mobx";
 import { observer } from "mobx-react-lite";
 import { useEffect, useState } from "react";
@@ -22,6 +22,8 @@ type BrowserItem = {
    icon?: string;
    size?: number;
    type: BrowserType;
+   dirResponse?: GetArtifactDirectoryEntryResponse;
+   fileResponse?: GetArtifactFileEntryResponse;
 }
 
 
@@ -232,7 +234,7 @@ class ArtifactsHandler {
       this.path = undefined;
       this.browse = undefined;
       this.artifact = undefined;
-      this.artifacts = undefined;      
+      this.artifacts = undefined;
       this.history = [];
       this.stepId = "";
       this.loading = false;
@@ -251,7 +253,7 @@ class ArtifactsHandler {
    browse?: GetArtifactDirectoryResponse;
 
    artifact?: GetArtifactResponseV2;
-   artifacts?: GetArtifactResponseV2[];   
+   artifacts?: GetArtifactResponseV2[];
 
    readonly context: ArtifactContextType;
 
@@ -282,7 +284,7 @@ const getStyles = () => {
                height: "unset !important",
             },
             '.ms-List-cell': {
-   
+
                borderTop: border,
                borderRight: border,
                borderLeft: border
@@ -303,7 +305,7 @@ const getStyles = () => {
    _styles = styles;
 
    return styles;
-   
+
 }
 
 const BrowseHistory: React.FC<{ handler: ArtifactsHandler }> = observer(({ handler }) => {
@@ -512,6 +514,8 @@ const JobDetailArtifactsInner: React.FC<{ jobId: string; stepId: string, artifac
    // eslint-disable-next-line
    const handler = ArtifactsHandler.current ?? new ArtifactsHandler(jobId, stepId, contextType, artifactPath, artifacts);
 
+   const [infoKey, setInfoKey] = useState("");
+
    useEffect(() => {
       return () => {
          handler?.clear();
@@ -550,7 +554,7 @@ const JobDetailArtifactsInner: React.FC<{ jobId: string; stepId: string, artifac
    const items: BrowserItem[] = [];
 
    const getFileHRef = (item: BrowserItem) => {
-      
+
       if (item.type !== BrowserType.File) {
          return undefined;
       }
@@ -558,7 +562,7 @@ const JobDetailArtifactsInner: React.FC<{ jobId: string; stepId: string, artifac
       const path = encodeURI((handler.path ? handler.path + "/" : "") + item.text);
       const server = backend.serverUrl;
       return `${server}/api/v2/artifacts/${handler.artifact!.id}/file?path=${path}`;
-      
+
    }
 
    // use the up arrow instead
@@ -569,9 +573,9 @@ const JobDetailArtifactsInner: React.FC<{ jobId: string; stepId: string, artifac
    browse.directories?.forEach(d => {
 
       function recurseDirectories(dir: GetArtifactDirectoryEntryResponse, flattened: GetArtifactDirectoryEntryResponse[]) {
-         if (!dir.directories || dir.directories.length > 1  || dir.files?.length) {
+         if (!dir.directories || dir.directories.length > 1 || dir.files?.length) {
             const name = flattened.length ? flattened.map(d => d.name).join("/") + "/" + dir.name : dir.name;
-            items.push({ key: dir.hash, text: name, icon: "Folder", type: BrowserType.Directory, size: dir.length });
+            items.push({ key: dir.hash, text: name, icon: "Folder", type: BrowserType.Directory, size: dir.length, dirResponse: dir });
          } else {
             flattened.push(dir);
             dir.directories.forEach(d => recurseDirectories(d, [...flattened]));
@@ -584,13 +588,13 @@ const JobDetailArtifactsInner: React.FC<{ jobId: string; stepId: string, artifac
    });
 
    browse.files?.forEach(d => {
-      items.push({ key: d.hash, text: d.name, icon: "Document", type: BrowserType.File, size: d.length });
+      items.push({ key: d.hash, text: d.name, icon: "Document", type: BrowserType.File, size: d.length, fileResponse: d });
    });
 
    const columns: IColumn[] = [
-      { key: 'column1', name: 'Name', minWidth: 794, maxWidth: 794, isResizable: false, isPadded: false },
+      { key: 'column1', name: 'Name', minWidth: 794 - 32, maxWidth: 794 - 32, isResizable: false, isPadded: false },
       { key: 'column2', name: 'Size', minWidth: 128, maxWidth: 128, isResizable: false, isPadded: false },
-      { key: 'column3', name: 'View_Download', minWidth: 64, maxWidth: 64, isResizable: false, isPadded: false }
+      { key: 'column3', name: 'View_Download', minWidth: 64 + 32, maxWidth: 64 + 32, isResizable: false, isPadded: false }
    ];
 
    const renderItem = (item: any, index?: number, column?: IColumn) => {
@@ -611,15 +615,51 @@ const JobDetailArtifactsInner: React.FC<{ jobId: string; stepId: string, artifac
 
       if (column.name === "View_Download") {
 
-         if (item.type !== BrowserType.File) {
+         if (item.text === "..") {
             return null;
          }
 
+         const isFile = item.type === BrowserType.File;
+
          const href = getFileHRef(item);
 
-         return <Stack data-selection-disabled verticalAlign="center" verticalFill horizontal horizontalAlign="end" style={{ paddingTop: 0, paddingBottom: 0 }}>
-            <IconButton id="artifactview" href={`${href}&inline=true`} target="_blank" style={{ paddingTop: 1, color: "#106EBE" }} iconProps={{ iconName: "Eye", styles: { root: { fontSize: "14px" } } }} />
-            <IconButton id="artifactview" href={href} target="_blank" style={{ paddingTop: 1, color: "#106EBE" }} iconProps={{ iconName: "CloudDownload", styles: { root: { fontSize: "14px" } } }} />
+         return <Stack>
+            <Stack id={`callout_target_${item.key}`} data-selection-disabled verticalAlign="center" verticalFill horizontal horizontalAlign="end" style={{ paddingTop: 0, paddingBottom: 0 }}>
+               {isFile && <IconButton id="artifactview" href={`${href}&inline=true`} target="_blank" style={{ paddingTop: 1, color: "#106EBE" }} iconProps={{ iconName: "Eye", styles: { root: { fontSize: "14px" } } }} />}
+               {isFile && <IconButton id="artifactview" href={href} target="_blank" style={{ paddingTop: 1, color: "#106EBE" }} iconProps={{ iconName: "CloudDownload", styles: { root: { fontSize: "14px" } } }} />}
+               <IconButton id="artifactview" onClick={() => setInfoKey(item.key)} style={{ paddingTop: 1, color: "#106EBE", height: "18px" }} iconProps={{ iconName: "Info", styles: { root: { fontSize: "14px" } } }} />
+            </Stack>
+            {infoKey === item.key && <Callout
+               styles={{ root: { padding: "32px 24px", maxWidth: 1300 } }}
+               role="dialog"
+               gapSpace={0}
+               target={`#callout_target_${item.key}`}
+               isBeakVisible={true}
+               beakWidth={12}
+               onDismiss={() => {
+                  setInfoKey("")
+               }}
+               directionalHint={DirectionalHint.leftCenter}
+               setInitialFocus>
+               <Stack style={{ maxWidth: 1140 }}>
+                  <Stack style={{ paddingBottom: 12 }}>
+                     <Stack horizontal verticalAlign="center" verticalFill={true}>
+                        <Stack>
+                           <Text style={{ fontSize: 14, fontFamily: "Horde Open Sans SemiBold" }}>Artifact Info</Text>
+                        </Stack>
+                        <Stack grow />
+                        <Stack>
+                           <IconButton
+                              iconProps={{ iconName: 'Cancel', styles: { root: { fontSize: "14px" } } }}
+                              onClick={() => { setInfoKey("") }}
+                           />
+                        </Stack>
+                     </Stack>
+                  </Stack>
+                  <Stack style={{ paddingLeft: 0 }}>
+                     <Text style={{ fontSize: 11, whiteSpace: "pre-wrap", fontFamily: "Horde Cousine Regular" }}>{JSON.stringify(item.dirResponse ?? item.fileResponse ?? {}, undefined, 2).replaceAll("\\r", "").replaceAll("\\n", "\n")}</Text></Stack>
+               </Stack>
+            </Callout>}
          </Stack>
       }
 
@@ -700,14 +740,14 @@ const JobDetailArtifactsInner: React.FC<{ jobId: string; stepId: string, artifac
                   selection={handler.selection}
                   selectionPreservedOnEmptyClick={true}
                   onShouldVirtualize={() => false}
-                  onItemInvoked={(item:BrowserItem) => {
+                  onItemInvoked={(item: BrowserItem) => {
                      if (item?.type !== BrowserType.File) {
                         return;
                      }
                      const href = getFileHRef(item);
                      if (href) {
                         window.open(href + "&inline=true", "_blank");
-                     }                     
+                     }
                   }}
                   onRenderItemColumn={renderItem}
                />
diff --git a/Engine/Source/Programs/Horde/HordeDashboard/src/styles/Styles.tsx b/Engine/Source/Programs/Horde/HordeDashboard/src/styles/Styles.tsx
--- a/Engine/Source/Programs/Horde/HordeDashboard/src/styles/Styles.tsx
+++ b/Engine/Source/Programs/Horde/HordeDashboard/src/styles/Styles.tsx
@@ -104,7 +104,8 @@ registerIcons({
       "Play": IconContents.PLAY,
       "Maximize" : IconContents.MAXIMIZE,
       "Star": IconContents.STAR,
-      "Clean" : IconContents.CLEAN
+      "Clean": IconContents.CLEAN,
+      "Info" : IconContents.INFO_SIGN
    }
 });
 
