Horde: Add artifact id to info box
#rnx

diff --git a/Engine/Source/Programs/Horde/HordeDashboard/src/components/artifacts/ArtifactsModal.tsx b/Engine/Source/Programs/Horde/HordeDashboard/src/components/artifacts/ArtifactsModal.tsx
--- a/Engine/Source/Programs/Horde/HordeDashboard/src/components/artifacts/ArtifactsModal.tsx
+++ b/Engine/Source/Programs/Horde/HordeDashboard/src/components/artifacts/ArtifactsModal.tsx
@@ -645,7 +645,7 @@ const JobDetailArtifactsInner: React.FC<{ jobId: string; stepId: string, artifac
                   <Stack style={{ paddingBottom: 12 }}>
                      <Stack horizontal verticalAlign="center" verticalFill={true}>
                         <Stack>
-                           <Text style={{ fontSize: 14, fontFamily: "Horde Open Sans SemiBold" }}>Artifact Info</Text>
+                           <Text style={{ fontSize: 14, fontFamily: "Horde Open Sans SemiBold" }}>{`Artifact ID ${handler.artifact?.id}`}</Text>
                         </Stack>
                         <Stack grow />
                         <Stack>
