Horde: Add a header to each blob returned from the unsync-blobs endpoint. Each blob is prefixed by:

* 8 bytes: magic number (0x4C5C2AABA992610C)
* 8 bytes: payload size, excluding this header
* 8 bytes: size of the payload when decompressed
* 20 bytes: hash of the decompressed data

Values are all stored little endian.

#rnx

diff --git a/Engine/Source/Programs/Horde/Horde.Server/Artifacts/ArtifactsController.cs b/Engine/Source/Programs/Horde/Horde.Server/Artifacts/ArtifactsController.cs
--- a/Engine/Source/Programs/Horde/Horde.Server/Artifacts/ArtifactsController.cs
+++ b/Engine/Source/Programs/Horde/Horde.Server/Artifacts/ArtifactsController.cs
@@ -1,8 +1,11 @@
 // Copyright Epic Games, Inc. All Rights Reserved.
 
 using System;
+using System.Buffers.Binary;
 using System.Collections.Generic;
+using System.Diagnostics;
 using System.IO;
+using System.IO.Pipelines;
 using System.Linq;
 using System.Net;
 using System.Net.Mime;
@@ -15,6 +18,7 @@ using EpicGames.Horde.Agents.Leases;
 using EpicGames.Horde.Artifacts;
 using EpicGames.Horde.Jobs;
 using EpicGames.Horde.Storage;
+using EpicGames.Horde.Storage.Bundles;
 using EpicGames.Horde.Storage.Nodes;
 using EpicGames.Horde.Streams;
 using Google.Protobuf.WellKnownTypes;
@@ -660,11 +664,12 @@ namespace Horde.Server.Artifacts
 		/// </summary>
 		/// <param name="id">The artifact id</param>
 		/// <param name="request">Information about the blobs to return</param>
+		/// <param name="compress">Whether to compress the output data</param>
 		/// <param name="cancellationToken">Cancellation token for the request</param>
 		/// <returns>An unsync manifest stream</returns>
 		[HttpPost]
 		[Route("/api/v2/artifacts/{id}/unsync-blobs")]
-		public async Task<ActionResult> GetUnsyncBlobsAsync(ArtifactId id, [FromBody] GetUnsyncDataRequest request, CancellationToken cancellationToken = default)
+		public async Task<ActionResult> GetUnsyncBlobsAsync(ArtifactId id, [FromBody] GetUnsyncDataRequest request, [FromQuery] bool compress = true, CancellationToken cancellationToken = default)
 		{
 			IArtifact? artifact = await _artifactCollection.GetAsync(id, cancellationToken);
 			if (artifact == null)
@@ -699,9 +704,11 @@ namespace Horde.Server.Artifacts
 
 			// Send the response headers
 			HttpResponse response = HttpContext.Response;
-			response.Headers["x-chunk-content-encoding"] = "identity";
+			response.Headers["x-chunk-content-encoding"] = compress ? "zstd" : "identity";
 			response.StatusCode = (int)HttpStatusCode.OK;
 
+			ArrayMemoryWriter? compressedWriter = null;
+
 			await response.StartAsync(cancellationToken);
 			foreach (string block in request.Blocks)
 			{
@@ -716,9 +723,18 @@ namespace Horde.Server.Artifacts
 					throw new InvalidOperationException();
 				}
 
-				Memory<byte> data = response.BodyWriter.GetMemory(blobData.Data.Length);
-				blobData.Data.CopyTo(data);
-				response.BodyWriter.Advance(blobData.Data.Length);
+				if (compress)
+				{
+					compressedWriter ??= new ArrayMemoryWriter(300 * 1024);
+					compressedWriter.Clear();
+					BundleData.Compress(BundleCompressionFormat.Zstd, blobData.Data, compressedWriter);
+
+					WriteBlock(response.BodyWriter, hash, blobData.Data.Length, compressedWriter.WrittenSpan);
+				}
+				else
+				{
+					WriteBlock(response.BodyWriter, hash, blobData.Data.Length, blobData.Data.Span);
+				}
 
 				await response.BodyWriter.FlushAsync(cancellationToken);
 			}
@@ -727,6 +743,29 @@ namespace Horde.Server.Artifacts
 			return Ok();
 		}
 
+		static void WriteBlock(PipeWriter writer, IoHash decompressedHash, long decompressedSize, ReadOnlySpan<byte> payloadData)
+		{
+			int length = (sizeof(long) * 3) + IoHash.NumBytes + payloadData.Length;
+
+			Span<byte> data = writer.GetSpan(length);
+			data = data.Slice(0, length);
+
+			BinaryPrimitives.WriteUInt64LittleEndian(data, 0x_4C5C_2AAB_A992_610C);
+			data = data.Slice(8);
+
+			BinaryPrimitives.WriteUInt64LittleEndian(data, (ulong)payloadData.Length);
+			data = data.Slice(8);
+
+			BinaryPrimitives.WriteUInt64LittleEndian(data, (ulong)decompressedSize);
+			data = data.Slice(8);
+
+			decompressedHash.CopyTo(data);
+			data = data.Slice(IoHash.NumBytes);
+
+			payloadData.CopyTo(data);
+			Debug.Assert(data.Length == payloadData.Length);
+		}
+
 		/// <summary>
 		/// Downloads an individual file from an artifact
 		/// </summary>
