Horde: Add a custom content type for unsync blob results (application/x-horde-unsync-blob)

#rnx

diff --git a/Engine/Source/Programs/Horde/Horde.Server/Artifacts/ArtifactsController.cs b/Engine/Source/Programs/Horde/Horde.Server/Artifacts/ArtifactsController.cs
--- a/Engine/Source/Programs/Horde/Horde.Server/Artifacts/ArtifactsController.cs
+++ b/Engine/Source/Programs/Horde/Horde.Server/Artifacts/ArtifactsController.cs
@@ -704,6 +704,7 @@ namespace Horde.Server.Artifacts
 
 			// Send the response headers
 			HttpResponse response = HttpContext.Response;
+			response.ContentType = "application/x-horde-unsync-blob";
 			response.Headers["x-chunk-content-encoding"] = compress ? "zstd" : "identity";
 			response.StatusCode = (int)HttpStatusCode.OK;
 
