Horde: Fix body length not being updated for unsync responses.

#rnx

diff --git a/Engine/Source/Programs/Horde/Horde.Server/Artifacts/ArtifactsController.cs b/Engine/Source/Programs/Horde/Horde.Server/Artifacts/ArtifactsController.cs
--- a/Engine/Source/Programs/Horde/Horde.Server/Artifacts/ArtifactsController.cs
+++ b/Engine/Source/Programs/Horde/Horde.Server/Artifacts/ArtifactsController.cs
@@ -765,6 +765,8 @@ namespace Horde.Server.Artifacts
 
 			payloadData.CopyTo(data);
 			Debug.Assert(data.Length == payloadData.Length);
+
+			writer.Advance(length);
 		}
 
 		/// <summary>
