Horde: Include the file modtime in unsync responses.

#rnx

diff --git a/Engine/Source/Programs/Horde/Horde.Server/Artifacts/ArtifactsController.cs b/Engine/Source/Programs/Horde/Horde.Server/Artifacts/ArtifactsController.cs
--- a/Engine/Source/Programs/Horde/Horde.Server/Artifacts/ArtifactsController.cs
+++ b/Engine/Source/Programs/Horde/Horde.Server/Artifacts/ArtifactsController.cs
@@ -635,6 +635,7 @@ namespace Horde.Server.Artifacts
 					{
 						writer.WriteBoolean("read_only", file.ReadOnly);
 					}
+					writer.WriteNumber("mtime", file.ModTime.ToFileTimeUtc());
 					writer.WriteNumber("size", file.Length);
 					writer.WriteNumber("block_size", LeafChunkedDataNodeOptions.Default.TargetSize);
 					writer.WriteStartArray("blocks");
