Horde: Updating schema documentation

#rnx

diff --git a/Engine/Source/Programs/Horde/Docs/Deployment/ServerSettings.md b/Engine/Source/Programs/Horde/Docs/Deployment/ServerSettings.md
--- a/Engine/Source/Programs/Horde/Docs/Deployment/ServerSettings.md
+++ b/Engine/Source/Programs/Horde/Docs/Deployment/ServerSettings.md
@@ -228,6 +228,7 @@ Name | Description
 `Gzip` | Gzip compression
 `Oodle` | Oodle compression (Selkie)
 `Brotli` | Brotli compression
+`Zstd` | ZStandard compression
 
 ## ChunkingOptions
 
