Horde: Add more logging to agent session setup

Resolving of capabilities has been known to stall at times. Adding trace spans to these should help see which part is causing slowdowns.

diff --git a/Engine/Source/Programs/Horde/Horde.Agent/Services/CapabilitiesService.cs b/Engine/Source/Programs/Horde/Horde.Agent/Services/CapabilitiesService.cs
--- a/Engine/Source/Programs/Horde/Horde.Agent/Services/CapabilitiesService.cs
+++ b/Engine/Source/Programs/Horde/Horde.Agent/Services/CapabilitiesService.cs
@@ -17,6 +17,8 @@ using HordeCommon.Rpc.Messages;
 using Microsoft.Extensions.Logging;
 using Microsoft.Extensions.Options;
 using Microsoft.Management.Infrastructure;
+using OpenTracing;
+using OpenTracing.Util;
 using AddressFamily = System.Net.Sockets.AddressFamily;
 
 namespace Horde.Agent.Services
@@ -62,7 +64,9 @@ namespace Horde.Agent.Services
 						_logger.LogWarning("GetCapabilitiesInternalAsync() has been running for {Time}", timer.Elapsed);
 					}
 				}
-				_logger.LogInformation("Agent capabilities queried in {Time} ms", timer.ElapsedMilliseconds);
+				
+				LogLevel logLevel = timer.ElapsedMilliseconds >= 90000 ? LogLevel.Error : LogLevel.Information;
+				_logger.Log(logLevel, "Agent capabilities queried in {Time} ms", timer.ElapsedMilliseconds);
 
 				return await task;
 			}
@@ -75,6 +79,7 @@ namespace Horde.Agent.Services
 
 		async Task<RpcAgentCapabilities> GetCapabilitiesInternalAsync(DirectoryReference? workingDir)
 		{
+			using IScope scope = GlobalTracer.Instance.BuildSpan($"{nameof(CapabilitiesService)}.{nameof(GetCapabilitiesInternalAsync)}").StartActive();
 			ILogger logger = _logger;
 
 			// Create the primary device
@@ -84,6 +89,7 @@ namespace Horde.Agent.Services
 			List<RpcDeviceCapabilities> otherDevices = new List<RpcDeviceCapabilities>();
 			if (RuntimeInformation.IsOSPlatform(OSPlatform.Windows))
 			{
+				using IScope winScope = GlobalTracer.Instance.BuildSpan("GetWindowsCapabilities").StartActive();
 				primaryDevice.Properties.Add("Platform=Win64");
 				primaryDevice.Properties.Add("PlatformGroup=Windows");
 				primaryDevice.Properties.Add("PlatformGroup=Microsoft");
@@ -96,101 +102,115 @@ namespace Horde.Agent.Services
 				{
 					const string QueryNamespace = @"root\cimv2";
 					const string QueryDialect = "WQL";
-
-					// Add OS info
-					foreach (CimInstance instance in session.QueryInstances(QueryNamespace, QueryDialect, "select * from Win32_OperatingSystem"))
+					
 					{
-						foreach (CimProperty property in instance.CimInstanceProperties)
+						// Add OS info
+						using IScope _ = GlobalTracer.Instance.BuildSpan("GetOsInfo").StartActive();
+						foreach (CimInstance instance in session.QueryInstances(QueryNamespace, QueryDialect, "select * from Win32_OperatingSystem"))
 						{
-							string name = property.Name;
-							if (name.Equals("Caption", StringComparison.OrdinalIgnoreCase))
-							{
-								primaryDevice.Properties.Add($"OSDistribution={property.Value}");
-							}
-							else if (name.Equals("Version", StringComparison.OrdinalIgnoreCase))
+							foreach (CimProperty property in instance.CimInstanceProperties)
 							{
-								primaryDevice.Properties.Add($"OSKernelVersion={property.Value}");
+								string name = property.Name;
+								if (name.Equals("Caption", StringComparison.OrdinalIgnoreCase))
+								{
+									primaryDevice.Properties.Add($"OSDistribution={property.Value}");
+								}
+								else if (name.Equals("Version", StringComparison.OrdinalIgnoreCase))
+								{
+									primaryDevice.Properties.Add($"OSKernelVersion={property.Value}");
+								}
 							}
 						}
 					}
-
-					// Add CPU info
-					Dictionary<string, int> cpuNameToCount = new Dictionary<string, int>();
-					int totalPhysicalCores = 0;
-					int totalLogicalCores = 0;
-
-					foreach (CimInstance instance in session.QueryInstances(QueryNamespace, QueryDialect, "select * from Win32_Processor"))
+					
 					{
-						foreach (CimProperty property in instance.CimInstanceProperties)
+						// Add CPU info
+						using IScope _ = GlobalTracer.Instance.BuildSpan("GetCpuInfo").StartActive();
+						Dictionary<string, int> cpuNameToCount = new Dictionary<string, int>();
+						int totalPhysicalCores = 0;
+						int totalLogicalCores = 0;
+						
+						foreach (CimInstance instance in session.QueryInstances(QueryNamespace, QueryDialect, "select * from Win32_Processor"))
 						{
-							string name = property.Name;
-							if (name.Equals("Name", StringComparison.OrdinalIgnoreCase))
+							foreach (CimProperty property in instance.CimInstanceProperties)
 							{
-								string cpuName = property.Value.ToString() ?? String.Empty;
-								int count;
-								cpuNameToCount.TryGetValue(cpuName, out count);
-								cpuNameToCount[cpuName] = count + 1;
-							}
-							else if (name.Equals("NumberOfEnabledCore", StringComparison.OrdinalIgnoreCase) || name.Equals("NumberOfCores", StringComparison.OrdinalIgnoreCase))
-							{
-								if (property.Value is uint numCores)
+								string name = property.Name;
+								if (name.Equals("Name", StringComparison.OrdinalIgnoreCase))
 								{
-									totalPhysicalCores += (int)numCores;
+									string cpuName = property.Value.ToString() ?? String.Empty;
+									int count;
+									cpuNameToCount.TryGetValue(cpuName, out count);
+									cpuNameToCount[cpuName] = count + 1;
 								}
-							}
-							else if (name.Equals("NumberOfLogicalProcessors", StringComparison.OrdinalIgnoreCase))
-							{
-								if (property.Value is uint numCores)
+								else if (name.Equals("NumberOfEnabledCore", StringComparison.OrdinalIgnoreCase) ||
+								         name.Equals("NumberOfCores", StringComparison.OrdinalIgnoreCase))
 								{
-									totalLogicalCores += (int)numCores;
+									if (property.Value is uint numCores)
+									{
+										totalPhysicalCores += (int)numCores;
+									}
+								}
+								else if (name.Equals("NumberOfLogicalProcessors", StringComparison.OrdinalIgnoreCase))
+								{
+									if (property.Value is uint numCores)
+									{
+										totalLogicalCores += (int)numCores;
+									}
 								}
 							}
 						}
+						
+						AddCpuInfo(primaryDevice, cpuNameToCount, totalLogicalCores, totalPhysicalCores);
 					}
-
-					AddCpuInfo(primaryDevice, cpuNameToCount, totalLogicalCores, totalPhysicalCores);
-
-					// Add RAM info
-					ulong totalCapacity = 0;
-					foreach (CimInstance instance in session.QueryInstances(QueryNamespace, QueryDialect, "select Capacity from Win32_PhysicalMemory"))
+					
 					{
-						foreach (CimProperty property in instance.CimInstanceProperties)
+						// Add RAM info
+						using IScope _ = GlobalTracer.Instance.BuildSpan("GetRamInfo").StartActive();
+						ulong totalCapacity = 0;
+						foreach (CimInstance instance in session.QueryInstances(QueryNamespace, QueryDialect, "select Capacity from Win32_PhysicalMemory"))
 						{
-							if (property.Name.Equals("Capacity", StringComparison.OrdinalIgnoreCase) && property.Value is ulong capacity)
+							foreach (CimProperty property in instance.CimInstanceProperties)
 							{
-								totalCapacity += capacity;
+								if (property.Name.Equals("Capacity", StringComparison.OrdinalIgnoreCase) && property.Value is ulong capacity)
+								{
+									totalCapacity += capacity;
+								}
 							}
 						}
+						
+						primaryDevice.Properties.Add($"RAM={totalCapacity / (1024 * 1024 * 1024)}");
 					}
-					primaryDevice.Properties.Add($"RAM={totalCapacity / (1024 * 1024 * 1024)}");
-
-					// Add GPU info
-					int index = 0;
-					foreach (CimInstance instance in session.QueryInstances(QueryNamespace, QueryDialect, "select Name, DriverVersion, AdapterRAM from Win32_VideoController"))
+					
 					{
-						string? name = null;
-						string? driverVersion = null;
-
-						foreach (CimProperty property in instance.CimInstanceProperties)
+						// Add GPU info
+						using IScope _ = GlobalTracer.Instance.BuildSpan("GetGpuInfo").StartActive();
+						int index = 0;
+						foreach (CimInstance instance in session.QueryInstances(QueryNamespace, QueryDialect, "select Name, DriverVersion, AdapterRAM from Win32_VideoController"))
 						{
-							if (property.Name.Equals("Name", StringComparison.OrdinalIgnoreCase))
+							string? name = null;
+							string? driverVersion = null;
+							
+							foreach (CimProperty property in instance.CimInstanceProperties)
 							{
-								name = property.Value.ToString();
-							}
-							else if (property.Name.Equals("DriverVersion", StringComparison.OrdinalIgnoreCase))
-							{
-								driverVersion = property.Value.ToString();
+								if (property.Name.Equals("Name", StringComparison.OrdinalIgnoreCase))
+								{
+									name = property.Value.ToString();
+								}
+								else if (property.Name.Equals("DriverVersion", StringComparison.OrdinalIgnoreCase))
+								{
+									driverVersion = property.Value.ToString();
+								}
 							}
-						}
-
-						if (name != null)
-						{
-							string prefix = $"GPU-{++index}";
-							primaryDevice.Properties.Add($"{prefix}-Name={name}");
-
-							if (driverVersion != null)
+							
+							if (name != null)
 							{
-								primaryDevice.Properties.Add($"{prefix}-DriverVersion={driverVersion}");
+								string prefix = $"GPU-{++index}";
+								primaryDevice.Properties.Add($"{prefix}-Name={name}");
+								
+								if (driverVersion != null)
+								{
+									primaryDevice.Properties.Add($"{prefix}-DriverVersion={driverVersion}");
+								}
 							}
 						}
 					}
@@ -210,6 +230,7 @@ namespace Horde.Agent.Services
 			}
 			else if (RuntimeInformation.IsOSPlatform(OSPlatform.Linux))
 			{
+				using IScope _ = GlobalTracer.Instance.BuildSpan("GetLinuxCapabilities").StartActive();
 				primaryDevice.Properties.Add("Platform=Linux");
 				primaryDevice.Properties.Add("PlatformGroup=Linux");
 				primaryDevice.Properties.Add("PlatformGroup=Unix");
@@ -281,6 +302,7 @@ namespace Horde.Agent.Services
 			}
 			else if (RuntimeInformation.IsOSPlatform(OSPlatform.OSX))
 			{
+				using IScope _ = GlobalTracer.Instance.BuildSpan("GetMacCapabilities").StartActive();
 				primaryDevice.Properties.Add("Platform=Mac");
 				primaryDevice.Properties.Add("PlatformGroup=Apple");
 				primaryDevice.Properties.Add("PlatformGroup=Desktop");
@@ -361,6 +383,7 @@ namespace Horde.Agent.Services
 			// Get the IP addresses
 			try
 			{
+				using IScope _ = GlobalTracer.Instance.BuildSpan("ResolveIp").StartActive();
 				using CancellationTokenSource dnsCts = new(3000);
 				IPHostEntry entry = await Dns.GetHostEntryAsync(Dns.GetHostName(), dnsCts.Token);
 				foreach (IPAddress address in entry.AddressList)
@@ -486,6 +509,7 @@ namespace Horde.Agent.Services
 		/// <returns>Local IP address of this machine</returns>
 		public static async Task<IPAddress?> GetLocalIpAddressAsync(string hostname, int timeoutMs = 2000)
 		{
+			using IScope _ = GlobalTracer.Instance.BuildSpan(nameof(GetLocalIpAddressAsync)).StartActive();
 			try
 			{
 				using Socket socket = new(AddressFamily.InterNetwork, SocketType.Dgram, ProtocolType.IP);
@@ -583,6 +607,7 @@ namespace Horde.Agent.Services
 
 		static async Task AddAwsPropertiesInternalAsync(IList<string> properties, ILogger logger)
 		{
+			using IScope scope = GlobalTracer.Instance.BuildSpan(nameof(AddAwsPropertiesInternalAsync)).StartActive();
 			if (EC2InstanceMetadata.IdentityDocument != null)
 			{
 				properties.Add("EC2=1");
diff --git a/Engine/Source/Programs/Horde/Horde.Agent/Services/WorkerService.cs b/Engine/Source/Programs/Horde/Horde.Agent/Services/WorkerService.cs
--- a/Engine/Source/Programs/Horde/Horde.Agent/Services/WorkerService.cs
+++ b/Engine/Source/Programs/Horde/Horde.Agent/Services/WorkerService.cs
@@ -128,7 +128,7 @@ namespace Horde.Agent.Services
 						}
 
 						TimeSpan backOffTime = s_sessionBackOffTime[Math.Min(failureCount - 1, s_sessionBackOffTime.Length - 1)];
-						_logger.LogInformation(ex, "Session failure #{FailureNum}. Waiting {Time} and restarting. ({Message})", failureCount, backOffTime, ex.Message);
+						_logger.LogWarning(ex, "Session failure #{FailureNum}. Waiting {Time} and restarting. ({Message})", failureCount, backOffTime, ex.Message);
 						_statusService.Set(false, 0, $"Unable to start session: {ex.Message}");
 						await Task.Delay(backOffTime, stoppingToken);
 					}
