Horde: Fix newlines in Additional Args advanced job options
#rnx

diff --git a/Engine/Source/Programs/Horde/HordeDashboard/src/components/build/NewBuildV2.tsx b/Engine/Source/Programs/Horde/HordeDashboard/src/components/build/NewBuildV2.tsx
--- a/Engine/Source/Programs/Horde/HordeDashboard/src/components/build/NewBuildV2.tsx
+++ b/Engine/Source/Programs/Horde/HordeDashboard/src/components/build/NewBuildV2.tsx
@@ -1805,7 +1805,7 @@ const BuildModal: React.FC<{ setUseLegacyDialog: (value: boolean) => void }> = o
       let additionalArgs: string[] = [];
       if (options!.advAdditionalArgs) {
          const argRegex = /"(\\"|[^"])*?"|[^ ]+/g;
-         options.advAdditionalArgs.trim().match(argRegex)?.forEach(arg => additionalArgs.push(arg.replace(/"/g, "")));
+         options.advAdditionalArgs.trim().replaceAll("\n", " ").match(argRegex)?.forEach(arg => additionalArgs.push(arg.replace(/"/g, "")));
       }
 
 
