Horde: Show use who started the job in compact views when available
#rnx

diff --git a/Engine/Source/Programs/Horde/HordeDashboard/src/components/JobViewIncremental.tsx b/Engine/Source/Programs/Horde/HordeDashboard/src/components/JobViewIncremental.tsx
--- a/Engine/Source/Programs/Horde/HordeDashboard/src/components/JobViewIncremental.tsx
+++ b/Engine/Source/Programs/Horde/HordeDashboard/src/components/JobViewIncremental.tsx
@@ -50,12 +50,12 @@ const buildColumns = (jobTab: JobsTabData, streamId: string): IColumn[] => {
    const fixedWidths: Record<string, number | undefined> = {
       "Status": 16, // note this doesn't have header text, and need to match the status dot width of 20 after sizing
       "Time": 60,
-      "Author": 180,
+      "Started By": 180,
    };
 
    const minWidths: Record<string, number | undefined> = {};
 
-   let cnames = ["Status", "Author"];
+   let cnames = ["Status", "Started By"];
 
    if (jobTab.columns) {
 
@@ -499,11 +499,16 @@ const JobList: React.FC<{ tab: string; filter: JobFilterSimple, controller: Call
 
       const commit = commitCache.getCommit(streamId!, item.job!.change!);
 
-      if (column!.key === "Author") {
+      if (column!.key === "Started By") {
 
-         let authorName = commit?.authorInfo?.name;
-         if (item.job.preflightChange) {
-            authorName = item.job.startedByUserInfo?.name;
+         let authorName = item.job.startedByUserInfo?.name;
+
+         if (!authorName) {
+            authorName = commit?.authorInfo?.name;
+         }
+
+         if (!authorName) {
+            authorName = "Scheduler";
          }
 
          let change = item.job?.change?.toString() ?? "Latest";
@@ -557,7 +562,7 @@ const JobList: React.FC<{ tab: string; filter: JobFilterSimple, controller: Call
 
    function onRenderItemColumn(item: JobItem, index?: number, column?: IColumn) {
 
-      if (column?.key === "Author") {
+      if (column?.key === "Started By") {
          return onRenderItemColumnInner(item, index, column);
       }
 
