Horde: Bump dashboard version
#rnx
#skipci

diff --git a/Engine/Source/Programs/Horde/HordeDashboard/package.json b/Engine/Source/Programs/Horde/HordeDashboard/package.json
--- a/Engine/Source/Programs/Horde/HordeDashboard/package.json
+++ b/Engine/Source/Programs/Horde/HordeDashboard/package.json
@@ -1,7 +1,7 @@
 {
   "name": "horde-dashboard",
   "private": true,
-  "version": "0.1.1",
+  "version": "0.1.2",
   "type": "module",  
   "scripts": {
     "dev": "vite",
