Horde: Remove background task which resolves blob imports to ids. This should no longer be necessary since we specify imports at upload time.

diff --git a/Engine/Source/Programs/Horde/Horde.Server/Storage/StorageService.cs b/Engine/Source/Programs/Horde/Horde.Server/Storage/StorageService.cs
--- a/Engine/Source/Programs/Horde/Horde.Server/Storage/StorageService.cs
+++ b/Engine/Source/Programs/Horde/Horde.Server/Storage/StorageService.cs
@@ -383,7 +383,6 @@ namespace Horde.Server.Storage
 		readonly IMongoCollection<BlobInfo> _blobCollection;
 		readonly IMongoCollection<RefInfo> _refCollection;
 
-		readonly ITicker _blobTicker;
 		readonly ITicker _refTicker;
 
 		readonly SingletonDocument<GcState> _gcState;
@@ -422,7 +421,6 @@ namespace Horde.Server.Storage
 			refIndexes.Add(keys => keys.Descending(x => x.ExpiresAtUtc), sparse: true);
 			_refCollection = mongoService.GetCollection<RefInfo>("Storage.Refs", refIndexes);
 
-			_blobTicker = clock.AddSharedTicker("Storage:Blobs", TimeSpan.FromMinutes(5.0), TickBlobsAsync, _logger);
 			_refTicker = clock.AddSharedTicker("Storage:Refs", TimeSpan.FromMinutes(5.0), TickRefsAsync, _logger);
 
 			_gcState = new SingletonDocument<GcState>(mongoService);
@@ -432,7 +430,6 @@ namespace Horde.Server.Storage
 		/// <inheritdoc/>
 		public async ValueTask DisposeAsync()
 		{
-			await _blobTicker.DisposeAsync();
 			await _refTicker.DisposeAsync();
 			await _gcTicker.DisposeAsync();
 		}
@@ -463,7 +460,6 @@ namespace Horde.Server.Storage
 		/// <inheritdoc/>
 		public async Task StartAsync(CancellationToken cancellationToken)
 		{
-			await _blobTicker.StartAsync();
 			await _refTicker.StartAsync();
 			await _gcTicker.StartAsync();
 		}
@@ -473,7 +469,6 @@ namespace Horde.Server.Storage
 		{
 			await _gcTicker.StopAsync();
 			await _refTicker.StopAsync();
-			await _blobTicker.StopAsync();
 		}
 
 		/// <inheritdoc/>
@@ -621,110 +616,6 @@ namespace Horde.Server.Storage
 			return false;
 		}
 
-		/// <summary>
-		/// Finds blobs at least 30 minutes old and computes import metadata for them. Done with a delay to allow write redirects.
-		/// </summary>
-		/// <param name="cancellationToken">Cancellation token for the operation</param>
-		async ValueTask TickBlobsAsync(CancellationToken cancellationToken)
-		{
-			using TelemetrySpan span = _tracer.StartActiveSpan($"{nameof(StorageService)}.{nameof(TickBlobsAsync)}");
-
-			GcState gcState = await _gcState.GetAsync(cancellationToken);
-			DateTime utcNow = _clock.UtcNow;
-
-			// Get the current state of the storage system
-			State state = CreateState(_globalConfig.CurrentValue);
-
-			Dictionary<NamespaceId, BundleStorageClient> cachedClients = new();
-			try
-			{
-				// Compute missing import info, by searching for blobs with an ObjectId timestamp after the last import compute cycle
-				ObjectId latestInfoId = ObjectId.GenerateNewId(utcNow - TimeSpan.FromMinutes(30.0));
-				using (IAsyncCursor<BlobInfo> cursor = await _blobCollection.Find(x => x.Id >= gcState.LastImportBlobInfoId && x.Id < latestInfoId).ToCursorAsync(cancellationToken))
-				{
-					while (await cursor.MoveNextAsync(cancellationToken))
-					{
-						// Find imports, and add a check record for each new blob
-						foreach (BlobInfo blobInfo in cursor.Current)
-						{
-							NamespaceInfo? namespaceInfo;
-							if (state.Namespaces.TryGetValue(blobInfo.NamespaceId, out namespaceInfo))
-							{
-								BundleStorageClient? storageClient;
-								if (!cachedClients.TryGetValue(namespaceInfo.Id, out storageClient))
-								{
-									storageClient = new BundleStorageClient(namespaceInfo.Backend, _bundleCache, null, _logger);
-									cachedClients.Add(namespaceInfo.Id, storageClient);
-								}
-
-								try
-								{
-									await TickBlobAsync(storageClient, blobInfo, cancellationToken);
-								}
-								catch (ObjectNotFoundException ex)
-								{
-									_logger.LogInformation(ex, "Unable to read references for {NamespaceId} blob {BlobId}: {Message}", blobInfo.NamespaceId, blobInfo.Id, ex.Message);
-								}
-								catch (Exception ex)
-								{
-									_logger.LogWarning(ex, "Unable to read references for {NamespaceId} blob {BlobId} (key: {ObjectKey}): {Message}", blobInfo.NamespaceId, blobInfo.Id, GetObjectKey(blobInfo.Locator), ex.Message);
-								}
-							}
-						}
-
-						// Update the last imported blob id
-						await _gcState.UpdateAsync(state => state.LastImportBlobInfoId = latestInfoId, cancellationToken);
-					}
-				}
-			}
-			finally
-			{
-				foreach (BundleStorageClient client in cachedClients.Values)
-				{
-					client.Dispose();
-				}
-			}
-		}
-
-		async Task TickBlobAsync(BundleStorageClient storageClient, BlobInfo blobInfo, CancellationToken cancellationToken)
-		{
-			List<ObjectId> importInfoIds = new List<ObjectId>();
-
-			IEnumerable<BlobLocator> importLocators = await storageClient.ReadBundleReferencesAsync(blobInfo.Locator, cancellationToken);
-			foreach (BlobLocator importLocator in importLocators)
-			{
-				string importPath = importLocator.BaseLocator.ToString();
-
-				FilterDefinition<BlobInfo> filter = Builders<BlobInfo>.Filter.Expr(x => x.NamespaceId == blobInfo.NamespaceId && x.Path == importPath);
-				UpdateDefinition<BlobInfo> update = Builders<BlobInfo>.Update.SetOnInsert(x => x.Imports, null);
-				BlobInfo blobInfoDoc = await _blobCollection.FindOneAndUpdateAsync(filter, update, new FindOneAndUpdateOptions<BlobInfo> { IsUpsert = true, ReturnDocument = ReturnDocument.After }, cancellationToken);
-
-				importInfoIds.Add(blobInfoDoc.Id);
-			}
-
-			importInfoIds.Sort();
-
-			if (blobInfo.Imports != null && !Enumerable.SequenceEqual(importInfoIds, blobInfo.Imports))
-			{
-				List<ObjectId> missing = importInfoIds.Except(blobInfo.Imports).ToList();
-				if (missing.Count > 0)
-				{
-					_logger.LogWarning("Missing imports for blob {Locator}: {Missing}", blobInfo.Path, String.Join(", ", missing.Select(x => x.ToString())));
-				}
-
-				List<ObjectId> extra = blobInfo.Imports.Except(importInfoIds).ToList();
-				if (extra.Count > 0)
-				{
-					_logger.LogWarning("Extra imports for blob {Locator}: {Extra}", blobInfo.Path, String.Join(", ", extra.Select(x => x.ToString())));
-				}
-			}
-
-			await _blobCollection.UpdateOneAsync(x => x.Id == blobInfo.Id, Builders<BlobInfo>.Update.Set(x => x.Imports, importInfoIds), null, cancellationToken);
-
-			AddGcCheckRecord(blobInfo.NamespaceId, blobInfo.Id);
-			_logger.LogDebug("Added {Count} imports for {NamespaceId} blob {BlobId}", importInfoIds.Count, blobInfo.NamespaceId, blobInfo.Id);
-		}
-
 		#endregion
 
 		#region Nodes
