Horde: Add wait time for agent to pool view
#rnx
#skipci

diff --git a/Engine/Source/Programs/Horde/HordeDashboard/src/components/PoolView.tsx b/Engine/Source/Programs/Horde/HordeDashboard/src/components/PoolView.tsx
--- a/Engine/Source/Programs/Horde/HordeDashboard/src/components/PoolView.tsx
+++ b/Engine/Source/Programs/Horde/HordeDashboard/src/components/PoolView.tsx
@@ -523,9 +523,10 @@ const BatchPanel: React.FC = () => {
    let columns: IColumn[] = [];
 
    columns = [
-      { key: 'column1', name: 'Job', minWidth: 800, maxWidth: 800 },
+      { key: 'column1', name: 'Job', minWidth: 700, maxWidth: 700 },
       { key: 'column2', name: 'Job Created', minWidth: 200, maxWidth: 200 },
       { key: 'column3', name: 'Status', minWidth: 200, maxWidth: 200 },
+      { key: 'column4', name: 'WaitTime', minWidth: 100, maxWidth: 100 },
    ];
 
    const batchItems: BatchItem[] = batches.map(b => {
@@ -554,9 +555,27 @@ const BatchPanel: React.FC = () => {
             statusText = `Waiting for agent`;
          }
 
-         return <Text>{statusText}</Text>
+         return <Stack horizontalAlign="end"><Text>{statusText}</Text></Stack>
       }
 
+      if (column.name === "WaitTime") {
+
+         let statusText = "";
+
+         if (batch.state === JobStepBatchState.Ready) {            
+            if (batch.readyTime) {
+               statusText += `${getElapsedString(moment(batch.readyTime), moment(Date.now()), false)}`;
+            }
+         }
+
+         if (!statusText) {
+            return null;
+         }
+
+         return <Stack style={{ paddingRight: 32 }} horizontalAlign="end"><Text>{statusText}</Text></Stack>
+      }
+
+
       if (column.name === "Job Created") {
 
          return <Stack ><Text style={{ fontSize: "13px", paddingRight: 12 }}>{getShortNiceTime(job.createTime, true, true)}</Text></Stack>;
