Horde: Sort requested unsync blobs to optimize download/extract order.

#rnx

diff --git a/Engine/Source/Programs/Horde/Horde.Server.Tests/Artifacts/UnsyncCacheTests.cs b/Engine/Source/Programs/Horde/Horde.Server.Tests/Artifacts/UnsyncCacheTests.cs
--- a/Engine/Source/Programs/Horde/Horde.Server.Tests/Artifacts/UnsyncCacheTests.cs
+++ b/Engine/Source/Programs/Horde/Horde.Server.Tests/Artifacts/UnsyncCacheTests.cs
@@ -65,8 +65,10 @@ namespace Horde.Server.Tests.Artifacts
 			int offset = 0;
 			foreach (UnsyncBlock block in file.Blocks)
 			{
-				using BlobData? blobData = await cache.ReadBlobAsync(artifact.Object, block.Blob.Hash);
-				Assert.IsNotNull(blobData);
+				IBlobRef? blobRef = await cache.ReadBlobRefAsync(artifact.Object, block.Blob.Hash);
+				Assert.IsNotNull(blobRef);
+
+				using BlobData blobData = await blobRef.ReadBlobDataAsync();
 
 				Assert.IsTrue(blobData.Data.Span.SequenceEqual(source.AsSpan(offset, blobData.Data.Length)));
 				offset += (int)block.Length;
diff --git a/Engine/Source/Programs/Horde/Horde.Server/Artifacts/ArtifactsController.cs b/Engine/Source/Programs/Horde/Horde.Server/Artifacts/ArtifactsController.cs
--- a/Engine/Source/Programs/Horde/Horde.Server/Artifacts/ArtifactsController.cs
+++ b/Engine/Source/Programs/Horde/Horde.Server/Artifacts/ArtifactsController.cs
@@ -706,6 +706,27 @@ namespace Horde.Server.Artifacts
 				}
 			}
 
+			// Find all the blob refs that we need to fetch
+			List<IBlobRef> blobRefs = new List<IBlobRef>();
+			foreach (string block in request.Blocks)
+			{
+				if (!IoHash.TryParse(block, out IoHash hash))
+				{
+					return BadRequest($"Invalid IoHash value: {block}");
+				}
+
+				IBlobRef? blobRef = await _unsyncCache.ReadBlobRefAsync(artifact, hash, cancellationToken);
+				if (blobRef == null)
+				{
+					return NotFound($"Hash '{hash}' is not part of artifact {artifact.Id}");
+				}
+
+				blobRefs.Add(blobRef);
+			}
+
+			// Sort them to optimize for coherency
+			blobRefs.Sort((x, y) => DirectoryNodeExtract.CompareBlobs(x, y));
+
 			// Send the response headers
 			HttpResponse response = HttpContext.Response;
 			response.ContentType = "application/x-horde-unsync-blob";
@@ -715,32 +736,21 @@ namespace Horde.Server.Artifacts
 			ArrayMemoryWriter? compressedWriter = null;
 
 			await response.StartAsync(cancellationToken);
-			foreach (string block in request.Blocks)
+			foreach (IBlobRef blobRef in blobRefs)
 			{
-				if (!IoHash.TryParse(block, out IoHash hash))
-				{
-					throw new InvalidOperationException();
-				}
-
-				BlobData? blobData = await _unsyncCache.ReadBlobAsync(artifact, hash, cancellationToken);
-				if (blobData == null)
-				{
-					throw new InvalidOperationException();
-				}
-
+				BlobData? blobData = await blobRef.ReadBlobDataAsync(cancellationToken);
 				if (compress)
 				{
 					compressedWriter ??= new ArrayMemoryWriter(300 * 1024);
 					compressedWriter.Clear();
 					BundleData.Compress(BundleCompressionFormat.Zstd, blobData.Data, compressedWriter);
 
-					WriteBlock(response.BodyWriter, hash, blobData.Data.Length, compressedWriter.WrittenSpan);
+					WriteBlock(response.BodyWriter, blobRef.Hash, blobData.Data.Length, compressedWriter.WrittenSpan);
 				}
 				else
 				{
-					WriteBlock(response.BodyWriter, hash, blobData.Data.Length, blobData.Data.Span);
+					WriteBlock(response.BodyWriter, blobRef.Hash, blobData.Data.Length, blobData.Data.Span);
 				}
-
 				await response.BodyWriter.FlushAsync(cancellationToken);
 			}
 			await response.CompleteAsync();
diff --git a/Engine/Source/Programs/Horde/Horde.Server/Artifacts/UnsyncCache.cs b/Engine/Source/Programs/Horde/Horde.Server/Artifacts/UnsyncCache.cs
--- a/Engine/Source/Programs/Horde/Horde.Server/Artifacts/UnsyncCache.cs
+++ b/Engine/Source/Programs/Horde/Horde.Server/Artifacts/UnsyncCache.cs
@@ -124,7 +124,7 @@ namespace Horde.Server.Artifacts
 		/// <param name="blobHash"></param>
 		/// <param name="cancellationToken"></param>
 		/// <returns></returns>
-		public async ValueTask<BlobData?> ReadBlobAsync(IArtifact artifact, IoHash blobHash, CancellationToken cancellationToken = default)
+		public async ValueTask<IBlobRef?> ReadBlobRefAsync(IArtifact artifact, IoHash blobHash, CancellationToken cancellationToken = default)
 		{
 			ArtifactInfo? artifactInfo = await GetArtifactInfoAsync(artifact, cancellationToken);
 			if (artifactInfo == null)
@@ -138,7 +138,7 @@ namespace Horde.Server.Artifacts
 				return null;
 			}
 
-			return await blobRef.ReadBlobDataAsync(cancellationToken);
+			return blobRef;
 		}
 
 		async Task<ArtifactInfo?> ReadArtifactAsync(IArtifact artifact, CancellationToken cancellationToken)
diff --git a/Engine/Source/Programs/Shared/EpicGames.Horde/Storage/Nodes/DirectoryNode.Extract.cs b/Engine/Source/Programs/Shared/EpicGames.Horde/Storage/Nodes/DirectoryNode.Extract.cs
--- a/Engine/Source/Programs/Shared/EpicGames.Horde/Storage/Nodes/DirectoryNode.Extract.cs
+++ b/Engine/Source/Programs/Shared/EpicGames.Horde/Storage/Nodes/DirectoryNode.Extract.cs
@@ -1,6 +1,7 @@
 // Copyright Epic Games, Inc. All Rights Reserved.
 
 using System;
+using System.Collections;
 using System.Collections.Generic;
 using System.Diagnostics;
 using System.Diagnostics.CodeAnalysis;
@@ -430,6 +431,35 @@ namespace EpicGames.Horde.Storage.Nodes
 
 		#region Group requests by bundles
 
+		/// <summary>
+		/// Compare two blob handles to allow sorting for optimized reads
+		/// </summary>
+		public static int CompareBlobs(IBlobHandle a, IBlobHandle b)
+		{
+			if (a.Innermost is not ExportHandle exportHandleA || exportHandleA.Packet is not FlushedPacketHandle packetHandleA)
+			{
+				return Comparer.Default.Compare(a, b);
+			}
+
+			if (b.Innermost is not ExportHandle exportHandleB || exportHandleB.Packet is not FlushedPacketHandle packetHandleB)
+			{
+				return Comparer.Default.Compare(a, b);
+			}
+
+			if (packetHandleA.Bundle != packetHandleB.Bundle)
+			{
+				return Comparer.Default.Compare(packetHandleA.Bundle, packetHandleB.Bundle);
+			}
+			else if (packetHandleA.PacketOffset != packetHandleB.PacketOffset)
+			{
+				return Comparer.Default.Compare(packetHandleA.PacketOffset, packetHandleB.PacketOffset);
+			}
+			else
+			{
+				return Comparer.Default.Compare(exportHandleA.ExportIdx, exportHandleB.ExportIdx);
+			}
+		}
+
 		record class OutputExport(BundleHandle BundleHandle, int PacketOffset, int ExportIdx, OutputChunk Chunk);
 		record class OutputExportBatch(BundleHandle BundleHandle, List<OutputExport> Exports);
 
