Horde: Fix overflow on retry steps
#rnx
#skipci

diff --git a/Engine/Source/Programs/Horde/HordeDashboard/src/components/jobDetailsV2/StepRetryModal.tsx b/Engine/Source/Programs/Horde/HordeDashboard/src/components/jobDetailsV2/StepRetryModal.tsx
--- a/Engine/Source/Programs/Horde/HordeDashboard/src/components/jobDetailsV2/StepRetryModal.tsx
+++ b/Engine/Source/Programs/Horde/HordeDashboard/src/components/jobDetailsV2/StepRetryModal.tsx
@@ -1,5 +1,5 @@
 // Copyright Epic Games, Inc. All Rights Reserved.
-import { Checkbox, DefaultButton, IconButton, Modal, PrimaryButton, Spinner, SpinnerSize, Stack, Text, TextField } from '@fluentui/react';
+import { Checkbox, DefaultButton, IconButton, Modal, PrimaryButton, ScrollablePane, ScrollbarVisibility, Spinner, SpinnerSize, Stack, Text, TextField } from '@fluentui/react';
 import React, { useState } from 'react';
 import backend from '../../backend';
 import { CreateJobRequest } from '../../backend/Api';
@@ -85,7 +85,7 @@ export const RetryStepsModal: React.FC<{ stepIds: string[]; jobDetails: JobDetai
       return null;
    }
 
-   const headerText = retrySteps.size > 1 ? "Retry Steps?" : "Retry Step?";
+   const headerText = stepIds.length > 1 ? "Retry Steps?" : "Retry Step?";
 
    const onRetry = async () => {
 
@@ -170,9 +170,14 @@ export const RetryStepsModal: React.FC<{ stepIds: string[]; jobDetails: JobDetai
       }} /><Text>{name}</Text></Stack>
    }))
 
-   return <Modal className={hordeClasses.modal} isOpen={true} styles={{ main: { padding: 8, width: 800 } }} onDismiss={() => { onClose() }}>
+   let height = Math.max(stepElements.length * 32, 128);
+   if (height > 432) {
+      height = 432;
+   }
+
+   return <Modal className={hordeClasses.modal} isOpen={true} styles={{ main: { padding: 8, width: 800, maxHeight: 640 } }} onDismiss={() => { onClose() }}>
       <Stack horizontal styles={{ root: { padding: 8 } }}>
-         <Stack.Item grow={2}>
+         <Stack.Item grow={2} style={{paddingLeft: 8}}>
             <Text variant="mediumPlus">{headerText}</Text>
          </Stack.Item>
          <Stack.Item grow={0}>
@@ -184,12 +189,18 @@ export const RetryStepsModal: React.FC<{ stepIds: string[]; jobDetails: JobDetai
          </Stack.Item>
       </Stack>
 
-      <Stack styles={{ root: { paddingLeft: 32 } }} tokens={{ childrenGap: 8 }}>
-         {stepElements}
+      <Stack style={{marginRight: 18, marginTop: 18}}>
+         <Stack styles={{ root: { position: "relative", height: height } }}>
+            <ScrollablePane scrollbarVisibility={ScrollbarVisibility.auto}>
+               <Stack style={{ paddingLeft: 24, paddingRight: 24, height: height - 12 }} tokens={{ childrenGap: 8 }}>
+                  {stepElements}
+               </Stack>
+            </ScrollablePane>
+         </Stack>
       </Stack>
 
       <Stack styles={{ root: { padding: 8 } }}>
-         <Stack horizontal tokens={{ childrenGap: 16 }} styles={{ root: { paddingTop: 12, paddingLeft: 8, paddingBottom: 8 } }}>
+         <Stack horizontal tokens={{ childrenGap: 16 }} styles={{ root: { paddingTop: 24, paddingLeft: 8, paddingBottom: 8, paddingRight: 12 } }}>
             <Stack grow />
             <PrimaryButton text="Retry" disabled={submitting || retrySteps.size === 0} onClick={() => { onRetry(); setSubmitting(true) }} />
             <DefaultButton text="Cancel" disabled={submitting} onClick={() => { onClose(); }} />
