Horde: Tweak MongoDB client's config defaults to curb wait queue full errors

diff --git a/Engine/Source/Programs/Horde/Horde.Server/Server/MongoService.cs b/Engine/Source/Programs/Horde/Horde.Server/Server/MongoService.cs
--- a/Engine/Source/Programs/Horde/Horde.Server/Server/MongoService.cs
+++ b/Engine/Source/Programs/Horde/Horde.Server/Server/MongoService.cs
@@ -343,7 +343,12 @@ namespace Horde.Server.Server
 				mongoSettings.LinqProvider = LinqProvider.V2;
 				mongoSettings.SslSettings = new SslSettings();
 				mongoSettings.SslSettings.ServerCertificateValidationCallback = CertificateValidationCallBack;
+				
+				// Increase defaults to handle larger deployments of Horde interfacing with MongoDB
+				// Making room for slower running commands potentially hogging the queue 
+				mongoSettings.MinConnectionPoolSize = 50; // Default is 0
 				mongoSettings.MaxConnectionPoolSize = 300; // Default is 100
+				mongoSettings.MaxConnecting = 10; // Default is 2
 
 				//TestSslConnection(MongoSettings.Server.Host, MongoSettings.Server.Port, Logger);
 
@@ -683,7 +688,16 @@ namespace Horde.Server.Server
 			_logger.LogError("TLS certificate validation failed ({Errors}).{AdditionalInfo}", sslPolicyErrors, StringUtils.Indent(builder.ToString(), "    "));
 			return false;
 		}
-
+		
+		/// <summary>
+		/// Get the MongoDB client
+		/// </summary>
+		/// <returns>A MongoDB client instance</returns>
+		public MongoClient GetClient()
+		{
+			return _client;
+		}
+		
 		/// <summary>
 		/// Get a MongoDB collection from database
 		/// </summary>
