Horde: Disable sorting for unsync blobs; failing for items without a comparer.

#rnx

diff --git a/Engine/Source/Programs/Horde/Horde.Server/Artifacts/ArtifactsController.cs b/Engine/Source/Programs/Horde/Horde.Server/Artifacts/ArtifactsController.cs
--- a/Engine/Source/Programs/Horde/Horde.Server/Artifacts/ArtifactsController.cs
+++ b/Engine/Source/Programs/Horde/Horde.Server/Artifacts/ArtifactsController.cs
@@ -725,7 +725,7 @@ namespace Horde.Server.Artifacts
 			}
 
 			// Sort them to optimize for coherency
-			blobRefs.Sort((x, y) => DirectoryNodeExtract.CompareBlobs(x, y));
+//			blobRefs.Sort((x, y) => DirectoryNodeExtract.CompareBlobs(x, y));
 
 			// Send the response headers
 			HttpResponse response = HttpContext.Response;
