Horde: Fix overflow on new build dialog
#skipci
#rnx

diff --git a/Engine/Source/Programs/Horde/HordeDashboard/src/components/build/NewBuildV2.tsx b/Engine/Source/Programs/Horde/HordeDashboard/src/components/build/NewBuildV2.tsx
--- a/Engine/Source/Programs/Horde/HordeDashboard/src/components/build/NewBuildV2.tsx
+++ b/Engine/Source/Programs/Horde/HordeDashboard/src/components/build/NewBuildV2.tsx
@@ -368,7 +368,7 @@ class BuildOptions {
    private changed = 0;
 
    @action
-   setChanged() {      
+   setChanged() {
       this.changed++;
    }
 
@@ -1299,7 +1299,11 @@ const BuildParametersPanel: React.FC = observer(() => {
       return null;
    }
 
-   const estimatedHeight = options.estimateHeight();
+   let estimatedHeight = options.estimateHeight();
+
+   if (estimatedHeight > window.innerHeight - 360) {
+      estimatedHeight = window.innerHeight - 360;
+   }
 
    const renderParameter = (param: ParameterData, disabled?: boolean) => {
 
@@ -1504,7 +1508,11 @@ const PreviewPanel: React.FC = observer(() => {
       options.onPreviewSourceChanged(csource);
    }
 
-   const estimatedHeight = options.estimateHeight(true);
+   let estimatedHeight = options.estimateHeight(true);
+
+   if (estimatedHeight > window.innerHeight - 360) {
+      estimatedHeight = window.innerHeight - 360;
+   }
 
    return <Stack>
       <Stack style={{ paddingTop: 64, paddingBottom: 12 }}>
@@ -1512,7 +1520,7 @@ const PreviewPanel: React.FC = observer(() => {
             Template Source
          </Label>
          <Stack>
-            <TextField key={`text_counter_${textCounter++}`} componentRef={textRef} style={{ width: parameterWidth + 4, height: estimatedHeight }} multiline spellCheck={false} resizable={false} defaultValue={options.previewSource} />
+            <TextField key={`text_counter_${textCounter++}`} componentRef={textRef} style={{ width: 670, height: estimatedHeight }} multiline spellCheck={false} resizable={false} defaultValue={options.previewSource} />
          </Stack>
          <Stack>
          </Stack>
@@ -1605,6 +1613,10 @@ const AdvancedPanel: React.FC = observer(() => {
    // targets
    height += 92;
 
+   if (height > window.innerHeight - 360) {
+      height = window.innerHeight - 360;
+   }
+
    // target options for picker
    type TargetPickerItem = {
       key: string;
@@ -1619,129 +1631,130 @@ const AdvancedPanel: React.FC = observer(() => {
    });
    return <Stack style={{
       height: height,
-      position: 'relative',
+      position: 'relative',      
       width: parameterWidth,
-      maxHeight: 'calc(100vh - 360px)'
    }}>
-      <Stack style={{ paddingBottom: 12 }} tokens={{ childrenGap: 12 }}>
-         {!!stream.configRevision && <Stack>
-            <TextField label="Template Path" readOnly={true} value={stream.configPath ?? ""} />
-         </Stack>
-         }
-         <Stack>
-            <Dropdown disabled={options.readOnly} key={`key_adv_priority_${options.currentRenderKey}`} defaultValue={options.advJobPriority} label="Priority" options={priorityOptions} onChange={(ev, option) => {
-               options.advJobPriority = option?.key as Priority;
-               options.setChanged();
+      <ScrollablePane scrollbarVisibility={ScrollbarVisibility.auto}>
+         <Stack style={{ paddingBottom: 12, paddingRight: 12 }} tokens={{ childrenGap: 12 }}>
+            {!!stream.configRevision && <Stack>
+               <TextField label="Template Path" readOnly={true} value={stream.configPath ?? ""} />
+            </Stack>
             }
-            } />
-         </Stack>
-         <Stack>
-            <Checkbox disabled={options.readOnly || options.template?.updateIssues} key={`key_adv_update_issues_${options.currentRenderKey}`}
-               label="Update Build Health Issues"
-               defaultChecked={options.template?.updateIssues ? true : options.advUpdateIssues}
-               onChange={(ev, checked) => {
-                  options.advUpdateIssues = checked ? true : undefined;
+            <Stack>
+               <Dropdown disabled={options.readOnly} key={`key_adv_priority_${options.currentRenderKey}`} defaultValue={options.advJobPriority} label="Priority" options={priorityOptions} onChange={(ev, option) => {
+                  options.advJobPriority = option?.key as Priority;
+                  options.setChanged();
+               }
+               } />
+            </Stack>
+            <Stack>
+               <Checkbox disabled={options.readOnly || options.template?.updateIssues} key={`key_adv_update_issues_${options.currentRenderKey}`}
+                  label="Update Build Health Issues"
+                  defaultChecked={options.template?.updateIssues ? true : options.advUpdateIssues}
+                  onChange={(ev, checked) => {
+                     options.advUpdateIssues = checked ? true : undefined;
+                     options.setChanged();
+                  }} />
+            </Stack>
+            <Stack>
+               <TextField key={`key_adv_job_name_${options.currentRenderKey}`} disabled={options.readOnly} spellCheck={false} defaultValue={options.advJobName} label="Job Name" onChange={(ev, newValue) => {
+                  options.advJobName = newValue;
                   options.setChanged();
                }} />
-         </Stack>
-         <Stack>
-            <TextField key={`key_adv_job_name_${options.currentRenderKey}`} disabled={options.readOnly} spellCheck={false} defaultValue={options.advJobName} label="Job Name" onChange={(ev, newValue) => {
-               options.advJobName = newValue;
-               options.setChanged();
-            }} />
-         </Stack>
-         <Stack>
-            <Label>Targets</Label>
-            <TagPicker
-
-               componentRef={targetPicker}
+            </Stack>
+            <Stack>
+               <Label>Targets</Label>
+               <TagPicker
 
-               disabled={options.readOnly}
+                  componentRef={targetPicker}
 
-               onBlur={(ev) => {
-                  if (ev.target.value && ev.target.value.trim()) {
-                     options.addTarget(ev.target.value.trim())
-                  }
+                  disabled={options.readOnly}
 
-                  // This is using undocumented behavior to clear the input when you lose focus
-                  // It could very well break
-                  if (targetPicker.current) {
-                     try {
-                        (targetPicker.current as any).input.current._updateValue("");
-                     } catch (reason) {
-                        console.error("There was an error adding target to list and clearing the input in process\n" + reason);
+                  onBlur={(ev) => {
+                     if (ev.target.value && ev.target.value.trim()) {
+                        options.addTarget(ev.target.value.trim())
                      }
 
-                  }
+                     // This is using undocumented behavior to clear the input when you lose focus
+                     // It could very well break
+                     if (targetPicker.current) {
+                        try {
+                           (targetPicker.current as any).input.current._updateValue("");
+                        } catch (reason) {
+                           console.error("There was an error adding target to list and clearing the input in process\n" + reason);
+                        }
 
-               }}
+                     }
 
-               onResolveSuggestions={(filter, selected) => {
-                  return [];
-               }}
+                  }}
 
-               onEmptyResolveSuggestions={(selected) => {
-                  return [];
-               }}
+                  onResolveSuggestions={(filter, selected) => {
+                     return [];
+                  }}
 
-               onRemoveSuggestion={(item) => { }}
+                  onEmptyResolveSuggestions={(selected) => {
+                     return [];
+                  }}
 
-               createGenericItem={(input: string, ValidationState: ValidationState) => {
-                  return {
-                     item: { name: input, key: input },
-                     selected: true
-                  };
-               }}
+                  onRemoveSuggestion={(item) => { }}
 
-               onChange={(items) => {
-                  if (items) {
-                     options.setTargets(items.map(i => i.name))
-                  }
-               }}
+                  createGenericItem={(input: string, ValidationState: ValidationState) => {
+                     return {
+                        item: { name: input, key: input },
+                        selected: true
+                     };
+                  }}
 
-               onValidateInput={(input?: string) => input ? ValidationState.valid : ValidationState.invalid}
+                  onChange={(items) => {
+                     if (items) {
+                        options.setTargets(items.map(i => i.name))
+                     }
+                  }}
 
-               selectedItems={targetItems}
+                  onValidateInput={(input?: string) => input ? ValidationState.valid : ValidationState.invalid}
 
-               onItemSelected={(item) => {
-                  if (!item || !item.name?.trim()) {
+                  selectedItems={targetItems}
+
+                  onItemSelected={(item) => {
+                     if (!item || !item.name?.trim()) {
+                        return null;
+                     }
+                     options.addTarget(item.name.trim());
                      return null;
-                  }
-                  options.addTarget(item.name.trim());
-                  return null;
-               }}
-            />
-         </Stack>
-         {showAdditionalArgs && <TextField key={`key_adv_add_args_${options.currentRenderKey}`} multiline style={{ height: 48 }} resizable={false} readOnly={options.readOnly} spellCheck={false} defaultValue={options.advAdditionalArgs} label="Additional Arguments" onChange={(ev, newValue) => options.advAdditionalArgs = newValue} />}
-         {showArgumentClipboardButton && <DefaultButton text="Copy Job Arguments to Clipboard" style={{ width: 240 }} onClick={() => copyToClipboard(
-            options.jobDetails?.jobData?.arguments.map(arg => {
-               if (arg.indexOf("=") !== -1) {
-                  const components = arg.split("=");
-                  if (components[1].indexOf(" ") === -1) {
+                  }}
+               />
+            </Stack>
+            {showAdditionalArgs && <TextField key={`key_adv_add_args_${options.currentRenderKey}`} multiline style={{ height: 48 }} resizable={false} readOnly={options.readOnly} spellCheck={false} defaultValue={options.advAdditionalArgs} label="Additional Arguments" onChange={(ev, newValue) => options.advAdditionalArgs = newValue} />}
+            {showArgumentClipboardButton && <DefaultButton text="Copy Job Arguments to Clipboard" style={{ width: 240 }} onClick={() => copyToClipboard(
+               options.jobDetails?.jobData?.arguments.map(arg => {
+                  if (arg.indexOf("=") !== -1) {
+                     const components = arg.split("=");
+                     if (components[1].indexOf(" ") === -1) {
+                        return arg;
+                     }
+                     return `${components[0]}="${components[1]}"`;
+                  } else {
                      return arg;
                   }
-                  return `${components[0]}="${components[1]}"`;
-               } else {
-                  return arg;
-               }
-            }).join(" ")
-         )} />}
-         {!!farguments && <TextField style={{ height: fargumentsHeight }} key={`key_adv_job_arguments_${options.currentRenderKey}`} defaultValue={farguments} readOnly={true} label="Job Arguments" multiline resizable={false} />}
-         {!!fparameters && <TextField style={{ height: fparametersHeight }} key={`key_adv_job_parameters_${options.currentRenderKey}`} defaultValue={fparameters} readOnly={true} label="Job Parameters" multiline resizable={false} />}
+               }).join(" ")
+            )} />}
+            {!!farguments && <TextField style={{ height: fargumentsHeight }} key={`key_adv_job_arguments_${options.currentRenderKey}`} defaultValue={farguments} readOnly={true} label="Job Arguments" multiline resizable={false} />}
+            {!!fparameters && <TextField style={{ height: fparametersHeight }} key={`key_adv_job_parameters_${options.currentRenderKey}`} defaultValue={fparameters} readOnly={true} label="Job Parameters" multiline resizable={false} />}
+
+            <Stack>
+               <Checkbox disabled={options.readOnly}
+                  label="Template Editor"
+                  checked={options.preview}
+                  onChange={(ev, checked) => {
+                     if (checked) {
+                        options.onModeChanged("Basic");
+                     }
+                     options.onPreviewChanged(checked ?? false);
+                  }} />
+            </Stack>
 
-         <Stack>
-            <Checkbox disabled={options.readOnly}
-               label="Template Editor"
-               checked={options.preview}
-               onChange={(ev, checked) => {
-                  if (checked) {
-                     options.onModeChanged("Basic");
-                  }
-                  options.onPreviewChanged(checked ?? false);
-               }} />
          </Stack>
-
-      </Stack>
+      </ScrollablePane>
    </Stack>
 
 })
@@ -1976,7 +1989,7 @@ const BuildModal: React.FC<{ setUseLegacyDialog: (value: boolean) => void }> = o
 
                      <Stack>
                         <TextField
-                           key={`build_shelved_change_key_${options.currentRenderKey}`} 
+                           key={`build_shelved_change_key_${options.currentRenderKey}`}
                            label="Shelved Change"
                            style={{ width: 146 }}
                            placeholder={!template!.allowPreflights ? "Disabled by template" : "None"}
