Horde: Remove group parameter template
#rnx
#skipci

diff --git a/Engine/Source/Programs/Horde/Horde.Server/Jobs/Templates/ITemplate.cs b/Engine/Source/Programs/Horde/Horde.Server/Jobs/Templates/ITemplate.cs
--- a/Engine/Source/Programs/Horde/Horde.Server/Jobs/Templates/ITemplate.cs
+++ b/Engine/Source/Programs/Horde/Horde.Server/Jobs/Templates/ITemplate.cs
@@ -13,7 +13,7 @@ namespace Horde.Server.Jobs.Templates
 	/// Base class for parameters used to configure templates via the new build dialog
 	/// </summary>
 	[BsonDiscriminator(RootClass = true)]
-	[BsonKnownTypes(typeof(GroupParameter), typeof(TextParameter), typeof(ListParameter), typeof(BoolParameter))]
+	[BsonKnownTypes(typeof(TextParameter), typeof(ListParameter), typeof(BoolParameter))]
 	public abstract class Parameter
 	{
 		/// <summary>
@@ -38,76 +38,6 @@ namespace Horde.Server.Jobs.Templates
 		public abstract ParameterData ToData();
 	}
 
-	/// <summary>
-	/// Used to group a number of other parameters
-	/// </summary>
-	public class GroupParameter : Parameter
-	{
-		/// <summary>
-		/// Label to display next to this parameter
-		/// </summary>
-		public string Label { get; set; }
-
-		/// <summary>
-		/// How to display this group
-		/// </summary>
-		public GroupParameterStyle Style { get; set; }
-
-		/// <summary>
-		/// List of child parameters
-		/// </summary>
-		public List<Parameter> Children { get; set; }
-
-		/// <summary>
-		/// Private constructor for serialization
-		/// </summary>
-		public GroupParameter()
-		{
-			Label = null!;
-			Children = null!;
-		}
-
-		/// <summary>
-		/// Constructor
-		/// </summary>
-		/// <param name="label">Name of the group</param>
-		/// <param name="style">How to display this group</param>
-		/// <param name="children">List of child parameters</param>
-		public GroupParameter(string label, GroupParameterStyle style, List<Parameter> children)
-		{
-			Label = label;
-			Style = style;
-			Children = children;
-		}
-
-		/// <inheritdoc/>
-		public override void GetArguments(IReadOnlyDictionary<ParameterId, string> parameters, bool scheduledBuild, List<string> arguments)
-		{
-			foreach (Parameter child in Children)
-			{
-				child.GetArguments(parameters, scheduledBuild, arguments);
-			}
-		}
-
-		/// <inheritdoc/>
-		public override void GetDefaultParameters(Dictionary<ParameterId, string> parameters, bool scheduledBuild)
-		{
-			foreach (Parameter child in Children)
-			{
-				child.GetDefaultParameters(parameters, scheduledBuild);
-			}
-		}
-
-		/// <summary>
-		/// Convert this parameter to data for serialization
-		/// </summary>
-		/// <returns>Serializable parameter data</returns>
-		public override ParameterData ToData()
-		{
-			return new GroupParameterData(Label, Style, Children.ConvertAll(x => x.ToData()));
-		}
-	}
-
 	/// <summary>
 	/// Free-form text entry parameter
 	/// </summary>
diff --git a/Engine/Source/Programs/Horde/Horde.Server/Jobs/Templates/TemplateConfig.cs b/Engine/Source/Programs/Horde/Horde.Server/Jobs/Templates/TemplateConfig.cs
--- a/Engine/Source/Programs/Horde/Horde.Server/Jobs/Templates/TemplateConfig.cs
+++ b/Engine/Source/Programs/Horde/Horde.Server/Jobs/Templates/TemplateConfig.cs
@@ -98,7 +98,7 @@ namespace Horde.Server.Jobs.Templates
 	/// <summary>
 	/// Base class for template parameters
 	/// </summary>
-	[JsonKnownTypes(typeof(GroupParameterData), typeof(TextParameterData), typeof(ListParameterData), typeof(BoolParameterData))]
+	[JsonKnownTypes(typeof(TextParameterData), typeof(ListParameterData), typeof(BoolParameterData))]
 	public abstract class ParameterData
 	{
 		/// <summary>
@@ -164,77 +164,6 @@ namespace Horde.Server.Jobs.Templates
 		}
 	}
 
-	/// <summary>
-	/// Describes how to render a group parameter
-	/// </summary>
-	public enum GroupParameterStyle
-	{
-		/// <summary>
-		/// Separate tab on the form
-		/// </summary>
-		Tab,
-
-		/// <summary>
-		/// Section with heading
-		/// </summary>
-		Section,
-	}
-
-	/// <summary>
-	/// Used to group a number of other parameters
-	/// </summary>
-	[JsonDiscriminator("Group")]
-	public class GroupParameterData : ParameterData
-	{
-		/// <summary>
-		/// Label to display next to this parameter
-		/// </summary>
-		public string Label { get; set; }
-
-		/// <summary>
-		/// How to display this group
-		/// </summary>
-		public GroupParameterStyle Style { get; set; }
-
-		/// <summary>
-		/// List of child parameters
-		/// </summary>
-		public List<ParameterData> Children { get; set; }
-
-		/// <summary>
-		/// Private constructor for serialization
-		/// </summary>
-		public GroupParameterData()
-		{
-			Label = null!;
-			Children = null!;
-		}
-
-		/// <summary>
-		/// Constructor
-		/// </summary>
-		/// <param name="label">Name of the group</param>
-		/// <param name="style">How to display this group</param>
-		/// <param name="children">List of child parameters</param>
-		public GroupParameterData(string label, GroupParameterStyle style, List<ParameterData> children)
-		{
-			Label = label;
-			Style = style;
-			Children = children;
-		}
-
-		/// <inheritdoc/>
-		public override void PostLoad(HashSet<ParameterId> parameterIds)
-		{
-		}
-
-		/// <inheritdoc/>
-		public override Parameter ToModel()
-		{
-			return new GroupParameter(Label, Style, Children.ConvertAll(x => x.ToModel()));
-		}
-	}
-
 	/// <summary>
 	/// Free-form text entry parameter
 	/// </summary>
