Horde: Add a disk-based cache for compressed Unsync blocks. Uses memory mapped files to read/write to underlying storage.

diff --git a/Engine/Source/Programs/Horde/Horde.Server.Tests/Artifacts/BlockCacheTests.cs b/Engine/Source/Programs/Horde/Horde.Server.Tests/Artifacts/BlockCacheTests.cs
new file mode 100644
--- /dev/null
+++ b/Engine/Source/Programs/Horde/Horde.Server.Tests/Artifacts/BlockCacheTests.cs
@@ -0,0 +1,99 @@
+// Copyright Epic Games, Inc. All Rights Reserved.
+
+using System;
+using System.Buffers;
+using System.Linq;
+using System.Security.Cryptography;
+using EpicGames.Core;
+using Horde.Server.Artifacts;
+using Microsoft.VisualStudio.TestTools.UnitTesting;
+
+namespace Horde.Server.Tests.Artifacts
+{
+	[TestClass]
+	public class BlockCacheTests
+	{
+		[TestMethod]
+		public void RecycleTest()
+		{
+			using BlockCache blockCache = BlockCache.CreateInMemory(1, 2, 4096);
+			blockCache.Add("foo", new byte[] { 1, 2, 3 });
+			blockCache.Add("bar", new byte[] { 1, 2, 3 });
+
+			Assert.IsTrue(blockCache.Contains("foo"));
+			Assert.IsTrue(blockCache.Contains("bar"));
+
+			blockCache.Add("baz", new byte[] { 4, 5, 6 });
+
+			Assert.IsFalse(blockCache.Contains("foo"));
+			Assert.IsTrue(blockCache.Contains("bar"));
+			Assert.IsTrue(blockCache.Contains("baz"));
+
+			blockCache.Get("bar")?.Dispose();
+			blockCache.Add("foo", new byte[] { 4, 5, 6 });
+
+			Assert.IsTrue(blockCache.Contains("foo"));
+			Assert.IsTrue(blockCache.Contains("bar"));
+			Assert.IsFalse(blockCache.Contains("baz"));
+		}
+
+		[TestMethod]
+		public void BasicTest()
+		{
+			using BlockCache blockCache = BlockCache.CreateInMemory(1, 4096, 4096);
+			blockCache.Add("hello", new byte[] { 1, 2, 3 });
+
+			IBlockCacheValue? value = blockCache.Get("hello");
+			Assert.IsNotNull(value);
+
+			Assert.IsTrue(value.Data.ToArray().SequenceEqual(new byte[] { 1, 2, 3 }));
+		}
+
+		[TestMethod]
+		public void LargeBlockTest()
+		{
+			byte[] data = RandomNumberGenerator.GetBytes(100000);
+
+			using BlockCache blockCache = BlockCache.CreateInMemory(1);
+			blockCache.Add("hello", data);
+
+			IBlockCacheValue? value = blockCache.Get("hello");
+			Assert.IsNotNull(value);
+
+			Assert.IsTrue(value.Data.ToArray().SequenceEqual(data));
+		}
+
+		[TestMethod]
+		public void MultiBlockTest()
+		{
+			string[] keys = new string[1000];
+			byte[][] values = new byte[keys.Length][];
+
+			using BlockCache blockCache = BlockCache.CreateInMemory(1);
+
+			Random random = new Random(0);
+			for (int idx = 0; idx < keys.Length; idx++)
+			{
+				keys[idx] = $"key{idx}";
+
+				int length = random.Next(4000) + 1;
+
+				byte[] buffer = new byte[length];
+				random.NextBytes(buffer);
+
+				values[idx] = buffer;
+				blockCache.Add(keys[idx], values[idx]);
+			}
+
+			for (int idx = 0; idx < 2000; idx++)
+			{
+				int keyIdx = random.Next(keys.Length);
+
+				IBlockCacheValue? value = blockCache.Get(keys[keyIdx]);
+				Assert.IsNotNull(value);
+
+				Assert.IsTrue(value.Data.ToArray().SequenceEqual(values[keyIdx]));
+			}
+		}
+	}
+}
diff --git a/Engine/Source/Programs/Horde/Horde.Server/Artifacts/BlockCache.cs b/Engine/Source/Programs/Horde/Horde.Server/Artifacts/BlockCache.cs
new file mode 100644
--- /dev/null
+++ b/Engine/Source/Programs/Horde/Horde.Server/Artifacts/BlockCache.cs
@@ -0,0 +1,485 @@
+// Copyright Epic Games, Inc. All Rights Reserved.
+
+using System;
+using System.Buffers;
+using System.Buffers.Binary;
+using System.Collections.Concurrent;
+using System.Collections.Generic;
+using System.IO;
+using System.IO.MemoryMappedFiles;
+using System.Numerics;
+using System.Text;
+using System.Threading;
+using EpicGames.Core;
+
+namespace Horde.Server.Artifacts
+{
+	/// <summary>
+	/// Stats for the cache
+	/// </summary>
+	public record class BlockCacheStats(int NumItems, long TotalSize, long UsedSize, long AllocatedSize, long FreeSize);
+
+	/// <summary>
+	/// Implements a simple key-value cache using memory mapped files.
+	/// </summary>
+	public sealed class BlockCache : IBlockCache, IDisposable
+	{
+		// Size of a page for allocations
+		const int PageSizeLog2 = 10;
+		const int PageSize = 1 << PageSizeLog2; // 4kb
+
+		// Size of a block; the smallest chunk of data that can be allocated
+		const int DefaultBlockSize = 32768;
+
+		// Default number of blocks in each partition
+		const int DefaultNumBlocksPerPartition = 32768;
+
+		// Key used to lookup a particular block, consisting of a user-defined key and chunk index.
+		record struct BlockKey(IoHash Hash, int Index)
+		{
+			public static BlockKey Empty => default;
+		}
+
+		// Header information for a block
+		record struct BlockHeader(IoHash Hash, uint PackedData)
+		{
+			public const int NumBytes = 24;
+
+			public BlockKey Key => new BlockKey(Hash, Index);
+			public int Index => (int)((PackedData >> 24) & 0xff);
+			public int Count => (int)((PackedData >> 16) & 0xff);
+			public int Length => ((int)(PackedData & 0xffff)) + 1;
+
+			public BlockHeader(BlockKey key, int count, int length)
+				: this(key.Hash, ((uint)key.Index << 24) | ((uint)count << 16) | (uint)(length - 1))
+			{ }
+
+			public static BlockHeader Read(ReadOnlySpan<byte> data)
+			{
+				IoHash hash = new IoHash(data);
+				uint packed = BinaryPrimitives.ReadUInt32LittleEndian(data[IoHash.NumBytes..]);
+				return new BlockHeader(hash, packed);
+			}
+
+			public void Write(Span<byte> data)
+			{
+				Key.Hash.CopyTo(data);
+				BinaryPrimitives.WriteUInt32LittleEndian(data[IoHash.NumBytes..], PackedData);
+			}
+		}
+
+		class Partition : IDisposable
+		{
+			public Memory<byte> AccessCounts { get; }
+			public Memory<byte> BlockHeaders { get; }
+			public Memory<byte> BlockData { get; }
+
+			public Partition(int numBlocks, Memory<byte> data)
+			{
+				Memory<byte> remainingData = data;
+
+				AccessCounts = remainingData.Slice(0, numBlocks);
+				remainingData = remainingData.Slice(PageAlign(AccessCounts.Length));
+
+				BlockHeaders = remainingData.Slice(0, BlockHeader.NumBytes * numBlocks);
+				remainingData = remainingData.Slice(PageAlign(BlockHeaders.Length));
+
+				BlockData = remainingData;
+			}
+
+			public virtual void Dispose()
+			{ }
+		}
+
+		class MemoryMappedFilePartition : Partition
+		{
+			readonly MemoryMappedFile _memoryMappedFile;
+			readonly MemoryMappedView _memoryMappedView;
+
+			private MemoryMappedFilePartition(MemoryMappedFile memoryMappedFile, MemoryMappedView memoryMappedView, int numBlocks, Memory<byte> data)
+				: base(numBlocks, data)
+			{
+				_memoryMappedFile = memoryMappedFile;
+				_memoryMappedView = memoryMappedView;
+			}
+
+			public static MemoryMappedFilePartition Create(FileReference file, int numBlocks, int blockSize)
+			{
+				int partitionSize = GetPartitionSize(numBlocks, blockSize);
+
+				MemoryMappedFile memoryMappedFile = MemoryMappedFile.CreateFromFile(file.FullName, FileMode.OpenOrCreate, null, partitionSize, MemoryMappedFileAccess.ReadWrite);
+				MemoryMappedView memoryMappedView = new MemoryMappedView(memoryMappedFile, 0, partitionSize);
+
+				return new MemoryMappedFilePartition(memoryMappedFile, memoryMappedView, numBlocks, memoryMappedView.GetMemory(0, partitionSize));
+			}
+
+			public override void Dispose()
+			{
+				base.Dispose();
+
+				_memoryMappedFile.Dispose();
+				_memoryMappedView.Dispose();
+			}
+		}
+
+		class BlockCacheValue : IBlockCacheValue
+		{
+			ReadOnlySequence<byte> _data;
+			BlockState[] _values;
+
+			public BlockCacheValue(ReadOnlySequence<byte> data, BlockState[] values)
+			{
+				_data = data;
+				_values = values;
+			}
+
+			public ReadOnlySequence<byte> Data => _data;
+
+			public void Dispose()
+			{
+				for (int idx = 0; idx < _values.Length; idx++)
+				{
+					_values[idx].ReleaseReadLock();
+				}
+
+				_data = ReadOnlySequence<byte>.Empty;
+				_values = Array.Empty<BlockState>();
+			}
+		}
+
+		class BlockState
+		{
+			public int BlockIdx { get; }
+			public int BlockLength { get; }
+			public int NumBlocks { get; }
+
+			int _readerCount;
+
+			public BlockState(int blockIdx, int blockLength, int numBlocks)
+			{
+				BlockIdx = blockIdx;
+				BlockLength = blockLength;
+				NumBlocks = numBlocks;
+			}
+
+			public bool TryAddReadLock()
+			{
+				int lockCount = _readerCount;
+				while (lockCount >= 0)
+				{
+					if (Interlocked.CompareExchange(ref _readerCount, lockCount + 1, lockCount) == lockCount)
+					{
+						return true;
+					}
+					lockCount = Interlocked.CompareExchange(ref _readerCount, 0, 0);
+				}
+				return false;
+			}
+
+			public bool TryAddWriteLock()
+				=> Interlocked.CompareExchange(ref _readerCount, -1, 0) == 0;
+
+			public void ReleaseReadLock()
+				=> Interlocked.Decrement(ref _readerCount);
+		}
+
+		readonly Random _rng = new Random(0);
+		readonly Partition[] _partitions;
+		readonly int _numBlocks;
+		readonly int _numBlocksPerPartition;
+		readonly int _numBlocksPerPartitionLog2;
+		readonly int _blockSize;
+		readonly int _blockSizeLog2;
+		readonly ConcurrentDictionary<BlockKey, BlockState> _lookup = new ConcurrentDictionary<BlockKey, BlockState>();
+
+		int _scanFreeBlockIdx = -1;
+
+		/// <summary>
+		/// Constructor
+		/// </summary>
+		/// <param name="partitions">Allocated partition data</param>
+		/// <param name="numBlocksPerPartition">Number of blocks in each partition</param>
+		/// <param name="blockSize">The size of each block. Must be a power of two.</param>
+		BlockCache(Partition[] partitions, int numBlocksPerPartition, int blockSize)
+		{
+			if ((numBlocksPerPartition & (numBlocksPerPartition - 1)) != 0)
+			{
+				throw new ArgumentException("Number of blocks per partition must be a power of two", nameof(numBlocksPerPartition));
+			}
+			if ((blockSize & (blockSize - 1)) != 0)
+			{
+				throw new ArgumentException("Block size must be a power of two", nameof(blockSize));
+			}
+
+			_partitions = partitions;
+			_numBlocks = partitions.Length * numBlocksPerPartition;
+			_numBlocksPerPartition = numBlocksPerPartition;
+			_numBlocksPerPartitionLog2 = BitOperations.Log2((uint)numBlocksPerPartition);
+			_blockSize = blockSize;
+			_blockSizeLog2 = BitOperations.Log2((uint)blockSize);
+		}
+
+		/// <summary>
+		/// Creates a new cache from memory
+		/// </summary>
+		/// <param name="numPartitions">Number of partitions in the cache</param>
+		/// <param name="numBlocksPerPartition">Number of blocks in each partition</param>
+		/// <param name="blockSize">Size of each block</param>
+		public static BlockCache CreateInMemory(int numPartitions, int numBlocksPerPartition = DefaultNumBlocksPerPartition, int blockSize = DefaultBlockSize)
+		{
+			int partitionSize = GetPartitionSize(numBlocksPerPartition, blockSize);
+
+			Partition[] partitions = new Partition[numPartitions];
+			for (int idx = 0; idx < numPartitions; idx++)
+			{
+				byte[] data = new byte[partitionSize];
+				partitions[idx] = new Partition(numBlocksPerPartition, data);
+			}
+
+			return new BlockCache(partitions, numBlocksPerPartition, blockSize);
+		}
+
+		/// <summary>
+		/// Create a new cache backed by files on disk
+		/// </summary>
+		/// <param name="rootDir">Directory containing the partitions</param>
+		/// <param name="numPartitions">Number of partitions to create. Each partition is ~1gb.</param>
+		/// <param name="numBlocksPerPartition">Number of blocks in each partition</param>
+		/// <param name="blockSize">Size of a block</param>
+		public static BlockCache Create(DirectoryReference rootDir, int numPartitions, int numBlocksPerPartition = DefaultNumBlocksPerPartition, int blockSize = DefaultBlockSize)
+		{
+			MemoryMappedFilePartition[] partitions = new MemoryMappedFilePartition[numPartitions];
+			try
+			{
+				DirectoryReference.CreateDirectory(rootDir);
+				for (int idx = 0; idx < numPartitions; idx++)
+				{
+					FileReference file = FileReference.Combine(rootDir, $"partition{idx:0000}.dat");
+					partitions[idx] = MemoryMappedFilePartition.Create(file, numBlocksPerPartition, blockSize);
+				}
+				return new BlockCache(partitions, numBlocksPerPartition, blockSize);
+			}
+			catch
+			{
+				for (int idx = 0; idx < numPartitions; idx++)
+				{
+					partitions[idx]?.Dispose();
+				}
+				throw;
+			}
+		}
+
+		static int GetPartitionSize(int numBlocks, int blockSize)
+			=> PageAlign(numBlocks) + PageAlign(numBlocks * BlockHeader.NumBytes) + (numBlocks * blockSize);
+
+		static int PageAlign(int value)
+			=> (value + (PageSize - 1)) & ~(PageSize - 1);
+
+		/// <inheritdoc/>
+		public void Dispose()
+		{
+			foreach (Partition partition in _partitions)
+			{
+				partition.Dispose();
+			}
+		}
+
+		/// <summary>
+		/// Gets stats for the cache
+		/// </summary>
+		public BlockCacheStats GetStats()
+		{
+			long usedSize = 0;
+			long allocatedSize = 0;
+			long freeSize = 0;
+			for (int blockIdx = 0; blockIdx < _numBlocks; blockIdx++)
+			{
+				BlockHeader header = GetBlockHeader(blockIdx);
+				if (header.Hash == IoHash.Zero)
+				{
+					freeSize += _blockSize;
+				}
+				else
+				{
+					usedSize += header.Length;
+					allocatedSize += _blockSize;
+				}
+			}
+			return new BlockCacheStats(_lookup.Count, _numBlocks << _blockSizeLog2, usedSize, allocatedSize, freeSize);
+		}
+
+		int GetNextFreeBlockIdx()
+		{
+			int scanLength = Math.Min(32, _numBlocks);
+			for (; ; )
+			{
+				int minBlockIdx = 0;
+				int minAccessCount = 256;
+
+				// Check the next 32 entries and find the least referenced
+				for (int offset = 0; offset < scanLength; offset++)
+				{
+					_scanFreeBlockIdx++;
+					if (_scanFreeBlockIdx == _numBlocks)
+					{
+						_scanFreeBlockIdx = 0;
+					}
+
+					int accessCount = GetBlockAccessCount(_scanFreeBlockIdx);
+					SetBlockAccessCount(_scanFreeBlockIdx, 0);
+
+					if (accessCount < minAccessCount)
+					{
+						minBlockIdx = _scanFreeBlockIdx;
+						minAccessCount = accessCount;
+
+						if (accessCount == 0)
+						{
+							break;
+						}
+					}
+				}
+
+				// Try to lock this block. If we can't, just try again with the next span of entries.
+				BlockHeader header = GetBlockHeader(minBlockIdx);
+				if (header.Hash == IoHash.Zero || !_lookup.TryGetValue(header.Key, out BlockState? blockState))
+				{
+					return minBlockIdx;
+				}
+
+				// Try to lock the block, and return it
+				if (blockState.TryAddWriteLock())
+				{
+					SetBlockHeader(minBlockIdx, default);
+					_lookup.Remove(header.Key, out _);
+					return minBlockIdx;
+				}
+			}
+		}
+
+		byte GetBlockAccessCount(int blockIdx)
+		{
+			Partition partition = _partitions[blockIdx >> _numBlocksPerPartitionLog2];
+			return partition.AccessCounts.Span[blockIdx & (_numBlocksPerPartition - 1)];
+		}
+
+		void SetBlockAccessCount(int blockIdx, byte value)
+		{
+			Partition partition = _partitions[blockIdx >> _numBlocksPerPartitionLog2];
+			partition.AccessCounts.Span[blockIdx & (_numBlocksPerPartition - 1)] = value;
+		}
+
+		BlockHeader GetBlockHeader(int blockIdx)
+		{
+			Partition partition = _partitions[blockIdx >> _numBlocksPerPartitionLog2];
+			ReadOnlySpan<byte> span = partition.BlockHeaders.Span.Slice((blockIdx & (_numBlocksPerPartition - 1)) * BlockHeader.NumBytes, BlockHeader.NumBytes);
+			return BlockHeader.Read(span);
+		}
+
+		void SetBlockHeader(int blockIdx, BlockHeader header)
+		{
+			Partition partition = _partitions[blockIdx >> _numBlocksPerPartitionLog2];
+			Span<byte> span = partition.BlockHeaders.Span.Slice((blockIdx & (_numBlocksPerPartition - 1)) * BlockHeader.NumBytes, BlockHeader.NumBytes);
+			header.Write(span);
+		}
+
+		Memory<byte> GetBlockData(int blockIdx, int blockLength)
+		{
+			Partition partition = _partitions[blockIdx >> _numBlocksPerPartitionLog2];
+			int blockIdxInPartition = blockIdx & (_numBlocksPerPartition - 1);
+			return partition.BlockData.Slice(blockIdxInPartition << _blockSizeLog2, blockLength);
+		}
+
+		/// <inheritdoc/>
+		public bool Add(string key, ReadOnlyMemory<byte> value)
+		{
+			if (value.Length == 0)
+			{
+				return false;
+			}
+
+			IoHash hash = IoHash.Compute(Encoding.UTF8.GetBytes(key));
+
+			int numBlocks = (value.Length + (_blockSize - 1)) >> _blockSizeLog2;
+			for (int idx = 0; idx < numBlocks; idx++)
+			{
+				BlockKey blockKey = new BlockKey(hash, idx);
+				if (!_lookup.ContainsKey(blockKey))
+				{
+					int offset = idx << _blockSizeLog2;
+					int length = Math.Min(value.Length - offset, _blockSize);
+
+					int freeBlockIdx = GetNextFreeBlockIdx();
+
+					Memory<byte> blockData = GetBlockData(freeBlockIdx, _blockSize);
+					value.Slice(offset, length).CopyTo(blockData);
+
+					BlockHeader header = new BlockHeader(blockKey, numBlocks, length);
+					SetBlockHeader(freeBlockIdx, header);
+
+					BlockState state = new BlockState(freeBlockIdx, length, numBlocks);
+					_lookup.TryAdd(blockKey, state);
+				}
+			}
+
+			return true;
+		}
+
+		/// <summary>
+		/// Gets the data associated with a given key
+		/// </summary>
+		/// <param name="key">Key to search for</param>
+		/// <returns>Data corresponding to the given key</returns>
+		public IBlockCacheValue? Get(string key)
+		{
+			IoHash hash = IoHash.Compute(Encoding.UTF8.GetBytes(key));
+
+			BlockState? firstBlock;
+			if (!_lookup.TryGetValue(new BlockKey(hash, 0), out firstBlock) || !firstBlock.TryAddReadLock())
+			{
+				return null;
+			}
+
+			BlockState[] blocks = new BlockState[firstBlock.NumBlocks];
+			blocks[0] = firstBlock;
+
+			try
+			{
+				ReadOnlySequenceBuilder<byte> builder = new ReadOnlySequenceBuilder<byte>();
+				builder.Append(GetBlockData(firstBlock.BlockIdx, firstBlock.BlockLength));
+
+				for (int idx = 1; idx < firstBlock.NumBlocks; idx++)
+				{
+					BlockState? nextBlock;
+					if (!_lookup.TryGetValue(new BlockKey(hash, (ushort)idx), out nextBlock) || !nextBlock.TryAddReadLock())
+					{
+						return null;
+					}
+
+					blocks[idx] = nextBlock;
+					builder.Append(GetBlockData(nextBlock.BlockIdx, nextBlock.BlockLength));
+				}
+
+				foreach (BlockState block in blocks)
+				{
+					int nextValue = GetBlockAccessCount(block.BlockIdx) + 1;
+					if (nextValue < 256)
+					{
+						SetBlockAccessCount(block.BlockIdx, (byte)nextValue);
+					}
+				}
+
+				BlockCacheValue result = new BlockCacheValue(builder.Construct(), blocks);
+				blocks = Array.Empty<BlockState>();
+				return result;
+			}
+			finally
+			{
+				for (int idx = 0; idx < blocks.Length; idx++)
+				{
+					blocks[idx]?.ReleaseReadLock();
+				}
+			}
+		}
+	}
+}
diff --git a/Engine/Source/Programs/Horde/Horde.Server/Artifacts/IBlockCache.cs b/Engine/Source/Programs/Horde/Horde.Server/Artifacts/IBlockCache.cs
new file mode 100644
--- /dev/null
+++ b/Engine/Source/Programs/Horde/Horde.Server/Artifacts/IBlockCache.cs
@@ -0,0 +1,54 @@
+// Copyright Epic Games, Inc. All Rights Reserved.
+
+using System;
+using System.Buffers;
+
+namespace Horde.Server.Artifacts
+{
+	/// <summary>
+	/// Implements a key/value cache store.
+	/// </summary>
+	public interface IBlockCache
+	{
+		/// <summary>
+		/// Adds a new item to the cache
+		/// </summary>
+		/// <param name="key">Key for finding the item</param>
+		/// <param name="value">Data to be stored</param>
+		/// <returns>True if the item was added, false if not</returns>
+		bool Add(string key, ReadOnlyMemory<byte> value);
+
+		/// <summary>
+		/// Retrieves an item from the cache. If successful, the value must be disposed by the caller.
+		/// </summary>
+		/// <param name="key">Key to search for</param>
+		/// <returns>Data corresponding to the given key</returns>
+		IBlockCacheValue? Get(string key);
+	}
+
+	/// <summary>
+	/// Handle to a value in the block cache.
+	/// </summary>
+	public interface IBlockCacheValue : IDisposable
+	{
+		/// <summary>
+		/// The underlying data
+		/// </summary>
+		ReadOnlySequence<byte> Data { get; }
+	}
+
+	/// <summary>
+	/// Extension methods for <see cref="IBlockCache"/>
+	/// </summary>
+	public static class BlockCacheExtensions
+	{
+		/// <summary>
+		/// Determines if the cache contains a particular key
+		/// </summary>
+		public static bool Contains(this IBlockCache blockCache, string key)
+		{
+			using IBlockCacheValue? value = blockCache.Get(key);
+			return value != null;
+		}
+	}
+}
