Horde: Use random sampling and last access timestamps to evict items from the block cache.

diff --git a/Engine/Source/Programs/Horde/Horde.Server.Tests/Artifacts/BlockCacheTests.cs b/Engine/Source/Programs/Horde/Horde.Server.Tests/Artifacts/BlockCacheTests.cs
--- a/Engine/Source/Programs/Horde/Horde.Server.Tests/Artifacts/BlockCacheTests.cs
+++ b/Engine/Source/Programs/Horde/Horde.Server.Tests/Artifacts/BlockCacheTests.cs
@@ -13,30 +13,6 @@ namespace Horde.Server.Tests.Artifacts
 	[TestClass]
 	public class BlockCacheTests
 	{
-		[TestMethod]
-		public void RecycleTest()
-		{
-			using BlockCache blockCache = BlockCache.CreateInMemory(1, 2, 4096);
-			blockCache.Add("foo", new byte[] { 1, 2, 3 });
-			blockCache.Add("bar", new byte[] { 1, 2, 3 });
-
-			Assert.IsTrue(blockCache.Contains("foo"));
-			Assert.IsTrue(blockCache.Contains("bar"));
-
-			blockCache.Add("baz", new byte[] { 4, 5, 6 });
-
-			Assert.IsFalse(blockCache.Contains("foo"));
-			Assert.IsTrue(blockCache.Contains("bar"));
-			Assert.IsTrue(blockCache.Contains("baz"));
-
-			blockCache.Get("bar")?.Dispose();
-			blockCache.Add("foo", new byte[] { 4, 5, 6 });
-
-			Assert.IsTrue(blockCache.Contains("foo"));
-			Assert.IsTrue(blockCache.Contains("bar"));
-			Assert.IsFalse(blockCache.Contains("baz"));
-		}
-
 		[TestMethod]
 		public void BasicTest()
 		{
diff --git a/Engine/Source/Programs/Horde/Horde.Server/Artifacts/BlockCache.cs b/Engine/Source/Programs/Horde/Horde.Server/Artifacts/BlockCache.cs
--- a/Engine/Source/Programs/Horde/Horde.Server/Artifacts/BlockCache.cs
+++ b/Engine/Source/Programs/Horde/Horde.Server/Artifacts/BlockCache.cs
@@ -5,6 +5,7 @@ using System.Buffers;
 using System.Buffers.Binary;
 using System.Collections.Concurrent;
 using System.Collections.Generic;
+using System.Diagnostics;
 using System.IO;
 using System.IO.MemoryMappedFiles;
 using System.Numerics;
@@ -17,7 +18,7 @@ namespace Horde.Server.Artifacts
 	/// <summary>
 	/// Stats for the cache
 	/// </summary>
-	public record class BlockCacheStats(int NumItems, long TotalSize, long UsedSize, long AllocatedSize, long FreeSize);
+	public record class BlockCacheStats(int NumItems, int NumBlocks, int NumFreeBlocks, int BlockSize, long TotalSize, long UsedSize, long AllocatedSize);
 
 	/// <summary>
 	/// Implements a simple key-value cache using memory mapped files.
@@ -34,24 +35,21 @@ namespace Horde.Server.Artifacts
 		// Default number of blocks in each partition
 		const int DefaultNumBlocksPerPartition = 32768;
 
-		// Key used to lookup a particular block, consisting of a user-defined key and chunk index.
-		record struct BlockKey(IoHash Hash, int Index)
-		{
-			public static BlockKey Empty => default;
-		}
+		// When creating free space in the cache, define the number of samples to take, and the number of entries to remove from the examined items.
+		const int NumEvictionSamples = 10;
+		const int NumEvictionEntries = 4;
 
 		// Header information for a block
 		record struct BlockHeader(IoHash Hash, uint PackedData)
 		{
 			public const int NumBytes = 24;
 
-			public BlockKey Key => new BlockKey(Hash, Index);
 			public int Index => (int)((PackedData >> 24) & 0xff);
 			public int Count => (int)((PackedData >> 16) & 0xff);
 			public int Length => ((int)(PackedData & 0xffff)) + 1;
 
-			public BlockHeader(BlockKey key, int count, int length)
-				: this(key.Hash, ((uint)key.Index << 24) | ((uint)count << 16) | (uint)(length - 1))
+			public BlockHeader(IoHash hash, int index, int count, int length)
+				: this(hash, ((uint)index << 24) | ((uint)count << 16) | (uint)(length - 1))
 			{ }
 
 			public static BlockHeader Read(ReadOnlySpan<byte> data)
@@ -63,7 +61,7 @@ namespace Horde.Server.Artifacts
 
 			public void Write(Span<byte> data)
 			{
-				Key.Hash.CopyTo(data);
+				Hash.CopyTo(data);
 				BinaryPrimitives.WriteUInt32LittleEndian(data[IoHash.NumBytes..], PackedData);
 			}
 		}
@@ -125,43 +123,46 @@ namespace Horde.Server.Artifacts
 		class BlockCacheValue : IBlockCacheValue
 		{
 			ReadOnlySequence<byte> _data;
-			BlockState[] _values;
+			BlockState? _state;
 
-			public BlockCacheValue(ReadOnlySequence<byte> data, BlockState[] values)
+			public BlockCacheValue(ReadOnlySequence<byte> data, BlockState state)
 			{
 				_data = data;
-				_values = values;
+				_state = state;
 			}
 
 			public ReadOnlySequence<byte> Data => _data;
 
 			public void Dispose()
 			{
-				for (int idx = 0; idx < _values.Length; idx++)
+				_data = ReadOnlySequence<byte>.Empty;
+
+				if (_state != null)
 				{
-					_values[idx].ReleaseReadLock();
+					_state.ReleaseReadLock();
+					_state = null;
 				}
-
-				_data = ReadOnlySequence<byte>.Empty;
-				_values = Array.Empty<BlockState>();
 			}
 		}
 
 		class BlockState
 		{
-			public int BlockIdx { get; }
-			public int BlockLength { get; }
-			public int NumBlocks { get; }
+			public int[] BlockIdxs { get; }
+			public long LastAccessTicks { get; private set; }
+			public int LastBlockLength { get; }
 
 			int _readerCount;
 
-			public BlockState(int blockIdx, int blockLength, int numBlocks)
+			public BlockState(int[] blockIdxs, int lastBlockLength)
 			{
-				BlockIdx = blockIdx;
-				BlockLength = blockLength;
-				NumBlocks = numBlocks;
+				BlockIdxs = blockIdxs;
+				LastAccessTicks = Stopwatch.GetTimestamp();
+				LastBlockLength = lastBlockLength;
 			}
 
+			public void Touch()
+				=> LastAccessTicks = Stopwatch.GetTimestamp();
+
 			public bool TryAddReadLock()
 			{
 				int lockCount = _readerCount;
@@ -190,9 +191,10 @@ namespace Horde.Server.Artifacts
 		readonly int _numBlocksPerPartitionLog2;
 		readonly int _blockSize;
 		readonly int _blockSizeLog2;
-		readonly ConcurrentDictionary<BlockKey, BlockState> _lookup = new ConcurrentDictionary<BlockKey, BlockState>();
+		readonly ConcurrentDictionary<IoHash, BlockState> _lookup = new ConcurrentDictionary<IoHash, BlockState>();
+		readonly ConcurrentQueue<int> _freeBlocks = new ConcurrentQueue<int>();
 
-		int _scanFreeBlockIdx = -1;
+		int _numAllocatedBlocks = 0;
 
 		/// <summary>
 		/// Constructor
@@ -291,13 +293,14 @@ namespace Horde.Server.Artifacts
 		{
 			long usedSize = 0;
 			long allocatedSize = 0;
-			long freeSize = 0;
-			for (int blockIdx = 0; blockIdx < _numBlocks; blockIdx++)
+
+			int numFreeBlocks = (_numBlocks - _numAllocatedBlocks) + _freeBlocks.Count;
+			for (int blockIdx = 0; blockIdx < _numAllocatedBlocks; blockIdx++)
 			{
 				BlockHeader header = GetBlockHeader(blockIdx);
 				if (header.Hash == IoHash.Zero)
 				{
-					freeSize += _blockSize;
+					numFreeBlocks++;
 				}
 				else
 				{
@@ -305,56 +308,46 @@ namespace Horde.Server.Artifacts
 					allocatedSize += _blockSize;
 				}
 			}
-			return new BlockCacheStats(_lookup.Count, _numBlocks << _blockSizeLog2, usedSize, allocatedSize, freeSize);
+			return new BlockCacheStats(_lookup.Count, _numBlocks, numFreeBlocks, _blockSize, _numBlocks << _blockSizeLog2, usedSize, allocatedSize);
 		}
 
 		int GetNextFreeBlockIdx()
 		{
-			int scanLength = Math.Min(32, _numBlocks);
-			for (; ; )
+			if (_numAllocatedBlocks < _numBlocks)
 			{
-				int minBlockIdx = 0;
-				int minAccessCount = 256;
+				return _numAllocatedBlocks++;
+			}
 
-				// Check the next 32 entries and find the least referenced
-				for (int offset = 0; offset < scanLength; offset++)
+			int blockIdx;
+			while (!_freeBlocks.TryDequeue(out blockIdx))
+			{
+				PriorityQueue<(IoHash, BlockState), long> priorityQueue = new PriorityQueue<(IoHash, BlockState), long>(NumEvictionSamples);
+				for (int sampleIdx = 0; sampleIdx < NumEvictionSamples; sampleIdx++)
 				{
-					_scanFreeBlockIdx++;
-					if (_scanFreeBlockIdx == _numBlocks)
+					blockIdx = _rng.Next(_numBlocks);
+
+					BlockHeader header = GetBlockHeader(blockIdx);
+					if (_lookup.TryGetValue(header.Hash, out BlockState? state))
 					{
-						_scanFreeBlockIdx = 0;
+						priorityQueue.Enqueue((header.Hash, state), state.LastAccessTicks);
 					}
+				}
 
-					int accessCount = GetBlockAccessCount(_scanFreeBlockIdx);
-					SetBlockAccessCount(_scanFreeBlockIdx, 0);
-
-					if (accessCount < minAccessCount)
+				for (int idx = 0; idx < NumEvictionEntries && priorityQueue.Count > 0; idx++)
+				{
+					(IoHash hash, BlockState evictState) = priorityQueue.Dequeue();
+					if (evictState.TryAddWriteLock())
 					{
-						minBlockIdx = _scanFreeBlockIdx;
-						minAccessCount = accessCount;
-
-						if (accessCount == 0)
+						_lookup.TryRemove(hash, out _);
+						foreach (int evictBlockIdx in evictState.BlockIdxs)
 						{
-							break;
+							SetBlockHeader(evictBlockIdx, default);
+							_freeBlocks.Enqueue(evictBlockIdx);
 						}
 					}
 				}
-
-				// Try to lock this block. If we can't, just try again with the next span of entries.
-				BlockHeader header = GetBlockHeader(minBlockIdx);
-				if (header.Hash == IoHash.Zero || !_lookup.TryGetValue(header.Key, out BlockState? blockState))
-				{
-					return minBlockIdx;
-				}
-
-				// Try to lock the block, and return it
-				if (blockState.TryAddWriteLock())
-				{
-					SetBlockHeader(minBlockIdx, default);
-					_lookup.Remove(header.Key, out _);
-					return minBlockIdx;
-				}
 			}
+			return blockIdx;
 		}
 
 		byte GetBlockAccessCount(int blockIdx)
@@ -401,24 +394,29 @@ namespace Horde.Server.Artifacts
 			IoHash hash = IoHash.Compute(Encoding.UTF8.GetBytes(key));
 
 			int numBlocks = (value.Length + (_blockSize - 1)) >> _blockSizeLog2;
+			int[] blockIdxs = new int[numBlocks];
+
 			for (int idx = 0; idx < numBlocks; idx++)
 			{
-				BlockKey blockKey = new BlockKey(hash, idx);
-				if (!_lookup.ContainsKey(blockKey))
-				{
-					int offset = idx << _blockSizeLog2;
-					int length = Math.Min(value.Length - offset, _blockSize);
+				int offset = idx << _blockSizeLog2;
+				int length = Math.Min(value.Length - offset, _blockSize);
 
-					int freeBlockIdx = GetNextFreeBlockIdx();
+				int freeBlockIdx = GetNextFreeBlockIdx();
+				blockIdxs[idx] = freeBlockIdx;
 
-					Memory<byte> blockData = GetBlockData(freeBlockIdx, _blockSize);
-					value.Slice(offset, length).CopyTo(blockData);
+				Memory<byte> blockData = GetBlockData(freeBlockIdx, _blockSize);
+				value.Slice(offset, length).CopyTo(blockData);
 
-					BlockHeader header = new BlockHeader(blockKey, numBlocks, length);
-					SetBlockHeader(freeBlockIdx, header);
+				BlockHeader header = new BlockHeader(hash, idx, numBlocks, length);
+				SetBlockHeader(freeBlockIdx, header);
+			}
 
-					BlockState state = new BlockState(freeBlockIdx, length, numBlocks);
-					_lookup.TryAdd(blockKey, state);
+			BlockState state = new BlockState(blockIdxs, value.Length & (_blockSize - 1));
+			if (!_lookup.TryAdd(hash, state))
+			{
+				foreach (int blockIdx in blockIdxs)
+				{
+					_freeBlocks.Enqueue(blockIdx);
 				}
 			}
 
@@ -434,52 +432,22 @@ namespace Horde.Server.Artifacts
 		{
 			IoHash hash = IoHash.Compute(Encoding.UTF8.GetBytes(key));
 
-			BlockState? firstBlock;
-			if (!_lookup.TryGetValue(new BlockKey(hash, 0), out firstBlock) || !firstBlock.TryAddReadLock())
+			BlockState? blockState;
+			if (!_lookup.TryGetValue(hash, out blockState) || !blockState.TryAddReadLock())
 			{
 				return null;
 			}
 
-			BlockState[] blocks = new BlockState[firstBlock.NumBlocks];
-			blocks[0] = firstBlock;
-
-			try
-			{
-				ReadOnlySequenceBuilder<byte> builder = new ReadOnlySequenceBuilder<byte>();
-				builder.Append(GetBlockData(firstBlock.BlockIdx, firstBlock.BlockLength));
-
-				for (int idx = 1; idx < firstBlock.NumBlocks; idx++)
-				{
-					BlockState? nextBlock;
-					if (!_lookup.TryGetValue(new BlockKey(hash, (ushort)idx), out nextBlock) || !nextBlock.TryAddReadLock())
-					{
-						return null;
-					}
-
-					blocks[idx] = nextBlock;
-					builder.Append(GetBlockData(nextBlock.BlockIdx, nextBlock.BlockLength));
-				}
+			blockState.Touch();
 
-				foreach (BlockState block in blocks)
-				{
-					int nextValue = GetBlockAccessCount(block.BlockIdx) + 1;
-					if (nextValue < 256)
-					{
-						SetBlockAccessCount(block.BlockIdx, (byte)nextValue);
-					}
-				}
-
-				BlockCacheValue result = new BlockCacheValue(builder.Construct(), blocks);
-				blocks = Array.Empty<BlockState>();
-				return result;
-			}
-			finally
+			ReadOnlySequenceBuilder<byte> builder = new ReadOnlySequenceBuilder<byte>();
+			for (int idx = 0; idx + 1 < blockState.BlockIdxs.Length; idx++)
 			{
-				for (int idx = 0; idx < blocks.Length; idx++)
-				{
-					blocks[idx]?.ReleaseReadLock();
-				}
+				builder.Append(GetBlockData(blockState.BlockIdxs[idx], _blockSize));
 			}
+			builder.Append(GetBlockData(blockState.BlockIdxs[^1], blockState.LastBlockLength));
+
+			return new BlockCacheValue(builder.Construct(), blockState);
 		}
 	}
 }
