Horde: Updating schema documentation

#rnx

diff --git a/Engine/Source/Programs/Horde/Docs/Config/Schema/Streams.md b/Engine/Source/Programs/Horde/Docs/Config/Schema/Streams.md
--- a/Engine/Source/Programs/Horde/Docs/Config/Schema/Streams.md
+++ b/Engine/Source/Programs/Horde/Docs/Config/Schema/Streams.md
@@ -239,7 +239,7 @@ Name | Description
 `submitDescription` | `string`<br>Description for new changelists
 `defaultChange` | [`ChangeQueryConfig`](#changequeryconfig)`[]`<br>Default change to build at. Each object has a condition parameter which can evaluated by the server to determine which change to use.
 `arguments` | `string[]`<br>Fixed arguments for the new job
-`parameters` | [`GroupParameterData`](#groupparameterdata)/[`TextParameterData`](#textparameterdata)/[`ListParameterData`](#listparameterdata)/[`BoolParameterData`](#boolparameterdata)`[]`<br>Parameters for this template
+`parameters` | [`TextParameterData`](#textparameterdata)/[`ListParameterData`](#listparameterdata)/[`BoolParameterData`](#boolparameterdata)`[]`<br>Parameters for this template
 `jobOptions` | [`JobOptions`](#joboptions)<br>Default settings for jobs
 
 ## ScheduleConfig
@@ -366,26 +366,6 @@ Name | Description
 `High` | High priority
 `Highest` | Highest priority
 
-## GroupParameterData
-
-Used to group a number of other parameters
-
-Name | Description
----- | -----------
-`type` | Group<br>Type discriminator
-`label` | `string`<br>Label to display next to this parameter
-`style` | [`GroupParameterStyle`](#groupparameterstyle-enum)<br>How to display this group
-`children` | [`GroupParameterData`](#groupparameterdata)/[`TextParameterData`](#textparameterdata)/[`ListParameterData`](#listparameterdata)/[`BoolParameterData`](#boolparameterdata)`[]`<br>List of child parameters
-
-## GroupParameterStyle (Enum)
-
-Describes how to render a group parameter
-
-Name | Description
----- | -----------
-`Tab` | Separate tab on the form
-`Section` | Section with heading
-
 ## TextParameterData
 
 Free-form text entry parameter
