Horde: Only direct to existing preflight job if streams match
#rnx
#skipci

diff --git a/Engine/Source/Programs/Horde/HordeDashboard/src/components/Preflight.tsx b/Engine/Source/Programs/Horde/HordeDashboard/src/components/Preflight.tsx
--- a/Engine/Source/Programs/Horde/HordeDashboard/src/components/Preflight.tsx
+++ b/Engine/Source/Programs/Horde/HordeDashboard/src/components/Preflight.tsx
@@ -109,26 +109,29 @@ export const PreflightRedirector: React.FC = () => {
 
          if (result && result.length === 1) {
 
-            let url = `/job/${result[0].id}?newbuild=true&allowtemplatechange=true&shelvedchange=${change}&p4v=true`;
-
-            if (autosubmit === "true") {
-               url += "&autosubmit=true";
-            }
-
-            if (templateId) {
-               url += `&templateId=${templateId}`;
-            }
-
-            if (parameters.length) {
-               url += ("&" + parameters.join("&"));
-            }
-
-            if (version === "2") {
-               url += "&newbuildversion=2"
-            }
-
-            navigate(url, { replace: true });
-            return;
+            if (stream.id === result[0].streamId) {
+               
+               let url = `/job/${result[0].id}?newbuild=true&allowtemplatechange=true&shelvedchange=${change}&p4v=true`;
+
+               if (autosubmit === "true") {
+                  url += "&autosubmit=true";
+               }
+   
+               if (templateId) {
+                  url += `&templateId=${templateId}`;
+               }
+   
+               if (parameters.length) {
+                  url += ("&" + parameters.join("&"));
+               }
+   
+               if (version === "2") {
+                  url += "&newbuildversion=2"
+               }
+   
+               navigate(url, { replace: true });
+               return;   
+            } 
          }
 
          let url = `/stream/${stream!.id}?tab=${tab}&newbuild=true&shelvedchange=${change}&p4v=true`;
