Horde: Cache breaking for dashboard deploys
#rnx
#skipci

diff --git a/Engine/Source/Programs/Horde/HordeDashboard/package.json b/Engine/Source/Programs/Horde/HordeDashboard/package.json
--- a/Engine/Source/Programs/Horde/HordeDashboard/package.json
+++ b/Engine/Source/Programs/Horde/HordeDashboard/package.json
@@ -1,7 +1,7 @@
 {
   "name": "horde-dashboard",
   "private": true,
-  "version": "0.1.2",
+  "version": "0.1.3",
   "type": "module",  
   "scripts": {
     "dev": "vite",
diff --git a/Engine/Source/Programs/Horde/HordeDashboard/vite.config.ts b/Engine/Source/Programs/Horde/HordeDashboard/vite.config.ts
--- a/Engine/Source/Programs/Horde/HordeDashboard/vite.config.ts
+++ b/Engine/Source/Programs/Horde/HordeDashboard/vite.config.ts
@@ -1,5 +1,6 @@
 import { defineConfig } from 'vite'
 import react from '@vitejs/plugin-react'
+import { version } from './package.json'
 
 const proxyTarget = "http://127.0.0.1:13340"
 const debug = false;
@@ -14,7 +15,14 @@ export default defineConfig({
     },
   })],
   build: {
-    chunkSizeWarningLimit: 8192
+    chunkSizeWarningLimit: 8192,
+    rollupOptions: {
+      output: {
+          entryFileNames: `[name].${version}.js`,
+          chunkFileNames: `[name].${version}.js`,
+          assetFileNames: `[name].${version}.[ext]`
+      }
+    }
   },
   server: {
     proxy: {
