Horde: Filter out new test data schema ingestion
#rnx

diff --git a/Engine/Source/Programs/Horde/Horde.Server/Jobs/TestData/TestDataCollection.cs b/Engine/Source/Programs/Horde/Horde.Server/Jobs/TestData/TestDataCollection.cs
--- a/Engine/Source/Programs/Horde/Horde.Server/Jobs/TestData/TestDataCollection.cs
+++ b/Engine/Source/Programs/Horde/Horde.Server/Jobs/TestData/TestDataCollection.cs
@@ -1049,8 +1049,12 @@ namespace Horde.Server.Jobs.TestData
 			foreach (TestDataDocument item in documents)
 			{
 				BsonDocument testData = item.Data;
-				version = testData.GetValue("Version", new BsonInt32(0)).AsInt32;
-				if (version > 0)
+				if (!testData.TryGetInt32("Version", out version))
+				{
+					testData.TryGetInt32("version", out version);
+				}
+				
+				if (version != 0)
 				{
 					break;
 				}
@@ -1062,6 +1066,12 @@ namespace Horde.Server.Jobs.TestData
 				return;
 			}
 
+			if (version > 1)
+			{
+				// Don't support v2 ingestion yet
+				return;
+			}
+
 			foreach (TestDataDocument item in documents)
 			{
 				BsonDocument testData = item.Data;
