Horde: Always cache break when building production site
#rnx
#skipci

diff --git a/Engine/Source/Programs/Horde/HordeDashboard/vite.config.ts b/Engine/Source/Programs/Horde/HordeDashboard/vite.config.ts
--- a/Engine/Source/Programs/Horde/HordeDashboard/vite.config.ts
+++ b/Engine/Source/Programs/Horde/HordeDashboard/vite.config.ts
@@ -1,10 +1,11 @@
 import { defineConfig } from 'vite'
 import react from '@vitejs/plugin-react'
-import { version } from './package.json'
 
 const proxyTarget = "http://127.0.0.1:13340"
 const debug = false;
 
+const cacheBreak = Date.now();
+
 // https://vitejs.dev/config/
 export default defineConfig({
   plugins: [react({
@@ -18,9 +19,9 @@ export default defineConfig({
     chunkSizeWarningLimit: 8192,
     rollupOptions: {
       output: {
-          entryFileNames: `[name].${version}.js`,
-          chunkFileNames: `[name].${version}.js`,
-          assetFileNames: `[name].${version}.[ext]`
+          entryFileNames: `[name].${cacheBreak}.js`,
+          chunkFileNames: `[name].${cacheBreak}.js`,
+          assetFileNames: `[name].${cacheBreak}.[ext]`
       }
     }
   },
