Horde: Missing call to BlobPipeline.FinishAdding().

#rnx

diff --git a/Engine/Source/Programs/Horde/Horde.Server/Artifacts/ArtifactsController.cs b/Engine/Source/Programs/Horde/Horde.Server/Artifacts/ArtifactsController.cs
--- a/Engine/Source/Programs/Horde/Horde.Server/Artifacts/ArtifactsController.cs
+++ b/Engine/Source/Programs/Horde/Horde.Server/Artifacts/ArtifactsController.cs
@@ -723,6 +723,7 @@ namespace Horde.Server.Artifacts
 
 				pipeline.Add(new BlobRequest<IoHash>(blobRef, hash));
 			}
+			pipeline.FinishAdding();
 
 			// Send the response headers
 			HttpResponse response = HttpContext.Response;
