Horde: Filter out v2 test data earlier
#rnx

diff --git a/Engine/Source/Programs/Horde/Horde.Server/Jobs/TestData/TestDataCollection.cs b/Engine/Source/Programs/Horde/Horde.Server/Jobs/TestData/TestDataCollection.cs
--- a/Engine/Source/Programs/Horde/Horde.Server/Jobs/TestData/TestDataCollection.cs
+++ b/Engine/Source/Programs/Horde/Horde.Server/Jobs/TestData/TestDataCollection.cs
@@ -509,7 +509,26 @@ namespace Horde.Server.Jobs.TestData
 			for (int i = 0; i < data.Length; i++)
 			{
 				(string key, BsonDocument document) = data[i];
-				documents.Add(new TestDataDocument(job, step, key, document));
+
+				// Get test document version
+				int version;
+				if (!document.TryGetInt32("Version", out version))
+				{
+					if (!document.TryGetInt32("version", out version))
+					{
+						continue;
+					}
+				}
+
+				if (version == 1)
+				{
+					documents.Add(new TestDataDocument(job, step, key, document));
+				}				
+			}
+
+			if (documents.Count == 0)
+			{
+				return documents;
 			}
 
 			await _testDataDocuments.InsertManyAsync(documents, null, cancellationToken);
@@ -1044,34 +1063,6 @@ namespace Horde.Server.Jobs.TestData
 			List<AutomatedTestSessionData> sessions = new List<AutomatedTestSessionData>();
 			List<UnrealAutomatedTestData> tests = new List<UnrealAutomatedTestData>();
 
-			// Get test document version
-			int version = 0;
-			foreach (TestDataDocument item in documents)
-			{
-				BsonDocument testData = item.Data;
-				if (!testData.TryGetInt32("Version", out version))
-				{
-					testData.TryGetInt32("version", out version);
-				}
-				
-				if (version != 0)
-				{
-					break;
-				}
-			}
-
-			if (version < 1)
-			{
-				_logger.LogWarning("Test data does not have version and needs to be updated in stream for job {JobId} step {StepId}", job.Id, step.Id);
-				return;
-			}
-
-			if (version > 1)
-			{
-				// Don't support v2 ingestion yet
-				return;
-			}
-
 			foreach (TestDataDocument item in documents)
 			{
 				BsonDocument testData = item.Data;
