Update release notes.

#rnx

diff --git a/Engine/Source/Programs/Horde/Docs/ReleaseNotes.md b/Engine/Source/Programs/Horde/Docs/ReleaseNotes.md
--- a/Engine/Source/Programs/Horde/Docs/ReleaseNotes.md
+++ b/Engine/Source/Programs/Horde/Docs/ReleaseNotes.md
@@ -1,5 +1,22 @@
 # Release Notes
 
+## 2024-06-13
+
+* Add a pipelined blob read class for storage blobs. (34322019)
+* Add a disk-based cache for compressed Unsync blocks. Uses memory mapped files to read/write to underlying storage, LRU eviction via random sampling. (34288152, 34292867)
+* Tweak MongoDB client's config defaults to curb wait queue full errors (34276157)
+* Resolving of capabilities has been known to stall at times. Adding trace spans to these should help see which part is causing slowdowns. (34232453)
+* Add zstd as a compression format for bundles. Also add some basic compression tests, and fix an isuse with gzip streams being truncated. (34211536)
+* Include rolling hashes in bundle archives to support unsync manifests. DirectoryNode/ChunkedDataNodeRef now use HordeApiVersion for versioning purposes, rather than their own custom versioning scheme. These blob types will serialize the most relevant Horde API version number they support into archives. (34207581)
+* Ignore cancellations in AWS instance lifecycle (34207087)
+* Include modification times for file entries in bundles by default, but do not write them unless the configured API version allows it. (34205938)
+* Add compute endpoint that automatically finds the best cluster. Currently uses IP of requester together with networks from global config. (34201862)
+* Fix some artifacts not showing up in job responses, due to default cap of 100 artifacts returned. (34200332)
+* Check agent's compute cluster membership on assignment. Previously it was only checked by the task source. (34199317)
+* Set user ID in compute assignment request track usage by user (34195976)
+* Disable config updates when running locally in Redis read-only mode, preventing an exception on startup. (34169346)
+* Make horde chunk boundary calculation consistent with unsync, add unit tests to validate determinism (34168267)
+
 ## 2024-06-06
 
 * Add a retry policy to the Epic telemetry sink, to stop 502 errors causing server log entries. (34141477)
