Additional prebuilt Horde server Buildgraph for non-windows
#tests Built new server architectures & installer locally, short confidence test on osx-arm64
#rnx

diff --git a/Engine/Source/Programs/Horde/BuildHorde.xml b/Engine/Source/Programs/Horde/BuildHorde.xml
--- a/Engine/Source/Programs/Horde/BuildHorde.xml
+++ b/Engine/Source/Programs/Horde/BuildHorde.xml
@@ -114,6 +114,9 @@
 	<Property Name="ServerToolsStagingDir_UnrealGameSync" Value="$(ServerToolsStagingDir)/UnrealGameSync"/>
 	<Property Name="ServerToolsStagingDir_UnrealGameSyncCmd" Value="$(ServerToolsStagingDir)/UnrealGameSyncCmd"/>
 	<Property Name="ServerToolsStagingDir_UnrealGameSyncMsi" Value="$(ServerToolsStagingDir)/UnrealGameSyncMsi"/>
+
+    <!-- Changing the windows architecture value here means updates elsewhere in this XML, and in the Installer wxs file -->
+	<Property Name="ServerArchList" Value="win-x64;osx-arm64;linux-x64"/>
 	
 	<Agent Name="HordeInstallerTools" Type="LicenseeWin64;Win64">
 		<Node Name="Build HordeCmd" Produces="#HordeCmdFiles">
@@ -240,39 +243,41 @@
 
 		</Node>
 
-		<Node Name="Stage Horde Server Installer" Requires="#DashboardOutputFiles;#HordeCmdFiles;#HordeAgentFiles;#HordeAgentMsiFiles;#P4VUtilsFiles;#UnrealGameSyncFiles;#UnrealGameSyncCmdFiles;#UnrealGameSyncMsiFiles;" Produces="#ServerWithToolsFiles">
+		<ForEach Name="ServerArch" Values="$(ServerArchList)">
+			<Node Name="Stage Horde Server Installer ($(ServerArch))" Requires="#DashboardOutputFiles;#HordeCmdFiles;#HordeAgentFiles;#HordeAgentMsiFiles;#P4VUtilsFiles;#UnrealGameSyncFiles;#UnrealGameSyncCmdFiles;#UnrealGameSyncMsiFiles;" Produces="#ServerWithToolsFiles_$(ServerArch)">
 		
-			<!-- Stage the server -->
-			<Delete Files="$(StagingDir)/Server/..."/>
-			<DotNet Arguments="publish &quot;$(RootDir)/Engine/Source/Programs/Horde/Horde.Server/Horde.Server.csproj&quot; --output &quot;$(StagingDir)/Server&quot; --runtime win-x64 --self-contained $(VersionArguments)" DotNetPath="$(DotNetPath)" />
-
-			<!-- Remove the local config file. This is only used for local iteration. -->
-			<Delete Files="$(StagingDir)/Server/appsettings.Local.json"/>
-
-			<!-- Add the bundled tools into the installer -->
-			<DotNet Arguments="run --project &quot;$(RootDir)/Engine/Source/Programs/Horde/Horde.Server/Horde.Server.csproj&quot; -- generate tooldata -id=horde-cmd -name=&quot;Horde (Command-Line)&quot; -description=&quot;Update package for the Horde command-line tool.&quot; -version=&quot;$(InformationalVersion)&quot; -inputdir=&quot;$(ServerToolsStagingDir_HordeCmd)&quot; -serverdir=&quot;$(StagingDir)/Server&quot; -public -showinugs -Category=Internal" DotNetPath="$(DotNetPath)" />
-			<DotNet Arguments="run --project &quot;$(RootDir)/Engine/Source/Programs/Horde/Horde.Server/Horde.Server.csproj&quot; -- generate tooldata -id=horde-agent -name=&quot;Horde Agent (Cross-Platform)&quot; -description=&quot;Cross-platform Horde Agent distribution. Allows hosts to take on work from the Horde Server.&quot; -version=&quot;$(InformationalVersion)&quot; -inputdir=&quot;$(ServerToolsStagingDir_HordeAgent)&quot; -serverdir=&quot;$(StagingDir)/Server&quot; -public" DotNetPath="$(DotNetPath)" />
-			<DotNet Arguments="run --project &quot;$(RootDir)/Engine/Source/Programs/Horde/Horde.Server/Horde.Server.csproj&quot; -- generate tooldata -id=horde-agent-win64 -name=&quot;Horde Agent (Windows)&quot; -description=&quot;Update package for older installations of Horde Agent through the MSI installer.&quot; -version=&quot;$(InformationalVersion)&quot; -inputdir=&quot;$(ServerToolsStagingDir_HordeAgentWin64)&quot; -serverdir=&quot;$(StagingDir)/Server&quot; -public -showindashboard=false -category=Internal" DotNetPath="$(DotNetPath)" />
-			<DotNet Arguments="run --project &quot;$(RootDir)/Engine/Source/Programs/Horde/Horde.Server/Horde.Server.csproj&quot; -- generate tooldata -id=horde-agent-win64-installer -name=&quot;Horde Agent (Windows Installer)&quot; -description=&quot;Installer for the Horde Agent on Windows, with tray application.&quot; -version=&quot;$(InformationalVersion)&quot; -inputdir=&quot;$(ServerToolsStagingDir_HordeAgentMsi)&quot; -serverdir=&quot;$(StagingDir)/Server&quot; -public" DotNetPath="$(DotNetPath)" />
-			<DotNet Arguments="run --project &quot;$(RootDir)/Engine/Source/Programs/Horde/Horde.Server/Horde.Server.csproj&quot; -- generate tooldata -id=p4v-utils -name=&quot;P4V Extensions&quot; -description=&quot;Update package for the P4VUtils tool.&quot; -version=&quot;$(InformationalVersion)&quot; -inputdir=&quot;$(ServerToolsStagingDir_P4VUtils)&quot; -serverdir=&quot;$(StagingDir)/Server&quot; -public -showinugs -category=Internal" DotNetPath="$(DotNetPath)" />
-			<DotNet Arguments="run --project &quot;$(RootDir)/Engine/Source/Programs/Horde/Horde.Server/Horde.Server.csproj&quot; -- generate tooldata -id=ugs-cmd -name=&quot;UnrealGameSync (Command-Line)&quot; -description=&quot;Command-line version of UnrealGameSync. Supports Windows, MacOS and Linux.&quot; -version=&quot;$(InformationalVersion)&quot; -inputdir=&quot;$(ServerToolsStagingDir_UnrealGameSyncCmd)&quot; -serverdir=&quot;$(StagingDir)/Server&quot; -public" DotNetPath="$(DotNetPath)" />
-			<DotNet Arguments="run --project &quot;$(RootDir)/Engine/Source/Programs/Horde/Horde.Server/Horde.Server.csproj&quot; -- generate tooldata -id=ugs-msi -name=&quot;UnrealGameSync (Windows Installer)&quot; -description=&quot;Installer for UnrealGameSync.&quot; -version=&quot;$(InformationalVersion)&quot; -inputdir=&quot;$(ServerToolsStagingDir_UnrealGameSyncMsi)&quot; -serverdir=&quot;$(StagingDir)/Server&quot; -public" DotNetPath="$(DotNetPath)" />
-			<DotNet Arguments="run --project &quot;$(RootDir)/Engine/Source/Programs/Horde/Horde.Server/Horde.Server.csproj&quot; -- generate tooldata -id=ugs-win -name=&quot;UnrealGameSync (Windows Update Package)&quot; -description=&quot;Update package for UGS on Windows.&quot; -version=&quot;$(InformationalVersion)&quot; -inputdir=&quot;$(ServerToolsStagingDir_UnrealGameSync)&quot; -serverdir=&quot;$(StagingDir)/Server&quot; -public -showindashboard=false -category=Internal" DotNetPath="$(DotNetPath)" />
-
-			<!-- Separate copy of the command line tool for installation with the server on Windows -->
-			<Delete Files="$(StagingDir)/Tools/..."/>
-			<DotNet Arguments="publish &quot;$(RootDir)/Engine/Source/Programs/Horde/Horde/Horde.csproj&quot; --output &quot;$(StagingDir)/Tools&quot; --runtime win-x64 --self-contained $(VersionArguments)" DotNetPath="$(DotNetPath)" />
-
-			<!-- Copy the Dashboard -->
-			<Copy Files="#DashboardOutputFiles" From="$(StagingDir)/Dashboard/..." To="$(StagingDir)/Server/DashboardApp/..." />
-			
-			<!-- Copy the VCREDIST files -->
-			<Copy From="$(StagingDir)/Server/ThirdParty/Mongo/*.dll" To="$(StagingDir)/Server"/>
-
-			<Tag Files="$(StagingDir)/Server/..." With="#ServerWithToolsFiles" />
-		</Node>
-
-		<Node Name="Build Horde Server Installer" Requires="Stage Horde Server Installer">
+				<!-- Stage the server -->
+				<Delete Files="$(StagingDir)/Server/$(ServerArch)/..."/>
+				<DotNet Arguments="publish &quot;$(RootDir)/Engine/Source/Programs/Horde/Horde.Server/Horde.Server.csproj&quot; --output &quot;$(StagingDir)/Server/$(ServerArch)&quot; --runtime $(ServerArch) --self-contained $(VersionArguments)" DotNetPath="$(DotNetPath)" />
+
+				<!-- Remove the local config file. This is only used for local iteration. -->
+				<Delete Files="$(StagingDir)/Server/$(ServerArch)/appsettings.Local.json"/>
+
+				<!-- Add the bundled tools into the installer -->
+				<DotNet Arguments="run --project &quot;$(RootDir)/Engine/Source/Programs/Horde/Horde.Server/Horde.Server.csproj&quot; -- generate tooldata -id=horde-cmd -name=&quot;Horde (Command-Line)&quot; -description=&quot;Update package for the Horde command-line tool.&quot; -version=&quot;$(InformationalVersion)&quot; -inputdir=&quot;$(ServerToolsStagingDir_HordeCmd)&quot; -serverdir=&quot;$(StagingDir)/Server/$(ServerArch)&quot; -public -showinugs -Category=Internal" DotNetPath="$(DotNetPath)" />
+				<DotNet Arguments="run --project &quot;$(RootDir)/Engine/Source/Programs/Horde/Horde.Server/Horde.Server.csproj&quot; -- generate tooldata -id=horde-agent -name=&quot;Horde Agent (Cross-Platform)&quot; -description=&quot;Cross-platform Horde Agent distribution. Allows hosts to take on work from the Horde Server.&quot; -version=&quot;$(InformationalVersion)&quot; -inputdir=&quot;$(ServerToolsStagingDir_HordeAgent)&quot; -serverdir=&quot;$(StagingDir)/Server/$(ServerArch)&quot; -public" DotNetPath="$(DotNetPath)" />
+				<DotNet Arguments="run --project &quot;$(RootDir)/Engine/Source/Programs/Horde/Horde.Server/Horde.Server.csproj&quot; -- generate tooldata -id=horde-agent-win64 -name=&quot;Horde Agent (Windows)&quot; -description=&quot;Update package for older installations of Horde Agent through the MSI installer.&quot; -version=&quot;$(InformationalVersion)&quot; -inputdir=&quot;$(ServerToolsStagingDir_HordeAgentWin64)&quot; -serverdir=&quot;$(StagingDir)/Server/$(ServerArch)&quot; -public -showindashboard=false -category=Internal" DotNetPath="$(DotNetPath)" />
+				<DotNet Arguments="run --project &quot;$(RootDir)/Engine/Source/Programs/Horde/Horde.Server/Horde.Server.csproj&quot; -- generate tooldata -id=horde-agent-win64-installer -name=&quot;Horde Agent (Windows Installer)&quot; -description=&quot;Installer for the Horde Agent on Windows, with tray application.&quot; -version=&quot;$(InformationalVersion)&quot; -inputdir=&quot;$(ServerToolsStagingDir_HordeAgentMsi)&quot; -serverdir=&quot;$(StagingDir)/Server/$(ServerArch)&quot; -public" DotNetPath="$(DotNetPath)" />
+				<DotNet Arguments="run --project &quot;$(RootDir)/Engine/Source/Programs/Horde/Horde.Server/Horde.Server.csproj&quot; -- generate tooldata -id=p4v-utils -name=&quot;P4V Extensions&quot; -description=&quot;Update package for the P4VUtils tool.&quot; -version=&quot;$(InformationalVersion)&quot; -inputdir=&quot;$(ServerToolsStagingDir_P4VUtils)&quot; -serverdir=&quot;$(StagingDir)/Server/$(ServerArch)&quot; -public -showinugs -category=Internal" DotNetPath="$(DotNetPath)" />
+				<DotNet Arguments="run --project &quot;$(RootDir)/Engine/Source/Programs/Horde/Horde.Server/Horde.Server.csproj&quot; -- generate tooldata -id=ugs-cmd -name=&quot;UnrealGameSync (Command-Line)&quot; -description=&quot;Command-line version of UnrealGameSync. Supports Windows, MacOS and Linux.&quot; -version=&quot;$(InformationalVersion)&quot; -inputdir=&quot;$(ServerToolsStagingDir_UnrealGameSyncCmd)&quot; -serverdir=&quot;$(StagingDir)/Server/$(ServerArch)&quot; -public" DotNetPath="$(DotNetPath)" />
+				<DotNet Arguments="run --project &quot;$(RootDir)/Engine/Source/Programs/Horde/Horde.Server/Horde.Server.csproj&quot; -- generate tooldata -id=ugs-msi -name=&quot;UnrealGameSync (Windows Installer)&quot; -description=&quot;Installer for UnrealGameSync.&quot; -version=&quot;$(InformationalVersion)&quot; -inputdir=&quot;$(ServerToolsStagingDir_UnrealGameSyncMsi)&quot; -serverdir=&quot;$(StagingDir)/Server/$(ServerArch)&quot; -public" DotNetPath="$(DotNetPath)" />
+				<DotNet Arguments="run --project &quot;$(RootDir)/Engine/Source/Programs/Horde/Horde.Server/Horde.Server.csproj&quot; -- generate tooldata -id=ugs-win -name=&quot;UnrealGameSync (Windows Update Package)&quot; -description=&quot;Update package for UGS on Windows.&quot; -version=&quot;$(InformationalVersion)&quot; -inputdir=&quot;$(ServerToolsStagingDir_UnrealGameSync)&quot; -serverdir=&quot;$(StagingDir)/Server/$(ServerArch)&quot; -public -showindashboard=false -category=Internal" DotNetPath="$(DotNetPath)" />
+
+				<!-- Separate copy of the command line tool for installation with the server on Windows -->
+				<Delete Files="$(StagingDir)/Tools/..."/>
+				<DotNet Arguments="publish &quot;$(RootDir)/Engine/Source/Programs/Horde/Horde/Horde.csproj&quot; --output &quot;$(StagingDir)/Tools&quot; --runtime win-x64 --self-contained $(VersionArguments)" DotNetPath="$(DotNetPath)" />
+
+				<!-- Copy the Dashboard -->
+				<Copy Files="#DashboardOutputFiles" From="$(StagingDir)/Dashboard/..." To="$(StagingDir)/Server/$(ServerArch)/DashboardApp/..." />
+				
+				<!-- Copy the VCREDIST files -->
+				<Copy From="$(StagingDir)/Server/$(ServerArch)/ThirdParty/Mongo/*.dll" To="$(StagingDir)/Server/$(ServerArch)"/>
+
+				<Tag Files="$(StagingDir)/Server/$(ServerArch)/..." With="#ServerWithToolsFiles_$(ServerArch)" />
+			</Node>
+		</ForEach>
+
+		<Node Name="Build Horde Server Installer" Requires="Stage Horde Server Installer (win-x64)">
 
 			<!-- Create the tools file list -->
 			<Delete Files="$(StagingDir)/Tools-Bulk/..."/>
@@ -284,7 +289,7 @@
 
 			<!-- Create the server file list -->
 			<Delete Files="$(StagingDir)/Server-Bulk/..."/>
-			<Tag Files="$(StagingDir)/Server/..." Except="HordeServer.exe" With="#ServerFiles"/>
+			<Tag Files="$(StagingDir)/Server/win-x64/..." Except="HordeServer.exe" With="#ServerFiles"/>
 			<Copy Files="#ServerFiles" From="$(StagingDir)" To="$(StagingDir)/Server-Bulk"/>
 
 			<Spawn Exe="$(WixDir)/heat.exe" Arguments="dir &quot;$(StagingDir)/Server-Bulk&quot; -cg HordeServerFiles -scom -sreg -gg -dr InstallDir -srd -var var.SourceDir -out &quot;$(InstallerDir)/HordeServerFiles.wxs&quot;"/>
@@ -304,7 +309,7 @@
 	<Agent Name="HordeInstallerLinux" Type="Linux">
 		<Property Name="StagingDir" Value="$(RootDir)/Staging"/>
 		
-		<Node Name="Build Bundled Docker Image" Requires="Stage Horde Server Installer;#ServerWithToolsFiles">
+		<Node Name="Build Bundled Docker Image" Requires="Stage Horde Server Installer (win-x64);#ServerWithToolsFiles_win-x64">
 			<Tag Files="$(StagePaths)" Except=".../.vs/...;.../.git/...;.../bin/...;.../obj/..." With="#InputFiles"/>
 			<Docker-Build BaseDir="Engine" Files="#InputFiles" Tag="horde-server-bare" DockerFile="Engine/Source/Programs/Horde/Horde.Server/Dockerfile" Arguments="--build-arg msbuild_args=&quot;$(VersionArguments)&quot;"/>
 
