Horde: Fix static analysis warnings in Horde.Server.

#rnx

diff --git a/Engine/Source/Programs/Horde/Horde.Server/Artifacts/BlockCache.cs b/Engine/Source/Programs/Horde/Horde.Server/Artifacts/BlockCache.cs
--- a/Engine/Source/Programs/Horde/Horde.Server/Artifacts/BlockCache.cs
+++ b/Engine/Source/Programs/Horde/Horde.Server/Artifacts/BlockCache.cs
@@ -68,7 +68,6 @@ namespace Horde.Server.Artifacts
 
 		class Partition : IDisposable
 		{
-			public Memory<byte> AccessCounts { get; }
 			public Memory<byte> BlockHeaders { get; }
 			public Memory<byte> BlockData { get; }
 
@@ -76,9 +75,6 @@ namespace Horde.Server.Artifacts
 			{
 				Memory<byte> remainingData = data;
 
-				AccessCounts = remainingData.Slice(0, numBlocks);
-				remainingData = remainingData.Slice(PageAlign(AccessCounts.Length));
-
 				BlockHeaders = remainingData.Slice(0, BlockHeader.NumBytes * numBlocks);
 				remainingData = remainingData.Slice(PageAlign(BlockHeaders.Length));
 
@@ -272,7 +268,7 @@ namespace Horde.Server.Artifacts
 		}
 
 		static int GetPartitionSize(int numBlocks, int blockSize)
-			=> PageAlign(numBlocks) + PageAlign(numBlocks * BlockHeader.NumBytes) + (numBlocks * blockSize);
+			=> PageAlign(numBlocks * BlockHeader.NumBytes) + (numBlocks * blockSize);
 
 		static int PageAlign(int value)
 			=> (value + (PageSize - 1)) & ~(PageSize - 1);
@@ -350,18 +346,6 @@ namespace Horde.Server.Artifacts
 			return blockIdx;
 		}
 
-		byte GetBlockAccessCount(int blockIdx)
-		{
-			Partition partition = _partitions[blockIdx >> _numBlocksPerPartitionLog2];
-			return partition.AccessCounts.Span[blockIdx & (_numBlocksPerPartition - 1)];
-		}
-
-		void SetBlockAccessCount(int blockIdx, byte value)
-		{
-			Partition partition = _partitions[blockIdx >> _numBlocksPerPartitionLog2];
-			partition.AccessCounts.Span[blockIdx & (_numBlocksPerPartition - 1)] = value;
-		}
-
 		BlockHeader GetBlockHeader(int blockIdx)
 		{
 			Partition partition = _partitions[blockIdx >> _numBlocksPerPartitionLog2];
