Horde: Remove unused using directive.

#rnx

diff --git a/Engine/Source/Programs/Horde/Horde.Server.Tests/Artifacts/BlockCacheTests.cs b/Engine/Source/Programs/Horde/Horde.Server.Tests/Artifacts/BlockCacheTests.cs
--- a/Engine/Source/Programs/Horde/Horde.Server.Tests/Artifacts/BlockCacheTests.cs
+++ b/Engine/Source/Programs/Horde/Horde.Server.Tests/Artifacts/BlockCacheTests.cs
@@ -4,7 +4,6 @@ using System;
 using System.Buffers;
 using System.Linq;
 using System.Security.Cryptography;
-using EpicGames.Core;
 using Horde.Server.Artifacts;
 using Microsoft.VisualStudio.TestTools.UnitTesting;
 
