Horde: Use the block cache to cache compressed unsync blob responses.

diff --git a/Engine/Source/Programs/Horde/Horde.Server/Artifacts/ArtifactsController.cs b/Engine/Source/Programs/Horde/Horde.Server/Artifacts/ArtifactsController.cs
--- a/Engine/Source/Programs/Horde/Horde.Server/Artifacts/ArtifactsController.cs
+++ b/Engine/Source/Programs/Horde/Horde.Server/Artifacts/ArtifactsController.cs
@@ -1,11 +1,10 @@
 // Copyright Epic Games, Inc. All Rights Reserved.
 
 using System;
+using System.Buffers;
 using System.Buffers.Binary;
 using System.Collections.Generic;
-using System.Diagnostics;
 using System.IO;
-using System.IO.Pipelines;
 using System.Linq;
 using System.Net;
 using System.Net.Mime;
@@ -53,13 +52,14 @@ namespace Horde.Server.Artifacts
 		readonly IJobCollection _jobCollection;
 		readonly AclService _aclService;
 		readonly UnsyncCache _unsyncCache;
+		readonly IBlockCache _blockCache;
 		readonly GlobalConfig _globalConfig;
 		readonly ILogger _logger;
 
 		/// <summary>
 		/// Constructor
 		/// </summary>
-		public ArtifactsController(IArtifactCollection artifactCollection, StorageService storageService, ILeaseCollection leaseCollection, IJobCollection jobCollection, AclService aclService, UnsyncCache unsyncCache, IOptionsSnapshot<GlobalConfig> globalConfig, ILogger<ArtifactsController> logger)
+		public ArtifactsController(IArtifactCollection artifactCollection, StorageService storageService, ILeaseCollection leaseCollection, IJobCollection jobCollection, AclService aclService, UnsyncCache unsyncCache, IBlockCache blockCache, IOptionsSnapshot<GlobalConfig> globalConfig, ILogger<ArtifactsController> logger)
 		{
 			_artifactCollection = artifactCollection;
 			_storageService = storageService;
@@ -67,6 +67,7 @@ namespace Horde.Server.Artifacts
 			_jobCollection = jobCollection;
 			_aclService = aclService;
 			_unsyncCache = unsyncCache;
+			_blockCache = blockCache;
 			_globalConfig = globalConfig.Value;
 			_logger = logger;
 		}
@@ -706,79 +707,126 @@ namespace Horde.Server.Artifacts
 				}
 			}
 
-			// Find all the blob refs that we need to fetch
-			await using BlobPipeline<IoHash> pipeline = new BlobPipeline<IoHash>();
-			foreach (string block in request.Blocks)
-			{
-				if (!IoHash.TryParse(block, out IoHash hash))
-				{
-					return BadRequest($"Invalid IoHash value: {block}");
-				}
-
-				IBlobRef? blobRef = await _unsyncCache.ReadBlobRefAsync(artifact, hash, cancellationToken);
-				if (blobRef == null)
-				{
-					return NotFound($"Hash '{hash}' is not part of artifact {artifact.Id}");
-				}
-
-				pipeline.Add(new BlobRequest<IoHash>(blobRef, hash));
-			}
-			pipeline.FinishAdding();
-
 			// Send the response headers
 			HttpResponse response = HttpContext.Response;
 			response.ContentType = "application/x-horde-unsync-blob";
 			response.Headers["x-chunk-content-encoding"] = compress ? "zstd" : "identity";
 			response.StatusCode = (int)HttpStatusCode.OK;
 
-			ArrayMemoryWriter? compressedWriter = null;
-
 			await response.StartAsync(cancellationToken);
-			await foreach (BlobResponse<IoHash> blobResponse in pipeline.ReadAllAsync(cancellationToken))
+			long responseLength = 0;
+
+			// Create a pipeline for downloading new blocks
+			await using BlobPipeline<IoHash> pipeline = new BlobPipeline<IoHash>();
+
+			// Cache of blocks that have already been extracted and compressed
+			List<(IoHash, IBlockCacheValue)> cachedBlocks = new List<(IoHash, IBlockCacheValue)>();
+			try
 			{
-				using BlobData blobData = blobResponse.BlobData;
-				if (compress)
+				// Find all the blob refs that we need to fetch
+				foreach (string block in request.Blocks)
 				{
-					compressedWriter ??= new ArrayMemoryWriter(300 * 1024);
-					compressedWriter.Clear();
-					BundleData.Compress(BundleCompressionFormat.Zstd, blobData.Data, compressedWriter);
+					if (!IoHash.TryParse(block, out IoHash hash))
+					{
+						return BadRequest($"Invalid IoHash value: {block}");
+					}
+
+					IBlobRef? blobRef = await _unsyncCache.ReadBlobRefAsync(artifact, hash, cancellationToken);
+					if (blobRef == null)
+					{
+						return NotFound($"Hash '{hash}' is not part of artifact {artifact.Id}");
+					}
 
-					WriteBlock(response.BodyWriter, blobResponse.UserData, blobData.Data.Length, compressedWriter.WrittenSpan);
+					IBlockCacheValue? cacheValue = _blockCache.Get(GetUnsyncBlockKey(hash, compress));
+					if (cacheValue != null)
+					{
+						cachedBlocks.Add((hash, cacheValue));
+					}
+					else
+					{
+						pipeline.Add(new BlobRequest<IoHash>(blobRef, hash));
+					}
 				}
-				else
+				pipeline.FinishAdding();
+
+				// Write all the cached blocks to the response
+				foreach ((IoHash hash, IBlockCacheValue value) in cachedBlocks)
 				{
-					WriteBlock(response.BodyWriter, blobResponse.UserData, blobData.Data.Length, blobData.Data.Span);
+					Memory<byte> buffer = response.BodyWriter.GetMemory((int)value.Data.Length);
+					value.Data.CopyTo(buffer.Span);
+					response.BodyWriter.Advance((int)value.Data.Length);
+					await response.BodyWriter.FlushAsync(cancellationToken);
+					responseLength += value.Data.Length;
+				}
+			}
+			finally
+			{
+				foreach ((_, IBlockCacheValue value) in cachedBlocks)
+				{
+					value.Dispose();
+				}
+			}
+
+			// Read all the uncached blocks
+			if (cachedBlocks.Count < request.Blocks.Count)
+			{
+				ArrayMemoryWriter blockWriter = new ArrayMemoryWriter(128 * 1024);
+				await foreach (BlobResponse<IoHash> blobResponse in pipeline.ReadAllAsync(cancellationToken))
+				{
+					using BlobData blobData = blobResponse.BlobData;
+					blockWriter.Clear();
+
+					// Skip the header to start with. We'll write it once we know the compressed size
+					blockWriter.Advance(BlockHeaderSize);
+
+					// Write the payload
+					if (compress)
+					{
+						BundleData.Compress(BundleCompressionFormat.Zstd, blobData.Data, blockWriter);
+					}
+					else
+					{
+						blockWriter.WriteFixedLengthBytes(blobData.Data.Span);
+					}
+
+					// Write the header
+					WriteBlockHeader(blockWriter.WrittenSpan, blockWriter.WrittenSpan.Length - BlockHeaderSize, blobResponse.UserData, blobData.Data.Length);
+
+					// Add it to the cache for future requests
+					string cacheKey = GetUnsyncBlockKey(blobResponse.UserData, compress);
+					_blockCache.Add(cacheKey, blockWriter.WrittenMemory);
+
+					// Copy it to the response output
+					await response.BodyWriter.WriteAsync(blockWriter.WrittenMemory, cancellationToken);
+					await response.BodyWriter.FlushAsync(cancellationToken);
+					responseLength += blockWriter.WrittenMemory.Length;
 				}
-				await response.BodyWriter.FlushAsync(cancellationToken);
 			}
+
+			// Finish the response
 			await response.CompleteAsync();
+			_logger.LogDebug("Unsync request for {NumBlocks} blocks (compressed={Compressed}); {NumCachedBlocks} cached, {NumFetchedBlocks} fetched, response length = {Length} bytes", request.Blocks.Count, compress, cachedBlocks.Count, request.Blocks.Count - cachedBlocks.Count, responseLength);
 
-			return Ok();
+			return Empty;
 		}
 
-		static void WriteBlock(PipeWriter writer, IoHash decompressedHash, long decompressedSize, ReadOnlySpan<byte> payloadData)
-		{
-			int length = (sizeof(long) * 3) + IoHash.NumBytes + payloadData.Length;
+		static string GetUnsyncBlockKey(IoHash hash, bool compressed)
+			=> compressed ? $"unsync-zstd:{hash}" : $"unsync:{hash}";
 
-			Span<byte> data = writer.GetSpan(length);
-			data = data.Slice(0, length);
+		const int BlockHeaderSize = (sizeof(long) * 3) + IoHash.NumBytes;
 
+		static void WriteBlockHeader(Span<byte> data, long compressedSize, IoHash decompressedHash, long decompressedSize)
+		{
 			BinaryPrimitives.WriteUInt64LittleEndian(data, 0x_4C5C_2AAB_A992_610C);
 			data = data.Slice(8);
 
-			BinaryPrimitives.WriteUInt64LittleEndian(data, (ulong)payloadData.Length);
+			BinaryPrimitives.WriteUInt64LittleEndian(data, (ulong)compressedSize);
 			data = data.Slice(8);
 
 			BinaryPrimitives.WriteUInt64LittleEndian(data, (ulong)decompressedSize);
 			data = data.Slice(8);
 
 			decompressedHash.CopyTo(data);
-			data = data.Slice(IoHash.NumBytes);
-
-			payloadData.CopyTo(data);
-			Debug.Assert(data.Length == payloadData.Length);
-
-			writer.Advance(length);
 		}
 
 		/// <summary>
diff --git a/Engine/Source/Programs/Horde/Horde.Server/ServerSettings.cs b/Engine/Source/Programs/Horde/Horde.Server/ServerSettings.cs
--- a/Engine/Source/Programs/Horde/Horde.Server/ServerSettings.cs
+++ b/Engine/Source/Programs/Horde/Horde.Server/ServerSettings.cs
@@ -936,7 +936,7 @@ namespace Horde.Server
 		public bool OpenBrowser { get; set; } = false;
 
 		/// <summary>
-		/// Directory to use for cache data
+		/// Directory to use for the coarse-grained backend cache. This caches full bundles downloaded from the upstream object store.
 		/// </summary>
 		public string? BundleCacheDir { get; set; }
 
@@ -945,6 +945,16 @@ namespace Horde.Server
 		/// </summary>
 		public long BundleCacheSize { get; set; } = 1024;
 
+		/// <summary>
+		/// Directory to store the fine-grained block cache. This caches individual exports embedded in bundles.
+		/// </summary>
+		public string? BlockCacheDir { get; set; }
+
+		/// <summary>
+		/// Maximum size of the block cache, in megabytes. Currently only allocates in multiples of 1024mb.
+		/// </summary>
+		public long BlockCacheSize { get; set; } = 4096;
+
 		/// <summary>
 		/// Experimental features to enable on the server.
 		/// </summary>
diff --git a/Engine/Source/Programs/Horde/Horde.Server/Startup.cs b/Engine/Source/Programs/Horde/Horde.Server/Startup.cs
--- a/Engine/Source/Programs/Horde/Horde.Server/Startup.cs
+++ b/Engine/Source/Programs/Horde/Horde.Server/Startup.cs
@@ -614,6 +614,7 @@ namespace Horde.Server
 			services.AddSingleton<StorageService>();
 			services.AddSingleton<IStorageService>(sp => sp.GetRequiredService<StorageService>());
 			services.AddScoped(sp => sp.GetRequiredService<StorageService>().CreateStorageClientFactory(sp.GetRequiredService<IOptionsSnapshot<GlobalConfig>>().Value));
+			services.AddSingleton<IBlockCache>(sp => CreateBlockCache(sp));
 			services.AddSingleton<TestDataService>();
 			services.AddSingleton<BundleCache>();
 			services.AddSingleton<StorageBackendCache>(CreateStorageBackendCache);
@@ -1010,6 +1011,13 @@ namespace Horde.Server
 			OnAddHealthChecks(services);
 		}
 
+		static BlockCache CreateBlockCache(IServiceProvider serviceProvider)
+		{
+			ServerSettings serverSettings = serviceProvider.GetRequiredService<IOptions<ServerSettings>>().Value;
+			DirectoryReference cacheDir = DirectoryReference.Combine(ServerApp.DataDir, String.IsNullOrEmpty(serverSettings.BlockCacheDir) ? "BlockCache" : serverSettings.BlockCacheDir);
+			return BlockCache.Create(cacheDir, (int)(serverSettings.BlockCacheSize / 1024));
+		}
+
 		static StorageBackendCache CreateStorageBackendCache(IServiceProvider serviceProvider)
 		{
 			ServerSettings serverSettings = serviceProvider.GetRequiredService<IOptions<ServerSettings>>().Value;
