Horde: Allow overriding storage backends from the server configuration file, to make it easier to run against a production server.

#rnx

diff --git a/Engine/Source/Programs/Horde/Horde.Server/Artifacts/ArtifactsController.cs b/Engine/Source/Programs/Horde/Horde.Server/Artifacts/ArtifactsController.cs
--- a/Engine/Source/Programs/Horde/Horde.Server/Artifacts/ArtifactsController.cs
+++ b/Engine/Source/Programs/Horde/Horde.Server/Artifacts/ArtifactsController.cs
@@ -752,11 +752,14 @@ namespace Horde.Server.Artifacts
 				// Write all the cached blocks to the response
 				foreach ((IoHash hash, IBlockCacheValue value) in cachedBlocks)
 				{
-					Memory<byte> buffer = response.BodyWriter.GetMemory((int)value.Data.Length);
+					int length = (int)value.Data.Length;
+
+					Memory<byte> buffer = response.BodyWriter.GetMemory(length);
 					value.Data.CopyTo(buffer.Span);
-					response.BodyWriter.Advance((int)value.Data.Length);
+					response.BodyWriter.Advance(length);
+					responseLength += length;
+
 					await response.BodyWriter.FlushAsync(cancellationToken);
-					responseLength += value.Data.Length;
 				}
 			}
 			finally
diff --git a/Engine/Source/Programs/Horde/Horde.Server/ServerSettings.cs b/Engine/Source/Programs/Horde/Horde.Server/ServerSettings.cs
--- a/Engine/Source/Programs/Horde/Horde.Server/ServerSettings.cs
+++ b/Engine/Source/Programs/Horde/Horde.Server/ServerSettings.cs
@@ -16,6 +16,7 @@ using EpicGames.Horde.Tools;
 using EpicGames.Perforce;
 using Horde.Server.Agents.Fleet;
 using Horde.Server.Server;
+using Horde.Server.Storage;
 using Horde.Server.Storage.ObjectStores;
 using Horde.Server.Telemetry.Sinks;
 using Horde.Server.Tools;
@@ -706,6 +707,11 @@ namespace Horde.Server
 		/// </summary>
 		public string LogServiceWriteCacheType { get; set; } = "InMemory";
 
+		/// <summary>
+		/// Overridden settings for storage backends. Useful for running against a production server with custom backends.
+		/// </summary>
+		public List<BackendConfig> StorageBackends { get; set; } = new List<BackendConfig>();
+
 		/// <summary>
 		/// Settings for artifact storage
 		/// </summary>
diff --git a/Engine/Source/Programs/Horde/Horde.Server/Storage/StorageConfig.cs b/Engine/Source/Programs/Horde/Horde.Server/Storage/StorageConfig.cs
--- a/Engine/Source/Programs/Horde/Horde.Server/Storage/StorageConfig.cs
+++ b/Engine/Source/Programs/Horde/Horde.Server/Storage/StorageConfig.cs
@@ -84,6 +84,10 @@ namespace Horde.Server.Storage
 			{
 				_backendLookup[backendConfig.Id] = backendConfig;
 			}
+			foreach (BackendConfig backendConfig in globalConfig.ServerSettings.StorageBackends)
+			{
+				_backendLookup[backendConfig.Id] = backendConfig;
+			}
 
 			// Fixup all the inherited properties
 			Dictionary<BackendId, BackendConfig> mergedBackendConfigs = new Dictionary<BackendId, BackendConfig>(Backends.Count);
