Horde: Fix race condition in block cache when allocating new blocks.

#rnx

diff --git a/Engine/Source/Programs/Horde/Horde.Server/Artifacts/BlockCache.cs b/Engine/Source/Programs/Horde/Horde.Server/Artifacts/BlockCache.cs
--- a/Engine/Source/Programs/Horde/Horde.Server/Artifacts/BlockCache.cs
+++ b/Engine/Source/Programs/Horde/Horde.Server/Artifacts/BlockCache.cs
@@ -309,11 +309,19 @@ namespace Horde.Server.Artifacts
 
 		int GetNextFreeBlockIdx()
 		{
-			if (_numAllocatedBlocks < _numBlocks)
+			// Try to allocate a new block
+			int numAllocatedBlocks = _numAllocatedBlocks;
+			while(numAllocatedBlocks < _numBlocks)
 			{
-				return _numAllocatedBlocks++;
+				int initialNumAllocatedBlocks = Interlocked.CompareExchange(ref _numAllocatedBlocks, numAllocatedBlocks + 1, numAllocatedBlocks);
+				if (initialNumAllocatedBlocks == numAllocatedBlocks)
+				{
+					return initialNumAllocatedBlocks;
+				}
+				numAllocatedBlocks = initialNumAllocatedBlocks;
 			}
 
+			// Dequeue a block from the free list
 			int blockIdx;
 			while (!_freeBlocks.TryDequeue(out blockIdx))
 			{
