Horde: Add stats for serving unsync blobs.

#rnx

diff --git a/Engine/Source/Programs/Horde/Horde.Server/Artifacts/ArtifactsController.cs b/Engine/Source/Programs/Horde/Horde.Server/Artifacts/ArtifactsController.cs
--- a/Engine/Source/Programs/Horde/Horde.Server/Artifacts/ArtifactsController.cs
+++ b/Engine/Source/Programs/Horde/Horde.Server/Artifacts/ArtifactsController.cs
@@ -4,6 +4,7 @@ using System;
 using System.Buffers;
 using System.Buffers.Binary;
 using System.Collections.Generic;
+using System.Diagnostics;
 using System.IO;
 using System.Linq;
 using System.Net;
@@ -691,6 +692,8 @@ namespace Horde.Server.Artifacts
 				return BadRequest($"Unsupported hash algorithm: {request.HashStrong}");
 			}
 
+			Stopwatch totalResponseTimer = Stopwatch.StartNew();
+
 			// Disable buffering for the response
 			IHttpResponseBodyFeature? responseBodyFeature = HttpContext.Features.Get<IHttpResponseBodyFeature>();
 			responseBodyFeature?.DisableBuffering();
@@ -714,12 +717,13 @@ namespace Horde.Server.Artifacts
 			response.StatusCode = (int)HttpStatusCode.OK;
 
 			await response.StartAsync(cancellationToken);
-			long responseLength = 0;
 
 			// Create a pipeline for downloading new blocks
 			await using BlobPipeline<IoHash> pipeline = new BlobPipeline<IoHash>();
 
 			// Cache of blocks that have already been extracted and compressed
+			long cachedResponseLength = 0;
+			double cachedResponseTime = 0.0;
 			List<(IoHash, IBlockCacheValue)> cachedBlocks = new List<(IoHash, IBlockCacheValue)>();
 			try
 			{
@@ -750,6 +754,7 @@ namespace Horde.Server.Artifacts
 				pipeline.FinishAdding();
 
 				// Write all the cached blocks to the response
+				Stopwatch cachedResponseTimer = Stopwatch.StartNew();
 				foreach ((IoHash hash, IBlockCacheValue value) in cachedBlocks)
 				{
 					int length = (int)value.Data.Length;
@@ -757,10 +762,11 @@ namespace Horde.Server.Artifacts
 					Memory<byte> buffer = response.BodyWriter.GetMemory(length);
 					value.Data.CopyTo(buffer.Span);
 					response.BodyWriter.Advance(length);
-					responseLength += length;
+					cachedResponseLength += length;
 
 					await response.BodyWriter.FlushAsync(cancellationToken);
 				}
+				cachedResponseTime = cachedResponseTimer.Elapsed.TotalSeconds;
 			}
 			finally
 			{
@@ -771,8 +777,12 @@ namespace Horde.Server.Artifacts
 			}
 
 			// Read all the uncached blocks
+			long uncachedResponseLength = 0;
+			double uncachedResponseTime = 0.0;
 			if (cachedBlocks.Count < request.Blocks.Count)
 			{
+				Stopwatch uncachedResponseTimer = Stopwatch.StartNew();
+
 				ArrayMemoryWriter blockWriter = new ArrayMemoryWriter(128 * 1024);
 				await foreach (BlobResponse<IoHash> blobResponse in pipeline.ReadAllAsync(cancellationToken))
 				{
@@ -802,13 +812,31 @@ namespace Horde.Server.Artifacts
 					// Copy it to the response output
 					await response.BodyWriter.WriteAsync(blockWriter.WrittenMemory, cancellationToken);
 					await response.BodyWriter.FlushAsync(cancellationToken);
-					responseLength += blockWriter.WrittenMemory.Length;
+					uncachedResponseLength += blockWriter.WrittenMemory.Length;
 				}
+
+				uncachedResponseTime = uncachedResponseTimer.Elapsed.TotalSeconds;
 			}
 
 			// Finish the response
 			await response.CompleteAsync();
-			_logger.LogDebug("Unsync request for {NumBlocks} blocks (compressed={Compressed}); {NumCachedBlocks} cached, {NumFetchedBlocks} fetched, response length = {Length} bytes", request.Blocks.Count, compress, cachedBlocks.Count, request.Blocks.Count - cachedBlocks.Count, responseLength);
+
+			long totalResponseLength = cachedResponseLength + uncachedResponseLength;
+			double totalResponseTime = totalResponseTimer.Elapsed.TotalSeconds;
+			_logger.LogInformation("Unsync request for {NumBlocks} blocks ({Type}, {TotalSize:n1}mb, {TotalSpeed:n1}mb/s); {NumCachedBlocks} cached ({CachedSize:n1}mb, {CachedSpeed:n1}mb/s), {NumUncachedBlocks} uncached ({UncachedSize:n1}mb, {UncachedSpeed:n1}mb/s)",
+				request.Blocks.Count,
+				compress ? "compressed" : "not compressed",
+				totalResponseLength / (1024.0 * 1024.0),
+				(totalResponseTime > 0.0) ? (totalResponseLength / (1024 * 1024 * totalResponseTime)) : 0.0,
+
+				cachedBlocks.Count,
+				cachedResponseLength / (1024.0 * 1024.0),
+				(cachedResponseTime > 0.0) ? (cachedResponseLength / (1024 * 1024 * cachedResponseTime)) : 0.0,
+
+				request.Blocks.Count - cachedBlocks.Count,
+				uncachedResponseLength / (1024.0 * 1024.0),
+				(uncachedResponseTime > 0.0) ? (uncachedResponseLength / (1024 * 1024 * uncachedResponseTime)) : 0.0
+			);
 
 			return Empty;
 		}
