Horde: Add timing log statements for creating the unsync manifest for an artifact.

#rnx

diff --git a/Engine/Source/Programs/Horde/Horde.Server.Tests/Artifacts/UnsyncCacheTests.cs b/Engine/Source/Programs/Horde/Horde.Server.Tests/Artifacts/UnsyncCacheTests.cs
--- a/Engine/Source/Programs/Horde/Horde.Server.Tests/Artifacts/UnsyncCacheTests.cs
+++ b/Engine/Source/Programs/Horde/Horde.Server.Tests/Artifacts/UnsyncCacheTests.cs
@@ -54,7 +54,7 @@ namespace Horde.Server.Tests.Artifacts
 			artifact.SetupGet(x => x.NamespaceId).Returns(ns);
 			artifact.SetupGet(x => x.RefName).Returns(refName);
 
-			using UnsyncCache cache = new UnsyncCache(factory.Object);
+			using UnsyncCache cache = new UnsyncCache(factory.Object, NullLogger<UnsyncCache>.Instance);
 
 			UnsyncManifest? manifest = await cache.GetManifestAsync(artifact.Object);
 			Assert.IsNotNull(manifest);
diff --git a/Engine/Source/Programs/Horde/Horde.Server/Artifacts/UnsyncCache.cs b/Engine/Source/Programs/Horde/Horde.Server/Artifacts/UnsyncCache.cs
--- a/Engine/Source/Programs/Horde/Horde.Server/Artifacts/UnsyncCache.cs
+++ b/Engine/Source/Programs/Horde/Horde.Server/Artifacts/UnsyncCache.cs
@@ -3,6 +3,7 @@
 using System;
 using System.Collections.Frozen;
 using System.Collections.Generic;
+using System.Diagnostics;
 using System.Linq;
 using System.Threading;
 using System.Threading.Tasks;
@@ -11,6 +12,7 @@ using EpicGames.Horde.Storage;
 using EpicGames.Horde.Storage.Nodes;
 using Horde.Server.Storage;
 using Microsoft.Extensions.Caching.Memory;
+using Microsoft.Extensions.Logging;
 
 namespace Horde.Server.Artifacts
 {
@@ -65,14 +67,16 @@ namespace Horde.Server.Artifacts
 		readonly IStorageService _storageService;
 		readonly MemoryCache _cache;
 		readonly object _lockObject = new object();
+		readonly ILogger _logger;
 
 		/// <summary>
 		/// Constructor
 		/// </summary>
-		public UnsyncCache(IStorageService storageService)
+		public UnsyncCache(IStorageService storageService, ILogger<UnsyncCache> logger)
 		{
 			_storageService = storageService;
 			_cache = new MemoryCache(new MemoryCacheOptions());
+			_logger = logger;
 		}
 
 		/// <inheritdoc/>
@@ -146,6 +150,7 @@ namespace Horde.Server.Artifacts
 			IStorageClient? storageClient = null;
 			try
 			{
+				Stopwatch timer = Stopwatch.StartNew();
 				storageClient = _storageService.CreateClient(artifact.NamespaceId);
 
 				IBlobRef<DirectoryNode>? target = await storageClient.TryReadRefAsync<DirectoryNode>(artifact.RefName, cancellationToken: cancellationToken);
@@ -163,6 +168,8 @@ namespace Horde.Server.Artifacts
 					blocks[block.Blob.Hash] = block.Blob;
 				}
 
+				_logger.LogDebug("Generated Unsync manifest for artifact {ArtifactId} in {Time:n1}ms", timer.ElapsedMilliseconds);
+
 				ArtifactInfo artifactInfo = new ArtifactInfo(storageClient, new UnsyncManifest(files), blocks.ToFrozenDictionary());
 				storageClient = null;
 				return artifactInfo;
