Horde: Fix static analysis warning.

#rnx

diff --git a/Engine/Source/Programs/Horde/Horde.Server/Artifacts/UnsyncCache.cs b/Engine/Source/Programs/Horde/Horde.Server/Artifacts/UnsyncCache.cs
--- a/Engine/Source/Programs/Horde/Horde.Server/Artifacts/UnsyncCache.cs
+++ b/Engine/Source/Programs/Horde/Horde.Server/Artifacts/UnsyncCache.cs
@@ -168,7 +168,7 @@ namespace Horde.Server.Artifacts
 					blocks[block.Blob.Hash] = block.Blob;
 				}
 
-				_logger.LogDebug("Generated Unsync manifest for artifact {ArtifactId} in {Time:n1}ms", timer.ElapsedMilliseconds);
+				_logger.LogDebug("Generated Unsync manifest for artifact {ArtifactId} in {Time:n1}ms", artifact.Id, timer.ElapsedMilliseconds);
 
 				ArtifactInfo artifactInfo = new ArtifactInfo(storageClient, new UnsyncManifest(files), blocks.ToFrozenDictionary());
 				storageClient = null;
