Horde: Build server on a separate agent to the one running the analyze step.

#rnx

diff --git a/Engine/Source/Programs/Horde/BuildHorde.xml b/Engine/Source/Programs/Horde/BuildHorde.xml
--- a/Engine/Source/Programs/Horde/BuildHorde.xml
+++ b/Engine/Source/Programs/Horde/BuildHorde.xml
@@ -46,10 +46,8 @@
 		Engine/Source/Programs/AutomationTool/AutomationUtils/Matchers/...
 		Engine/Source/Programs/UnrealBuildTool/Matchers/...
 	</Property>
-	
-	<Agent Name="HordeServer" Type="Linux;Win64_Docker">
-		<Property Name="StagingDir" Value="$(RootDir)/Engine/Saved/Horde.Server"/>
 
+	<Agent Name="HordeServer Analyze" Type="Linux;Win64_Docker">
 		<Node Name="Check Documentation">
 			<CheckMarkdown Files="$(RootDir)/Engine/Source/Programs/Horde/README.md;$(RootDir)/Engine/Source/Programs/Horde/Docs/..."/>
 		</Node>
@@ -68,8 +66,12 @@
 			</Property>
 			<NuGet-LicenseCheck BaseDir="Engine/Source/Programs/Horde/Horde.Server" IgnorePackages="$(IgnorePackages)" LicenseDir="Engine/Source/Programs/Horde/ThirdParty/Licenses" DotNetPath="$(DotNetPath)"/>
 		</Node>
+	</Agent>
+	
+	<Agent Name="HordeServer" Type="Linux;Win64_Docker">
+		<Property Name="StagingDir" Value="$(RootDir)/Engine/Saved/Horde.Server"/>
 
-		<Node Name="Build HordeServer" Requires="Check HordeServer Licenses">
+		<Node Name="Build HordeServer">
 			<!-- Tag all the files that need to be staged to build -->
 			<Tag Files="$(StagePaths)" Except=".../.vs/...;.../.git/...;.../bin/...;.../obj/..." With="#InputFiles"/>
 
