Horde: Updating schema documentation

#rnx

diff --git a/Engine/Source/Programs/Horde/Docs/Deployment/ServerSettings.md b/Engine/Source/Programs/Horde/Docs/Deployment/ServerSettings.md
--- a/Engine/Source/Programs/Horde/Docs/Deployment/ServerSettings.md
+++ b/Engine/Source/Programs/Horde/Docs/Deployment/ServerSettings.md
@@ -56,6 +56,7 @@ Name | Description
 `redisConnectionConfig` | `string`<br>
 `redisReadOnlyMode` | `boolean`<br>Whether to disable writes to Redis.
 `logServiceWriteCacheType` | `string`<br>Type of write cache to use in log service Currently Supported: "InMemory" or "Redis"
+`storageBackends` | [`BackendConfig`](#backendconfig)`[]`<br>Overridden settings for storage backends. Useful for running against a production server with custom backends.
 `logStorage` | [`StorageBackendOptions`](#storagebackendoptions)<br>Settings for artifact storage
 `artifactStorage` | [`StorageBackendOptions`](#storagebackendoptions)<br>Settings for artifact storage
 `commitStorage` | [`TreeStoreOptions`](#treestoreoptions)<br>Configuration of tree storage
@@ -101,8 +102,10 @@ Name | Description
 `enableConformTasks` | `boolean`<br>Whether to enable the conform task source.
 `forceConfigUpdateOnStartup` | `boolean`<br>Forces configuration data to be read and updated as part of appplication startup, rather than on a schedule. Useful when running locally.
 `openBrowser` | `boolean`<br>Whether to open a browser on startup
-`bundleCacheDir` | `string`<br>Directory to use for cache data
+`bundleCacheDir` | `string`<br>Directory to use for the coarse-grained backend cache. This caches full bundles downloaded from the upstream object store.
 `bundleCacheSize` | `integer`<br>Maximum size of the storage cache on disk, in megabytes
+`blockCacheDir` | `string`<br>Directory to store the fine-grained block cache. This caches individual exports embedded in bundles.
+`blockCacheSize` | `integer`<br>Maximum size of the block cache, in megabytes. Currently only allocates in multiples of 1024mb.
 `featureFlags` | [`FeatureFlagSettings`](#featureflagsettings)<br>Experimental features to enable on the server.
 `commits` | [`CommitSettings`](#commitsettings)<br>Options for the commit service
 `telemetry` | [`BaseTelemetryConfig`](#basetelemetryconfig)`[]`<br>Settings for sending telemetry events to external services (for example Snowflake, ClickHouse etc)
@@ -130,12 +133,14 @@ Name | Description
 `OpenIdConnect` | Generic OpenID Connect authentication, recommended for most
 `Horde` | Authenticate using username and password credentials stored in Horde OpenID Connect (OIDC) is first and foremost recommended. But if you have a small installation (less than ~10 users) or lacking an OIDC provider, this is an option.
 
-## StorageBackendOptions
+## BackendConfig
 
 Common settings object for different providers
 
 Name | Description
 ---- | -----------
+`id` | `string`<br>The storage backend ID
+`base` | `string`<br>Base backend to copy default settings from
 `type` | [`StorageBackendType`](#storagebackendtype-enum)<br>The type of storage backend to use
 `baseDir` | `string`<br>Base directory for filesystem storage
 `awsBucketName` | `string`<br>Name of the bucket to use
@@ -146,6 +151,8 @@ Name | Description
 `awsRegion` | `string`<br>Region to connect to
 `azureConnectionString` | `string`<br>Connection string for Azure
 `azureContainerName` | `string`<br>Name of the container
+`relayServer` | `string`<br>
+`relayToken` | `string`<br>
 
 ## StorageBackendType (Enum)
 
@@ -169,6 +176,23 @@ Name | Description
 `AssumeRole` | Assume a particular role. Should specify ARN in
 `AssumeRoleWebIdentity` | Assume a particular role using the current environment variables.
 
+## StorageBackendOptions
+
+Common settings object for different providers
+
+Name | Description
+---- | -----------
+`type` | [`StorageBackendType`](#storagebackendtype-enum)<br>The type of storage backend to use
+`baseDir` | `string`<br>Base directory for filesystem storage
+`awsBucketName` | `string`<br>Name of the bucket to use
+`awsBucketPath` | `string`<br>Base path within the bucket
+`awsCredentials` | [`AwsCredentialsType`](#awscredentialstype-enum)<br>Type of credentials to use
+`awsRole` | `string`<br>ARN of a role to assume
+`awsProfile` | `string`<br>The AWS profile to read credentials form
+`awsRegion` | `string`<br>Region to connect to
+`azureConnectionString` | `string`<br>Connection string for Azure
+`azureContainerName` | `string`<br>Name of the container
+
 ## TreeStoreOptions
 
 Options for storing trees
