Horde: Manually split GC scan into explicit batches and add more scopes/logging. Using cursors seems to be hanging on DocDB.

diff --git a/Engine/Source/Programs/Horde/Horde.Server/Storage/StorageService.cs b/Engine/Source/Programs/Horde/Horde.Server/Storage/StorageService.cs
--- a/Engine/Source/Programs/Horde/Horde.Server/Storage/StorageService.cs
+++ b/Engine/Source/Programs/Horde/Horde.Server/Storage/StorageService.cs
@@ -630,7 +630,7 @@ namespace Horde.Server.Storage
 			using TelemetrySpan span = _tracer.StartActiveSpan($"{nameof(StorageService)}.{nameof(TickBlobsAsync)}");
 
 			GcState gcState = await _gcState.GetAsync(cancellationToken);
-			DateTime utcNow = _clock.UtcNow;
+			DateTime ingestTimeUtc = _clock.UtcNow - TimeSpan.FromMinutes(30.0);
 
 			// Get the current state of the storage system
 			State state = CreateState(_globalConfig.CurrentValue);
@@ -638,44 +638,58 @@ namespace Horde.Server.Storage
 			Dictionary<NamespaceId, BundleStorageClient> cachedClients = new();
 			try
 			{
+				long ingestedCount = 0;
+
 				// Compute missing import info, by searching for blobs with an ObjectId timestamp after the last import compute cycle
-				ObjectId latestInfoId = ObjectId.GenerateNewId(utcNow - TimeSpan.FromMinutes(30.0));
-				using (IAsyncCursor<BlobInfo> cursor = await _blobCollection.Find(x => x.Id >= gcState.LastImportBlobInfoId && x.Id < latestInfoId).ToCursorAsync(cancellationToken))
+				ObjectId latestInfoId = ObjectId.GenerateNewId(ingestTimeUtc);
+				for (; ; )
 				{
-					while (await cursor.MoveNextAsync(cancellationToken))
+					// Fetch the next batch of blobs
+					List<BlobInfo> current = await _blobCollection.Find(x => x.Id > gcState.LastImportBlobInfoId && x.Id < latestInfoId).SortBy(x => x.Id).Limit(500).ToListAsync();
+					if (current.Count == 0)
 					{
-						// Find imports, and add a check record for each new blob
-						foreach (BlobInfo blobInfo in cursor.Current)
+						break;
+					}
+
+					// Find imports, and add a check record for each new blob
+					using TelemetrySpan innerSpan = _tracer.StartActiveSpan($"{nameof(StorageService)}.{nameof(TickBlobsAsync)}.Batch");
+					innerSpan.SetAttribute("First", current[0].Id.ToString());
+					innerSpan.SetAttribute("Last", current[^1].Id.ToString());
+					innerSpan.SetAttribute("Count", current.Count);
+
+					foreach (BlobInfo blobInfo in current)
+					{
+						NamespaceInfo? namespaceInfo;
+						if (state.Namespaces.TryGetValue(blobInfo.NamespaceId, out namespaceInfo))
 						{
-							NamespaceInfo? namespaceInfo;
-							if (state.Namespaces.TryGetValue(blobInfo.NamespaceId, out namespaceInfo))
+							BundleStorageClient? storageClient;
+							if (!cachedClients.TryGetValue(namespaceInfo.Id, out storageClient))
 							{
-								BundleStorageClient? storageClient;
-								if (!cachedClients.TryGetValue(namespaceInfo.Id, out storageClient))
-								{
-									storageClient = new BundleStorageClient(namespaceInfo.Backend, _bundleCache, null, _logger);
-									cachedClients.Add(namespaceInfo.Id, storageClient);
-								}
+								storageClient = new BundleStorageClient(namespaceInfo.Backend, _bundleCache, null, _logger);
+								cachedClients.Add(namespaceInfo.Id, storageClient);
+							}
 
-								try
-								{
-									await TickBlobAsync(storageClient, blobInfo, cancellationToken);
-								}
-								catch (ObjectNotFoundException ex)
-								{
-									_logger.LogInformation(ex, "Unable to read references for {NamespaceId} blob {BlobId}: {Message}", blobInfo.NamespaceId, blobInfo.Id, ex.Message);
-								}
-								catch (Exception ex)
-								{
-									_logger.LogWarning(ex, "Unable to read references for {NamespaceId} blob {BlobId} (key: {ObjectKey}): {Message}", blobInfo.NamespaceId, blobInfo.Id, GetObjectKey(blobInfo.Locator), ex.Message);
-								}
+							try
+							{
+								await TickBlobAsync(storageClient, blobInfo, cancellationToken);
+							}
+							catch (ObjectNotFoundException ex)
+							{
+								_logger.LogInformation(ex, "Unable to read references for {NamespaceId} blob {BlobId}: {Message}", blobInfo.NamespaceId, blobInfo.Id, ex.Message);
+							}
+							catch (Exception ex)
+							{
+								_logger.LogWarning(ex, "Unable to read references for {NamespaceId} blob {BlobId} (key: {ObjectKey}): {Message}", blobInfo.NamespaceId, blobInfo.Id, GetObjectKey(blobInfo.Locator), ex.Message);
 							}
 						}
-
-						// Update the last imported blob id
-						await _gcState.UpdateAsync(state => state.LastImportBlobInfoId = latestInfoId, cancellationToken);
 					}
+
+					// Update the last imported blob id
+					gcState = await _gcState.UpdateAsync(state => state.LastImportBlobInfoId = current[^1].Id, cancellationToken);
+					ingestedCount += current.Count;
 				}
+
+				_logger.LogInformation("Added {NumBlobs} blobs for GC (upper time: {Time})", ingestedCount, ingestTimeUtc);
 			}
 			finally
 			{
@@ -719,7 +733,7 @@ namespace Horde.Server.Storage
 				}
 			}
 
-			await _blobCollection.UpdateOneAsync(x => x.Id == blobInfo.Id, Builders<BlobInfo>.Update.Set(x => x.Imports, importInfoIds), null, cancellationToken);
+			await _blobCollection.UpdateOneAsync(x => x.Id == blobInfo.Id, Builders<BlobInfo>.Update.Set(x => x.Imports, importInfoIds).Unset(x => x.GcVersion), null, cancellationToken);
 
 			AddGcCheckRecord(blobInfo.NamespaceId, blobInfo.Id);
 			_logger.LogDebug("Added {Count} imports for {NamespaceId} blob {BlobId}", importInfoIds.Count, blobInfo.NamespaceId, blobInfo.Id);
