Horde: Add missing cancellation token argument.

#rnx

diff --git a/Engine/Source/Programs/Horde/Horde.Server/Storage/StorageService.cs b/Engine/Source/Programs/Horde/Horde.Server/Storage/StorageService.cs
--- a/Engine/Source/Programs/Horde/Horde.Server/Storage/StorageService.cs
+++ b/Engine/Source/Programs/Horde/Horde.Server/Storage/StorageService.cs
@@ -645,7 +645,7 @@ namespace Horde.Server.Storage
 				for (; ; )
 				{
 					// Fetch the next batch of blobs
-					List<BlobInfo> current = await _blobCollection.Find(x => x.Id > gcState.LastImportBlobInfoId && x.Id < latestInfoId).SortBy(x => x.Id).Limit(500).ToListAsync();
+					List<BlobInfo> current = await _blobCollection.Find(x => x.Id > gcState.LastImportBlobInfoId && x.Id < latestInfoId).SortBy(x => x.Id).Limit(500).ToListAsync(cancellationToken);
 					if (current.Count == 0)
 					{
 						break;
