Horde: Fix GC tests.

#rnx

diff --git a/Engine/Source/Programs/Horde/Horde.Server.Tests/Storage/GcServiceTests.cs b/Engine/Source/Programs/Horde/Horde.Server.Tests/Storage/GcServiceTests.cs
--- a/Engine/Source/Programs/Horde/Horde.Server.Tests/Storage/GcServiceTests.cs
+++ b/Engine/Source/Programs/Horde/Horde.Server.Tests/Storage/GcServiceTests.cs
@@ -88,14 +88,17 @@ namespace Horde.Server.Tests.Storage
 				children[idx] = new List<int>();
 			}
 
-			double maxParents = ((numRoots + numInterior) * avgChildren) / (numInterior + numLeaves);
-			for (int idx = numRoots; idx < numNodes; idx++)
+			if (numRoots < numNodes)
 			{
-				int numParents = 1 + (int)(random.NextDouble() * maxParents);
-				for (; numParents > 0; numParents--)
+				double maxParents = ((numRoots + numInterior) * avgChildren) / (numInterior + numLeaves);
+				for (int idx = numRoots; idx < numNodes; idx++)
 				{
-					int parentIdx = Math.Min((int)(random.NextDouble() * Math.Min(idx, numRoots + numInterior)), idx - 1);
-					children[parentIdx].Add(idx);
+					int numParents = 1 + (int)(random.NextDouble() * maxParents);
+					for (; numParents > 0; numParents--)
+					{
+						int parentIdx = Math.Min((int)(random.NextDouble() * Math.Min(idx, numRoots + numInterior)), idx - 1);
+						children[parentIdx].Add(idx);
+					}
 				}
 			}
 
diff --git a/Engine/Source/Programs/Horde/Horde.Server/Storage/StorageService.cs b/Engine/Source/Programs/Horde/Horde.Server/Storage/StorageService.cs
--- a/Engine/Source/Programs/Horde/Horde.Server/Storage/StorageService.cs
+++ b/Engine/Source/Programs/Horde/Horde.Server/Storage/StorageService.cs
@@ -1101,7 +1101,7 @@ namespace Horde.Server.Storage
 					using TelemetrySpan span = _tracer.StartActiveSpan($"{nameof(StorageService)}.{nameof(TickGcForNamespaceAsync)}");
 					span.SetAttribute("BlobId", blobInfoId.ToString());
 
-					if (blobInfoId < lastImportBlobInfoId && !await IsBlobReferencedAsync(blobInfoId, cancellationToken))
+					if (blobInfoId <= lastImportBlobInfoId && !await IsBlobReferencedAsync(blobInfoId, cancellationToken))
 					{
 						BlobInfo? info;
 						if (deleteObjects)
