Horde: Relax running jobs without parameters again now that V2 dialog has been out a couple weeks
#rnx

diff --git a/Engine/Source/Programs/Horde/HordeDashboard/src/components/jobDetailsV2/JobOperationsBar.tsx b/Engine/Source/Programs/Horde/HordeDashboard/src/components/jobDetailsV2/JobOperationsBar.tsx
--- a/Engine/Source/Programs/Horde/HordeDashboard/src/components/jobDetailsV2/JobOperationsBar.tsx
+++ b/Engine/Source/Programs/Horde/HordeDashboard/src/components/jobDetailsV2/JobOperationsBar.tsx
@@ -42,8 +42,6 @@ export const JobOperations: React.FC<{ jobDetails: JobDetailsV2 }> = observer(({
    const batchFilter = query.get("batch");
    let newBuildVersion: string = query.get("newbuildversion") ? query.get("newbuildversion")! : "1";
 
-   // show new build version 2 if we have new parameters store
-   const hasJobParameters = !!(Object.keys(jobDetails?.jobData?.parameters ?? {}).length)
    newBuildVersion = "2";
 
    // subscribe
@@ -58,7 +56,6 @@ export const JobOperations: React.FC<{ jobDetails: JobDetailsV2 }> = observer(({
    }
 
    const abortDisabled = jobData.state === JobState.Complete;
-   const runAgainDisabled = !hasJobParameters;
 
    const failedSteps = jobDetails.getSteps().filter(s => {
 
@@ -124,7 +121,6 @@ export const JobOperations: React.FC<{ jobDetails: JobDetailsV2 }> = observer(({
    opsList.push({
       key: 'jobops_runagain',
       text: "Run Again",
-      disabled: runAgainDisabled,
       iconProps: { iconName: "Duplicate" },
       onClick: () => { setParametersState(ParameterState.Clone); }
    });
