Horde: Change GC verification to check against Path property rather than Locator, which is synthesized on access.

#rnx

diff --git a/Engine/Source/Programs/Horde/Horde.Server/Storage/StorageService.cs b/Engine/Source/Programs/Horde/Horde.Server/Storage/StorageService.cs
--- a/Engine/Source/Programs/Horde/Horde.Server/Storage/StorageService.cs
+++ b/Engine/Source/Programs/Horde/Horde.Server/Storage/StorageService.cs
@@ -597,9 +597,17 @@ namespace Horde.Server.Storage
 		{
 			if (_globalConfig.CurrentValue.Storage.EnableGcVerification)
 			{
-				if (await _blobCollection.Find(x => x.NamespaceId == namespaceId && x.Locator == locator && x.GcVersion >= CurrentGcVersion).AnyAsync(cancellationToken))
+				try
 				{
-					_logger.LogWarning("Blob {Locator} accessed after being garbage collected", locator);
+					string path = locator.Path.ToString();
+					if (await _blobCollection.Find(x => x.NamespaceId == namespaceId && x.Path == path && x.GcVersion >= CurrentGcVersion).AnyAsync(cancellationToken))
+					{
+						_logger.LogWarning("Blob {Locator} accessed after being garbage collected", locator);
+					}
+				}
+				catch (Exception ex)
+				{
+					_logger.LogWarning(ex, "Exception checking if blob {NamespaceId}:{Locator} exists", namespaceId, locator);
 				}
 			}
 		}
