Horde: More graceful handling of failed gc state updates.

#rnx

diff --git a/Engine/Source/Programs/Horde/Horde.Server/Storage/StorageService.cs b/Engine/Source/Programs/Horde/Horde.Server/Storage/StorageService.cs
--- a/Engine/Source/Programs/Horde/Horde.Server/Storage/StorageService.cs
+++ b/Engine/Source/Programs/Horde/Horde.Server/Storage/StorageService.cs
@@ -637,7 +637,6 @@ namespace Horde.Server.Storage
 		{
 			using TelemetrySpan span = _tracer.StartActiveSpan($"{nameof(StorageService)}.{nameof(TickBlobsAsync)}");
 
-			GcState gcState = await _gcState.GetAsync(cancellationToken);
 			DateTime ingestTimeUtc = _clock.UtcNow - TimeSpan.FromMinutes(30.0);
 
 			// Get the current state of the storage system
@@ -652,6 +651,8 @@ namespace Horde.Server.Storage
 				ObjectId latestInfoId = ObjectId.GenerateNewId(ingestTimeUtc);
 				for (; ; )
 				{
+					GcState gcState = await _gcState.GetAsync(cancellationToken);
+
 					// Fetch the next batch of blobs
 					List<BlobInfo> current = await _blobCollection.Find(x => x.Id > gcState.LastImportBlobInfoId && x.Id < latestInfoId).SortBy(x => x.Id).Limit(500).ToListAsync(cancellationToken);
 					if (current.Count == 0)
@@ -693,8 +694,15 @@ namespace Horde.Server.Storage
 					}
 
 					// Update the last imported blob id
-					gcState = await _gcState.UpdateAsync(state => state.LastImportBlobInfoId = current[^1].Id, cancellationToken);
-					ingestedCount += current.Count;
+					gcState.LastImportBlobInfoId = current[^1].Id;
+					if (await _gcState.TryUpdateAsync(gcState, cancellationToken))
+					{
+						ingestedCount += current.Count;
+					}
+					else
+					{
+						_logger.LogInformation("GC state has been modified externally (expected {Id} rev {Index}); may re-run over previous blobs", gcState.Id, gcState.Revision);
+					}
 				}
 
 				_logger.LogInformation("Added {NumBlobs} blobs for GC (upper time: {Time})", ingestedCount, ingestTimeUtc);
