Horde: Add a flag to allow resetting the GC blob sweep.

#rnx

diff --git a/Engine/Source/Programs/Horde/Horde.Server/Storage/StorageService.cs b/Engine/Source/Programs/Horde/Horde.Server/Storage/StorageService.cs
--- a/Engine/Source/Programs/Horde/Horde.Server/Storage/StorageService.cs
+++ b/Engine/Source/Programs/Horde/Horde.Server/Storage/StorageService.cs
@@ -352,6 +352,7 @@ namespace Horde.Server.Storage
 		{
 			public ObjectId LastImportBlobInfoId { get; set; }
 			public List<GcNamespaceState> Namespaces { get; set; } = new List<GcNamespaceState>();
+			public bool Reset { get; set; }
 
 			public GcNamespaceState FindOrAddNamespace(NamespaceId namespaceId)
 			{
@@ -637,6 +638,16 @@ namespace Horde.Server.Storage
 		{
 			using TelemetrySpan span = _tracer.StartActiveSpan($"{nameof(StorageService)}.{nameof(TickBlobsAsync)}");
 
+			GcState gcState = await _gcState.GetAsync(cancellationToken);
+			if(gcState.Reset)
+			{
+				_logger.LogInformation("Resetting scan for new blobs...");
+				gcState = await _gcState.UpdateAsync(x => { 
+					x.Reset = false; 
+					x.LastImportBlobInfoId = ObjectId.Empty; 
+				}, cancellationToken);
+			}
+
 			DateTime ingestTimeUtc = _clock.UtcNow - TimeSpan.FromMinutes(30.0);
 
 			// Get the current state of the storage system
@@ -651,8 +662,6 @@ namespace Horde.Server.Storage
 				ObjectId latestInfoId = ObjectId.GenerateNewId(ingestTimeUtc);
 				for (; ; )
 				{
-					GcState gcState = await _gcState.GetAsync(cancellationToken);
-
 					// Fetch the next batch of blobs
 					List<BlobInfo> current = await _blobCollection.Find(x => x.Id > gcState.LastImportBlobInfoId && x.Id < latestInfoId).SortBy(x => x.Id).Limit(500).ToListAsync(cancellationToken);
 					if (current.Count == 0)
@@ -694,15 +703,8 @@ namespace Horde.Server.Storage
 					}
 
 					// Update the last imported blob id
-					gcState.LastImportBlobInfoId = current[^1].Id;
-					if (await _gcState.TryUpdateAsync(gcState, cancellationToken))
-					{
-						ingestedCount += current.Count;
-					}
-					else
-					{
-						_logger.LogInformation("GC state has been modified externally (expected {Id} rev {Index}); may re-run over previous blobs", gcState.Id, gcState.Revision);
-					}
+					gcState = await _gcState.UpdateAsync(state => state.LastImportBlobInfoId = current[^1].Id, cancellationToken);
+					ingestedCount += current.Count;
 				}
 
 				_logger.LogInformation("Added {NumBlobs} blobs for GC (upper time: {Time})", ingestedCount, ingestTimeUtc);
