Horde: Remove erroneous device reservation warning
#rnx

diff --git a/Engine/Source/Programs/Horde/Horde.Server/Devices/DevicesController.cs b/Engine/Source/Programs/Horde/Horde.Server/Devices/DevicesController.cs
--- a/Engine/Source/Programs/Horde/Horde.Server/Devices/DevicesController.cs
+++ b/Engine/Source/Programs/Horde/Horde.Server/Devices/DevicesController.cs
@@ -869,11 +869,6 @@ namespace Horde.Server.Devices
 
 			if (reservation == null)
 			{
-				if (!String.IsNullOrWhiteSpace(errorMessage))
-				{
-					_logger.LogWarning("Device reservation failed: {ErrorMessage}", errorMessage);
-				}
-
 				return Conflict(String.IsNullOrEmpty(errorMessage) ? "Unable to allocated devices for reservation" : errorMessage);
 			}
 
