Horde: Add an AWS parameter store secret provider.

diff --git a/Engine/Source/Programs/Horde/Horde.Server/Horde.Server.csproj b/Engine/Source/Programs/Horde/Horde.Server/Horde.Server.csproj
--- a/Engine/Source/Programs/Horde/Horde.Server/Horde.Server.csproj
+++ b/Engine/Source/Programs/Horde/Horde.Server/Horde.Server.csproj
@@ -25,6 +25,7 @@
     <PackageReference Include="AWSSDK.EC2" Version="3.7.306.6" />
     <PackageReference Include="AWSSDK.Extensions.NETCore.Setup" Version="3.7.300" />
     <PackageReference Include="AWSSDK.S3" Version="3.7.304.1" />
+    <PackageReference Include="AWSSDK.SimpleSystemsManagement" Version="3.7.304.24" />
     <PackageReference Include="AWSSDK.SQS" Version="3.7.300.12" />
     <PackageReference Include="AWSSDK.SecurityToken" Version="3.7.300.13" />
     <PackageReference Include="Azure.Storage.Blobs" Version="12.19.1" />
diff --git a/Engine/Source/Programs/Horde/Horde.Server/Secrets/ISecret.cs b/Engine/Source/Programs/Horde/Horde.Server/Secrets/ISecret.cs
new file mode 100644
--- /dev/null
+++ b/Engine/Source/Programs/Horde/Horde.Server/Secrets/ISecret.cs
@@ -0,0 +1,23 @@
+// Copyright Epic Games, Inc. All Rights Reserved.
+
+using System.Collections.Generic;
+using EpicGames.Horde.Secrets;
+
+namespace Horde.Server.Secrets
+{
+	/// <summary>
+	/// Information about a secret
+	/// </summary>
+	public interface ISecret
+	{
+		/// <summary>
+		/// Identifier for the secret
+		/// </summary>
+		SecretId Id { get; }
+
+		/// <summary>
+		/// The secret values
+		/// </summary>
+		IReadOnlyDictionary<string, string> Data { get; }
+	}
+}
diff --git a/Engine/Source/Programs/Horde/Horde.Server/Secrets/ISecretCollection.cs b/Engine/Source/Programs/Horde/Horde.Server/Secrets/ISecretCollection.cs
new file mode 100644
--- /dev/null
+++ b/Engine/Source/Programs/Horde/Horde.Server/Secrets/ISecretCollection.cs
@@ -0,0 +1,21 @@
+﻿// Copyright Epic Games, Inc. All Rights Reserved.
+
+using System.Threading;
+using System.Threading.Tasks;
+using EpicGames.Horde.Secrets;
+
+namespace Horde.Server.Secrets
+{
+	/// <summary>
+	/// Collection of secrets
+	/// </summary>
+	public interface ISecretCollection
+	{
+		/// <summary>
+		/// Resolve a secret to concrete values
+		/// </summary>
+		/// <param name="secretId">Identifier for the secret</param>
+		/// <param name="cancellationToken">Cancellation token for the operation</param>
+		Task<ISecret?> GetAsync(SecretId secretId, CancellationToken cancellationToken);
+	}
+}
diff --git a/Engine/Source/Programs/Horde/Horde.Server/Secrets/ISecretProvider.cs b/Engine/Source/Programs/Horde/Horde.Server/Secrets/ISecretProvider.cs
--- a/Engine/Source/Programs/Horde/Horde.Server/Secrets/ISecretProvider.cs
+++ b/Engine/Source/Programs/Horde/Horde.Server/Secrets/ISecretProvider.cs
@@ -1,7 +1,5 @@
 // Copyright Epic Games, Inc. All Rights Reserved.
 
-using System.Collections.Generic;
-using System.Security.Claims;
 using System.Threading;
 using System.Threading.Tasks;
 
@@ -20,10 +18,8 @@ namespace Horde.Server.Secrets
 		/// <summary>
 		/// Gets values for a particular secret
 		/// </summary>
-		/// <param name="config">Configuration for the secret to fetch</param>
-		/// <param name="user">User requesting the data</param>
-		/// <param name="data">Dictionary to receive any enumerated properties</param>
+		/// <param name="path">Path to the secret to fetch</param>
 		/// <param name="cancellationToken">Cancellation token for the operation</param>
-		Task GetValuesAsync(ExternalSecretConfig config, ClaimsPrincipal user, Dictionary<string, string> data, CancellationToken cancellationToken);
+		Task<string> GetSecretAsync(string path, CancellationToken cancellationToken);
 	}
 }
diff --git a/Engine/Source/Programs/Horde/Horde.Server/Secrets/SecretCollection.cs b/Engine/Source/Programs/Horde/Horde.Server/Secrets/SecretCollection.cs
new file mode 100644
--- /dev/null
+++ b/Engine/Source/Programs/Horde/Horde.Server/Secrets/SecretCollection.cs
@@ -0,0 +1,115 @@
+// Copyright Epic Games, Inc. All Rights Reserved.
+
+using System;
+using System.Collections.Generic;
+using System.Linq;
+using System.Text.Json;
+using System.Threading;
+using System.Threading.Tasks;
+using EpicGames.Horde.Secrets;
+using Horde.Server.Server;
+using Microsoft.Extensions.Options;
+
+namespace Horde.Server.Secrets
+{
+	/// <summary>
+	/// Implementation of <see cref="ISecretCollection"/>
+	/// </summary>
+	public class SecretCollection : ISecretCollection
+	{
+		readonly SecretCollectionInternal _secretCollectionInternal;
+		readonly IOptionsSnapshot<GlobalConfig> _options;
+
+		/// <summary>
+		/// Constructor
+		/// </summary>
+		public SecretCollection(SecretCollectionInternal secretCollectionInternal, IOptionsSnapshot<GlobalConfig> options)
+		{
+			_secretCollectionInternal = secretCollectionInternal;
+			_options = options;
+		}
+
+		/// <inheritdoc/>
+		public async Task<ISecret?> GetAsync(SecretId secretId, CancellationToken cancellationToken)
+		{
+			SecretConfig? secretConfig;
+			if (!_options.Value.TryGetSecret(secretId, out secretConfig))
+			{
+				return null;
+			}
+			return await _secretCollectionInternal.GetAsync(secretConfig, cancellationToken);
+		}
+	}
+
+	/// <summary>
+	/// Core implementation of <see cref="ISecretCollection"/> which exists as a global service.
+	/// </summary>
+	public class SecretCollectionInternal
+	{
+		class Secret : ISecret
+		{
+			public SecretId Id { get; }
+			public IReadOnlyDictionary<string, string> Data { get; }
+
+			public Secret(SecretId id, IReadOnlyDictionary<string, string> data)
+			{
+				Id = id;
+				Data = data;
+			}
+		}
+
+		readonly Dictionary<string, ISecretProvider> _providers;
+
+		/// <summary>
+		/// Constructor
+		/// </summary>
+		public SecretCollectionInternal(IEnumerable<ISecretProvider> providers)
+		{
+			_providers = providers.ToDictionary(x => x.Name, x => x);
+		}
+
+		/// <summary>
+		/// Resolve a secret to concrete values
+		/// </summary>
+		/// <param name="config">Configuration for the secret</param>
+		/// <param name="cancellationToken">Cancellation token for the operation</param>
+		public async Task<ISecret> GetAsync(SecretConfig config, CancellationToken cancellationToken)
+		{
+			// Add the hard-coded secrets
+			Dictionary<string, string> data = new Dictionary<string, string>();
+			foreach ((string key, string value) in config.Data)
+			{
+				data.Add(key, value);
+			}
+
+			// Fetch all the secrets from external providers
+			foreach (ExternalSecretConfig externalConfig in config.Sources)
+			{
+				ISecretProvider provider = _providers[externalConfig.Provider];
+
+				string secret = await provider.GetSecretAsync(externalConfig.Path, cancellationToken);
+				if (externalConfig.Format == ExternalSecretFormat.Text)
+				{
+					data[externalConfig.Key ?? "default"] = secret;
+				}
+				else if (externalConfig.Format == ExternalSecretFormat.Json)
+				{
+					Dictionary<string, string>? values = JsonSerializer.Deserialize<Dictionary<string, string>>(secret, new JsonSerializerOptions { AllowTrailingCommas = true });
+					if (values != null)
+					{
+						foreach ((string key, string value) in values)
+						{
+							data[key] = value;
+						}
+					}
+				}
+				else
+				{
+					throw new NotImplementedException();
+				}
+			}
+
+			return new Secret(config.Id, data);
+		}
+	}
+}
diff --git a/Engine/Source/Programs/Horde/Horde.Server/Secrets/SecretConfig.cs b/Engine/Source/Programs/Horde/Horde.Server/Secrets/SecretConfig.cs
--- a/Engine/Source/Programs/Horde/Horde.Server/Secrets/SecretConfig.cs
+++ b/Engine/Source/Programs/Horde/Horde.Server/Secrets/SecretConfig.cs
@@ -53,6 +53,22 @@ namespace Horde.Server.Secrets
 			=> Acl.Authorize(action, user);
 	}
 
+	/// <summary>
+	/// Format describing how to parse external secret values
+	/// </summary>
+	public enum ExternalSecretFormat
+	{
+		/// <summary>
+		/// Secret is a plain text value which will be stored using the external secret key
+		/// </summary>
+		Text,
+
+		/// <summary>
+		/// Secret is a JSON formatted string containing key/value pairs
+		/// </summary>
+		Json,
+	}
+
 	/// <summary>
 	/// Configuration for an external secret provider
 	/// </summary>
@@ -64,18 +80,18 @@ namespace Horde.Server.Secrets
 		public string Provider { get; set; } = String.Empty;
 
 		/// <summary>
-		/// Optional key indicating the parameter to set in the resulting data array
+		/// Format of the secret
 		/// </summary>
-		public string? Key { get; set; }
+		public ExternalSecretFormat Format { get; set; }
 
 		/// <summary>
-		/// Optional value indicating what to fetch from the provider
+		/// Optional key indicating the parameter to set in the resulting data array. Required if if <see cref="Format"/> is <see cref="ExternalSecretFormat.Text"/>.
 		/// </summary>
-		public string? Path { get; set; }
+		public string? Key { get; set; }
 
 		/// <summary>
-		/// Additional provider-specific arguments
+		/// Optional value indicating what to fetch from the provider
 		/// </summary>
-		public Dictionary<string, string>? Arguments { get; set; }
+		public string Path { get; set; } = "default";
 	}
 }
diff --git a/Engine/Source/Programs/Horde/Horde.Server/Secrets/SecretService.cs b/Engine/Source/Programs/Horde/Horde.Server/Secrets/SecretService.cs
deleted file mode 100644
diff --git a/Engine/Source/Programs/Horde/Horde.Server/Secrets/SecretsController.cs b/Engine/Source/Programs/Horde/Horde.Server/Secrets/SecretsController.cs
--- a/Engine/Source/Programs/Horde/Horde.Server/Secrets/SecretsController.cs
+++ b/Engine/Source/Programs/Horde/Horde.Server/Secrets/SecretsController.cs
@@ -1,6 +1,9 @@
 // Copyright Epic Games, Inc. All Rights Reserved.
 
 using System.Collections.Generic;
+using System.Linq;
+using System.Threading;
+using System.Threading.Tasks;
 using EpicGames.Horde.Secrets;
 using Horde.Server.Server;
 using Horde.Server.Utilities;
@@ -11,20 +14,22 @@ using Microsoft.Extensions.Options;
 namespace Horde.Server.Secrets
 {
 	/// <summary>
-	/// Controller for the /api/v1/credentials endpoint
+	/// Controller for the /api/v1/secrets endpoint
 	/// </summary>
 	[ApiController]
 	[Authorize]
 	[Route("[controller]")]
 	public class SecretsController : HordeControllerBase
 	{
-		private readonly IOptionsSnapshot<GlobalConfig> _globalConfig;
+		readonly ISecretCollection _secretCollection;
+		readonly IOptionsSnapshot<GlobalConfig> _globalConfig;
 
 		/// <summary>
 		/// Constructor
 		/// </summary>
-		public SecretsController(IOptionsSnapshot<GlobalConfig> globalConfig)
+		public SecretsController(ISecretCollection secretCollection, IOptionsSnapshot<GlobalConfig> globalConfig)
 		{
+			_secretCollection = secretCollection;
 			_globalConfig = globalConfig;
 		}
 
@@ -51,22 +56,24 @@ namespace Horde.Server.Secrets
 		/// </summary>
 		/// <param name="secretId">Id of the secret to retrieve</param>
 		/// <param name="filter">Filter for properties to return</param>
+		/// <param name="cancellationToken">Cancellation token for the operation</param>
 		/// <returns>Information about the requested secret</returns>
 		[HttpGet]
 		[Route("/api/v1/secrets/{secretId}")]
 		[ProducesResponseType(typeof(GetSecretResponse), 200)]
-		public ActionResult<object> GetSecret(SecretId secretId, [FromQuery] PropertyFilter? filter = null)
+		public async Task<ActionResult<object>> GetSecretAsync(SecretId secretId, [FromQuery] PropertyFilter? filter = null, CancellationToken cancellationToken = default)
 		{
-			if (!_globalConfig.Value.TryGetSecret(secretId, out SecretConfig? secretConfig))
+			ISecret? secret = await _secretCollection.GetAsync(secretId, cancellationToken);
+			if (secret == null)
 			{
 				return NotFound(secretId);
 			}
-			if (!secretConfig.Authorize(SecretAclAction.ViewSecret, User))
+			if (!_globalConfig.Value.Authorize(secret.Id, SecretAclAction.ViewSecret, User))
 			{
 				return Forbid(SecretAclAction.ViewSecret, secretId);
 			}
 
-			return new GetSecretResponse(secretConfig.Id, secretConfig.Data).ApplyFilter(filter);
+			return new GetSecretResponse(secret.Id, secret.Data.ToDictionary(x => x.Key, x => x.Value)).ApplyFilter(filter);
 		}
 	}
 }
diff --git a/Engine/Source/Programs/Horde/Horde.Server/Server/GlobalConfig.cs b/Engine/Source/Programs/Horde/Horde.Server/Server/GlobalConfig.cs
--- a/Engine/Source/Programs/Horde/Horde.Server/Server/GlobalConfig.cs
+++ b/Engine/Source/Programs/Horde/Horde.Server/Server/GlobalConfig.cs
@@ -581,7 +581,8 @@ namespace Horde.Server.Server
 		/// <param name="clusterId">Compute cluster id</param>
 		/// <param name="config">Receives the cluster configuration on success</param>
 		/// <returns>True on success</returns>
-		public bool TryGetComputeCluster(ClusterId clusterId, [NotNullWhen(true)] out ComputeClusterConfig? config) => _computeClusterLookup.TryGetValue(clusterId, out config);
+		public bool TryGetComputeCluster(ClusterId clusterId, [NotNullWhen(true)] out ComputeClusterConfig? config)
+			=> _computeClusterLookup.TryGetValue(clusterId, out config);
 
 		/// <summary>
 		/// Attempts to get compute cluster configuration from this object
@@ -589,7 +590,14 @@ namespace Horde.Server.Server
 		/// <param name="secretId">Secret id</param>
 		/// <param name="config">Receives the secret configuration on success</param>
 		/// <returns>True on success</returns>
-		public bool TryGetSecret(SecretId secretId, [NotNullWhen(true)] out SecretConfig? config) => _secretLookup.TryGetValue(secretId, out config);
+		public bool TryGetSecret(SecretId secretId, [NotNullWhen(true)] out SecretConfig? config)
+			=> _secretLookup.TryGetValue(secretId, out config);
+
+		/// <summary>
+		/// Authorize access to a secret
+		/// </summary>
+		public bool Authorize(SecretId secretId, AclAction action, ClaimsPrincipal user)
+			=> TryGetSecret(secretId, out SecretConfig? secretConfig) && secretConfig.Authorize(action, user);
 
 		/// <summary>
 		/// Authorizes a user to perform a given action
@@ -597,9 +605,7 @@ namespace Horde.Server.Server
 		/// <param name="scopeName">Name of the scope to auth against</param>
 		/// <param name="scopeConfig">Configuration for the scope</param>
 		public bool TryGetAclScope(AclScopeName scopeName, [NotNullWhen(true)] out AclConfig? scopeConfig)
-		{
-			return _aclLookup.TryGetValue(scopeName, out scopeConfig);
-		}
+			=> _aclLookup.TryGetValue(scopeName, out scopeConfig);
 
 		/// <summary>
 		/// Authorizes a user to perform a given action
@@ -608,9 +614,7 @@ namespace Horde.Server.Server
 		/// <param name="action">The action being performed</param>
 		/// <param name="user">The principal to validate</param>
 		public bool Authorize(AclScopeName scopeName, AclAction action, ClaimsPrincipal user)
-		{
-			return _aclLookup.TryGetValue(scopeName, out AclConfig? scopeConfig) && scopeConfig.Authorize(action, user);
-		}
+			=> _aclLookup.TryGetValue(scopeName, out AclConfig? scopeConfig) && scopeConfig.Authorize(action, user);
 
 		/// <summary>
 		/// Determines whether the given user can masquerade as a given user
diff --git a/Engine/Source/Programs/Horde/Horde.Server/Startup.cs b/Engine/Source/Programs/Horde/Horde.Server/Startup.cs
--- a/Engine/Source/Programs/Horde/Horde.Server/Startup.cs
+++ b/Engine/Source/Programs/Horde/Horde.Server/Startup.cs
@@ -76,6 +76,7 @@ using Horde.Server.Notifications.Sinks;
 using Horde.Server.Perforce;
 using Horde.Server.Replicators;
 using Horde.Server.Secrets;
+using Horde.Server.Secrets.Providers;
 using Horde.Server.Server;
 using Horde.Server.Server.Notices;
 using Horde.Server.ServiceAccounts;
@@ -640,7 +641,9 @@ namespace Horde.Server
 			services.AddSingleton<CompressedBufferUtils>();
 
 			// Secret providers
-			services.AddSingleton<SecretService>();
+			services.AddSingleton<SecretCollectionInternal>();
+			services.AddScoped<ISecretCollection, SecretCollection>();
+			services.AddSingleton<ISecretProvider, AwsParameterStoreSecretProvider>();
 
 			// Storage providers
 			services.AddSingleton<IObjectStoreFactory, ObjectStoreFactory>();
