Horde: Missing AwsParameterStoreSecretProvider.cs.

diff --git a/Engine/Source/Programs/Horde/Horde.Server/Secrets/Providers/AwsParameterStoreSecretProvider.cs b/Engine/Source/Programs/Horde/Horde.Server/Secrets/Providers/AwsParameterStoreSecretProvider.cs
new file mode 100644
--- /dev/null
+++ b/Engine/Source/Programs/Horde/Horde.Server/Secrets/Providers/AwsParameterStoreSecretProvider.cs
@@ -0,0 +1,42 @@
+// Copyright Epic Games, Inc. All Rights Reserved.
+
+using System;
+using System.Threading;
+using System.Threading.Tasks;
+using Amazon.Extensions.NETCore.Setup;
+using Amazon.SimpleSystemsManagement;
+using Amazon.SimpleSystemsManagement.Model;
+
+namespace Horde.Server.Secrets.Providers
+{
+	/// <summary>
+	/// Fetches secrets from the AWS parameter store
+	/// </summary>
+	public class AwsParameterStoreSecretProvider : ISecretProvider
+	{
+		/// <inheritdoc/>
+		public string Name => "AwsParameterStore";
+
+		readonly IAmazonSimpleSystemsManagement _systemsManagement;
+
+		/// <summary>
+		/// Constructor
+		/// </summary>
+		public AwsParameterStoreSecretProvider()
+		{
+			AWSOptions awsOptions = new AWSOptions();
+			_systemsManagement = awsOptions.CreateServiceClient<IAmazonSimpleSystemsManagement>();
+		}
+
+		/// <inheritdoc/>
+		public async Task<string> GetSecretAsync(string path, CancellationToken cancellationToken)
+		{
+			GetParameterResponse response = await _systemsManagement.GetParameterAsync(new GetParameterRequest { Name = path, WithDecryption = true }, cancellationToken);
+			if (response.Parameter == null)
+			{
+				throw new InvalidOperationException($"Unable to fetch secret '{path}' from AWS parameter store ({response.HttpStatusCode})");
+			}
+			return response.Parameter.Value;
+		}
+	}
+}
