Horde: Remove Thread.Sleep() from ApplicationStopping() callback; set the appropriate shutdown timeout on the host options instead.

diff --git a/Engine/Source/Programs/Horde/Horde.Server/Commands/ServerCommand.cs b/Engine/Source/Programs/Horde/Horde.Server/Commands/ServerCommand.cs
--- a/Engine/Source/Programs/Horde/Horde.Server/Commands/ServerCommand.cs
+++ b/Engine/Source/Programs/Horde/Horde.Server/Commands/ServerCommand.cs
@@ -63,6 +63,7 @@ namespace Horde.Server.Commands
 
 			IHostBuilder hostBuilder = Host.CreateDefaultBuilder(args)
 				.UseSerilog()
+				.ConfigureHostOptions(options => options.ShutdownTimeout = TimeSpan.FromSeconds(30.0))
 				.ConfigureAppConfiguration(builder => builder.AddConfiguration(config))
 				.ConfigureWebHostDefaults(webBuilder =>
 				{
diff --git a/Engine/Source/Programs/Horde/Horde.Server/Server/LifetimeService.cs b/Engine/Source/Programs/Horde/Horde.Server/Server/LifetimeService.cs
--- a/Engine/Source/Programs/Horde/Horde.Server/Server/LifetimeService.cs
+++ b/Engine/Source/Programs/Horde/Horde.Server/Server/LifetimeService.cs
@@ -26,10 +26,10 @@ namespace Horde.Server.Server
 		/// </summary>
 		readonly TaskCompletionSource<bool> _stoppingTaskCompletionSource;
 
-		/// <summary>
-		/// Task source for the server stopping
-		/// </summary>
-		readonly TaskCompletionSource<bool> _preStoppingTaskCompletionSource;
+//		/// <summary>
+//		/// Task source for the server stopping
+//		/// </summary>
+//		readonly TaskCompletionSource<bool> _preStoppingTaskCompletionSource;
 
 		/// <summary>
 		/// Registration token for the stopping event
@@ -75,7 +75,7 @@ namespace Horde.Server.Server
 			_redisService = redisService;
 			_logger = logger;
 			_stoppingTaskCompletionSource = new TaskCompletionSource<bool>(TaskCreationOptions.RunContinuationsAsynchronously);
-			_preStoppingTaskCompletionSource = new TaskCompletionSource<bool>(TaskCreationOptions.RunContinuationsAsynchronously);
+//			_preStoppingTaskCompletionSource = new TaskCompletionSource<bool>(TaskCreationOptions.RunContinuationsAsynchronously);
 
 			if (env.IsProduction() || env.IsDevelopment())
 			{
@@ -135,14 +135,14 @@ namespace Horde.Server.Server
 			IsPreStopping = true;
 			IsStopping = true;
 
-			_preStoppingTaskCompletionSource.TrySetResult(true);
+//			_preStoppingTaskCompletionSource.TrySetResult(true);
 			_stoppingTaskCompletionSource.TrySetResult(true);
-
+/*
 			int shutdownDelayMs = 30 * 1000;
 			_logger.LogInformation("Delaying shutdown by sleeping {ShutdownDelayMs} ms...", shutdownDelayMs);
 			Thread.Sleep(shutdownDelayMs);
 			_logger.LogInformation("Server process now shutting down...");
-
+*/
 			/*
 			if (PreStoppingTaskCompletionSource.TrySetResult(true))
 			{
