Horde: Add the shutdown timeout on the web builder also.

diff --git a/Engine/Source/Programs/Horde/Horde.Server/Commands/ServerCommand.cs b/Engine/Source/Programs/Horde/Horde.Server/Commands/ServerCommand.cs
--- a/Engine/Source/Programs/Horde/Horde.Server/Commands/ServerCommand.cs
+++ b/Engine/Source/Programs/Horde/Horde.Server/Commands/ServerCommand.cs
@@ -67,6 +67,8 @@ namespace Horde.Server.Commands
 				.ConfigureAppConfiguration(builder => builder.AddConfiguration(config))
 				.ConfigureWebHostDefaults(webBuilder =>
 				{
+					webBuilder.UseShutdownTimeout(TimeSpan.FromSeconds(30.0));
+
 					webBuilder.UseUrls(); // Disable default URLs; we will configure each port directly.
 
 					webBuilder.UseWebRoot("DashboardApp");
