Horde: Allow specifying cache sizes using standard binary suffixes (eg. "4gb").

diff --git a/Engine/Source/Programs/Horde/Horde.Server/ServerSettings.cs b/Engine/Source/Programs/Horde/Horde.Server/ServerSettings.cs
--- a/Engine/Source/Programs/Horde/Horde.Server/ServerSettings.cs
+++ b/Engine/Source/Programs/Horde/Horde.Server/ServerSettings.cs
@@ -947,9 +947,14 @@ namespace Horde.Server
 		public string? BundleCacheDir { get; set; }
 
 		/// <summary>
-		/// Maximum size of the storage cache on disk, in megabytes
+		/// Maximum size of the storage cache on disk. Accepts standard binary suffixes (kb, mb, gb, tb, etc...)
 		/// </summary>
-		public long BundleCacheSize { get; set; } = 1024;
+		public string BundleCacheSize { get; set; } = "1gb";
+
+		/// <summary>
+		/// Accessor for the bundle cache size in bytes
+		/// </summary>
+		public long BundleCacheSizeBytes => StringUtils.ParseBytesString(BundleCacheSize);
 
 		/// <summary>
 		/// Directory to store the fine-grained block cache. This caches individual exports embedded in bundles.
@@ -957,9 +962,14 @@ namespace Horde.Server
 		public string? BlockCacheDir { get; set; }
 
 		/// <summary>
-		/// Maximum size of the block cache, in megabytes. Currently only allocates in multiples of 1024mb.
+		/// Maximum size of the block cache. Accepts standard binary suffixes. Currently only allocates in multiples of 1024mb.
+		/// </summary>
+		public string BlockCacheSize { get; set; } = "4gb";
+
+		/// <summary>
+		/// Accessor for the block cache size in bytes
 		/// </summary>
-		public long BlockCacheSize { get; set; } = 4096;
+		public long BlockCacheSizeBytes => StringUtils.ParseBytesString(BlockCacheSize);
 
 		/// <summary>
 		/// Experimental features to enable on the server.
diff --git a/Engine/Source/Programs/Horde/Horde.Server/Startup.cs b/Engine/Source/Programs/Horde/Horde.Server/Startup.cs
--- a/Engine/Source/Programs/Horde/Horde.Server/Startup.cs
+++ b/Engine/Source/Programs/Horde/Horde.Server/Startup.cs
@@ -1018,14 +1018,14 @@ namespace Horde.Server
 		{
 			ServerSettings serverSettings = serviceProvider.GetRequiredService<IOptions<ServerSettings>>().Value;
 			DirectoryReference cacheDir = DirectoryReference.Combine(ServerApp.DataDir, String.IsNullOrEmpty(serverSettings.BlockCacheDir) ? "BlockCache" : serverSettings.BlockCacheDir);
-			return BlockCache.Create(cacheDir, (int)(serverSettings.BlockCacheSize / 1024));
+			return BlockCache.Create(cacheDir, (int)(serverSettings.BlockCacheSizeBytes / (1024 * 1024 * 1024)));
 		}
 
 		static StorageBackendCache CreateStorageBackendCache(IServiceProvider serviceProvider)
 		{
 			ServerSettings serverSettings = serviceProvider.GetRequiredService<IOptions<ServerSettings>>().Value;
 			DirectoryReference cacheDir = DirectoryReference.Combine(ServerApp.DataDir, String.IsNullOrEmpty(serverSettings.BundleCacheDir) ? "Cache" : serverSettings.BundleCacheDir);
-			return new StorageBackendCache(cacheDir, serverSettings.BundleCacheSize * 1024 * 1024, serviceProvider.GetRequiredService<ILogger<StorageBackendCache>>());
+			return new StorageBackendCache(cacheDir, serverSettings.BundleCacheSizeBytes, serviceProvider.GetRequiredService<ILogger<StorageBackendCache>>());
 		}
 
 		public static void ConfigureFormatters()
