Horde: Updating schema documentation

#rnx

diff --git a/Engine/Source/Programs/Horde/Docs/Config/Schema/Globals.md b/Engine/Source/Programs/Horde/Docs/Config/Schema/Globals.md
--- a/Engine/Source/Programs/Horde/Docs/Config/Schema/Globals.md
+++ b/Engine/Source/Programs/Horde/Docs/Config/Schema/Globals.md
@@ -332,9 +332,18 @@ Configuration for an external secret provider
 Name | Description
 ---- | -----------
 `provider` | `string`<br>Name of the provider to use
-`key` | `string`<br>Optional key indicating the parameter to set in the resulting data array
+`format` | [`ExternalSecretFormat`](#externalsecretformat-enum)<br>Format of the secret
+`key` | `string`<br>Optional key indicating the parameter to set in the resulting data array. Required if if  is .
 `path` | `string`<br>Optional value indicating what to fetch from the provider
-`arguments` | `string` `->` `string`<br>Additional provider-specific arguments
+
+## ExternalSecretFormat (Enum)
+
+Format describing how to parse external secret values
+
+Name | Description
+---- | -----------
+`Text` | Secret is a plain text value which will be stored using the external secret key
+`Json` | Secret is a JSON formatted string containing key/value pairs
 
 ## DeviceConfig
 
diff --git a/Engine/Source/Programs/Horde/Docs/Deployment/ServerSettings.md b/Engine/Source/Programs/Horde/Docs/Deployment/ServerSettings.md
--- a/Engine/Source/Programs/Horde/Docs/Deployment/ServerSettings.md
+++ b/Engine/Source/Programs/Horde/Docs/Deployment/ServerSettings.md
@@ -103,9 +103,11 @@ Name | Description
 `forceConfigUpdateOnStartup` | `boolean`<br>Forces configuration data to be read and updated as part of appplication startup, rather than on a schedule. Useful when running locally.
 `openBrowser` | `boolean`<br>Whether to open a browser on startup
 `bundleCacheDir` | `string`<br>Directory to use for the coarse-grained backend cache. This caches full bundles downloaded from the upstream object store.
-`bundleCacheSize` | `integer`<br>Maximum size of the storage cache on disk, in megabytes
+`bundleCacheSize` | `string`<br>Maximum size of the storage cache on disk. Accepts standard binary suffixes (kb, mb, gb, tb, etc...)
+`bundleCacheSizeBytes` | `integer`<br>Accessor for the bundle cache size in bytes
 `blockCacheDir` | `string`<br>Directory to store the fine-grained block cache. This caches individual exports embedded in bundles.
-`blockCacheSize` | `integer`<br>Maximum size of the block cache, in megabytes. Currently only allocates in multiples of 1024mb.
+`blockCacheSize` | `string`<br>Maximum size of the block cache. Accepts standard binary suffixes. Currently only allocates in multiples of 1024mb.
+`blockCacheSizeBytes` | `integer`<br>Accessor for the block cache size in bytes
 `featureFlags` | [`FeatureFlagSettings`](#featureflagsettings)<br>Experimental features to enable on the server.
 `commits` | [`CommitSettings`](#commitsettings)<br>Options for the commit service
 `telemetry` | [`BaseTelemetryConfig`](#basetelemetryconfig)`[]`<br>Settings for sending telemetry events to external services (for example Snowflake, ClickHouse etc)
