Horde: Return cluster ID in compute task assignment

When using the generic compute assignment endpoint, client is not aware which cluster is assigned until the response is received. Subsequent calls to the API need this information.

#rnx

diff --git a/Engine/Source/Programs/Horde/Horde.Server/Compute/ComputeController.cs b/Engine/Source/Programs/Horde/Horde.Server/Compute/ComputeController.cs
--- a/Engine/Source/Programs/Horde/Horde.Server/Compute/ComputeController.cs
+++ b/Engine/Source/Programs/Horde/Horde.Server/Compute/ComputeController.cs
@@ -118,6 +118,7 @@ namespace Horde.Server.Compute
 			response.Nonce = StringUtils.FormatHexString(computeResource.Task.Nonce.Span);
 			response.Key = StringUtils.FormatHexString(computeResource.Task.Key.Span);
 			response.Certificate = StringUtils.FormatHexString(computeResource.Task.Certificate.Span);
+			response.ClusterId = clusterId;
 			response.AgentId = computeResource.AgentId;
 			response.LeaseId = computeResource.LeaseId;
 			response.Properties = computeResource.Properties;
diff --git a/Engine/Source/Programs/Shared/EpicGames.Horde/Compute/ComputeMessages.cs b/Engine/Source/Programs/Shared/EpicGames.Horde/Compute/ComputeMessages.cs
--- a/Engine/Source/Programs/Shared/EpicGames.Horde/Compute/ComputeMessages.cs
+++ b/Engine/Source/Programs/Shared/EpicGames.Horde/Compute/ComputeMessages.cs
@@ -162,6 +162,11 @@ namespace EpicGames.Horde.Compute
 		/// X.509 certificate used for SSL/TLS encryption
 		/// </summary>
 		public string Certificate { get; set; } = String.Empty;
+		
+		/// <summary>
+		/// Which cluster this remote machine belongs to
+		/// </summary>
+		public ClusterId ClusterId { get; set; }
 
 		/// <summary>
 		/// Identifier for the remote machine
