Horde: Fix skipping test events as those aren't versioned in v1
#rnx
#submitool bypass

diff --git a/Engine/Source/Programs/Horde/Horde.Server/Jobs/TestData/TestDataCollection.cs b/Engine/Source/Programs/Horde/Horde.Server/Jobs/TestData/TestDataCollection.cs
--- a/Engine/Source/Programs/Horde/Horde.Server/Jobs/TestData/TestDataCollection.cs
+++ b/Engine/Source/Programs/Horde/Horde.Server/Jobs/TestData/TestDataCollection.cs
@@ -505,25 +505,27 @@ namespace Horde.Server.Jobs.TestData
 		public async Task<IReadOnlyList<ITestData>> AddAsync(IJob job, IJobStep step, (string key, BsonDocument value)[] data, CancellationToken cancellationToken = default)
 		{
 			// detailed test data
+			bool skip = false;
 			List<TestDataDocument> documents = new List<TestDataDocument>();
 			for (int i = 0; i < data.Length; i++)
 			{
 				(string key, BsonDocument document) = data[i];
 
-				// Get test document version
 				int version;
-				if (!document.TryGetInt32("Version", out version))
+				if (document.TryGetInt32("version", out version) || document.TryGetInt32("Version", out version))
 				{
-					if (!document.TryGetInt32("version", out version))
+					if (version > 1)
 					{
-						continue;
+						skip = true;
 					}
 				}
 
-				if (version == 1)
-				{
-					documents.Add(new TestDataDocument(job, step, key, document));
-				}				
+				documents.Add(new TestDataDocument(job, step, key, document));
+			}
+
+			if (skip)
+			{
+				documents.Clear();
 			}
 
 			if (documents.Count == 0)
