Horde: Increase MongoDB client wait queue size from 300 -> 5000

diff --git a/Engine/Source/Programs/Horde/Horde.Server/Server/MongoService.cs b/Engine/Source/Programs/Horde/Horde.Server/Server/MongoService.cs
--- a/Engine/Source/Programs/Horde/Horde.Server/Server/MongoService.cs
+++ b/Engine/Source/Programs/Horde/Horde.Server/Server/MongoService.cs
@@ -348,6 +348,9 @@ namespace Horde.Server.Server
 				// Making room for slower running commands potentially hogging the queue 
 				mongoSettings.MinConnectionPoolSize = 50; // Default is 0
 				mongoSettings.MaxConnectionPoolSize = 300; // Default is 100
+#pragma warning disable CS0618 // Type or member is obsolete
+				mongoSettings.WaitQueueSize = 5000;
+#pragma warning restore CS0618 // Type or member is obsolete
 				mongoSettings.MaxConnecting = 10; // Default is 2
 
 				//TestSslConnection(MongoSettings.Server.Host, MongoSettings.Server.Port, Logger);
