Update release notes

#rnx

diff --git a/Engine/Source/Programs/Horde/Docs/ReleaseNotes.md b/Engine/Source/Programs/Horde/Docs/ReleaseNotes.md
--- a/Engine/Source/Programs/Horde/Docs/ReleaseNotes.md
+++ b/Engine/Source/Programs/Horde/Docs/ReleaseNotes.md
@@ -1,5 +1,13 @@
 # Release Notes
 
+## 2024-06-17
+
+* Add support for auto-assigning cluster ID during compute allocation requests (34423910)
+* Increase MongoDB client wait queue size from 300 -> 5000 (34423018)
+* Allow specifying cache sizes using standard binary suffixes (eg. "4gb") (34382251)
+* Add the shutdown timeout for Kestrel server (34379789, 34379660)
+* Add an AWS parameter store secret provider (34379595)
+
 ## 2024-06-13
 
 * Add a pipelined blob read class for storage blobs. (34322019)
