Horde: Add a debug endpoint which streams a random block of data to the caller. Usage is: /api/v1/debug/randomdata?size=1mb

diff --git a/Engine/Source/Programs/Horde/Horde.Server/Server/DebugController.cs b/Engine/Source/Programs/Horde/Horde.Server/Server/DebugController.cs
--- a/Engine/Source/Programs/Horde/Horde.Server/Server/DebugController.cs
+++ b/Engine/Source/Programs/Horde/Horde.Server/Server/DebugController.cs
@@ -14,6 +14,7 @@ using System.Text;
 using System.Threading;
 using System.Threading.Tasks;
 using System.Web;
+using EpicGames.Core;
 using EpicGames.Horde.Agents.Leases;
 using EpicGames.Horde.Compute;
 using EpicGames.Horde.Jobs;
@@ -32,6 +33,7 @@ using Horde.Server.Utilities;
 using JetBrains.Profiler.SelfApi;
 using Microsoft.AspNetCore.Authorization;
 using Microsoft.AspNetCore.Http;
+using Microsoft.AspNetCore.Http.Features;
 using Microsoft.AspNetCore.Mvc;
 using Microsoft.AspNetCore.Mvc.Filters;
 using Microsoft.Extensions.Logging;
@@ -276,7 +278,7 @@ namespace Horde.Server.Server
 		{
 		}
 	}
-	
+
 #if ENABLE_SECURE_DEBUG_CONTROLLER
 	/// <summary>
 	/// Controller managing account status
@@ -293,7 +295,7 @@ namespace Horde.Server.Server
 		/// Server settings (exposed for debug endpoint check filter)
 		/// </summary>
 		public IOptionsSnapshot<ServerSettings> ServerSettings { get; private set; }
-		
+
 		private readonly MongoService _mongoService;
 		private readonly ConfigService _configService;
 		private readonly AgentRelayService _agentRelayService;
@@ -427,6 +429,55 @@ namespace Horde.Server.Server
 			return config;
 		}
 
+		/// <summary>
+		/// Returns the fully parsed config object.
+		/// </summary>
+		[HttpGet]
+		[Route("/api/v1/debug/randomdata")]
+		public async Task<ActionResult> GetDataAsync([FromQuery] string size = "1mb", CancellationToken cancellationToken = default)
+		{
+			if (!_globalConfig.Value.Authorize(ServerAclAction.Debug, User))
+			{
+				return Forbid(ServerAclAction.Debug);
+			}
+
+			long sizeBytes;
+			try
+			{
+				sizeBytes = StringUtils.ParseBytesString(size);
+			}
+			catch
+			{
+				return BadRequest("Size must have a well-known binary suffix (eg. gb, mb, kb).");
+			}
+
+			// Disable buffering for the response
+			IHttpResponseBodyFeature? responseBodyFeature = HttpContext.Features.Get<IHttpResponseBodyFeature>();
+			responseBodyFeature?.DisableBuffering();
+
+			// Write the response directly to the writer
+			HttpResponse response = HttpContext.Response;
+			response.ContentType = "application/octet-stream";
+			response.StatusCode = (int)HttpStatusCode.OK;
+
+			await response.StartAsync(cancellationToken);
+
+			Random rnd = new Random();
+			for (long offsetBytes = 0; offsetBytes < sizeBytes;)
+			{
+				int chunkSize = (int)Math.Min(64 * 1024, sizeBytes - offsetBytes);
+
+				Memory<byte> chunkData = response.BodyWriter.GetMemory(chunkSize);
+				rnd.NextBytes(chunkData.Span.Slice(0, chunkSize));
+				response.BodyWriter.Advance(chunkSize);
+
+				offsetBytes += chunkSize;
+			}
+
+			await response.CompleteAsync();
+			return Empty;
+		}
+
 		/// <summary>
 		/// Get the network ID for a given IP address
 		/// </summary>
@@ -662,14 +713,14 @@ namespace Horde.Server.Server
 					throw new ArgumentException("StreamId or TemplateId do not match");
 				}
 
-				HashSet<string> newNames = [..StepNames.Union(jt.StepNames)];
-				return new JobTiming(StreamId, TemplateId, Name, newNames, 
+				HashSet<string> newNames = [.. StepNames.Union(jt.StepNames)];
+				return new JobTiming(StreamId, TemplateId, Name, newNames,
 					BatchSetupDuration + jt.BatchSetupDuration,
 					BatchWorkDuration + jt.BatchWorkDuration,
 					BatchTeardownDuration + jt.BatchTeardownDuration);
 			}
 		}
-		
+
 		/// <summary>
 		/// Display a table listing each template with job timings (setup, work and teardown durations)
 		/// </summary>
@@ -695,7 +746,7 @@ namespace Horde.Server.Server
 			{
 				allJobTimings = allJobTimings.Where(x => x.StepNames.SetEquals(["Setup Build"])).ToList();
 			}
-			
+
 			IReadOnlyList<JobTiming> jobTimingsByTemplate = GroupJobTimings(allJobTimings);
 			IEnumerable<JobTiming> sortedJobTimings = jobTimingsByTemplate.OrderBy(x => x.BatchSetupDuration).Reverse();
 
@@ -742,7 +793,7 @@ namespace Horde.Server.Server
 			sb.AppendLine("</table>");
 			return new ContentResult { ContentType = "text/html", StatusCode = (int)HttpStatusCode.OK, Content = sb.ToString() };
 		}
-		
+
 		private static ActionResult GetJobTimingsAsCsv(IEnumerable<JobTiming> jobTimings)
 		{
 			StringBuilder sb = new();
@@ -769,7 +820,7 @@ namespace Horde.Server.Server
 			foreach (IJob job in jobs)
 			{
 				IGraph graph = await _jobService.GetGraphAsync(job);
-				
+
 				foreach (IJobStepBatch batch in job.Batches)
 				{
 					if (batch.State != JobStepBatchState.Complete || batch.StartTimeUtc == null || batch.FinishTimeUtc == null)
@@ -786,7 +837,7 @@ namespace Horde.Server.Server
 						{
 							continue;
 						}
-						
+
 						INode node = graph.GetNode(new NodeRef(batch.GroupIdx, step.NodeIdx));
 						firstStepStartTime = step.StartTimeUtc.Value < firstStepStartTime ? step.StartTimeUtc.Value : firstStepStartTime;
 						lastStepFinishTime = step.FinishTimeUtc.Value > lastStepFinishTime ? step.FinishTimeUtc.Value : lastStepFinishTime;
@@ -1016,7 +1067,7 @@ namespace Horde.Server.Server
 
 			return PhysicalFile(snapshotZipFile, "application/zip", Path.GetFileName(snapshotZipFile));
 		}
-		
+
 		/// <summary>
 		/// Take a memory snapshot using dotTrace
 		/// </summary>
@@ -1046,7 +1097,7 @@ namespace Horde.Server.Server
 			config.SaveToDir(snapshotDir);
 			string workspaceFilePath = DotMemory.GetSnapshotOnce(config);
 			_logger.LogInformation("dotMemory snapshot captured in {CaptureTimeMs} ms", sw.ElapsedMilliseconds);
-			
+
 			if (!System.IO.File.Exists(workspaceFilePath))
 			{
 				return NotFound("The generated workspace file was not found");
