Horde: Remove unused GetJobStepFromClaim() method which introduces unnecessary dependencies on job system.

#rnx

diff --git a/Engine/Source/Programs/Horde/Horde.Server/Acls/AclService.cs b/Engine/Source/Programs/Horde/Horde.Server/Acls/AclService.cs
--- a/Engine/Source/Programs/Horde/Horde.Server/Acls/AclService.cs
+++ b/Engine/Source/Programs/Horde/Horde.Server/Acls/AclService.cs
@@ -10,13 +10,8 @@ using System.Threading.Tasks;
 using EpicGames.Horde.Agents;
 using EpicGames.Horde.Agents.Leases;
 using EpicGames.Horde.Agents.Sessions;
-using EpicGames.Horde.Jobs;
-using Google.Protobuf.WellKnownTypes;
-using Horde.Server.Agents.Leases;
-using Horde.Server.Jobs;
 using Horde.Server.Server;
 using Horde.Server.Utilities;
-using HordeCommon.Rpc.Tasks;
 using Microsoft.IdentityModel.Tokens;
 using MongoDB.Driver;
 
@@ -119,36 +114,6 @@ namespace Horde.Server.Acls
 			}
 		}
 
-		public static async Task<(IJob, IJobStep)?> GetJobStepFromClaimAsync(this ClaimsPrincipal user, ILeaseCollection leaseCollection, IJobCollection jobCollection, CancellationToken cancellationToken = default)
-		{
-			LeaseId? leaseId = user.GetLeaseClaim();
-			if (leaseId != null)
-			{
-				ILease? lease = await leaseCollection.GetAsync(leaseId.Value, cancellationToken);
-				if (lease != null)
-				{
-					Any payload = Any.Parser.ParseFrom(lease.Payload.ToArray());
-					if (payload.TryUnpack(out ExecuteJobTask jobTask))
-					{
-						IJob? job = await jobCollection.GetAsync(JobId.Parse(jobTask.JobId), cancellationToken);
-						if (job != null)
-						{
-							IJobStepBatch? batch = job.Batches.FirstOrDefault(x => x.LeaseId == leaseId);
-							if (batch != null && batch.State == JobStepBatchState.Running)
-							{
-								IJobStep? step = batch.Steps.FirstOrDefault(x => x.State == JobStepState.Running);
-								if (step != null)
-								{
-									return (job, step);
-								}
-							}
-						}
-					}
-				}
-			}
-			return null;
-		}
-
 		public static SessionId? GetSessionClaim(this ClaimsPrincipal user)
 		{
 			Claim? claim = user.FindFirst(HordeClaimTypes.AgentSessionId);
