Horde: Compute hashes for stream configurations in the StreamConfig.PostLoad() method rather than in the config service.

#rnx

diff --git a/Engine/Source/Programs/Horde/Horde.Server/Configuration/ConfigService.cs b/Engine/Source/Programs/Horde/Horde.Server/Configuration/ConfigService.cs
--- a/Engine/Source/Programs/Horde/Horde.Server/Configuration/ConfigService.cs
+++ b/Engine/Source/Programs/Horde/Horde.Server/Configuration/ConfigService.cs
@@ -12,9 +12,7 @@ using System.Threading.Tasks;
 using EpicGames.Core;
 using EpicGames.Horde.Users;
 using EpicGames.Redis;
-using Horde.Server.Projects;
 using Horde.Server.Server;
-using Horde.Server.Streams;
 using Horde.Server.Users;
 using Microsoft.Extensions.Diagnostics.HealthChecks;
 using Microsoft.Extensions.Hosting;
@@ -732,16 +730,6 @@ namespace Horde.Server.Configuration
 			GlobalConfig globalConfig = JsonSerializer.Deserialize<GlobalConfig>(snapshot.Data, _jsonOptions)!;
 			globalConfig.Revision = IoHash.Compute(data.Span).ToString();
 
-			// Compute hashes for all the stream objects
-			foreach (ProjectConfig projectConfig in globalConfig.Projects)
-			{
-				foreach (StreamConfig streamConfig in projectConfig.Streams)
-				{
-					byte[] streamData = JsonSerializer.SerializeToUtf8Bytes(streamConfig, _jsonOptions);
-					streamConfig.Revision = IoHash.Compute(streamData).ToString();
-				}
-			}
-
 			// Run the postload callbacks on all the config objects
 			globalConfig.PostLoad(_serverSettings);
 			return globalConfig;
diff --git a/Engine/Source/Programs/Horde/Horde.Server/Startup.cs b/Engine/Source/Programs/Horde/Horde.Server/Startup.cs
--- a/Engine/Source/Programs/Horde/Horde.Server/Startup.cs
+++ b/Engine/Source/Programs/Horde/Horde.Server/Startup.cs
@@ -1037,6 +1037,16 @@ namespace Horde.Server
 			LogValueFormatter.RegisterTypeAnnotation<UserId>("UserId");
 		}
 
+		public static JsonSerializerOptions JsonSerializerOptions { get; } = CreateJsonSerializerOptions();
+
+		static JsonSerializerOptions CreateJsonSerializerOptions()
+		{
+			JsonSerializerOptions options = new JsonSerializerOptions();
+			ConfigureJsonSerializer(options);
+			options.MakeReadOnly(true);
+			return options;
+		}
+
 		public static void ConfigureJsonSerializer(JsonSerializerOptions options)
 		{
 			HordeHttpClient.ConfigureJsonSerializer(options);
diff --git a/Engine/Source/Programs/Horde/Horde.Server/Streams/StreamConfig.cs b/Engine/Source/Programs/Horde/Horde.Server/Streams/StreamConfig.cs
--- a/Engine/Source/Programs/Horde/Horde.Server/Streams/StreamConfig.cs
+++ b/Engine/Source/Programs/Horde/Horde.Server/Streams/StreamConfig.cs
@@ -318,6 +318,10 @@ namespace Horde.Server.Streams
 					}
 				}
 			}
+
+			// Compute a hash of this stream revision to make it easier to detect changes
+			byte[] streamData = JsonSerializer.SerializeToUtf8Bytes(this, Startup.JsonSerializerOptions);
+			Revision = IoHash.Compute(streamData).ToString();
 		}
 
 		/// <summary>
