Horde: Move agent/pool functionality relating to Perforce into the Perforce namespace.

#rnx

diff --git a/Engine/Source/Programs/Horde/Horde.Server/Agents/IAgent.cs b/Engine/Source/Programs/Horde/Horde.Server/Agents/IAgent.cs
--- a/Engine/Source/Programs/Horde/Horde.Server/Agents/IAgent.cs
+++ b/Engine/Source/Programs/Horde/Horde.Server/Agents/IAgent.cs
@@ -18,10 +18,7 @@ using EpicGames.Horde.Logs;
 using EpicGames.Horde.Streams;
 using EpicGames.Horde.Tools;
 using Google.Protobuf.WellKnownTypes;
-using Horde.Server.Agents.Pools;
-using Horde.Server.Perforce;
 using Horde.Server.Server;
-using Horde.Server.Streams;
 using HordeCommon.Rpc;
 using HordeCommon.Rpc.Messages;
 using HordeCommon.Rpc.Tasks;
@@ -578,206 +575,6 @@ namespace Horde.Server.Agents
 			}
 			return true;
 		}
-
-		/// <summary>
-		/// Get the AutoSDK workspace required for an agent
-		/// </summary>
-		/// <param name="agent"></param>
-		/// <param name="cluster">The perforce cluster to get a workspace for</param>
-		/// <param name="pools">Pools that the agent belongs to</param>
-		/// <returns></returns>
-		public static AgentWorkspaceInfo? GetAutoSdkWorkspace(this IAgent agent, PerforceCluster cluster, IEnumerable<IPool> pools)
-		{
-			AutoSdkConfig? autoSdkConfig = null;
-			foreach (IPool pool in pools)
-			{
-				autoSdkConfig = AutoSdkConfig.Merge(autoSdkConfig, pool.AutoSdkConfig);
-			}
-			if (autoSdkConfig == null)
-			{
-				return null;
-			}
-
-			return GetAutoSdkWorkspace(agent, cluster, autoSdkConfig);
-		}
-
-		/// <summary>
-		/// Get the AutoSDK workspace required for an agent
-		/// </summary>
-		/// <param name="agent"></param>
-		/// <param name="cluster">The perforce cluster to get a workspace for</param>
-		/// <param name="autoSdkConfig">Configuration for autosdk</param>
-		/// <returns></returns>
-		public static AgentWorkspaceInfo? GetAutoSdkWorkspace(this IAgent agent, PerforceCluster cluster, AutoSdkConfig autoSdkConfig)
-		{
-			foreach (AutoSdkWorkspace autoSdk in cluster.AutoSdk)
-			{
-				if (autoSdk.Stream != null && autoSdk.Properties.All(x => agent.Properties.Contains(x)))
-				{
-					return new AgentWorkspaceInfo(cluster.Name, autoSdk.UserName, autoSdk.Name ?? "AutoSDK", autoSdk.Stream!, autoSdkConfig.View.ToList(), true, null, null);
-				}
-			}
-			return null;
-		}
-
-		/// <summary>
-		/// Converts this workspace to an RPC message
-		/// </summary>
-		/// <param name="agent">The agent to get a workspace for</param>
-		/// <param name="workspace">The workspace definition</param>
-		/// <param name="cluster">The global state</param>
-		/// <param name="loadBalancer">The Perforce load balancer</param>
-		/// <param name="workspaceMessages">List of messages</param>
-		/// <param name="cancellationToken">Cancellation token for the operation</param>
-		/// <returns>The RPC message</returns>
-		public static async Task<bool> TryAddWorkspaceMessageAsync(this IAgent agent, AgentWorkspaceInfo workspace, PerforceCluster cluster, PerforceLoadBalancer loadBalancer, IList<RpcAgentWorkspace> workspaceMessages, CancellationToken cancellationToken)
-		{
-			// Find a matching server, trying to use a previously selected one if possible
-			string? baseServerAndPort;
-			string? serverAndPort;
-			bool partitioned;
-
-			RpcAgentWorkspace? existingWorkspace = workspaceMessages.FirstOrDefault(x => x.ConfiguredCluster == workspace.Cluster);
-			if (existingWorkspace != null)
-			{
-				baseServerAndPort = existingWorkspace.BaseServerAndPort;
-				serverAndPort = existingWorkspace.ServerAndPort;
-				partitioned = existingWorkspace.Partitioned;
-			}
-			else
-			{
-				if (cluster == null)
-				{
-					return false;
-				}
-
-				IPerforceServer? server = await loadBalancer.SelectServerAsync(cluster, agent, cancellationToken);
-				if (server == null)
-				{
-					return false;
-				}
-
-				baseServerAndPort = server.BaseServerAndPort;
-				serverAndPort = server.ServerAndPort;
-				partitioned = server.SupportsPartitionedWorkspaces;
-			}
-
-			// Find the matching credentials for the desired user
-			PerforceCredentials? credentials = null;
-			if (cluster != null)
-			{
-				if (workspace.UserName == null)
-				{
-					credentials = cluster.Credentials.FirstOrDefault();
-				}
-				else
-				{
-					credentials = cluster.Credentials.FirstOrDefault(x => String.Equals(x.UserName, workspace.UserName, StringComparison.OrdinalIgnoreCase));
-				}
-			}
-
-			// Construct the message
-			RpcAgentWorkspace result = new RpcAgentWorkspace
-			{
-				ConfiguredCluster = workspace.Cluster,
-				ConfiguredUserName = workspace.UserName,
-				Cluster = cluster?.Name,
-				BaseServerAndPort = baseServerAndPort,
-				ServerAndPort = serverAndPort,
-				UserName = credentials?.UserName ?? workspace.UserName,
-				Password = credentials?.Password,
-				Ticket = credentials?.Ticket,
-				Identifier = workspace.Identifier,
-				Stream = workspace.Stream,
-				Incremental = workspace.Incremental,
-				Partitioned = partitioned,
-				Method = workspace.Method ?? String.Empty
-			};
-
-			if (workspace.View != null)
-			{
-				result.View.AddRange(workspace.View);
-			}
-
-			workspaceMessages.Add(result);
-			return true;
-		}
-
-		/// <summary>
-		/// Tries to get an agent workspace definition from the given type name
-		/// </summary>
-		/// <param name="streamConfig">The stream object</param>
-		/// <param name="agentType">The agent type</param>
-		/// <param name="workspace">Receives the agent workspace definition</param>
-		/// <param name="autoSdkConfig">Receives the autosdk workspace config</param>
-		/// <returns>True if the agent type was valid, and an agent workspace could be created</returns>
-		public static bool TryGetAgentWorkspace(this StreamConfig streamConfig, AgentConfig agentType, [NotNullWhen(true)] out AgentWorkspaceInfo? workspace, out AutoSdkConfig? autoSdkConfig)
-		{
-			// Get the workspace settings
-			if (agentType.Workspace == null)
-			{
-				// Use the default settings (fast switching workspace, clean)
-				workspace = new AgentWorkspaceInfo(streamConfig.ClusterName, null, streamConfig.GetDefaultWorkspaceIdentifier(), streamConfig.Name, null, false, null, null);
-				autoSdkConfig = AutoSdkConfig.Full;
-				return true;
-			}
-			else
-			{
-				// Try to get the matching workspace type
-				WorkspaceConfig? workspaceConfig;
-				if (!streamConfig.WorkspaceTypes.TryGetValue(agentType.Workspace, out workspaceConfig))
-				{
-					workspace = null;
-					autoSdkConfig = null;
-					return false;
-				}
-
-				// Get the workspace identifier
-				string identifier;
-				if (workspaceConfig.Identifier != null)
-				{
-					identifier = workspaceConfig.Identifier;
-				}
-				else if (workspaceConfig.Incremental ?? false)
-				{
-					identifier = $"{streamConfig.GetEscapedName()}+{agentType.Workspace}";
-				}
-				else
-				{
-					identifier = streamConfig.GetDefaultWorkspaceIdentifier();
-				}
-
-				// Create the new workspace
-				string cluster = workspaceConfig.Cluster ?? streamConfig.ClusterName;
-				workspace = new AgentWorkspaceInfo(cluster, workspaceConfig.UserName, identifier, workspaceConfig.Stream ?? streamConfig.Name, workspaceConfig.View, workspaceConfig.Incremental ?? false, workspaceConfig.Method, workspaceConfig.MinScratchSpace);
-				autoSdkConfig = GetAutoSdkConfig(workspaceConfig, streamConfig);
-
-				return true;
-			}
-
-			static AutoSdkConfig? GetAutoSdkConfig(WorkspaceConfig workspaceConfig, StreamConfig streamConfig)
-			{
-				AutoSdkConfig? autoSdkConfig = null;
-				if (workspaceConfig.UseAutoSdk ?? true)
-				{
-					List<string> view = new List<string>();
-					if (streamConfig.AutoSdkView != null)
-					{
-						view.AddRange(streamConfig.AutoSdkView);
-					}
-					if (workspaceConfig.AutoSdkView != null)
-					{
-						view.AddRange(workspaceConfig.AutoSdkView);
-					}
-					if (view.Count == 0)
-					{
-						view.Add("...");
-					}
-					autoSdkConfig = new AutoSdkConfig(view);
-				}
-				return autoSdkConfig;
-			}
-		}
 	}
 
 	/// <summary>
@@ -924,39 +721,6 @@ namespace Horde.Server.Agents
 			HashSet<AgentWorkspaceInfo> workspacesSetA = new HashSet<AgentWorkspaceInfo>(workspacesA);
 			return workspacesSetA.SetEquals(workspacesB);
 		}
-
-		/// <summary>
-		/// Converts this workspace to an RPC message
-		/// </summary>
-		/// <param name="server">The Perforce server</param>
-		/// <param name="credentials">Credentials for the server</param>
-		/// <returns>The RPC message</returns>
-		public RpcAgentWorkspace ToRpcMessage(IPerforceServer server, PerforceCredentials? credentials)
-		{
-			// Construct the message
-			RpcAgentWorkspace result = new RpcAgentWorkspace
-			{
-				ConfiguredCluster = Cluster,
-				ConfiguredUserName = UserName,
-				ServerAndPort = server.ServerAndPort,
-				UserName = credentials?.UserName ?? UserName,
-				Password = credentials?.Password,
-				Ticket = credentials?.Ticket,
-				Identifier = Identifier,
-				Stream = Stream,
-				Incremental = Incremental,
-				Partitioned = server.SupportsPartitionedWorkspaces,
-				Method = Method ?? String.Empty,
-				MinScratchSpace = MinScratchSpace ?? 0
-			};
-
-			if (View != null)
-			{
-				result.View.AddRange(View);
-			}
-
-			return result;
-		}
 	}
 
 	/// <summary>
diff --git a/Engine/Source/Programs/Horde/Horde.Server/Agents/Pools/PoolService.cs b/Engine/Source/Programs/Horde/Horde.Server/Agents/Pools/PoolService.cs
--- a/Engine/Source/Programs/Horde/Horde.Server/Agents/Pools/PoolService.cs
+++ b/Engine/Source/Programs/Horde/Horde.Server/Agents/Pools/PoolService.cs
@@ -8,7 +8,6 @@ using System.Threading.Tasks;
 using EpicGames.Core;
 using EpicGames.Horde;
 using EpicGames.Horde.Agents.Pools;
-using Horde.Server.Server;
 
 namespace Horde.Server.Agents.Pools
 {
@@ -111,7 +110,7 @@ namespace Horde.Server.Agents.Pools
 		/// <returns></returns>
 		public async Task<IPoolConfig?> GetPoolAsync(PoolId poolId, DateTime validAtTime, CancellationToken cancellationToken)
 		{
-			Dictionary<PoolId, IPoolConfig> poolMapping = await GetPoolLookupAsync(validAtTime, cancellationToken);
+			IReadOnlyDictionary<PoolId, IPoolConfig> poolMapping = await GetPoolLookupAsync(validAtTime, cancellationToken);
 			poolMapping.TryGetValue(poolId, out IPoolConfig? pool);
 			return pool;
 		}
@@ -125,7 +124,7 @@ namespace Horde.Server.Agents.Pools
 		/// <returns></returns>
 		public async Task<List<IPoolConfig>> GetPoolsAsync(IAgent agent, DateTime validAtTime, CancellationToken cancellationToken)
 		{
-			Dictionary<PoolId, IPoolConfig> poolMapping = await GetPoolLookupAsync(validAtTime, cancellationToken);
+			IReadOnlyDictionary<PoolId, IPoolConfig> poolMapping = await GetPoolLookupAsync(validAtTime, cancellationToken);
 
 			List<IPoolConfig> pools = new List<IPoolConfig>();
 			foreach (PoolId poolId in agent.Pools)
@@ -139,113 +138,13 @@ namespace Horde.Server.Agents.Pools
 			return pools;
 		}
 
-		/// <summary>
-		/// Get a list of workspaces for the given agent
-		/// </summary>
-		/// <param name="agent">The agent to return workspaces for</param>
-		/// <param name="validAtTime">Absolute time at which we expect the results to be valid. Values may be cached as long as they are after this time.</param>
-		/// <param name="globalConfig">Current configuration</param>
-		/// <param name="cancellationToken">Cancellation token for the operation</param>
-		/// <returns>List of workspaces</returns>
-		public async Task<HashSet<AgentWorkspaceInfo>> GetWorkspacesAsync(IAgent agent, DateTime validAtTime, GlobalConfig globalConfig, CancellationToken cancellationToken)
-		{
-			List<IPoolConfig> pools = await GetPoolsAsync(agent, validAtTime, cancellationToken);
-
-			HashSet<AgentWorkspaceInfo> workspaces = new HashSet<AgentWorkspaceInfo>();
-			foreach (IPoolConfig pool in pools)
-			{
-				workspaces.UnionWith(pool.Workspaces);
-			}
-
-			AutoSdkConfig? autoSdkConfig = GetAutoSdkConfig(pools);
-			if (autoSdkConfig != null)
-			{
-				foreach (string? clusterName in workspaces.Select(x => x.Cluster).Distinct().ToList())
-				{
-					PerforceCluster? cluster = globalConfig.FindPerforceCluster(clusterName);
-					if (cluster != null)
-					{
-						AgentWorkspaceInfo? autoSdkWorkspace = agent.GetAutoSdkWorkspace(cluster, autoSdkConfig);
-						if (autoSdkWorkspace != null)
-						{
-							workspaces.Add(autoSdkWorkspace);
-						}
-					}
-				}
-			}
-
-			return workspaces;
-		}
-
-		static AutoSdkConfig? GetAutoSdkConfig(IEnumerable<IPoolConfig> pools)
-		{
-			AutoSdkConfig? autoSdkConfig = null;
-			foreach (IPoolConfig pool in pools)
-			{
-				autoSdkConfig = AutoSdkConfig.Merge(autoSdkConfig, pool.AutoSdkConfig);
-			}
-			return autoSdkConfig;
-		}
-
-		/// <summary>
-		/// Gets all the autosdk workspaces required for an agent
-		/// </summary>
-		/// <param name="agent"></param>
-		/// <param name="cluster"></param>
-		/// <param name="cancellationToken">Cancellation token for the operation</param>
-		/// <returns></returns>
-		public async Task<AgentWorkspaceInfo?> GetAutoSdkWorkspaceAsync(IAgent agent, PerforceCluster cluster, CancellationToken cancellationToken)
-		{
-			List<IPoolConfig> pools = await GetPoolsAsync(agent, DateTime.UtcNow - TimeSpan.FromSeconds(10.0), cancellationToken);
-
-			AutoSdkConfig? autoSdkConfig = GetAutoSdkConfig(pools);
-			if (autoSdkConfig == null)
-			{
-				return null;
-			}
-
-			return agent.GetAutoSdkWorkspace(cluster, autoSdkConfig);
-		}
-
-		/// <summary>
-		/// Get a list of workspaces for the given agent
-		/// </summary>
-		/// <param name="agent">The agent to return workspaces for</param>
-		/// <param name="perforceCluster">The P4 cluster to find a workspace for</param>
-		/// <param name="validAtTime">Absolute time at which we expect the results to be valid. Values may be cached as long as they are after this time.</param>
-		/// <param name="cancellationToken">Cancellation token for the operation</param>
-		/// <returns>List of workspaces</returns>
-		public async Task<AgentWorkspaceInfo?> GetAutoSdkWorkspaceAsync(IAgent agent, PerforceCluster perforceCluster, DateTime validAtTime, CancellationToken cancellationToken)
-		{
-			AutoSdkConfig? autoSdkConfig = null;
-
-			Dictionary<PoolId, IPoolConfig> poolMapping = await GetPoolLookupAsync(validAtTime, cancellationToken);
-			foreach (PoolId poolId in agent.Pools)
-			{
-				IPoolConfig? pool;
-				if (poolMapping.TryGetValue(poolId, out pool))
-				{
-					autoSdkConfig = AutoSdkConfig.Merge(autoSdkConfig, pool.AutoSdkConfig);
-				}
-			}
-
-			if (autoSdkConfig == null)
-			{
-				return null;
-			}
-			else
-			{
-				return agent.GetAutoSdkWorkspace(perforceCluster, autoSdkConfig);
-			}
-		}
-
 		/// <summary>
 		/// Gets a mapping from pool identifiers to definitions
 		/// </summary>
 		/// <param name="validAtTime">Absolute time at which we expect the results to be valid. Values may be cached as long as they are after this time.</param>
 		/// <param name="cancellationToken">Cancellation token for the operation</param>
 		/// <returns>Map of pool ids to pool documents</returns>
-		private async Task<Dictionary<PoolId, IPoolConfig>> GetPoolLookupAsync(DateTime validAtTime, CancellationToken cancellationToken)
+		public async Task<IReadOnlyDictionary<PoolId, IPoolConfig>> GetPoolLookupAsync(DateTime validAtTime, CancellationToken cancellationToken)
 		{
 			Tuple<DateTime, Dictionary<PoolId, IPoolConfig>>? cachedPoolLookupCopy = _cachedPoolLookup;
 			if (cachedPoolLookupCopy == null || cachedPoolLookupCopy.Item1 < validAtTime)
diff --git a/Engine/Source/Programs/Horde/Horde.Server/Agents/Pools/PoolUpdateService.cs b/Engine/Source/Programs/Horde/Horde.Server/Agents/Pools/PoolUpdateService.cs
--- a/Engine/Source/Programs/Horde/Horde.Server/Agents/Pools/PoolUpdateService.cs
+++ b/Engine/Source/Programs/Horde/Horde.Server/Agents/Pools/PoolUpdateService.cs
@@ -8,6 +8,7 @@ using System.Threading.Tasks;
 using EpicGames.Core;
 using EpicGames.Horde.Agents;
 using EpicGames.Horde.Agents.Pools;
+using Horde.Server.Perforce;
 using Horde.Server.Server;
 using Horde.Server.Streams;
 using Microsoft.Extensions.Hosting;
diff --git a/Engine/Source/Programs/Horde/Horde.Server/Jobs/JobRpcService.cs b/Engine/Source/Programs/Horde/Horde.Server/Jobs/JobRpcService.cs
--- a/Engine/Source/Programs/Horde/Horde.Server/Jobs/JobRpcService.cs
+++ b/Engine/Source/Programs/Horde/Horde.Server/Jobs/JobRpcService.cs
@@ -29,6 +29,7 @@ using Horde.Server.Jobs.Graphs;
 using Horde.Server.Jobs.Templates;
 using Horde.Server.Jobs.TestData;
 using Horde.Server.Logs;
+using Horde.Server.Perforce;
 using Horde.Server.Server;
 using Horde.Server.Streams;
 using Horde.Server.Tasks;
diff --git a/Engine/Source/Programs/Horde/Horde.Server/Perforce/PerforceWorkspace.cs b/Engine/Source/Programs/Horde/Horde.Server/Perforce/PerforceWorkspace.cs
new file mode 100644
--- /dev/null
+++ b/Engine/Source/Programs/Horde/Horde.Server/Perforce/PerforceWorkspace.cs
@@ -0,0 +1,326 @@
+// Copyright Epic Games, Inc. All Rights Reserved.
+
+using System;
+using System.Collections.Generic;
+using System.Diagnostics.CodeAnalysis;
+using System.Linq;
+using System.Threading;
+using System.Threading.Tasks;
+using EpicGames.Horde.Agents.Pools;
+using Horde.Server.Agents;
+using Horde.Server.Agents.Pools;
+using Horde.Server.Server;
+using Horde.Server.Streams;
+using HordeCommon.Rpc.Messages;
+
+namespace Horde.Server.Perforce
+{
+	/// <summary>
+	/// Helper methods for dealing with Perforce workspaces
+	/// </summary>
+	public static class PerforceWorkspace
+	{
+		/// <summary>
+		/// Get the AutoSDK workspace required for an agent
+		/// </summary>
+		/// <param name="agent"></param>
+		/// <param name="cluster">The perforce cluster to get a workspace for</param>
+		/// <param name="pools">Pools that the agent belongs to</param>
+		/// <returns></returns>
+		public static AgentWorkspaceInfo? GetAutoSdkWorkspace(this IAgent agent, PerforceCluster cluster, IEnumerable<IPool> pools)
+		{
+			AutoSdkConfig? autoSdkConfig = null;
+			foreach (IPool pool in pools)
+			{
+				autoSdkConfig = AutoSdkConfig.Merge(autoSdkConfig, pool.AutoSdkConfig);
+			}
+			if (autoSdkConfig == null)
+			{
+				return null;
+			}
+
+			return GetAutoSdkWorkspace(agent, cluster, autoSdkConfig);
+		}
+
+		/// <summary>
+		/// Get the AutoSDK workspace required for an agent
+		/// </summary>
+		/// <param name="agent"></param>
+		/// <param name="cluster">The perforce cluster to get a workspace for</param>
+		/// <param name="autoSdkConfig">Configuration for autosdk</param>
+		/// <returns></returns>
+		public static AgentWorkspaceInfo? GetAutoSdkWorkspace(this IAgent agent, PerforceCluster cluster, AutoSdkConfig autoSdkConfig)
+		{
+			foreach (AutoSdkWorkspace autoSdk in cluster.AutoSdk)
+			{
+				if (autoSdk.Stream != null && autoSdk.Properties.All(x => agent.Properties.Contains(x)))
+				{
+					return new AgentWorkspaceInfo(cluster.Name, autoSdk.UserName, autoSdk.Name ?? "AutoSDK", autoSdk.Stream!, autoSdkConfig.View.ToList(), true, null, null);
+				}
+			}
+			return null;
+		}
+
+		/// <summary>
+		/// Converts this workspace to an RPC message
+		/// </summary>
+		/// <param name="agent">The agent to get a workspace for</param>
+		/// <param name="workspace">The workspace definition</param>
+		/// <param name="cluster">The global state</param>
+		/// <param name="loadBalancer">The Perforce load balancer</param>
+		/// <param name="workspaceMessages">List of messages</param>
+		/// <param name="cancellationToken">Cancellation token for the operation</param>
+		/// <returns>The RPC message</returns>
+		public static async Task<bool> TryAddWorkspaceMessageAsync(this IAgent agent, AgentWorkspaceInfo workspace, PerforceCluster cluster, PerforceLoadBalancer loadBalancer, IList<RpcAgentWorkspace> workspaceMessages, CancellationToken cancellationToken)
+		{
+			// Find a matching server, trying to use a previously selected one if possible
+			string? baseServerAndPort;
+			string? serverAndPort;
+			bool partitioned;
+
+			RpcAgentWorkspace? existingWorkspace = workspaceMessages.FirstOrDefault(x => x.ConfiguredCluster == workspace.Cluster);
+			if (existingWorkspace != null)
+			{
+				baseServerAndPort = existingWorkspace.BaseServerAndPort;
+				serverAndPort = existingWorkspace.ServerAndPort;
+				partitioned = existingWorkspace.Partitioned;
+			}
+			else
+			{
+				if (cluster == null)
+				{
+					return false;
+				}
+
+				IPerforceServer? server = await loadBalancer.SelectServerAsync(cluster, agent, cancellationToken);
+				if (server == null)
+				{
+					return false;
+				}
+
+				baseServerAndPort = server.BaseServerAndPort;
+				serverAndPort = server.ServerAndPort;
+				partitioned = server.SupportsPartitionedWorkspaces;
+			}
+
+			// Find the matching credentials for the desired user
+			PerforceCredentials? credentials = null;
+			if (cluster != null)
+			{
+				if (workspace.UserName == null)
+				{
+					credentials = cluster.Credentials.FirstOrDefault();
+				}
+				else
+				{
+					credentials = cluster.Credentials.FirstOrDefault(x => String.Equals(x.UserName, workspace.UserName, StringComparison.OrdinalIgnoreCase));
+				}
+			}
+
+			// Construct the message
+			RpcAgentWorkspace result = new RpcAgentWorkspace
+			{
+				ConfiguredCluster = workspace.Cluster,
+				ConfiguredUserName = workspace.UserName,
+				Cluster = cluster?.Name,
+				BaseServerAndPort = baseServerAndPort,
+				ServerAndPort = serverAndPort,
+				UserName = credentials?.UserName ?? workspace.UserName,
+				Password = credentials?.Password,
+				Ticket = credentials?.Ticket,
+				Identifier = workspace.Identifier,
+				Stream = workspace.Stream,
+				Incremental = workspace.Incremental,
+				Partitioned = partitioned,
+				Method = workspace.Method ?? String.Empty
+			};
+
+			if (workspace.View != null)
+			{
+				result.View.AddRange(workspace.View);
+			}
+
+			workspaceMessages.Add(result);
+			return true;
+		}
+
+		/// <summary>
+		/// Tries to get an agent workspace definition from the given type name
+		/// </summary>
+		/// <param name="streamConfig">The stream object</param>
+		/// <param name="agentType">The agent type</param>
+		/// <param name="workspace">Receives the agent workspace definition</param>
+		/// <param name="autoSdkConfig">Receives the autosdk workspace config</param>
+		/// <returns>True if the agent type was valid, and an agent workspace could be created</returns>
+		public static bool TryGetAgentWorkspace(this StreamConfig streamConfig, AgentConfig agentType, [NotNullWhen(true)] out AgentWorkspaceInfo? workspace, out AutoSdkConfig? autoSdkConfig)
+		{
+			// Get the workspace settings
+			if (agentType.Workspace == null)
+			{
+				// Use the default settings (fast switching workspace, clean)
+				workspace = new AgentWorkspaceInfo(streamConfig.ClusterName, null, streamConfig.GetDefaultWorkspaceIdentifier(), streamConfig.Name, null, false, null, null);
+				autoSdkConfig = AutoSdkConfig.Full;
+				return true;
+			}
+			else
+			{
+				// Try to get the matching workspace type
+				WorkspaceConfig? workspaceConfig;
+				if (!streamConfig.WorkspaceTypes.TryGetValue(agentType.Workspace, out workspaceConfig))
+				{
+					workspace = null;
+					autoSdkConfig = null;
+					return false;
+				}
+
+				// Get the workspace identifier
+				string identifier;
+				if (workspaceConfig.Identifier != null)
+				{
+					identifier = workspaceConfig.Identifier;
+				}
+				else if (workspaceConfig.Incremental ?? false)
+				{
+					identifier = $"{streamConfig.GetEscapedName()}+{agentType.Workspace}";
+				}
+				else
+				{
+					identifier = streamConfig.GetDefaultWorkspaceIdentifier();
+				}
+
+				// Create the new workspace
+				string cluster = workspaceConfig.Cluster ?? streamConfig.ClusterName;
+				workspace = new AgentWorkspaceInfo(cluster, workspaceConfig.UserName, identifier, workspaceConfig.Stream ?? streamConfig.Name, workspaceConfig.View, workspaceConfig.Incremental ?? false, workspaceConfig.Method, workspaceConfig.MinScratchSpace);
+				autoSdkConfig = GetAutoSdkConfig(workspaceConfig, streamConfig);
+
+				return true;
+			}
+
+			static AutoSdkConfig? GetAutoSdkConfig(WorkspaceConfig workspaceConfig, StreamConfig streamConfig)
+			{
+				AutoSdkConfig? autoSdkConfig = null;
+				if (workspaceConfig.UseAutoSdk ?? true)
+				{
+					List<string> view = new List<string>();
+					if (streamConfig.AutoSdkView != null)
+					{
+						view.AddRange(streamConfig.AutoSdkView);
+					}
+					if (workspaceConfig.AutoSdkView != null)
+					{
+						view.AddRange(workspaceConfig.AutoSdkView);
+					}
+					if (view.Count == 0)
+					{
+						view.Add("...");
+					}
+					autoSdkConfig = new AutoSdkConfig(view);
+				}
+				return autoSdkConfig;
+			}
+		}
+
+		/// <summary>
+		/// Get a list of workspaces for the given agent
+		/// </summary>
+		/// <param name="poolService">The pool service instance</param>
+		/// <param name="agent">The agent to return workspaces for</param>
+		/// <param name="validAtTime">Absolute time at which we expect the results to be valid. Values may be cached as long as they are after this time.</param>
+		/// <param name="globalConfig">Current configuration</param>
+		/// <param name="cancellationToken">Cancellation token for the operation</param>
+		/// <returns>List of workspaces</returns>
+		public static async Task<HashSet<AgentWorkspaceInfo>> GetWorkspacesAsync(this PoolService poolService, IAgent agent, DateTime validAtTime, GlobalConfig globalConfig, CancellationToken cancellationToken)
+		{
+			List<IPoolConfig> pools = await poolService.GetPoolsAsync(agent, validAtTime, cancellationToken);
+
+			HashSet<AgentWorkspaceInfo> workspaces = new HashSet<AgentWorkspaceInfo>();
+			foreach (IPoolConfig pool in pools)
+			{
+				workspaces.UnionWith(pool.Workspaces);
+			}
+
+			AutoSdkConfig? autoSdkConfig = GetAutoSdkConfig(pools);
+			if (autoSdkConfig != null)
+			{
+				foreach (string? clusterName in workspaces.Select(x => x.Cluster).Distinct().ToList())
+				{
+					PerforceCluster? cluster = globalConfig.FindPerforceCluster(clusterName);
+					if (cluster != null)
+					{
+						AgentWorkspaceInfo? autoSdkWorkspace = agent.GetAutoSdkWorkspace(cluster, autoSdkConfig);
+						if (autoSdkWorkspace != null)
+						{
+							workspaces.Add(autoSdkWorkspace);
+						}
+					}
+				}
+			}
+
+			return workspaces;
+		}
+
+		static AutoSdkConfig? GetAutoSdkConfig(IEnumerable<IPoolConfig> pools)
+		{
+			AutoSdkConfig? autoSdkConfig = null;
+			foreach (IPoolConfig pool in pools)
+			{
+				autoSdkConfig = AutoSdkConfig.Merge(autoSdkConfig, pool.AutoSdkConfig);
+			}
+			return autoSdkConfig;
+		}
+
+		/// <summary>
+		/// Gets all the autosdk workspaces required for an agent
+		/// </summary>
+		/// <param name="poolService">The pool service instance</param>
+		/// <param name="agent"></param>
+		/// <param name="cluster"></param>
+		/// <param name="cancellationToken">Cancellation token for the operation</param>
+		/// <returns></returns>
+		public static async Task<AgentWorkspaceInfo?> GetAutoSdkWorkspaceAsync(this PoolService poolService, IAgent agent, PerforceCluster cluster, CancellationToken cancellationToken)
+		{
+			List<IPoolConfig> pools = await poolService.GetPoolsAsync(agent, DateTime.UtcNow - TimeSpan.FromSeconds(10.0), cancellationToken);
+
+			AutoSdkConfig? autoSdkConfig = GetAutoSdkConfig(pools);
+			if (autoSdkConfig == null)
+			{
+				return null;
+			}
+
+			return agent.GetAutoSdkWorkspace(cluster, autoSdkConfig);
+		}
+
+		/// <summary>
+		/// Get a list of workspaces for the given agent
+		/// </summary>
+		/// <param name="poolService">Instance of the pool service</param>
+		/// <param name="agent">The agent to return workspaces for</param>
+		/// <param name="perforceCluster">The P4 cluster to find a workspace for</param>
+		/// <param name="validAtTime">Absolute time at which we expect the results to be valid. Values may be cached as long as they are after this time.</param>
+		/// <param name="cancellationToken">Cancellation token for the operation</param>
+		/// <returns>List of workspaces</returns>
+		public static async Task<AgentWorkspaceInfo?> GetAutoSdkWorkspaceAsync(this PoolService poolService, IAgent agent, PerforceCluster perforceCluster, DateTime validAtTime, CancellationToken cancellationToken)
+		{
+			AutoSdkConfig? autoSdkConfig = null;
+
+			IReadOnlyDictionary<PoolId, IPoolConfig> poolMapping = await poolService.GetPoolLookupAsync(validAtTime, cancellationToken);
+			foreach (PoolId poolId in agent.Pools)
+			{
+				IPoolConfig? pool;
+				if (poolMapping.TryGetValue(poolId, out pool))
+				{
+					autoSdkConfig = AutoSdkConfig.Merge(autoSdkConfig, pool.AutoSdkConfig);
+				}
+			}
+
+			if (autoSdkConfig == null)
+			{
+				return null;
+			}
+			else
+			{
+				return agent.GetAutoSdkWorkspace(perforceCluster, autoSdkConfig);
+			}
+		}
+	}
+}
diff --git a/Engine/Source/Programs/Horde/Horde.Server/Server/GlobalConfig.cs b/Engine/Source/Programs/Horde/Horde.Server/Server/GlobalConfig.cs
--- a/Engine/Source/Programs/Horde/Horde.Server/Server/GlobalConfig.cs
+++ b/Engine/Source/Programs/Horde/Horde.Server/Server/GlobalConfig.cs
@@ -31,6 +31,7 @@ using Horde.Server.Agents.Pools;
 using Horde.Server.Configuration;
 using Horde.Server.Dashboard;
 using Horde.Server.Devices;
+using Horde.Server.Perforce;
 using Horde.Server.Projects;
 using Horde.Server.Secrets;
 using Horde.Server.Storage;
