Horde: Remove tracing span from shared tickers. For long-running callbacks, this can create data which is difficult for Otel/DD to handle.

diff --git a/Engine/Source/Programs/Horde/Horde.Server/Utilities/Clock.cs b/Engine/Source/Programs/Horde/Horde.Server/Utilities/Clock.cs
--- a/Engine/Source/Programs/Horde/Horde.Server/Utilities/Clock.cs
+++ b/Engine/Source/Programs/Horde/Horde.Server/Utilities/Clock.cs
@@ -190,9 +190,6 @@ namespace HordeCommon
 			{
 				if (await sharedLock.AcquireAsync(interval, false))
 				{
-					using TelemetrySpan span = _tracer.StartActiveSpan(name);
-					span.SetAttribute("interval", interval.TotalSeconds);
-
 					await tickAsync(cancellationToken);
 				}
 			}
