Horde: Move debug endpoints for jobs into a separate controller, and nest them under /api/v1/debug/jobs. Generalize DebugEndpointAttribute to work with any controller.

#rnx

diff --git a/Engine/Source/Programs/Horde/Horde.Server/Jobs/DebugJobsController.cs b/Engine/Source/Programs/Horde/Horde.Server/Jobs/DebugJobsController.cs
new file mode 100644
--- /dev/null
+++ b/Engine/Source/Programs/Horde/Horde.Server/Jobs/DebugJobsController.cs
@@ -0,0 +1,459 @@
+// Copyright Epic Games, Inc. All Rights Reserved.
+
+using System;
+using System.Collections.Generic;
+using System.Linq;
+using System.Net;
+using System.Reflection;
+using System.Text;
+using System.Threading;
+using System.Threading.Tasks;
+using EpicGames.Horde.Jobs;
+using Horde.Server.Agents.Leases;
+using Horde.Server.Jobs.Graphs;
+using Horde.Server.Server;
+using Horde.Server.Streams;
+using Horde.Server.Utilities;
+using Microsoft.AspNetCore.Authorization;
+using Microsoft.AspNetCore.Http;
+using Microsoft.AspNetCore.Mvc;
+using Microsoft.Extensions.Logging;
+using Microsoft.Extensions.Options;
+
+namespace Horde.Server.Jobs
+{
+	/// <summary>
+	/// Debug functionality for jobs
+	/// </summary>
+	[ApiController]
+	[Authorize]
+	[DebugEndpoint]
+	[Tags("Debug")]
+	public class DebugJobsController : HordeControllerBase
+	{
+		private readonly JobService _jobService;
+		private readonly JobTaskSource _jobTaskSource;
+		private readonly ILeaseCollection _leaseCollection;
+		private readonly IOptionsSnapshot<GlobalConfig> _globalConfig;
+		private readonly ILogger<DebugJobsController> _logger;
+
+		/// <summary>
+		/// Constructor
+		/// </summary>
+		public DebugJobsController(
+			JobService jobService,
+			JobTaskSource jobTaskSource,
+			ILeaseCollection leaseCollection,
+			IOptionsSnapshot<GlobalConfig> globalConfig,
+			ILogger<DebugJobsController> logger)
+		{
+			_jobService = jobService;
+			_jobTaskSource = jobTaskSource;
+			_leaseCollection = leaseCollection;
+			_globalConfig = globalConfig;
+			_logger = logger;
+		}
+
+		/// <summary>
+		/// Returns diagnostic information about the current state of the queue
+		/// </summary>
+		/// <returns>Information about the queue</returns>
+		[HttpGet]
+		[Route("/api/v1/debug/queue")]
+		public ActionResult<object> GetQueueStatus()
+		{
+			if (!_globalConfig.Value.Authorize(ServerAclAction.Debug, User))
+			{
+				return Forbid(ServerAclAction.Debug);
+			}
+
+			return _jobTaskSource.GetStatus();
+		}
+
+		/// <summary>
+		/// Display a table listing each template with what job options are enabled
+		/// </summary>
+		/// <returns>Async task</returns>
+		[HttpGet]
+		[Route("/api/v1/debug/job-options")]
+		public ActionResult GetJobOptions([FromQuery] string? format = "html")
+		{
+			if (!_globalConfig.Value.Authorize(ServerAclAction.Debug, User))
+			{
+				return Forbid(ServerAclAction.Debug);
+			}
+
+			List<PropertyInfo> joProps = typeof(JobOptions).GetProperties(BindingFlags.Public | BindingFlags.Instance).OrderBy(x => x.Name).ToList();
+
+			if (format == "csv")
+			{
+				return GetJobOptionsAsCsv(joProps);
+			}
+
+			StringBuilder sb = new();
+
+			sb.AppendLine("<style>");
+			sb.AppendLine("body { font-family: 'Helvetica Neue', Helvetica, Arial, sans-serif; }");
+			sb.AppendLine("table { border-collapse: collapse; width: 100%; font-size: 12px; }");
+			sb.AppendLine("th, td { border: 1px solid black; text-align: left; padding: 8px; }");
+			sb.AppendLine("th { background-color: #f2f2f2; }");
+			sb.AppendLine("</style>");
+
+			sb.AppendLine("<h1>Job options enabled by stream + template</h1>");
+			sb.AppendLine("<table>");
+			sb.AppendLine("<thead><tr>");
+			sb.Append("<th>Stream</th>");
+			sb.Append("<th>Template</th>");
+			foreach (PropertyInfo prop in joProps)
+			{
+				sb.Append($"<th>{prop.Name}</th>");
+			}
+			sb.AppendLine("</tr></thead>");
+
+			foreach (StreamConfig sc in _globalConfig.Value.Streams)
+			{
+				foreach (TemplateRefConfig tpl in sc.Templates)
+				{
+					sb.AppendLine("<tr>");
+					sb.Append($"<td>{sc.Id}</td>");
+					sb.Append($"<td>{tpl.Id}</td>");
+					foreach (PropertyInfo prop in joProps)
+					{
+						sb.Append($"<td>{prop.GetValue(tpl.JobOptions)}</td>");
+					}
+					sb.AppendLine("</tr>");
+				}
+			}
+
+			sb.AppendLine("</table>");
+			return new ContentResult { ContentType = "text/html", StatusCode = (int)HttpStatusCode.OK, Content = sb.ToString() };
+		}
+
+		private ActionResult GetJobOptionsAsCsv(List<PropertyInfo> jobOptionsProps)
+		{
+			StringBuilder sb = new();
+
+			List<string> headers = new() { "Stream", "Template" };
+			headers.AddRange(jobOptionsProps.Select(prop => prop.Name));
+			sb.AppendLine(String.Join('\t', headers));
+
+			foreach (StreamConfig sc in _globalConfig.Value.Streams)
+			{
+				foreach (TemplateRefConfig tpl in sc.Templates)
+				{
+					List<string> row = new() { sc.Id.ToString(), tpl.Id.ToString() };
+					row.AddRange(jobOptionsProps.Select(prop => prop.GetValue(tpl.JobOptions)?.ToString() ?? ""));
+					sb.AppendLine(String.Join('\t', row));
+				}
+			}
+
+			return new ContentResult { ContentType = "text/csv", StatusCode = (int)HttpStatusCode.OK, Content = sb.ToString() };
+		}
+
+		record JobTiming(
+			string StreamId,
+			string TemplateId,
+			string Name,
+			IReadOnlySet<string> StepNames,
+			TimeSpan BatchSetupDuration,
+			TimeSpan BatchWorkDuration,
+			TimeSpan BatchTeardownDuration)
+		{
+			public JobTiming Merge(JobTiming jt)
+			{
+				if (StreamId != jt.StreamId || TemplateId != jt.TemplateId)
+				{
+					throw new ArgumentException("StreamId or TemplateId do not match");
+				}
+
+				HashSet<string> newNames = [.. StepNames.Union(jt.StepNames)];
+				return new JobTiming(StreamId, TemplateId, Name, newNames,
+					BatchSetupDuration + jt.BatchSetupDuration,
+					BatchWorkDuration + jt.BatchWorkDuration,
+					BatchTeardownDuration + jt.BatchTeardownDuration);
+			}
+		}
+
+		/// <summary>
+		/// Display a table listing each template with job timings (setup, work and teardown durations)
+		/// </summary>
+		/// <returns>Async task</returns>
+		[HttpGet]
+		[Route("/api/v1/debug/job-timings")]
+		public async Task<ActionResult> GetJobTimingsAsync(
+			[FromQuery] DateTimeOffset? minCreateTime = null,
+			[FromQuery] DateTimeOffset? maxCreateTime = null,
+			[FromQuery] bool onlySetupBuild = false,
+			[FromQuery] string? format = "html")
+		{
+			if (!_globalConfig.Value.Authorize(ServerAclAction.Debug, User))
+			{
+				return Forbid(ServerAclAction.Debug);
+			}
+
+			minCreateTime ??= DateTimeOffset.UtcNow.Subtract(TimeSpan.FromDays(1));
+			IReadOnlyList<IJob> jobs = await _jobService.FindJobsAsync(minCreateTime: minCreateTime, maxCreateTime: maxCreateTime);
+			List<JobTiming> allJobTimings = await CalculateJobTimingsAsync(jobs);
+
+			if (onlySetupBuild)
+			{
+				allJobTimings = allJobTimings.Where(x => x.StepNames.SetEquals(["Setup Build"])).ToList();
+			}
+
+			IReadOnlyList<JobTiming> jobTimingsByTemplate = GroupJobTimings(allJobTimings);
+			IEnumerable<JobTiming> sortedJobTimings = jobTimingsByTemplate.OrderBy(x => x.BatchSetupDuration).Reverse();
+
+			if (format == "csv")
+			{
+				return GetJobTimingsAsCsv(sortedJobTimings);
+			}
+
+			StringBuilder sb = new();
+
+			sb.AppendLine("<style>");
+			sb.AppendLine("body { font-family: 'Helvetica Neue', Helvetica, Arial, sans-serif; }");
+			sb.AppendLine("table { border-collapse: collapse; width: 100%; font-size: 12px; }");
+			sb.AppendLine("th, td { border: 1px solid black; text-align: left; padding: 8px; }");
+			sb.AppendLine("th { background-color: #f2f2f2; }");
+			sb.AppendLine("</style>");
+
+			sb.AppendLine("<h1>Job timings by template</h1>");
+			sb.AppendLine("<p>Durations specified in seconds.</p>");
+			sb.AppendLine("<table>");
+			sb.AppendLine("<thead><tr>");
+			sb.Append("<th>Stream</th>");
+			sb.Append("<th>Template</th>");
+			sb.Append("<th>Name</th>");
+			sb.Append("<th>Batch Setup</th>");
+			sb.Append("<th>Batch Work</th>");
+			sb.Append("<th>Batch Teardown</th>");
+			sb.Append("<th>Steps</th>");
+			sb.AppendLine("</tr></thead>");
+
+			foreach (JobTiming jt in sortedJobTimings)
+			{
+				sb.AppendLine("<tr>");
+				sb.Append($"<td>{jt.StreamId}</td>");
+				sb.Append($"<td>{jt.TemplateId}</td>");
+				sb.Append($"<td>{jt.Name}</td>");
+				sb.Append($"<td>{(int)jt.BatchSetupDuration.TotalSeconds}</td>");
+				sb.Append($"<td>{(int)jt.BatchWorkDuration.TotalSeconds}</td>");
+				sb.Append($"<td>{(int)jt.BatchTeardownDuration.TotalSeconds}</td>");
+				sb.Append($"<td>{String.Join(", ", jt.StepNames.Order())}</td>");
+				sb.AppendLine("</tr>");
+			}
+
+			sb.AppendLine("</table>");
+			return new ContentResult { ContentType = "text/html", StatusCode = (int)HttpStatusCode.OK, Content = sb.ToString() };
+		}
+
+		private static ActionResult GetJobTimingsAsCsv(IEnumerable<JobTiming> jobTimings)
+		{
+			StringBuilder sb = new();
+			sb.AppendJoin('\t', ["Stream", "Template", "Name", "Batch Setup", "Batch Work", "Batch Teardown", "Steps"]).AppendLine();
+
+			foreach (JobTiming jt in jobTimings)
+			{
+				sb.Append($"{jt.StreamId}\t");
+				sb.Append($"{jt.TemplateId}\t");
+				sb.Append($"{jt.Name}\t");
+				sb.Append($"{(int)jt.BatchSetupDuration.TotalSeconds}\t");
+				sb.Append($"{(int)jt.BatchWorkDuration.TotalSeconds}\t");
+				sb.Append($"{(int)jt.BatchTeardownDuration.TotalSeconds}\t");
+				sb.AppendJoin(',', jt.StepNames.Order());
+				sb.AppendLine();
+			}
+
+			return new ContentResult { ContentType = "text/csv", StatusCode = (int)HttpStatusCode.OK, Content = sb.ToString() };
+		}
+
+		private async Task<List<JobTiming>> CalculateJobTimingsAsync(IEnumerable<IJob> jobs)
+		{
+			List<JobTiming> timings = [];
+			foreach (IJob job in jobs)
+			{
+				IGraph graph = await _jobService.GetGraphAsync(job);
+
+				foreach (IJobStepBatch batch in job.Batches)
+				{
+					if (batch.State != JobStepBatchState.Complete || batch.StartTimeUtc == null || batch.FinishTimeUtc == null)
+					{
+						continue;
+					}
+					DateTime firstStepStartTime = DateTime.MaxValue;
+					DateTime lastStepFinishTime = DateTime.MinValue;
+
+					HashSet<string> stepNames = [];
+					foreach (IJobStep step in batch.Steps)
+					{
+						if (step.StartTimeUtc == null || step.FinishTimeUtc == null)
+						{
+							continue;
+						}
+
+						INode node = graph.GetNode(new NodeRef(batch.GroupIdx, step.NodeIdx));
+						firstStepStartTime = step.StartTimeUtc.Value < firstStepStartTime ? step.StartTimeUtc.Value : firstStepStartTime;
+						lastStepFinishTime = step.FinishTimeUtc.Value > lastStepFinishTime ? step.FinishTimeUtc.Value : lastStepFinishTime;
+						stepNames.Add(node.Name);
+					}
+
+					if (firstStepStartTime == DateTime.MaxValue || lastStepFinishTime == DateTime.MinValue)
+					{
+						continue;
+					}
+
+					TimeSpan batchSetupDuration = firstStepStartTime - batch.StartTimeUtc.Value;
+					TimeSpan batchWorkDuration = lastStepFinishTime - firstStepStartTime;
+					TimeSpan batchTeardownDuration = batch.FinishTimeUtc.Value - lastStepFinishTime;
+					timings.Add(new JobTiming(job.StreamId.ToString(), job.TemplateId.ToString(), job.Name, stepNames, batchSetupDuration, batchWorkDuration, batchTeardownDuration));
+				}
+			}
+
+			return timings;
+		}
+
+		private static IReadOnlyList<JobTiming> GroupJobTimings(IEnumerable<JobTiming> jobTimings)
+		{
+			Dictionary<string, JobTiming> groupedTimings = new();
+			foreach (JobTiming timing in jobTimings)
+			{
+				string key = $"{timing.StreamId}-{timing.TemplateId}";
+				if (!groupedTimings.TryGetValue(key, out JobTiming? groupTiming))
+				{
+					groupedTimings[key] = timing;
+				}
+				else
+				{
+					groupedTimings[key] = groupTiming.Merge(timing);
+				}
+			}
+
+			return groupedTimings.Values.ToList();
+		}
+
+		/// <summary>
+		/// Repairs any inconsistencies with jobs created in a particular time range 
+		/// </summary>
+		[HttpPost]
+		[Route("/api/v1/debug/repair-jobs")]
+		public async Task<ActionResult> RepairJobsAsync([FromQuery] DateTime? minTime, [FromQuery] DateTime? maxTime)
+		{
+			if (!_globalConfig.Value.Authorize(ServerAclAction.Debug, User))
+			{
+				return Forbid(ServerAclAction.Debug);
+			}
+
+			if (minTime == null || maxTime == null)
+			{
+				return BadRequest();
+			}
+
+			IReadOnlyList<IJob> jobs = await _jobService.FindJobsAsync(minCreateTime: minTime, maxCreateTime: maxTime);
+			foreach (IJob job in jobs)
+			{
+				StreamConfig? streamConfig;
+				if (_globalConfig.Value.TryGetStream(job.StreamId, out streamConfig))
+				{
+					_logger.LogInformation("Checking job {JobId}", job.Id);
+					await TryRepairJobAsync(streamConfig, job, HttpContext.RequestAborted);
+				}
+			}
+
+			_logger.LogInformation("Finished repair");
+			return Ok();
+		}
+
+		/// <summary>
+		/// Repairs any inconsistencies with a particular job 
+		/// </summary>
+		/// <param name="jobId">Id of the job to find</param>
+		[HttpPost]
+		[Route("/api/v1/debug/repair-job/{jobId}")]
+		public async Task<ActionResult> RepairJobAsync(JobId jobId)
+		{
+			IJob? job = await _jobService.GetJobAsync(jobId);
+			if (job == null || job.TemplateHash == null)
+			{
+				return NotFound(jobId);
+			}
+
+			StreamConfig? streamConfig;
+			if (!_globalConfig.Value.TryGetStream(job.StreamId, out streamConfig))
+			{
+				return NotFound(job.StreamId);
+			}
+			if (!_globalConfig.Value.Authorize(ServerAclAction.Debug, User))
+			{
+				return Forbid(ServerAclAction.Debug);
+			}
+
+			for (; ; )
+			{
+				IJob? newJob = await TryRepairJobAsync(streamConfig, job, HttpContext.RequestAborted);
+				if (newJob != null)
+				{
+					return Ok();
+				}
+
+				newJob = await _jobService.GetJobAsync(job.Id, HttpContext.RequestAborted);
+				if (newJob == null)
+				{
+					return NotFound();
+				}
+			}
+		}
+
+		async Task<IJob?> TryRepairJobAsync(StreamConfig streamConfig, IJob job, CancellationToken cancellationToken)
+		{
+			// Check the lease has not already completed. Workaround for issue where jobs collection came out of sync with leases collection due to Mongo timeouts while updating indexes.
+			List<JobStepBatchId> batchIds = job.Batches.Select(x => x.Id).ToList();
+			foreach (JobStepBatchId batchId in batchIds)
+			{
+				IJobStepBatch? batch;
+				if (job.TryGetBatch(batchId, out batch) && batch.LeaseId.HasValue && batch.State == JobStepBatchState.Running)
+				{
+					ILease? lease = await _leaseCollection.GetAsync(batch.LeaseId.Value, cancellationToken);
+					if (lease != null && lease.FinishTime.HasValue)
+					{
+						_logger.LogWarning("Job {JobId} batch {BatchId} is out of sync with lease {LeaseId}", job.Id, batch.Id, lease.Id);
+
+						IJob? newJob = await _jobService.UpdateBatchAsync(job, batch.Id, streamConfig, newState: JobStepBatchState.Complete, cancellationToken: cancellationToken);
+						if (newJob == null)
+						{
+							return null;
+						}
+
+						job = newJob;
+					}
+				}
+			}
+			return job;
+		}
+
+		/// <summary>
+		/// Forces an update of a job's batches to debug issues such as updating dependencies
+		/// </summary>
+		/// <returns></returns>
+		[HttpGet]
+		[Route("/api/v1/debug/batchupdate/{JobId}/{BatchId}")]
+		public async Task<ActionResult> DebugBatchUpdateAsync(string jobId, string batchId)
+		{
+			if (!_globalConfig.Value.Authorize(ServerAclAction.Debug, User))
+			{
+				return Forbid(ServerAclAction.Debug);
+			}
+
+			IJob? job = await _jobService.GetJobAsync(JobId.Parse(jobId));
+
+			if (job == null)
+			{
+				return NotFound();
+			}
+
+			await _jobService.TryUpdateBatchAsync(job, JobStepBatchId.Parse(batchId), newError: JobStepBatchError.None);
+
+			return Ok();
+
+		}
+	}
+}
diff --git a/Engine/Source/Programs/Horde/Horde.Server/Server/DebugController.cs b/Engine/Source/Programs/Horde/Horde.Server/Server/DebugController.cs
--- a/Engine/Source/Programs/Horde/Horde.Server/Server/DebugController.cs
+++ b/Engine/Source/Programs/Horde/Horde.Server/Server/DebugController.cs
@@ -1,15 +1,11 @@
 // Copyright Epic Games, Inc. All Rights Reserved.
 
-//#define ENABLE_PUBLIC_DEBUG_CONTROLLER
-#define ENABLE_SECURE_DEBUG_CONTROLLER
-
 using System;
 using System.Collections.Generic;
 using System.Diagnostics;
 using System.IO;
 using System.Linq;
 using System.Net;
-using System.Reflection;
 using System.Text;
 using System.Threading;
 using System.Threading.Tasks;
@@ -17,25 +13,19 @@ using System.Web;
 using EpicGames.Core;
 using EpicGames.Horde.Agents.Leases;
 using EpicGames.Horde.Compute;
-using EpicGames.Horde.Jobs;
 using EpicGames.Horde.Logs;
 using Google.Protobuf;
 using Horde.Common.Rpc;
-using Horde.Server.Agents.Leases;
 using Horde.Server.Agents.Relay;
 using Horde.Server.Configuration;
-using Horde.Server.Jobs;
-using Horde.Server.Jobs.Graphs;
 using Horde.Server.Logs;
 using Horde.Server.Projects;
-using Horde.Server.Streams;
 using Horde.Server.Utilities;
 using JetBrains.Profiler.SelfApi;
 using Microsoft.AspNetCore.Authorization;
 using Microsoft.AspNetCore.Http;
 using Microsoft.AspNetCore.Http.Features;
 using Microsoft.AspNetCore.Mvc;
-using Microsoft.AspNetCore.Mvc.Filters;
 using Microsoft.Extensions.Logging;
 using Microsoft.Extensions.Options;
 using MongoDB.Bson;
@@ -43,292 +33,39 @@ using MongoDB.Driver;
 
 namespace Horde.Server.Server
 {
-#if ENABLE_PUBLIC_DEBUG_CONTROLLER
-	/// <summary>
-	/// Public endpoints for the debug controller
-	/// </summary>
-	[ApiController]
-	public class PublicDebugController : ControllerBase
-	{
-		/// <summary>
-		/// The connection tracker service singleton
-		/// </summary>
-		RequestTrackerService RequestTrackerService;
-
-		IHostApplicationLifetime ApplicationLifetime;
-
-		IDogStatsd DogStatsd;
-
-		/// <summary>
-		/// Constructor
-		/// </summary>
-		/// <param name="RequestTrackerService"></param>
-		/// <param name="ApplicationLifetime"></param>
-		/// <param name="DogStatsd"></param>
-		public PublicDebugController(RequestTrackerService RequestTrackerService, IHostApplicationLifetime ApplicationLifetime, IDogStatsd DogStatsd)
-		{
-			RequestTrackerService = RequestTrackerService;
-			ApplicationLifetime = ApplicationLifetime;
-			DogStatsd = DogStatsd;
-		}
-
-		/// <summary>
-		/// Prints all the headers for the incoming request
-		/// </summary>
-		/// <returns>Http result</returns>
-		[HttpGet]
-		[Route("/api/v1/debug/headers")]
-		public IActionResult GetRequestHeaders()
-		{
-			StringBuilder Content = new StringBuilder();
-			Content.AppendLine("<html><body><pre>");
-			foreach (KeyValuePair<string, StringValues> Pair in HttpContext.Request.Headers)
-			{
-				foreach (string Value in Pair.Value)
-				{
-					Content.AppendLine(HttpUtility.HtmlEncode($"{Pair.Key}: {Value}"));
-				}
-			}
-			Content.Append("</pre></body></html>");
-			return new ContentResult { ContentType = "text/html", StatusCode = (int)HttpStatusCode.OK, Content = Content.ToString() };
-		}
-
-		/// <summary>
-		/// Waits specified number of milliseconds and then returns a response
-		/// Used for testing timeouts proxy settings.
-		/// </summary>
-		/// <returns>Http result</returns>
-		[HttpGet]
-		[Route("/api/v1/debug/wait")]
-		public async Task<ActionResult> GetAndWait([FromQuery] int WaitTimeMs = 1000)
-		{
-			await Task.Delay(WaitTimeMs);
-			string Content = $"Waited {WaitTimeMs} ms. " + new Random().Next(0, 10000000);
-			return new ContentResult { ContentType = "text/plain", StatusCode = (int)HttpStatusCode.OK, Content = Content };
-		}
-
-		/// <summary>
-		/// Waits specified number of milliseconds and then throws an exception
-		/// Used for testing graceful shutdown and interruption of outstanding requests.
-		/// </summary>
-		/// <returns>Http result</returns>
-		[HttpGet]
-		[Route("/api/v1/debug/exception")]
-		public async Task<ActionResult> ThrowException([FromQuery] int WaitTimeMs = 0)
-		{
-			await Task.Delay(WaitTimeMs);
-			throw new Exception("Test exception triggered by debug controller!");
-		}
-
-		/// <summary>
-		/// Trigger an increment of a DogStatsd metric
-		/// </summary>
-		/// <returns>Http result</returns>
-		[HttpGet]
-		[Route("/api/v1/debug/metric")]
-		public ActionResult TriggerMetric([FromQuery] int Value = 10)
-		{
-			DogStatsd.Increment("hordeMetricTest", Value);
-			return Ok("Incremented metric 'hordeMetricTest' Type: " + DogStatsd.GetType());
-		}
-
-		/// <summary>
-		/// Display metrics related to the .NET runtime
-		/// </summary>
-		/// <returns>Http result</returns>
-		[HttpGet]
-		[Route("/api/v1/debug/dotnet-metrics")]
-		public ActionResult DotNetMetrics()
-		{
-			ThreadPool.GetMaxThreads(out int MaxWorkerThreads, out int MaxIoThreads);
-			ThreadPool.GetAvailableThreads(out int FreeWorkerThreads, out int FreeIoThreads);
-			ThreadPool.GetMinThreads(out int MinWorkerThreads, out int MinIoThreads);
-
-			int BusyIoThreads = MaxIoThreads - FreeIoThreads;
-			int BusyWorkerThreads = MaxWorkerThreads - FreeWorkerThreads;
-
-			StringBuilder Content = new StringBuilder();
-			Content.AppendLine("Threads:");
-			Content.AppendLine("-------------------------------------------------------------");
-			Content.AppendLine("Worker busy={0,-5} free={1,-5} min={2,-5} max={3,-5}", BusyWorkerThreads, FreeWorkerThreads, MinWorkerThreads, MaxWorkerThreads);
-			Content.AppendLine("  IOCP busy={0,-5} free={1,-5} min={2,-5} max={3,-5}", BusyIoThreads, FreeIoThreads, MinIoThreads, MaxWorkerThreads);
-
-
-			NumberFormatInfo Nfi = (NumberFormatInfo)CultureInfo.InvariantCulture.NumberFormat.Clone();
-			Nfi.NumberGroupSeparator = " ";
-
-			string FormatBytes(long Number)
-			{
-				return (Number / 1024 / 1024).ToString("#,0", Nfi) + " MB";
-			}
-
-			GCMemoryInfo GcMemoryInfo = GC.GetGCMemoryInfo();
-			Content.AppendLine("");
-			Content.AppendLine("");
-			Content.AppendLine("Garbage collection (GC):");
-			Content.AppendLine("-------------------------------------------------------------");
-			Content.AppendLine("              Latency mode: " + GCSettings.LatencyMode);
-			Content.AppendLine("              Is server GC: " + GCSettings.IsServerGC);
-			Content.AppendLine("              Total memory: " + FormatBytes(GC.GetTotalMemory(false)));
-			Content.AppendLine("           Total allocated: " + FormatBytes(GC.GetTotalAllocatedBytes(false)));
-			Content.AppendLine("                 Heap size: " + FormatBytes(GcMemoryInfo.HeapSizeBytes));
-			Content.AppendLine("                Fragmented: " + FormatBytes(GcMemoryInfo.FragmentedBytes));
-			Content.AppendLine("               Memory Load: " + FormatBytes(GcMemoryInfo.MemoryLoadBytes));
-			Content.AppendLine("    Total available memory: " + FormatBytes(GcMemoryInfo.TotalAvailableMemoryBytes));
-			Content.AppendLine("High memory load threshold: " + FormatBytes(GcMemoryInfo.HighMemoryLoadThresholdBytes));
-
-			return Ok(Content.ToString());
-		}
-
-		/// <summary>
-		/// Force a full GC of all generations
-		/// </summary>
-		/// <returns>Prints time taken in ms</returns>
-		[HttpGet]
-		[Route("/api/v1/debug/force-gc")]
-		public ActionResult ForceTriggerGc()
-		{
-			Stopwatch Timer = new Stopwatch();
-			Timer.Start();
-			GC.Collect();
-			Timer.Stop();
-			return Ok($"Time taken: {Timer.Elapsed.TotalMilliseconds} ms");
-		}
-
-		/// <summary>
-		/// Lists requests in progress
-		/// </summary>
-		/// <returns>HTML result</returns>
-		[HttpGet]
-		[Route("/api/v1/debug/requests-in-progress")]
-		public ActionResult GetRequestsInProgress()
-		{
-			StringBuilder Content = new StringBuilder();
-			Content.AppendLine("<html><body>");
-			Content.AppendLine("<h1>Requests in progress</h1>");
-			Content.AppendLine("<table border=\"1\">");
-			Content.AppendLine("<tr>");
-			Content.AppendLine("<th>Request Trace ID</th>");
-			Content.AppendLine("<th>Path</th>");
-			Content.AppendLine("<th>Started At</th>");
-			Content.AppendLine("<th>Age</th>");
-			Content.AppendLine("</tr>");
-
-			List<KeyValuePair<string, TrackedRequest>> Requests = RequestTrackerService.GetRequestsInProgress().ToList();
-			Requests.Sort((A, B) => A.Value.StartedAt.CompareTo(B.Value.StartedAt));
-
-			foreach (KeyValuePair<string, TrackedRequest> Entry in Requests)
-			{
-				Content.Append("<tr>");
-				Content.AppendLine($"<td>{Entry.Key}</td>");
-				Content.AppendLine($"<td>{Entry.Value.Request.Path}</td>");
-				Content.AppendLine($"<td>{Entry.Value.StartedAt}</td>");
-				Content.AppendLine($"<td>{Entry.Value.GetTimeSinceStartInMs()} ms</td>");
-				Content.Append("</tr>");
-			}
-			Content.Append("</table>\n</body>\n</html>");
-
-			return new ContentResult { ContentType = "text/html", StatusCode = (int)HttpStatusCode.OK, Content = Content.ToString() };
-		}
-
-		/*
-		// Used during development only
-		[HttpGet]
-		[Route("/api/v1/debug/stop")]
-		public ActionResult StopApp()
-		{
-			Task.Run(async () =>
-			{
-				await Task.Delay(100);
-				ApplicationLifetime.StopApplication();
-			});
-			
-			return new ContentResult { ContentType = "text/plain", StatusCode = (int)HttpStatusCode.OK, Content = "App stopping..." };
-		}
-		/**/
-	}
-#endif
-
-	/// <summary>
-	/// Only requests to attached controller to pass if debug endpoint is enabled in settings
-	/// Adds extra security for not enabling these admin endpoints by accident.
-	/// </summary>
-	[AttributeUsage(AttributeTargets.Class)]
-	public sealed class DebugEndpointCheckAttribute : Attribute, IActionFilter
-	{
-		/// <inheritdoc />
-		public void OnActionExecuting(ActionExecutingContext context)
-		{
-			if (context.Controller is SecureDebugController controller)
-			{
-				if (!controller.ServerSettings.Value.EnableDebugEndpoints)
-				{
-					context.Result = new ForbidResult();
-				}
-			}
-			else
-			{
-				// Assume forbidden if controller is not resolved
-				context.Result = new ForbidResult();
-			}
-		}
-
-		/// <inheritdoc />
-		public void OnActionExecuted(ActionExecutedContext context)
-		{
-		}
-	}
-
-#if ENABLE_SECURE_DEBUG_CONTROLLER
 	/// <summary>
 	/// Controller managing account status
 	/// </summary>
 	[ApiController]
 	[Authorize]
-	[DebugEndpointCheck]
+	[DebugEndpoint]
 	[Tags("Debug")]
-	public class SecureDebugController : HordeControllerBase
+	public class DebugController : HordeControllerBase
 	{
 		private static readonly Random s_random = new();
 
-		/// <summary>
-		/// Server settings (exposed for debug endpoint check filter)
-		/// </summary>
-		public IOptionsSnapshot<ServerSettings> ServerSettings { get; private set; }
-
 		private readonly MongoService _mongoService;
 		private readonly ConfigService _configService;
 		private readonly AgentRelayService _agentRelayService;
-		private readonly JobService _jobService;
-		private readonly JobTaskSource _jobTaskSource;
 		private readonly ILogCollection _logCollection;
-		private readonly ILeaseCollection _leaseCollection;
 		private readonly IOptionsSnapshot<GlobalConfig> _globalConfig;
-		private readonly ILogger<SecureDebugController> _logger;
+		private readonly ILogger<DebugController> _logger;
 
 		/// <summary>
 		/// Constructor
 		/// </summary>
-		public SecureDebugController(
+		public DebugController(
 			MongoService mongoService,
 			ConfigService configService,
 			AgentRelayService agentRelayService,
-			JobService jobService,
-			JobTaskSource jobTaskSource,
 			ILogCollection logCollection,
-			ILeaseCollection leaseCollection,
-			IOptionsSnapshot<ServerSettings> serverSettings,
 			IOptionsSnapshot<GlobalConfig> globalConfig,
-			ILogger<SecureDebugController> logger)
+			ILogger<DebugController> logger)
 		{
 			_mongoService = mongoService;
 			_configService = configService;
-			_jobService = jobService;
 			_agentRelayService = agentRelayService;
-			_jobTaskSource = jobTaskSource;
 			_logCollection = logCollection;
-			_leaseCollection = leaseCollection;
-			ServerSettings = serverSettings;
 			_globalConfig = globalConfig;
 			_logger = logger;
 		}
@@ -359,22 +96,6 @@ namespace Horde.Server.Server
 			return new ContentResult { ContentType = "text/html", StatusCode = (int)HttpStatusCode.OK, Content = content.ToString() };
 		}
 
-		/// <summary>
-		/// Returns diagnostic information about the current state of the queue
-		/// </summary>
-		/// <returns>Information about the queue</returns>
-		[HttpGet]
-		[Route("/api/v1/debug/queue")]
-		public ActionResult<object> GetQueueStatus()
-		{
-			if (!_globalConfig.Value.Authorize(ServerAclAction.Debug, User))
-			{
-				return Forbid(ServerAclAction.Debug);
-			}
-
-			return _jobTaskSource.GetStatus();
-		}
-
 		/// <summary>
 		/// Converts all legacy pools into config entries
 		/// </summary>
@@ -617,366 +338,6 @@ namespace Horde.Server.Server
 			return log.ApplyFilter(filter);
 		}
 
-		/// <summary>
-		/// Display a table listing each template with what job options are enabled
-		/// </summary>
-		/// <returns>Async task</returns>
-		[HttpGet]
-		[Route("/api/v1/debug/job-options")]
-		public ActionResult GetJobOptions([FromQuery] string? format = "html")
-		{
-			if (!_globalConfig.Value.Authorize(ServerAclAction.Debug, User))
-			{
-				return Forbid(ServerAclAction.Debug);
-			}
-
-			List<PropertyInfo> joProps = typeof(JobOptions).GetProperties(BindingFlags.Public | BindingFlags.Instance).OrderBy(x => x.Name).ToList();
-
-			if (format == "csv")
-			{
-				return GetJobOptionsAsCsv(joProps);
-			}
-
-			StringBuilder sb = new();
-
-			sb.AppendLine("<style>");
-			sb.AppendLine("body { font-family: 'Helvetica Neue', Helvetica, Arial, sans-serif; }");
-			sb.AppendLine("table { border-collapse: collapse; width: 100%; font-size: 12px; }");
-			sb.AppendLine("th, td { border: 1px solid black; text-align: left; padding: 8px; }");
-			sb.AppendLine("th { background-color: #f2f2f2; }");
-			sb.AppendLine("</style>");
-
-			sb.AppendLine("<h1>Job options enabled by stream + template</h1>");
-			sb.AppendLine("<table>");
-			sb.AppendLine("<thead><tr>");
-			sb.Append("<th>Stream</th>");
-			sb.Append("<th>Template</th>");
-			foreach (PropertyInfo prop in joProps)
-			{
-				sb.Append($"<th>{prop.Name}</th>");
-			}
-			sb.AppendLine("</tr></thead>");
-
-			foreach (StreamConfig sc in _globalConfig.Value.Streams)
-			{
-				foreach (TemplateRefConfig tpl in sc.Templates)
-				{
-					sb.AppendLine("<tr>");
-					sb.Append($"<td>{sc.Id}</td>");
-					sb.Append($"<td>{tpl.Id}</td>");
-					foreach (PropertyInfo prop in joProps)
-					{
-						sb.Append($"<td>{prop.GetValue(tpl.JobOptions)}</td>");
-					}
-					sb.AppendLine("</tr>");
-				}
-			}
-
-			sb.AppendLine("</table>");
-			return new ContentResult { ContentType = "text/html", StatusCode = (int)HttpStatusCode.OK, Content = sb.ToString() };
-		}
-
-		private ActionResult GetJobOptionsAsCsv(List<PropertyInfo> jobOptionsProps)
-		{
-			StringBuilder sb = new();
-
-			List<string> headers = new() { "Stream", "Template" };
-			headers.AddRange(jobOptionsProps.Select(prop => prop.Name));
-			sb.AppendLine(String.Join('\t', headers));
-
-			foreach (StreamConfig sc in _globalConfig.Value.Streams)
-			{
-				foreach (TemplateRefConfig tpl in sc.Templates)
-				{
-					List<string> row = new() { sc.Id.ToString(), tpl.Id.ToString() };
-					row.AddRange(jobOptionsProps.Select(prop => prop.GetValue(tpl.JobOptions)?.ToString() ?? ""));
-					sb.AppendLine(String.Join('\t', row));
-				}
-			}
-
-			return new ContentResult { ContentType = "text/csv", StatusCode = (int)HttpStatusCode.OK, Content = sb.ToString() };
-		}
-
-		record JobTiming(
-			string StreamId,
-			string TemplateId,
-			string Name,
-			IReadOnlySet<string> StepNames,
-			TimeSpan BatchSetupDuration,
-			TimeSpan BatchWorkDuration,
-			TimeSpan BatchTeardownDuration)
-		{
-			public JobTiming Merge(JobTiming jt)
-			{
-				if (StreamId != jt.StreamId || TemplateId != jt.TemplateId)
-				{
-					throw new ArgumentException("StreamId or TemplateId do not match");
-				}
-
-				HashSet<string> newNames = [.. StepNames.Union(jt.StepNames)];
-				return new JobTiming(StreamId, TemplateId, Name, newNames,
-					BatchSetupDuration + jt.BatchSetupDuration,
-					BatchWorkDuration + jt.BatchWorkDuration,
-					BatchTeardownDuration + jt.BatchTeardownDuration);
-			}
-		}
-
-		/// <summary>
-		/// Display a table listing each template with job timings (setup, work and teardown durations)
-		/// </summary>
-		/// <returns>Async task</returns>
-		[HttpGet]
-		[Route("/api/v1/debug/job-timings")]
-		public async Task<ActionResult> GetJobTimingsAsync(
-			[FromQuery] DateTimeOffset? minCreateTime = null,
-			[FromQuery] DateTimeOffset? maxCreateTime = null,
-			[FromQuery] bool onlySetupBuild = false,
-			[FromQuery] string? format = "html")
-		{
-			if (!_globalConfig.Value.Authorize(ServerAclAction.Debug, User))
-			{
-				return Forbid(ServerAclAction.Debug);
-			}
-
-			minCreateTime ??= DateTimeOffset.UtcNow.Subtract(TimeSpan.FromDays(1));
-			IReadOnlyList<IJob> jobs = await _jobService.FindJobsAsync(minCreateTime: minCreateTime, maxCreateTime: maxCreateTime);
-			List<JobTiming> allJobTimings = await CalculateJobTimingsAsync(jobs);
-
-			if (onlySetupBuild)
-			{
-				allJobTimings = allJobTimings.Where(x => x.StepNames.SetEquals(["Setup Build"])).ToList();
-			}
-
-			IReadOnlyList<JobTiming> jobTimingsByTemplate = GroupJobTimings(allJobTimings);
-			IEnumerable<JobTiming> sortedJobTimings = jobTimingsByTemplate.OrderBy(x => x.BatchSetupDuration).Reverse();
-
-			if (format == "csv")
-			{
-				return GetJobTimingsAsCsv(sortedJobTimings);
-			}
-
-			StringBuilder sb = new();
-
-			sb.AppendLine("<style>");
-			sb.AppendLine("body { font-family: 'Helvetica Neue', Helvetica, Arial, sans-serif; }");
-			sb.AppendLine("table { border-collapse: collapse; width: 100%; font-size: 12px; }");
-			sb.AppendLine("th, td { border: 1px solid black; text-align: left; padding: 8px; }");
-			sb.AppendLine("th { background-color: #f2f2f2; }");
-			sb.AppendLine("</style>");
-
-			sb.AppendLine("<h1>Job timings by template</h1>");
-			sb.AppendLine("<p>Durations specified in seconds.</p>");
-			sb.AppendLine("<table>");
-			sb.AppendLine("<thead><tr>");
-			sb.Append("<th>Stream</th>");
-			sb.Append("<th>Template</th>");
-			sb.Append("<th>Name</th>");
-			sb.Append("<th>Batch Setup</th>");
-			sb.Append("<th>Batch Work</th>");
-			sb.Append("<th>Batch Teardown</th>");
-			sb.Append("<th>Steps</th>");
-			sb.AppendLine("</tr></thead>");
-
-			foreach (JobTiming jt in sortedJobTimings)
-			{
-				sb.AppendLine("<tr>");
-				sb.Append($"<td>{jt.StreamId}</td>");
-				sb.Append($"<td>{jt.TemplateId}</td>");
-				sb.Append($"<td>{jt.Name}</td>");
-				sb.Append($"<td>{(int)jt.BatchSetupDuration.TotalSeconds}</td>");
-				sb.Append($"<td>{(int)jt.BatchWorkDuration.TotalSeconds}</td>");
-				sb.Append($"<td>{(int)jt.BatchTeardownDuration.TotalSeconds}</td>");
-				sb.Append($"<td>{String.Join(", ", jt.StepNames.Order())}</td>");
-				sb.AppendLine("</tr>");
-			}
-
-			sb.AppendLine("</table>");
-			return new ContentResult { ContentType = "text/html", StatusCode = (int)HttpStatusCode.OK, Content = sb.ToString() };
-		}
-
-		private static ActionResult GetJobTimingsAsCsv(IEnumerable<JobTiming> jobTimings)
-		{
-			StringBuilder sb = new();
-			sb.AppendJoin('\t', ["Stream", "Template", "Name", "Batch Setup", "Batch Work", "Batch Teardown", "Steps"]).AppendLine();
-
-			foreach (JobTiming jt in jobTimings)
-			{
-				sb.Append($"{jt.StreamId}\t");
-				sb.Append($"{jt.TemplateId}\t");
-				sb.Append($"{jt.Name}\t");
-				sb.Append($"{(int)jt.BatchSetupDuration.TotalSeconds}\t");
-				sb.Append($"{(int)jt.BatchWorkDuration.TotalSeconds}\t");
-				sb.Append($"{(int)jt.BatchTeardownDuration.TotalSeconds}\t");
-				sb.AppendJoin(',', jt.StepNames.Order());
-				sb.AppendLine();
-			}
-
-			return new ContentResult { ContentType = "text/csv", StatusCode = (int)HttpStatusCode.OK, Content = sb.ToString() };
-		}
-
-		private async Task<List<JobTiming>> CalculateJobTimingsAsync(IEnumerable<IJob> jobs)
-		{
-			List<JobTiming> timings = [];
-			foreach (IJob job in jobs)
-			{
-				IGraph graph = await _jobService.GetGraphAsync(job);
-
-				foreach (IJobStepBatch batch in job.Batches)
-				{
-					if (batch.State != JobStepBatchState.Complete || batch.StartTimeUtc == null || batch.FinishTimeUtc == null)
-					{
-						continue;
-					}
-					DateTime firstStepStartTime = DateTime.MaxValue;
-					DateTime lastStepFinishTime = DateTime.MinValue;
-
-					HashSet<string> stepNames = [];
-					foreach (IJobStep step in batch.Steps)
-					{
-						if (step.StartTimeUtc == null || step.FinishTimeUtc == null)
-						{
-							continue;
-						}
-
-						INode node = graph.GetNode(new NodeRef(batch.GroupIdx, step.NodeIdx));
-						firstStepStartTime = step.StartTimeUtc.Value < firstStepStartTime ? step.StartTimeUtc.Value : firstStepStartTime;
-						lastStepFinishTime = step.FinishTimeUtc.Value > lastStepFinishTime ? step.FinishTimeUtc.Value : lastStepFinishTime;
-						stepNames.Add(node.Name);
-					}
-
-					if (firstStepStartTime == DateTime.MaxValue || lastStepFinishTime == DateTime.MinValue)
-					{
-						continue;
-					}
-
-					TimeSpan batchSetupDuration = firstStepStartTime - batch.StartTimeUtc.Value;
-					TimeSpan batchWorkDuration = lastStepFinishTime - firstStepStartTime;
-					TimeSpan batchTeardownDuration = batch.FinishTimeUtc.Value - lastStepFinishTime;
-					timings.Add(new JobTiming(job.StreamId.ToString(), job.TemplateId.ToString(), job.Name, stepNames, batchSetupDuration, batchWorkDuration, batchTeardownDuration));
-				}
-			}
-
-			return timings;
-		}
-
-		private static IReadOnlyList<JobTiming> GroupJobTimings(IEnumerable<JobTiming> jobTimings)
-		{
-			Dictionary<string, JobTiming> groupedTimings = new();
-			foreach (JobTiming timing in jobTimings)
-			{
-				string key = $"{timing.StreamId}-{timing.TemplateId}";
-				if (!groupedTimings.TryGetValue(key, out JobTiming? groupTiming))
-				{
-					groupedTimings[key] = timing;
-				}
-				else
-				{
-					groupedTimings[key] = groupTiming.Merge(timing);
-				}
-			}
-
-			return groupedTimings.Values.ToList();
-		}
-
-		/// <summary>
-		/// Repairs any inconsistencies with jobs created in a particular time range 
-		/// </summary>
-		[HttpPost]
-		[Route("/api/v1/debug/repair-jobs")]
-		public async Task<ActionResult> RepairJobsAsync([FromQuery] DateTime? minTime, [FromQuery] DateTime? maxTime)
-		{
-			if (!_globalConfig.Value.Authorize(ServerAclAction.Debug, User))
-			{
-				return Forbid(ServerAclAction.Debug);
-			}
-
-			if (minTime == null || maxTime == null)
-			{
-				return BadRequest();
-			}
-
-			IReadOnlyList<IJob> jobs = await _jobService.FindJobsAsync(minCreateTime: minTime, maxCreateTime: maxTime);
-			foreach (IJob job in jobs)
-			{
-				StreamConfig? streamConfig;
-				if (_globalConfig.Value.TryGetStream(job.StreamId, out streamConfig))
-				{
-					_logger.LogInformation("Checking job {JobId}", job.Id);
-					await TryRepairJobAsync(streamConfig, job, HttpContext.RequestAborted);
-				}
-			}
-
-			_logger.LogInformation("Finished repair");
-			return Ok();
-		}
-
-		/// <summary>
-		/// Repairs any inconsistencies with a particular job 
-		/// </summary>
-		/// <param name="jobId">Id of the job to find</param>
-		[HttpPost]
-		[Route("/api/v1/debug/repair-job/{jobId}")]
-		public async Task<ActionResult> RepairJobAsync(JobId jobId)
-		{
-			IJob? job = await _jobService.GetJobAsync(jobId);
-			if (job == null || job.TemplateHash == null)
-			{
-				return NotFound(jobId);
-			}
-
-			StreamConfig? streamConfig;
-			if (!_globalConfig.Value.TryGetStream(job.StreamId, out streamConfig))
-			{
-				return NotFound(job.StreamId);
-			}
-			if (!_globalConfig.Value.Authorize(ServerAclAction.Debug, User))
-			{
-				return Forbid(ServerAclAction.Debug);
-			}
-
-			for (; ; )
-			{
-				IJob? newJob = await TryRepairJobAsync(streamConfig, job, HttpContext.RequestAborted);
-				if (newJob != null)
-				{
-					return Ok();
-				}
-
-				newJob = await _jobService.GetJobAsync(job.Id, HttpContext.RequestAborted);
-				if (newJob == null)
-				{
-					return NotFound();
-				}
-			}
-		}
-
-		async Task<IJob?> TryRepairJobAsync(StreamConfig streamConfig, IJob job, CancellationToken cancellationToken)
-		{
-			// Check the lease has not already completed. Workaround for issue where jobs collection came out of sync with leases collection due to Mongo timeouts while updating indexes.
-			List<JobStepBatchId> batchIds = job.Batches.Select(x => x.Id).ToList();
-			foreach (JobStepBatchId batchId in batchIds)
-			{
-				IJobStepBatch? batch;
-				if (job.TryGetBatch(batchId, out batch) && batch.LeaseId.HasValue && batch.State == JobStepBatchState.Running)
-				{
-					ILease? lease = await _leaseCollection.GetAsync(batch.LeaseId.Value, cancellationToken);
-					if (lease != null && lease.FinishTime.HasValue)
-					{
-						_logger.LogWarning("Job {JobId} batch {BatchId} is out of sync with lease {LeaseId}", job.Id, batch.Id, lease.Id);
-
-						IJob? newJob = await _jobService.UpdateBatchAsync(job, batch.Id, streamConfig, newState: JobStepBatchState.Complete, cancellationToken: cancellationToken);
-						if (newJob == null)
-						{
-							return null;
-						}
-
-						job = newJob;
-					}
-				}
-			}
-			return job;
-		}
-
 		/// <summary>
 		/// Populate the database with test data
 		/// </summary>
@@ -1123,33 +484,5 @@ namespace Horde.Server.Server
 			string stringArg = "hello";
 			throw new Exception($"Message: numberArg:{numberArg}, stringArg:{stringArg}");
 		}
-
-		/// <summary>
-		/// Forces an update of a job's batches to debug issues such as updating dependencies
-		/// </summary>
-		/// <returns></returns>
-		[HttpGet]
-		[Route("/api/v1/debug/batchupdate/{JobId}/{BatchId}")]
-		public async Task<ActionResult> DebugBatchUpdateAsync(string jobId, string batchId)
-		{
-			if (!_globalConfig.Value.Authorize(ServerAclAction.Debug, User))
-			{
-				return Forbid(ServerAclAction.Debug);
-			}
-
-			IJob? job = await _jobService.GetJobAsync(JobId.Parse(jobId));
-
-			if (job == null)
-			{
-				return NotFound();
-			}
-
-			await _jobService.TryUpdateBatchAsync(job, JobStepBatchId.Parse(batchId), newError: JobStepBatchError.None);
-
-			return Ok();
-
-		}
 	}
 }
-
-#endif
diff --git a/Engine/Source/Programs/Horde/Horde.Server/Utilities/Clock.cs b/Engine/Source/Programs/Horde/Horde.Server/Utilities/Clock.cs
--- a/Engine/Source/Programs/Horde/Horde.Server/Utilities/Clock.cs
+++ b/Engine/Source/Programs/Horde/Horde.Server/Utilities/Clock.cs
@@ -10,7 +10,6 @@ using Horde.Server;
 using Horde.Server.Server;
 using Microsoft.Extensions.Logging;
 using Microsoft.Extensions.Options;
-using OpenTelemetry.Trace;
 using StackExchange.Redis;
 using TimeZoneConverter;
 
@@ -130,7 +129,6 @@ namespace HordeCommon
 		}
 
 		readonly RedisService _redis;
-		readonly Tracer _tracer;
 		readonly TimeZoneInfo _timeZone;
 		readonly List<TickerImpl> _tickers = new List<TickerImpl>();
 
@@ -143,10 +141,9 @@ namespace HordeCommon
 		/// <summary>
 		/// Constructor
 		/// </summary>
-		public Clock(RedisService redis, Tracer tracer, IOptions<ServerSettings> settings)
+		public Clock(RedisService redis, IOptions<ServerSettings> settings)
 		{
 			_redis = redis;
-			_tracer = tracer;
 
 			string? timeZoneName = settings.Value.ScheduleTimeZone;
 			_timeZone = (timeZoneName == null) ? TimeZoneInfo.Local : TZConvert.GetTimeZoneInfo(timeZoneName);
@@ -176,10 +173,10 @@ namespace HordeCommon
 		public ITicker AddSharedTicker(string name, TimeSpan delay, Func<CancellationToken, ValueTask> tickAsync, ILogger logger)
 		{
 			RedisKey key = new RedisKey($"tick/{name}");
-			return this.AddTicker(name, delay / 4, token => TriggerSharedAsync(name, key, delay, tickAsync, token), logger);
+			return this.AddTicker(name, delay / 4, token => TriggerSharedAsync(key, delay, tickAsync, token), logger);
 		}
 
-		async ValueTask TriggerSharedAsync(string name, RedisKey key, TimeSpan interval, Func<CancellationToken, ValueTask> tickAsync, CancellationToken cancellationToken)
+		async ValueTask TriggerSharedAsync(RedisKey key, TimeSpan interval, Func<CancellationToken, ValueTask> tickAsync, CancellationToken cancellationToken)
 		{
 			if (_redis.ReadOnlyMode)
 			{
diff --git a/Engine/Source/Programs/Horde/Horde.Server/Utilities/DebugEndpointAttribute.cs b/Engine/Source/Programs/Horde/Horde.Server/Utilities/DebugEndpointAttribute.cs
new file mode 100644
--- /dev/null
+++ b/Engine/Source/Programs/Horde/Horde.Server/Utilities/DebugEndpointAttribute.cs
@@ -0,0 +1,33 @@
+// Copyright Epic Games, Inc. All Rights Reserved.
+
+using System;
+using Microsoft.AspNetCore.Mvc;
+using Microsoft.AspNetCore.Mvc.Filters;
+using Microsoft.Extensions.DependencyInjection;
+using Microsoft.Extensions.Options;
+
+namespace Horde.Server.Utilities
+{
+	/// <summary>
+	/// Only requests to attached controller to pass if debug endpoint is enabled in settings
+	/// Adds extra security for not enabling these admin endpoints by accident.
+	/// </summary>
+	[AttributeUsage(AttributeTargets.Class)]
+	public sealed class DebugEndpointAttribute : Attribute, IActionFilter
+	{
+		/// <inheritdoc />
+		public void OnActionExecuting(ActionExecutingContext context)
+		{
+			IOptionsMonitor<ServerSettings> settings = context.HttpContext.RequestServices.GetRequiredService<IOptionsMonitor<ServerSettings>>();
+			if (!settings.CurrentValue.EnableDebugEndpoints)
+			{
+				context.Result = new ForbidResult();
+			}
+		}
+
+		/// <inheritdoc />
+		public void OnActionExecuted(ActionExecutedContext context)
+		{
+		}
+	}
+}
