Horde: Expose ShutdownIfDisabledGracePeriod in get pool response

#rnx

diff --git a/Engine/Source/Programs/Horde/Horde.Server/Agents/Pools/PoolMessages.cs b/Engine/Source/Programs/Horde/Horde.Server/Agents/Pools/PoolMessages.cs
--- a/Engine/Source/Programs/Horde/Horde.Server/Agents/Pools/PoolMessages.cs
+++ b/Engine/Source/Programs/Horde/Horde.Server/Agents/Pools/PoolMessages.cs
@@ -1,5 +1,6 @@
 // Copyright Epic Games, Inc. All Rights Reserved.
 
+using System;
 using System.Collections.Generic;
 using System.ComponentModel.DataAnnotations;
 using System.Linq;
@@ -509,6 +510,11 @@ namespace Horde.Server.Agents.Pools
 		/// The minimum number of idle agents in this pool, if autoscaling is enabled
 		/// </summary>
 		public int? NumReserveAgents { get; set; }
+		
+		/// <summary>
+		/// Time to wait before shutting down an agent that has been disabled
+		/// </summary>
+		public TimeSpan? ShutdownIfDisabledGracePeriod { get; set; }
 
 		/// <summary>
 		/// List of workspaces that this agent contains
@@ -549,6 +555,7 @@ namespace Horde.Server.Agents.Pools
 			ComputeQueueAwsMetricSettings = pool.ComputeQueueAwsMetricSettings == null ? null : new ComputeQueueAwsMetricSettingsMessage(pool.ComputeQueueAwsMetricSettings);
 			MinAgents = pool.MinAgents;
 			NumReserveAgents = pool.NumReserveAgents;
+			ShutdownIfDisabledGracePeriod = pool.ShutdownIfDisabledGracePeriod;
 			Workspaces = pool.Workspaces.Select(x => AgentsController.CreateGetAgentWorkspaceResponse(x)).ToList();
 			AutoSdkConfig = pool.AutoSdkConfig;
 			Properties = pool.Properties ?? new Dictionary<string, string>();
