Horde: Do not shutdown disabled agents by default

Prior to this change, this value defaulted to 8 hours which assumes you want this behavior in the first place. Making this nullable allow for opt-in instead.

diff --git a/Engine/Source/Programs/Horde/Horde.Server.Tests/Agents/Pools/PoolUpdateServiceTest.cs b/Engine/Source/Programs/Horde/Horde.Server.Tests/Agents/Pools/PoolUpdateServiceTest.cs
--- a/Engine/Source/Programs/Horde/Horde.Server.Tests/Agents/Pools/PoolUpdateServiceTest.cs
+++ b/Engine/Source/Programs/Horde/Horde.Server.Tests/Agents/Pools/PoolUpdateServiceTest.cs
@@ -52,9 +52,25 @@ public class PoolUpdateServiceTest : TestSetup
 		await _pus.DisposeAsync();
 	}
 
+	[TestMethod]
+	public async Task ShutdownDisabledAgents_WithDefaultGlobalGracePeriod_DoesNotRequestShutdownAsync()
+	{
+		// Act
+		await _pus.ShutdownDisabledAgentsAsync(CancellationToken.None);
+		await RefreshAgentsAsync();
+
+		// Assert
+		Assert.IsFalse(_enabledAgent.RequestShutdown);
+		Assert.IsFalse(_disabledAgent.RequestShutdown);
+		Assert.IsFalse(_disabledAgentBeyondGracePeriod.RequestShutdown);
+	}
+	
 	[TestMethod]
 	public async Task ShutdownDisabledAgents_WithGlobalGracePeriod_RequestsShutdownAsync()
 	{
+		// Arrange
+		UpdateConfig(config => config.AgentShutdownIfDisabledGracePeriod = TimeSpan.FromHours(3.0));
+		
 		// Act
 		await _pus.ShutdownDisabledAgentsAsync(CancellationToken.None);
 		await RefreshAgentsAsync();
diff --git a/Engine/Source/Programs/Horde/Horde.Server/Agents/Pools/PoolUpdateService.cs b/Engine/Source/Programs/Horde/Horde.Server/Agents/Pools/PoolUpdateService.cs
--- a/Engine/Source/Programs/Horde/Horde.Server/Agents/Pools/PoolUpdateService.cs
+++ b/Engine/Source/Programs/Horde/Horde.Server/Agents/Pools/PoolUpdateService.cs
@@ -99,15 +99,20 @@ namespace Horde.Server.Agents.Pools
 			span.SetAttribute("numShutdown", c);
 		}
 
-		private bool HasGracePeriodExpired(IAgent agent, IReadOnlyList<IPoolConfig> pools, TimeSpan globalGracePeriod)
+		private bool HasGracePeriodExpired(IAgent agent, IReadOnlyList<IPoolConfig> pools, TimeSpan? globalGracePeriod)
 		{
 			if (agent.LastStatusChange == null)
 			{
 				return false;
 			}
 
-			TimeSpan gracePeriod = GetGracePeriod(agent, pools) ?? globalGracePeriod;
-			DateTime expirationTime = agent.LastStatusChange.Value + gracePeriod;
+			TimeSpan? gracePeriod = GetGracePeriod(agent, pools) ?? globalGracePeriod;
+			if (gracePeriod == null)
+			{
+				return false;
+			}
+			
+			DateTime expirationTime = agent.LastStatusChange.Value + gracePeriod.Value;
 			return _clock.UtcNow > expirationTime;
 		}
 
diff --git a/Engine/Source/Programs/Horde/Horde.Server/Server/GlobalConfig.cs b/Engine/Source/Programs/Horde/Horde.Server/Server/GlobalConfig.cs
--- a/Engine/Source/Programs/Horde/Horde.Server/Server/GlobalConfig.cs
+++ b/Engine/Source/Programs/Horde/Horde.Server/Server/GlobalConfig.cs
@@ -235,7 +235,7 @@ namespace Horde.Server.Server
 		/// Time to wait before shutting down an agent that has been disabled
 		/// Used if no value is set on the actual pool.
 		/// </summary>
-		public TimeSpan AgentShutdownIfDisabledGracePeriod { get; set; } = TimeSpan.FromHours(8);
+		public TimeSpan? AgentShutdownIfDisabledGracePeriod { get; set; } = null;
 
 		/// <summary>
 		/// Storage configuration
