Horde: Add IPoolSizeStrategyFactory interface, which is used for creating IPoolSizeStrategy instances.

This allows us to remove the graph/job/stream collection interfaces from FleetService, allowing it to exist without any job handling code in the solution.

diff --git a/Engine/Source/Programs/Horde/Horde.Server.Tests/Fleet/FleetServiceTest.cs b/Engine/Source/Programs/Horde/Horde.Server.Tests/Fleet/FleetServiceTest.cs
--- a/Engine/Source/Programs/Horde/Horde.Server.Tests/Fleet/FleetServiceTest.cs
+++ b/Engine/Source/Programs/Horde/Horde.Server.Tests/Fleet/FleetServiceTest.cs
@@ -235,8 +235,8 @@ namespace Horde.Server.Tests.Fleet
 			serverSettingsOpt.Value.FleetManagerV2 = FleetManagerType.AwsReuse;
 
 			FleetService service = new(
-				AgentService, GraphCollection, JobCollection, LeaseCollection, PoolCollection, new DowntimeServiceStub(isDowntimeActive), StreamCollection, Meter,
-				new StubFleetManagerFactory(fleetManager), Clock, Cache, serverSettingsOpt, GlobalConfig, ServiceProvider, Tracer, loggerFactory.CreateLogger<FleetService>());
+				AgentService, PoolCollection, new DowntimeServiceStub(isDowntimeActive), Meter,
+				new StubFleetManagerFactory(fleetManager), Clock, ServiceProvider.GetRequiredService<IEnumerable<IPoolSizeStrategyFactory>>(), serverSettingsOpt, GlobalConfig, Tracer, loggerFactory.CreateLogger<FleetService>());
 
 			return service;
 		}
diff --git a/Engine/Source/Programs/Horde/Horde.Server.Tests/TestSetup.cs b/Engine/Source/Programs/Horde/Horde.Server.Tests/TestSetup.cs
--- a/Engine/Source/Programs/Horde/Horde.Server.Tests/TestSetup.cs
+++ b/Engine/Source/Programs/Horde/Horde.Server.Tests/TestSetup.cs
@@ -251,6 +251,11 @@ namespace Horde.Server.Tests
 			services.AddSingleton<IAmazonCloudWatch>(sp => new Mock<IAmazonCloudWatch>().Object);
 			services.AddSingleton<IFleetManagerFactory, FleetManagerFactory>();
 
+			services.AddSingleton<IPoolSizeStrategyFactory, NoOpPoolSizeStrategyFactory>();
+			services.AddSingleton<IPoolSizeStrategyFactory, JobQueueStrategyFactory>();
+			services.AddSingleton<IPoolSizeStrategyFactory, LeaseUtilizationStrategyFactory>();
+			services.AddSingleton<IPoolSizeStrategyFactory, LeaseUtilizationAwsMetricStrategyFactory>();
+
 			services.AddSingleton<AclService>();
 			services.AddSingleton<AgentService>();
 			services.AddSingleton(provider => new Lazy<AgentService>(provider.GetRequiredService<AgentService>));
diff --git a/Engine/Source/Programs/Horde/Horde.Server/Agents/Fleet/ComputeQueueAwsMetricStrategy.cs b/Engine/Source/Programs/Horde/Horde.Server/Agents/Fleet/ComputeQueueAwsMetricStrategy.cs
--- a/Engine/Source/Programs/Horde/Horde.Server/Agents/Fleet/ComputeQueueAwsMetricStrategy.cs
+++ b/Engine/Source/Programs/Horde/Horde.Server/Agents/Fleet/ComputeQueueAwsMetricStrategy.cs
@@ -157,4 +157,29 @@ namespace Horde.Server.Agents.Fleet
 			return new PoolSizeResult(numAgents, numAgents);
 		}
 	}
+
+	/// <summary>
+	/// Factory for <see cref="LeaseUtilizationStrategy"/>
+	/// </summary>
+	public class ComputeQueueAwsMetricStrategyFactory : PoolSizeStrategyFactory<ComputeQueueAwsMetricSettings>
+	{
+		readonly IAmazonCloudWatch _cloudWatch;
+		readonly ComputeTaskSource _computeTaskSource;
+		readonly ILogger<ComputeQueueAwsMetricStrategy> _logger;
+
+		/// <summary>
+		/// Constructor
+		/// </summary>
+		public ComputeQueueAwsMetricStrategyFactory(IAmazonCloudWatch cloudWatch, ComputeTaskSource computeTaskSource, ILogger<ComputeQueueAwsMetricStrategy> logger)
+			: base(PoolSizeStrategy.ComputeQueueAwsMetric)
+		{
+			_cloudWatch = cloudWatch;
+			_computeTaskSource = computeTaskSource;
+			_logger = logger;
+		}
+
+		/// <inheritdoc/>
+		public override IPoolSizeStrategy Create(ComputeQueueAwsMetricSettings settings)
+			=> new ComputeQueueAwsMetricStrategy(_cloudWatch, _computeTaskSource, settings, _logger);
+	}
 }
diff --git a/Engine/Source/Programs/Horde/Horde.Server/Agents/Fleet/FleetService.cs b/Engine/Source/Programs/Horde/Horde.Server/Agents/Fleet/FleetService.cs
--- a/Engine/Source/Programs/Horde/Horde.Server/Agents/Fleet/FleetService.cs
+++ b/Engine/Source/Programs/Horde/Horde.Server/Agents/Fleet/FleetService.cs
@@ -5,21 +5,14 @@ using System.Collections.Generic;
 using System.Diagnostics;
 using System.Diagnostics.Metrics;
 using System.Linq;
-using System.Text.Json;
 using System.Threading;
 using System.Threading.Tasks;
 using EpicGames.Core;
 using EpicGames.Horde.Agents;
 using EpicGames.Horde.Agents.Pools;
-using Horde.Server.Agents.Leases;
 using Horde.Server.Agents.Pools;
-using Horde.Server.Jobs;
-using Horde.Server.Jobs.Graphs;
 using Horde.Server.Server;
-using Horde.Server.Streams;
 using Horde.Server.Utilities;
-using Microsoft.Extensions.Caching.Memory;
-using Microsoft.Extensions.DependencyInjection;
 using Microsoft.Extensions.Hosting;
 using Microsoft.Extensions.Logging;
 using Microsoft.Extensions.Options;
@@ -65,23 +58,18 @@ namespace Horde.Server.Agents.Fleet
 		private const int MaxParallelTasks = 10;
 
 		private readonly AgentService _agentService;
-		private readonly IGraphCollection _graphCollection;
-		private readonly IJobCollection _jobCollection;
-		private readonly ILeaseCollection _leaseCollection;
 		private readonly IPoolCollection _poolCollection;
 		private readonly IDowntimeService _downtimeService;
-		private readonly IStreamCollection _streamCollection;
 		private readonly Meter _meter;
 		private readonly IFleetManagerFactory _fleetManagerFactory;
 		private readonly IClock _clock;
-		private readonly IMemoryCache _cache;
+		private readonly IEnumerable<IPoolSizeStrategyFactory> _poolSizeStrategyFactories;
 		private readonly ITicker _ticker;
 		private readonly ITicker _tickerHighFrequency;
 		private readonly TimeSpan _defaultScaleOutCooldown;
 		private readonly TimeSpan _defaultScaleInCooldown;
 		private readonly IOptions<ServerSettings> _settings;
 		private readonly IOptionsMonitor<GlobalConfig> _globalConfig;
-		private readonly IServiceProvider _provider;
 		private readonly Tracer _tracer;
 		private readonly ILogger<FleetService> _logger;
 
@@ -90,38 +78,28 @@ namespace Horde.Server.Agents.Fleet
 		/// </summary>
 		public FleetService(
 			AgentService agentService,
-			IGraphCollection graphCollection,
-			IJobCollection jobCollection,
-			ILeaseCollection leaseCollection,
 			IPoolCollection poolCollection,
 			IDowntimeService downtimeService,
-			IStreamCollection streamCollection,
 			Meter meter,
 			IFleetManagerFactory fleetManagerFactory,
 			IClock clock,
-			IMemoryCache cache,
+			IEnumerable<IPoolSizeStrategyFactory> poolSizeStrategyFactories,
 			IOptions<ServerSettings> settings,
 			IOptionsMonitor<GlobalConfig> globalConfig,
-			IServiceProvider provider,
 			Tracer tracer,
 			ILogger<FleetService> logger)
 		{
 			_agentService = agentService;
-			_graphCollection = graphCollection;
-			_jobCollection = jobCollection;
-			_leaseCollection = leaseCollection;
 			_poolCollection = poolCollection;
 			_downtimeService = downtimeService;
-			_streamCollection = streamCollection;
 			_meter = meter;
 			_fleetManagerFactory = fleetManagerFactory;
 			_clock = clock;
-			_cache = cache;
 			_globalConfig = globalConfig;
 			_tracer = tracer;
 			_logger = logger;
+			_poolSizeStrategyFactories = poolSizeStrategyFactories;
 			_settings = settings;
-			_provider = provider;
 			_defaultScaleOutCooldown = TimeSpan.FromSeconds(settings.Value.AgentPoolScaleOutCooldownSeconds);
 			_defaultScaleInCooldown = TimeSpan.FromSeconds(settings.Value.AgentPoolScaleInCooldownSeconds);
 
@@ -204,7 +182,7 @@ namespace Horde.Server.Agents.Fleet
 		{
 			List<IAgent> agents = (await _agentService.GetCachedAgentsAsync(cancellationToken))
 				.Where(x => x is { Status: AgentStatus.Ok, Enabled: true, RequestShutdown: false }).ToList();
-			
+
 			List<IAgent> GetAgentsInPool(PoolId poolId) => agents.FindAll(a => a.Pools.Any(p => p == poolId));
 			IReadOnlyList<IPoolConfig> pools = await _poolCollection.GetConfigsAsync(cancellationToken);
 
@@ -389,43 +367,30 @@ namespace Horde.Server.Agents.Fleet
 				{
 					if (info.Condition == null || info.Condition.Evaluate(GetPropValues))
 					{
-						switch (info.Type)
+						IPoolSizeStrategyFactory? factory = _poolSizeStrategyFactories.FirstOrDefault(x => x.Type == info.Type);
+						if (factory == null)
+						{
+							throw new ArgumentException("Invalid pool size strategy type " + info.Type);
+						}
+
+						IPoolSizeStrategy strategy = factory.Create(info.Config);
+						if (info.ExtraAgentCount != 0)
 						{
-							case PoolSizeStrategy.JobQueue:
-								JobQueueSettings jqSettings = DeserializeConfig<JobQueueSettings>(info.Config);
-								JobQueueStrategy jqStrategy = new(_jobCollection, _graphCollection, _streamCollection, _clock, _cache, _downtimeService.IsDowntimeActive, _globalConfig, jqSettings);
-								return info.ExtraAgentCount != 0 ? new ExtraAgentCountStrategy(jqStrategy, info.ExtraAgentCount) : jqStrategy;
-
-							case PoolSizeStrategy.LeaseUtilization:
-								LeaseUtilizationSettings luSettings = DeserializeConfig<LeaseUtilizationSettings>(info.Config);
-								LeaseUtilizationStrategy luStrategy = new(_agentService, _poolCollection, _leaseCollection, _clock, _cache, luSettings);
-								return info.ExtraAgentCount != 0 ? new ExtraAgentCountStrategy(luStrategy, info.ExtraAgentCount) : luStrategy;
-
-							case PoolSizeStrategy.ComputeQueueAwsMetric:
-								ComputeQueueAwsMetricSettings cqamSettings = DeserializeConfig<ComputeQueueAwsMetricSettings>(info.Config);
-								return ActivatorUtilities.CreateInstance<ComputeQueueAwsMetricStrategy>(_provider, cqamSettings);
-
-							case PoolSizeStrategy.LeaseUtilizationAwsMetric:
-								LeaseUtilizationAwsMetricSettings luamSettings = DeserializeConfig<LeaseUtilizationAwsMetricSettings>(info.Config);
-								return ActivatorUtilities.CreateInstance<LeaseUtilizationAwsMetricStrategy>(_provider, luamSettings);
-
-							case PoolSizeStrategy.NoOp:
-								NoOpPoolSizeStrategy noStrategy = new();
-								return info.ExtraAgentCount != 0 ? new ExtraAgentCountStrategy(noStrategy, info.ExtraAgentCount) : noStrategy;
-
-							default:
-								throw new ArgumentException("Invalid pool size strategy type " + info.Type);
+							strategy = new ExtraAgentCountStrategy(strategy, info.ExtraAgentCount);
 						}
+
+						return strategy;
 					}
 				}
 			}
 
 			// These is the legacy way of creating and configuring strategies (list-based approach above is preferred)
 #pragma warning disable CS0618
+
 			switch (pool.SizeStrategy ?? _settings.Value.DefaultAgentPoolSizeStrategy)
 			{
 				case PoolSizeStrategy.JobQueue:
-					return new JobQueueStrategy(_jobCollection, _graphCollection, _streamCollection, _clock, _cache, _downtimeService.IsDowntimeActive, _globalConfig, pool.JobQueueSettings);
+					return _poolSizeStrategyFactories.OfType<PoolSizeStrategyFactory<JobQueueSettings>>().First().Create(pool.JobQueueSettings ?? new JobQueueSettings());
 				case PoolSizeStrategy.LeaseUtilization:
 					LeaseUtilizationSettings luSettings = new();
 					if (pool.MinAgents != null)
@@ -436,7 +401,7 @@ namespace Horde.Server.Agents.Fleet
 					{
 						luSettings.NumReserveAgents = pool.NumReserveAgents.Value;
 					}
-					return new LeaseUtilizationStrategy(_agentService, _poolCollection, _leaseCollection, _clock, _cache, luSettings);
+					return _poolSizeStrategyFactories.OfType<PoolSizeStrategyFactory<LeaseUtilizationSettings>>().First().Create(luSettings);
 				case PoolSizeStrategy.NoOp:
 					return new NoOpPoolSizeStrategy();
 				default:
@@ -461,7 +426,7 @@ namespace Horde.Server.Agents.Fleet
 				.Where(x => x is { Status: AgentStatus.Ok, Enabled: true })
 				.Where(x => x.Pools.Contains(pool.Id))
 				.ToList();
-			
+
 			int numShutdownsCancelled = 0;
 			foreach (IAgent agent in agents)
 			{
@@ -490,16 +455,5 @@ namespace Horde.Server.Agents.Fleet
 			ScaleResult result = await scaleOutFunc(agentsToAdd);
 			return new ScaleResult(result.Outcome, result.AgentsAddedCount + numShutdownsCancelled, result.AgentsRemovedCount, result.Message);
 		}
-
-		private static T DeserializeConfig<T>(string json)
-		{
-			json = String.IsNullOrEmpty(json) ? "{}" : json;
-			T? config = JsonSerializer.Deserialize<T>(json, new JsonSerializerOptions { PropertyNameCaseInsensitive = true });
-			if (config == null)
-			{
-				throw new ArgumentException("Unable to deserialize config: " + json);
-			}
-			return config;
-		}
 	}
 }
diff --git a/Engine/Source/Programs/Horde/Horde.Server/Agents/Fleet/IPoolSizeStrategy.cs b/Engine/Source/Programs/Horde/Horde.Server/Agents/Fleet/IPoolSizeStrategy.cs
--- a/Engine/Source/Programs/Horde/Horde.Server/Agents/Fleet/IPoolSizeStrategy.cs
+++ b/Engine/Source/Programs/Horde/Horde.Server/Agents/Fleet/IPoolSizeStrategy.cs
@@ -1,6 +1,8 @@
 // Copyright Epic Games, Inc. All Rights Reserved.
 
+using System;
 using System.Collections.Generic;
+using System.Text.Json;
 using System.Threading;
 using System.Threading.Tasks;
 using Horde.Server.Agents.Pools;
@@ -117,6 +119,57 @@ namespace Horde.Server.Agents.Fleet
 		string Name { get; }
 	}
 
+	/// <summary>
+	/// Factory for creating <see cref="IPoolSizeStrategy"/> instances
+	/// </summary>
+	public interface IPoolSizeStrategyFactory
+	{
+		/// <summary>
+		/// Strategy type
+		/// </summary>
+		PoolSizeStrategy Type { get; }
+
+		/// <summary>
+		/// Create a new <see cref="IPoolSizeStrategy"/> instance from the given configuration
+		/// </summary>
+		IPoolSizeStrategy Create(string config);
+	}
+
+	/// <summary>
+	/// Base implementation of <see cref="IPoolSizeStrategyFactory"/> which deserializes options into a concrete settings type
+	/// </summary>
+	public abstract class PoolSizeStrategyFactory<TConfig> : IPoolSizeStrategyFactory where TConfig : new()
+	{
+		/// <inheritdoc/>
+		public PoolSizeStrategy Type { get; }
+
+		/// <summary>
+		/// Constructor
+		/// </summary>
+		protected PoolSizeStrategyFactory(PoolSizeStrategy type)
+			=> Type = type;
+
+		/// <inheritdoc/>
+		public IPoolSizeStrategy Create(string? config)
+			=> Create(DeserializeConfig<TConfig>(config));
+
+		/// <summary>
+		/// Creates a new strategy object from the given config
+		/// </summary>
+		public abstract IPoolSizeStrategy Create(TConfig config);
+
+		private static T DeserializeConfig<T>(string? json)
+		{
+			json = String.IsNullOrEmpty(json) ? "{}" : json;
+			T? config = JsonSerializer.Deserialize<T>(json, new JsonSerializerOptions { PropertyNameCaseInsensitive = true });
+			if (config == null)
+			{
+				throw new ArgumentException("Unable to deserialize config: " + json);
+			}
+			return config;
+		}
+	}
+
 	/// <summary>
 	/// No-operation strategy that won't resize pools, just return the existing count.
 	/// Used to ensure there's always a strategy available for dependency injection, even if it does nothing.
@@ -133,6 +186,14 @@ namespace Horde.Server.Agents.Fleet
 		public string Name { get; } = "NoOp";
 	}
 
+	class NoOpPoolSizeStrategyFactory : IPoolSizeStrategyFactory
+	{
+		public PoolSizeStrategy Type => PoolSizeStrategy.NoOp;
+
+		public IPoolSizeStrategy Create(string config)
+			=> new NoOpPoolSizeStrategy();
+	}
+
 	/// <summary>
 	/// Pool size strategy wrapping a normal strategy and
 	/// applies an extra agent count to the desired agent count.
diff --git a/Engine/Source/Programs/Horde/Horde.Server/Agents/Fleet/JobQueueStrategy.cs b/Engine/Source/Programs/Horde/Horde.Server/Agents/Fleet/JobQueueStrategy.cs
--- a/Engine/Source/Programs/Horde/Horde.Server/Agents/Fleet/JobQueueStrategy.cs
+++ b/Engine/Source/Programs/Horde/Horde.Server/Agents/Fleet/JobQueueStrategy.cs
@@ -249,4 +249,30 @@ namespace Horde.Server.Agents.Fleet
 			}
 		}
 	}
+
+	class JobQueueStrategyFactory : PoolSizeStrategyFactory<JobQueueSettings>
+	{
+		readonly IGraphCollection _graphCollection;
+		readonly IJobCollection _jobCollection;
+		readonly IStreamCollection _streamCollection;
+		readonly IDowntimeService _downtimeService;
+		readonly IClock _clock;
+		readonly IMemoryCache _memoryCache;
+		readonly IOptionsMonitor<GlobalConfig> _globalConfig;
+
+		public JobQueueStrategyFactory(IGraphCollection graphCollection, IJobCollection jobCollection, IStreamCollection streamCollection, IDowntimeService downtimeService, IClock clock, IMemoryCache memoryCache, IOptionsMonitor<GlobalConfig> globalConfig)
+			: base(PoolSizeStrategy.JobQueue)
+		{
+			_graphCollection = graphCollection;
+			_jobCollection = jobCollection;
+			_streamCollection = streamCollection;
+			_downtimeService = downtimeService;
+			_clock = clock;
+			_memoryCache = memoryCache;
+			_globalConfig = globalConfig;
+		}
+
+		public override IPoolSizeStrategy Create(JobQueueSettings settings)
+			=> new JobQueueStrategy(_jobCollection, _graphCollection, _streamCollection, _clock, _memoryCache, _downtimeService.IsDowntimeActive, _globalConfig, settings);
+	}
 }
diff --git a/Engine/Source/Programs/Horde/Horde.Server/Agents/Fleet/LeaseUtilizationAwsMetricStrategy.cs b/Engine/Source/Programs/Horde/Horde.Server/Agents/Fleet/LeaseUtilizationAwsMetricStrategy.cs
--- a/Engine/Source/Programs/Horde/Horde.Server/Agents/Fleet/LeaseUtilizationAwsMetricStrategy.cs
+++ b/Engine/Source/Programs/Horde/Horde.Server/Agents/Fleet/LeaseUtilizationAwsMetricStrategy.cs
@@ -138,3 +138,30 @@ public class LeaseUtilizationAwsMetricStrategy : IPoolSizeStrategy
 	}
 }
 
+/// <summary>
+/// Factory for <see cref="LeaseUtilizationAwsMetricStrategy"/>
+/// </summary>
+public class LeaseUtilizationAwsMetricStrategyFactory : PoolSizeStrategyFactory<LeaseUtilizationAwsMetricSettings>
+{
+	readonly ILeaseCollection _leaseCollection;
+	readonly IAmazonCloudWatch _cloudWatch;
+	readonly IClock _clock;
+	readonly ILogger<LeaseUtilizationAwsMetricStrategy> _logger;
+
+	/// <summary>
+	/// Constructor
+	/// </summary>
+	public LeaseUtilizationAwsMetricStrategyFactory(ILeaseCollection leaseCollection, IAmazonCloudWatch cloudWatch, IClock clock, ILogger<LeaseUtilizationAwsMetricStrategy> logger)
+		: base(PoolSizeStrategy.LeaseUtilizationAwsMetric)
+	{
+		_leaseCollection = leaseCollection;
+		_cloudWatch = cloudWatch;
+		_clock = clock;
+		_logger = logger;
+	}
+
+	/// <inheritdoc/>
+	public override IPoolSizeStrategy Create(LeaseUtilizationAwsMetricSettings settings)
+		=> new LeaseUtilizationAwsMetricStrategy(_leaseCollection, _cloudWatch, settings, _clock, _logger);
+}
+
diff --git a/Engine/Source/Programs/Horde/Horde.Server/Agents/Fleet/LeaseUtilizationStrategy.cs b/Engine/Source/Programs/Horde/Horde.Server/Agents/Fleet/LeaseUtilizationStrategy.cs
--- a/Engine/Source/Programs/Horde/Horde.Server/Agents/Fleet/LeaseUtilizationStrategy.cs
+++ b/Engine/Source/Programs/Horde/Horde.Server/Agents/Fleet/LeaseUtilizationStrategy.cs
@@ -288,4 +288,33 @@ namespace Horde.Server.Agents.Fleet
 			return new string(values.Select(x => Greyscale[Math.Clamp((int)(x * Greyscale.Length), 0, Greyscale.Length - 1)]).ToArray());
 		}
 	}
+
+	/// <summary>
+	/// Factory for <see cref="LeaseUtilizationStrategy"/>
+	/// </summary>
+	public class LeaseUtilizationStrategyFactory : PoolSizeStrategyFactory<LeaseUtilizationSettings>
+	{
+		readonly AgentService _agentService;
+		readonly IPoolCollection _poolCollection;
+		readonly ILeaseCollection _leaseCollection;
+		readonly IClock _clock;
+		readonly IMemoryCache _memoryCache;
+
+		/// <summary>
+		/// Constructor
+		/// </summary>
+		public LeaseUtilizationStrategyFactory(AgentService agentService, IPoolCollection poolCollection, ILeaseCollection leaseCollection, IClock clock, IMemoryCache memoryCache)
+			: base(PoolSizeStrategy.LeaseUtilization)
+		{
+			_agentService = agentService;
+			_poolCollection = poolCollection;
+			_leaseCollection = leaseCollection;
+			_clock = clock;
+			_memoryCache = memoryCache;
+		}
+
+		/// <inheritdoc/>
+		public override IPoolSizeStrategy Create(LeaseUtilizationSettings settings)
+			=> new LeaseUtilizationStrategy(_agentService, _poolCollection, _leaseCollection, _clock, _memoryCache, settings);
+	}
 }
diff --git a/Engine/Source/Programs/Horde/Horde.Server/Startup.cs b/Engine/Source/Programs/Horde/Horde.Server/Startup.cs
--- a/Engine/Source/Programs/Horde/Horde.Server/Startup.cs
+++ b/Engine/Source/Programs/Horde/Horde.Server/Startup.cs
@@ -541,6 +541,9 @@ namespace Horde.Server
 			services.AddSingleton<AwsAutoScalingLifecycleService>();
 			services.AddSingleton<FleetService>();
 			services.AddSingleton<IFleetManagerFactory, FleetManagerFactory>();
+			services.AddSingleton<IPoolSizeStrategyFactory, NoOpPoolSizeStrategyFactory>();
+			services.AddSingleton<IPoolSizeStrategyFactory, JobQueueStrategyFactory>();
+			services.AddSingleton<IPoolSizeStrategyFactory, LeaseUtilizationStrategyFactory>();
 
 			// Associate IFleetManager interface with the default implementation from config for convenience
 			// Though most fleet managers are created on a per-pool basis
@@ -667,6 +670,9 @@ namespace Horde.Server
 				services.AddAWSService<IAmazonEC2>();
 
 				services.AddSingleton<AwsCloudWatchMetricExporter>();
+
+				services.AddSingleton<IPoolSizeStrategyFactory, LeaseUtilizationAwsMetricStrategyFactory>();
+				services.AddSingleton<IPoolSizeStrategyFactory, ComputeQueueAwsMetricStrategyFactory>();
 			}
 
 			ConfigureLogStorage(services);
