Horde: Exclude blocks for empty files from the Unsync manifest.

diff --git a/Engine/Source/Programs/Horde/Horde.Server/Artifacts/UnsyncCache.cs b/Engine/Source/Programs/Horde/Horde.Server/Artifacts/UnsyncCache.cs
--- a/Engine/Source/Programs/Horde/Horde.Server/Artifacts/UnsyncCache.cs
+++ b/Engine/Source/Programs/Horde/Horde.Server/Artifacts/UnsyncCache.cs
@@ -189,7 +189,10 @@ namespace Horde.Server.Artifacts
 			foreach (FileEntry fileEntry in directoryNode.Files)
 			{
 				List<UnsyncBlock> blocks = new List<UnsyncBlock>();
-				await FindBlocksAsync(0, fileEntry.Target, blocks, cancellationToken);
+				if (fileEntry.Length > 0)
+				{
+					await FindBlocksAsync(0, fileEntry.Target, blocks, cancellationToken);
+				}
 
 				bool readOnly = (fileEntry.Flags & FileEntryFlags.ReadOnly) != 0;
 				bool executable = (fileEntry.Flags & FileEntryFlags.Executable) != 0;
