Horde: Remove job rpc calls from horde_rpc.proto.

#rnx

diff --git a/Engine/Source/Programs/Horde/Horde.Server/Server/RpcService.cs b/Engine/Source/Programs/Horde/Horde.Server/Server/RpcService.cs
--- a/Engine/Source/Programs/Horde/Horde.Server/Server/RpcService.cs
+++ b/Engine/Source/Programs/Horde/Horde.Server/Server/RpcService.cs
@@ -17,12 +17,10 @@ using Google.Protobuf;
 using Google.Protobuf.Reflection;
 using Google.Protobuf.WellKnownTypes;
 using Grpc.Core;
-using Horde.Common.Rpc;
 using Horde.Server.Acls;
 using Horde.Server.Agents;
 using Horde.Server.Agents.Sessions;
 using Horde.Server.Agents.Telemetry;
-using Horde.Server.Jobs;
 using Horde.Server.Telemetry;
 using Horde.Server.Tools;
 using Horde.Server.Utilities;
@@ -30,7 +28,6 @@ using HordeCommon.Rpc;
 using HordeCommon.Rpc.Messages;
 using HordeCommon.Rpc.Messages.Telemetry;
 using Microsoft.AspNetCore.Authorization;
-using Microsoft.Extensions.DependencyInjection;
 using Microsoft.Extensions.Logging;
 using Microsoft.Extensions.Options;
 
@@ -164,14 +161,6 @@ namespace Horde.Server.Server
 			}
 		}
 
-		/// <inheritdoc/>
-		[Obsolete("Use JobRpc instead")]
-		public override Task<RpcUpdateAgentWorkspacesResponse> UpdateAgentWorkspaces(RpcUpdateAgentWorkspacesRequest request, ServerCallContext context)
-		{
-			JobRpcService jobRpcService = ActivatorUtilities.CreateInstance<JobRpcService>(_serviceProvider);
-			return jobRpcService.UpdateAgentWorkspaces(request, context);
-		}
-
 		static void CopyPropertyToResource(string name, List<string> properties, Dictionary<string, int> resources)
 		{
 			foreach (string property in properties)
