Horde: Move functionality for preflighting config changes into a separate controller.

#rnx

diff --git a/Engine/Source/Programs/Horde/Horde.Server/Perforce/ServerPerforceController.cs b/Engine/Source/Programs/Horde/Horde.Server/Perforce/ServerPerforceController.cs
new file mode 100644
--- /dev/null
+++ b/Engine/Source/Programs/Horde/Horde.Server/Perforce/ServerPerforceController.cs
@@ -0,0 +1,91 @@
+// Copyright Epic Games, Inc. All Rights Reserved.
+
+using System;
+using System.Collections.Generic;
+using System.Linq;
+using System.Threading;
+using System.Threading.Tasks;
+using EpicGames.Core;
+using EpicGames.Horde.Server;
+using EpicGames.Perforce;
+using Horde.Server.Configuration;
+using Horde.Server.Utilities;
+using Microsoft.AspNetCore.Authorization;
+using Microsoft.AspNetCore.Mvc;
+
+namespace Horde.Server.Perforce
+{
+	/// <summary>
+	/// Implements preflight of config changes with Perforce
+	/// </summary>
+	[ApiController]
+	[Authorize]
+	[Route("[controller]")]
+	public class ServerPerforceController : HordeControllerBase
+	{
+		readonly ConfigService _configService;
+		readonly IPerforceService _perforceService;
+
+		/// <summary>
+		/// Constructor
+		/// </summary>
+		public ServerPerforceController(ConfigService configService, IPerforceService perforceService)
+		{
+			_configService = configService;
+			_perforceService = perforceService;
+		}
+
+		/// <summary>
+		/// Returns settings for automating auth against this server
+		/// </summary>
+		[HttpPost]
+		[Route("/api/v1/server/preflightconfig")]
+		public async Task<ActionResult<PreflightConfigResponse>> PreflightConfigAsync(PreflightConfigRequest request, CancellationToken cancellationToken)
+		{
+			string cluster = request.Cluster ?? "default";
+
+			IPooledPerforceConnection perforce = await _perforceService.ConnectAsync(cluster, cancellationToken: cancellationToken);
+
+			PerforceResponse<DescribeRecord> describeResponse = await perforce.TryDescribeAsync(DescribeOptions.Shelved, -1, request.ShelvedChange, cancellationToken);
+			if (!describeResponse.Succeeded)
+			{
+				return BadRequest(KnownLogEvents.Horde_InvalidPreflight, "CL {Change} does not exist.", request.ShelvedChange);
+			}
+
+			DescribeRecord record = describeResponse.Data;
+
+			List<string> configFiles = new List<string> { "/globals.json", "global.json", ".project.json", ".stream.json", ".dashboard.json", ".telemetry.json" };
+
+			Dictionary<Uri, byte[]> files = new Dictionary<Uri, byte[]>();
+			foreach (DescribeFileRecord fileRecord in record.Files)
+			{
+				if (configFiles.FirstOrDefault(config => fileRecord.DepotFile.EndsWith(config, StringComparison.OrdinalIgnoreCase)) != null)
+				{
+					PerforceResponse<PrintRecord<byte[]>> printRecordResponse = await perforce.TryPrintDataAsync($"{fileRecord.DepotFile}@={request.ShelvedChange}", cancellationToken);
+					if (!printRecordResponse.Succeeded || printRecordResponse.Data.Contents == null)
+					{
+						return BadRequest($"Unable to print contents of {fileRecord.DepotFile}@={request.ShelvedChange}");
+					}
+
+					PrintRecord<byte[]> printRecord = printRecordResponse.Data;
+
+					Uri uri = new Uri($"perforce://{cluster}{printRecord.DepotFile}");
+					files.Add(uri, printRecord.Contents);
+				}
+			}
+
+			if (files.Count == 0)
+			{
+				return BadRequest(KnownLogEvents.Horde_InvalidPreflight, "No config files found in CL {Change}.", request.ShelvedChange);
+			}
+
+			string? message = await _configService.ValidateAsync(files, cancellationToken);
+
+			PreflightConfigResponse response = new PreflightConfigResponse();
+			response.Result = message == null;
+			response.Message = message;
+
+			return response;
+		}
+	}
+}
diff --git a/Engine/Source/Programs/Horde/Horde.Server/Server/RpcService.cs b/Engine/Source/Programs/Horde/Horde.Server/Server/RpcService.cs
--- a/Engine/Source/Programs/Horde/Horde.Server/Server/RpcService.cs
+++ b/Engine/Source/Programs/Horde/Horde.Server/Server/RpcService.cs
@@ -50,14 +50,13 @@ namespace Horde.Server.Server
 		readonly IToolCollection _toolCollection;
 		readonly IAgentTelemetryCollection _agentTelemetryCollection;
 		readonly AclService _aclService;
-		readonly IServiceProvider _serviceProvider;
 		readonly IOptionsSnapshot<GlobalConfig> _globalConfig;
 		readonly ILogger _logger;
 
 		/// <summary>
 		/// Constructor
 		/// </summary>
-		public RpcService(AgentService agentService, LifetimeService lifetimeService, ITelemetrySink telemetrySink, IToolCollection toolCollection, IAgentTelemetryCollection agentTelemetryCollection, AclService aclService, IServiceProvider serviceProvider, IOptionsSnapshot<GlobalConfig> globalConfig, ILogger<RpcService> logger)
+		public RpcService(AgentService agentService, LifetimeService lifetimeService, ITelemetrySink telemetrySink, IToolCollection toolCollection, IAgentTelemetryCollection agentTelemetryCollection, AclService aclService, IOptionsSnapshot<GlobalConfig> globalConfig, ILogger<RpcService> logger)
 		{
 			_agentService = agentService;
 			_lifetimeService = lifetimeService;
@@ -65,7 +64,6 @@ namespace Horde.Server.Server
 			_toolCollection = toolCollection;
 			_agentTelemetryCollection = agentTelemetryCollection;
 			_aclService = aclService;
-			_serviceProvider = serviceProvider;
 			_globalConfig = globalConfig;
 			_logger = logger;
 		}
diff --git a/Engine/Source/Programs/Horde/Horde.Server/Server/ServerController.cs b/Engine/Source/Programs/Horde/Horde.Server/Server/ServerController.cs
--- a/Engine/Source/Programs/Horde/Horde.Server/Server/ServerController.cs
+++ b/Engine/Source/Programs/Horde/Horde.Server/Server/ServerController.cs
@@ -11,11 +11,8 @@ using EpicGames.Core;
 using EpicGames.Horde;
 using EpicGames.Horde.Agents.Pools;
 using EpicGames.Horde.Server;
-using EpicGames.Perforce;
 using Horde.Server.Agents;
 using Horde.Server.Agents.Pools;
-using Horde.Server.Configuration;
-using Horde.Server.Perforce;
 using Horde.Server.Tools;
 using Horde.Server.Utilities;
 using Microsoft.AspNetCore.Authorization;
@@ -36,20 +33,16 @@ namespace Horde.Server.Server
 		readonly IServiceProvider _serviceProvider;
 		readonly IToolCollection _toolCollection;
 		readonly IClock _clock;
-		readonly ConfigService _configService;
-		readonly IPerforceService _perforceService;
 		readonly IOptionsSnapshot<GlobalConfig> _globalConfig;
 
 		/// <summary>
 		/// Constructor
 		/// </summary>
-		public ServerController(IServiceProvider serviceProvider, IToolCollection toolCollection, IClock clock, ConfigService configService, IPerforceService perforceService, IOptionsSnapshot<GlobalConfig> globalConfig)
+		public ServerController(IServiceProvider serviceProvider, IToolCollection toolCollection, IClock clock, IOptionsSnapshot<GlobalConfig> globalConfig)
 		{
 			_serviceProvider = serviceProvider;
 			_toolCollection = toolCollection;
 			_clock = clock;
-			_configService = configService;
-			_perforceService = perforceService;
 			_globalConfig = globalConfig;
 		}
 
@@ -151,59 +144,6 @@ namespace Horde.Server.Server
 			return response;
 		}
 
-		/// <summary>
-		/// Returns settings for automating auth against this server
-		/// </summary>
-		[HttpPost]
-		[Route("/api/v1/server/preflightconfig")]
-		public async Task<ActionResult<PreflightConfigResponse>> PreflightConfigAsync(PreflightConfigRequest request, CancellationToken cancellationToken)
-		{
-			string cluster = request.Cluster ?? "default";
-
-			IPooledPerforceConnection perforce = await _perforceService.ConnectAsync(cluster, cancellationToken: cancellationToken);
-
-			PerforceResponse<DescribeRecord> describeResponse = await perforce.TryDescribeAsync(DescribeOptions.Shelved, -1, request.ShelvedChange, cancellationToken);
-			if (!describeResponse.Succeeded)
-			{
-				return BadRequest(KnownLogEvents.Horde_InvalidPreflight, "CL {Change} does not exist.", request.ShelvedChange);
-			}
-
-			DescribeRecord record = describeResponse.Data;
-
-			List<string> configFiles = new List<string> { "/globals.json", "global.json", ".project.json", ".stream.json", ".dashboard.json", ".telemetry.json" };
-
-			Dictionary<Uri, byte[]> files = new Dictionary<Uri, byte[]>();
-			foreach (DescribeFileRecord fileRecord in record.Files)
-			{
-				if (configFiles.FirstOrDefault(config => fileRecord.DepotFile.EndsWith(config, StringComparison.OrdinalIgnoreCase)) != null)
-				{
-					PerforceResponse<PrintRecord<byte[]>> printRecordResponse = await perforce.TryPrintDataAsync($"{fileRecord.DepotFile}@={request.ShelvedChange}", cancellationToken);
-					if (!printRecordResponse.Succeeded || printRecordResponse.Data.Contents == null)
-					{
-						return BadRequest($"Unable to print contents of {fileRecord.DepotFile}@={request.ShelvedChange}");
-					}
-
-					PrintRecord<byte[]> printRecord = printRecordResponse.Data;
-
-					Uri uri = new Uri($"perforce://{cluster}{printRecord.DepotFile}");
-					files.Add(uri, printRecord.Contents);
-				}
-			}
-
-			if (files.Count == 0)
-			{
-				return BadRequest(KnownLogEvents.Horde_InvalidPreflight, "No config files found in CL {Change}.", request.ShelvedChange);
-			}
-
-			string? message = await _configService.ValidateAsync(files, cancellationToken);
-
-			PreflightConfigResponse response = new PreflightConfigResponse();
-			response.Result = message == null;
-			response.Message = message;
-
-			return response;
-		}
-
 		/// <summary>
 		/// Converts all legacy pools into config entries
 		/// </summary>
