Gauntlet Event - Grab the summary if there is enough information instead of using the full callstack

#rnx

diff --git a/Engine/Source/Programs/Horde/Horde.Server.Tests/Issues/IssueServiceTests.cs b/Engine/Source/Programs/Horde/Horde.Server.Tests/Issues/IssueServiceTests.cs
--- a/Engine/Source/Programs/Horde/Horde.Server.Tests/Issues/IssueServiceTests.cs
+++ b/Engine/Source/Programs/Horde/Horde.Server.Tests/Issues/IssueServiceTests.cs
@@ -1972,7 +1972,7 @@ namespace Horde.Server.Tests.Issues
 			}
 			// #2
 			// Scenario: Gauntlet Fatal event
-			// Expected: Gauntlet fingerprint using hash prefix
+			// Expected: Gauntlet fingerprint using hash prefix and generate only one issue as they share the same summary
 			{
 				string logMessage = "Assertion failed: State.bGfxPSOSet [File:D:\\build\\U5M+Inc\\Sync\\Engine\\Source\\Runtime\\RHI\\Public\\RHIValidationContext.h] [Line: 809]\n"
 				+ "A Graphics PSO has to be set to set resources into a shader!";
@@ -1997,20 +1997,30 @@ namespace Horde.Server.Tests.Issues
 				+ "	0x00007fff5089a7bf UnrealEditor-Core.dll!FRunnableThreadWin::GuardedRun() [Unknown File]\n"
 				+ "	0x00007fff8e304ed0 KERNEL32.DLL!UnknownFunction [Unknown File]\n"
 				+ "	0x00007fff8f26e39b ntdll.dll!UnknownFunction [Unknown File]\n";
+				string logAlternateCallstack = "	0x00007fff4b43a1b4 UnrealEditor-RHI.dll!FValidationContext::RHISetShaderParameters() [Unknown File]\n"
+				+ "	0x00007fff4b3e5a01 UnrealEditor-RHI.dll!FRHICommandSetShaderParameters<FRHIGraphicsShader>::Execute() [Unknown File]\n"
+				+ "	0x00007fff4b3e9a5a UnrealEditor-RHI.dll!FRHICommand<FRHICommandSetShaderParameters<FRHIGraphicsShader>,FRHICommandSetShaderParametersString1159>::ExecuteAndDestruct() [Unknown File]\n"
+				+ "	0x00007fff4b3e74f7 UnrealEditor - RHI.dll!FRHICommandListBase::Execute()[Unknown File]\n";
 
 				IJob job = CreateJob(_mainStreamId, 120, "Test Build", _graph);
 				await using (TestJsonLogger logger = await CreateLoggerAsync(job, 0, 0))
 				{
-					logger.LogError(KnownLogEvents.Gauntlet_FatalEvent, "{Message}\n{Callstack}", logMessage, logCallstack);
+					logger.LogError(KnownLogEvents.Gauntlet_FatalEvent, "{Summary}\n{Callstack}", logMessage, logCallstack);
 				}
 				await UpdateCompleteStepAsync(job, 0, 0, JobStepOutcome.Failure);
+				await using (TestJsonLogger logger = await CreateLoggerAsync(job, 0, 1))
+				{
+					logger.LogError(KnownLogEvents.Gauntlet_FatalEvent, "Engine encounter a critical failure\n{Message}\n{Callstack}", logMessage, logAlternateCallstack);
+				}
+				await UpdateCompleteStepAsync(job, 0, 1, JobStepOutcome.Failure);
 
 				IReadOnlyList<IIssue> issues = await IssueCollection.FindIssuesAsync();
 				Assert.AreEqual(1, issues.Count);
+				Assert.AreEqual(1, issues[0].Fingerprints.Count);
 				Assert.AreEqual("Gauntlet:fatal:with-callstack", issues[0].Fingerprints[0].Type);
 				Assert.AreEqual("hash:", issues[0].Fingerprints[0].Keys.First().Name.Substring(0, 5));
 				Assert.AreEqual(37, issues[0].Fingerprints[0].Keys.First().Name.Length);
-				Assert.AreEqual("Automation fatal errors in Update Version Files", issues[0].Summary);
+				Assert.AreEqual("Automation fatal errors in Update Version Files and Compile UnrealHeaderTool Win64", issues[0].Summary);
 			}
 			// #3
 			// Scenario: Gauntlet Test + Fatal event
@@ -2040,6 +2050,7 @@ namespace Horde.Server.Tests.Issues
 					}
 				}
 				await UpdateCompleteStepAsync(job, 0, 0, JobStepOutcome.Failure);
+				await UpdateCompleteStepAsync(job, 0, 1, JobStepOutcome.Success);
 
 				IReadOnlyList<IIssue> issues = await IssueCollection.FindIssuesAsync();
 				Assert.AreEqual(2, issues.Count);
diff --git a/Engine/Source/Programs/Shared/EpicGames.Horde/Issues/Handlers/GauntletIssueHandler.cs b/Engine/Source/Programs/Shared/EpicGames.Horde/Issues/Handlers/GauntletIssueHandler.cs
--- a/Engine/Source/Programs/Shared/EpicGames.Horde/Issues/Handlers/GauntletIssueHandler.cs
+++ b/Engine/Source/Programs/Shared/EpicGames.Horde/Issues/Handlers/GauntletIssueHandler.cs
@@ -46,6 +46,16 @@ namespace EpicGames.Horde.Issues.Handlers
 		/// </summary>
 		static readonly Utf8String s_callstackLogType = new Utf8String("Callstack");
 
+		/// <summary>
+		/// Summary log type property
+		/// </summary>
+		static readonly Utf8String s_summaryLogType = new Utf8String("Summary");
+
+		/// <summary>
+		/// Regex pattern to match assertion, ensure and other crash report summary from UE
+		/// </summary>
+		static readonly Regex s_summaryPattern = new Regex(@"\w[\w ]+:\s+.+\s+\[[^]]+\.[^]]+\](?:\s*\[[^]]+\])?(?:\s*\n\w+.+)?", RegexOptions.Multiline | RegexOptions.ExplicitCapture);
+
 		/// <summary>
 		/// Max Message Length to hash
 		/// </summary>
@@ -123,10 +133,22 @@ namespace EpicGames.Horde.Issues.Handlers
 
 		private static bool TryGetHash(IssueEvent issueEvent, out Md5Hash hash)
 		{
-			string sanitized = issueEvent.Message.ToUpperInvariant();
+			string? sanitized = GetSummaryProperty(issueEvent);
+			// Use only the summary if one is found instead of the full callstack
+			if (sanitized == null)
+			{
+				sanitized = issueEvent.Message;
+				// Let's try with regex
+				Match matchSummary = s_summaryPattern.Match(sanitized);
+				if (matchSummary.Success)
+				{
+					sanitized = matchSummary.Groups[0].Value;
+				}
+			}
 			sanitized = sanitized.Length > MaxMessageLength ? sanitized.Substring(0, MaxMessageLength) : sanitized;
-			sanitized = Regex.Replace(sanitized, @"(?<![a-zA-Z])(?:[A-Z]:|/)[^ :]+[/\\]SYNC[/\\]", "{root}/"); // Redact things that look like workspace roots; may be different between agents
-			sanitized = Regex.Replace(sanitized, @"0[xX][0-9a-fA-F]+", "H"); // Redact hex strings
+			sanitized = sanitized.Trim().ToUpperInvariant();
+			sanitized = Regex.Replace(sanitized, @"(?<![A-Z])(?:[A-Z]:|/)[^ :]+[/\\]SYNC[/\\]", "{root}/"); // Redact things that look like workspace roots; may be different between agents
+			sanitized = Regex.Replace(sanitized, @"0X[0-9A-F]+", "H"); // Redact hex strings
 			sanitized = Regex.Replace(sanitized, @"\d[\d.,:]*", "n"); // Redact numbers and timestamp like things
 
 			if (sanitized.Length > 30)
@@ -146,6 +168,19 @@ namespace EpicGames.Horde.Issues.Handlers
 			return issueEvent.Lines.Any(x => FindNestedPropertyOfType(x, s_callstackLogType) != null);
 		}
 
+		private static string? GetSummaryProperty(IssueEvent issueEvent)
+		{
+			foreach (JsonLogEvent logEvent in issueEvent.Lines)
+			{
+				JsonProperty? property = FindNestedPropertyOfType(logEvent, s_summaryLogType);
+				if (property != null)
+				{
+					return property.Value.Value.GetString();
+				}
+			}
+			return null;
+		}
+
 		private static JsonProperty? FindNestedPropertyOfType(JsonLogEvent logEvent, Utf8String type)
 		{
 			JsonElement line = JsonDocument.Parse(logEvent.Data).RootElement;
