Horde: Move configuration for devices into a separate file.

#rnx

diff --git a/Engine/Source/Programs/Horde/Horde.Server/Devices/DeviceConfig.cs b/Engine/Source/Programs/Horde/Horde.Server/Devices/DeviceConfig.cs
new file mode 100644
--- /dev/null
+++ b/Engine/Source/Programs/Horde/Horde.Server/Devices/DeviceConfig.cs
@@ -0,0 +1,95 @@
+// Copyright Epic Games, Inc. All Rights Reserved.
+
+using System;
+using System.Collections.Generic;
+using System.ComponentModel.DataAnnotations;
+using System.Diagnostics;
+using System.Text.Json.Serialization;
+using EpicGames.Horde.Devices;
+using EpicGames.Horde.Projects;
+
+namespace Horde.Server.Devices
+{
+	/// <summary>
+	/// Configuration for devices
+	/// </summary>
+	public class DeviceConfig
+	{
+		/// <summary>
+		/// List of device platforms
+		/// </summary>
+		public List<DevicePlatformConfig> Platforms { get; set; } = new List<DevicePlatformConfig>();
+
+		/// <summary>
+		/// List of device pools
+		/// </summary>
+		public List<DevicePoolConfig> Pools { get; set; } = new List<DevicePoolConfig>();
+	}
+
+	/// <summary>
+	/// Configuration for a device platform 
+	/// </summary>
+	[DebuggerDisplay("{Id}")]
+	public class DevicePlatformConfig : IDevicePlatform
+	{
+		/// <summary>
+		/// The id for this platform 
+		/// </summary>
+		[Required]
+		public DevicePlatformId Id { get; set; } = new DevicePlatformId("default");
+
+		/// <summary>
+		/// Name of the platform
+		/// </summary>
+		[Required]
+		public string Name { get; set; } = String.Empty;
+
+		/// <summary>
+		/// A list of platform models 
+		/// </summary>
+		public List<string>? Models { get; set; }
+		IReadOnlyList<string>? IDevicePlatform.Models => Models;
+
+		/// <summary>
+		/// Legacy names which older versions of Gauntlet may be using
+		/// </summary>
+		public List<string>? LegacyNames { get; set; }
+		IReadOnlyList<string>? IDevicePlatform.LegacyNames => LegacyNames;
+
+		/// <summary>
+		/// Model name for the high perf spec, which may be requested by Gauntlet
+		/// </summary>
+		public string? LegacyPerfSpecHighModel { get; set; }
+	}
+
+	/// <summary>
+	/// Configuration for a device pool
+	/// </summary>
+	[DebuggerDisplay("{Id}")]
+	public class DevicePoolConfig : IDevicePool
+	{
+		/// <summary>
+		/// The id for this platform 
+		/// </summary>
+		[Required]
+		public DevicePoolId Id { get; set; } = new DevicePoolId("default");
+
+		/// <summary>
+		/// The name of the pool
+		/// </summary>
+		[Required]
+		public string Name { get; set; } = String.Empty;
+
+		/// <summary>
+		/// The type of the pool
+		/// </summary>
+		[Required]
+		[JsonConverter(typeof(JsonStringEnumConverter))]
+		public DevicePoolType PoolType { get; set; } = DevicePoolType.Automation;
+
+		/// <summary>
+		/// List of project ids associated with pool
+		/// </summary>
+		public List<ProjectId>? ProjectIds { get; set; }
+	}
+}
diff --git a/Engine/Source/Programs/Horde/Horde.Server/Server/GlobalConfig.cs b/Engine/Source/Programs/Horde/Horde.Server/Server/GlobalConfig.cs
--- a/Engine/Source/Programs/Horde/Horde.Server/Server/GlobalConfig.cs
+++ b/Engine/Source/Programs/Horde/Horde.Server/Server/GlobalConfig.cs
@@ -799,90 +799,6 @@ namespace Horde.Server.Server
 		}
 	}
 
-	/// <summary>
-	/// Configuration for a device platform 
-	/// </summary>
-	[DebuggerDisplay("{Id}")]
-	public class DevicePlatformConfig : IDevicePlatform
-	{
-		/// <summary>
-		/// The id for this platform 
-		/// </summary>
-		[Required]
-		public DevicePlatformId Id { get; set; } = new DevicePlatformId("default");
-
-		/// <summary>
-		/// Name of the platform
-		/// </summary>
-		[Required]
-		public string Name { get; set; } = String.Empty;
-
-		/// <summary>
-		/// A list of platform models 
-		/// </summary>
-		public List<string>? Models { get; set; }
-		IReadOnlyList<string>? IDevicePlatform.Models => Models;
-
-		/// <summary>
-		/// Legacy names which older versions of Gauntlet may be using
-		/// </summary>
-		public List<string>? LegacyNames { get; set; }
-		IReadOnlyList<string>? IDevicePlatform.LegacyNames => LegacyNames;
-
-		/// <summary>
-		/// Model name for the high perf spec, which may be requested by Gauntlet
-		/// </summary>
-		public string? LegacyPerfSpecHighModel { get; set; }
-	}
-
-	/// <summary>
-	/// Configuration for a device pool
-	/// </summary>
-	[DebuggerDisplay("{Id}")]
-	public class DevicePoolConfig : IDevicePool
-	{
-		/// <summary>
-		/// The id for this platform 
-		/// </summary>
-		[Required]
-		public DevicePoolId Id { get; set; } = new DevicePoolId("default");
-
-		/// <summary>
-		/// The name of the pool
-		/// </summary>
-		[Required]
-		public string Name { get; set; } = String.Empty;
-
-		/// <summary>
-		/// The type of the pool
-		/// </summary>
-		[Required]
-		[JsonConverter(typeof(JsonStringEnumConverter))]
-		public DevicePoolType PoolType { get; set; } = DevicePoolType.Automation;
-
-		/// <summary>
-		/// List of project ids associated with pool
-		/// </summary>
-		public List<ProjectId>? ProjectIds { get; set; }
-	}
-
-	/// <summary>
-	/// Configuration for devices
-	/// </summary>
-	public class DeviceConfig
-	{
-		/// <summary>
-		/// List of device platforms
-		/// </summary>
-		public List<DevicePlatformConfig> Platforms { get; set; } = new List<DevicePlatformConfig>();
-
-		/// <summary>
-		/// List of device pools
-		/// </summary>
-		public List<DevicePoolConfig> Pools { get; set; } = new List<DevicePoolConfig>();
-
-	}
-
 	/// <summary>
 	/// Selects different agent software versions by evaluating a condition
 	/// </summary>
