Horde: Add multithreaded reader/writer test for the block cache.

#rnx

diff --git a/Engine/Source/Programs/Horde/Horde.Server.Tests/Artifacts/BlockCacheTests.cs b/Engine/Source/Programs/Horde/Horde.Server.Tests/Artifacts/BlockCacheTests.cs
--- a/Engine/Source/Programs/Horde/Horde.Server.Tests/Artifacts/BlockCacheTests.cs
+++ b/Engine/Source/Programs/Horde/Horde.Server.Tests/Artifacts/BlockCacheTests.cs
@@ -2,8 +2,11 @@
 
 using System;
 using System.Buffers;
+using System.Collections.Generic;
 using System.Linq;
 using System.Security.Cryptography;
+using System.Threading.Tasks;
+using EpicGames.Core;
 using Horde.Server.Artifacts;
 using Microsoft.VisualStudio.TestTools.UnitTesting;
 
@@ -70,5 +73,47 @@ namespace Horde.Server.Tests.Artifacts
 				Assert.IsTrue(value.Data.ToArray().SequenceEqual(values[keyIdx]));
 			}
 		}
+
+		[TestMethod]
+		public void RandomBlockTest()
+		{
+			Random sizeRng = new Random(0);
+
+			List<(string Name, byte[] Data)> items = new List<(string, byte[])>();
+			for (int idx = 0; idx < 512; idx++)
+			{
+				int size = sizeRng.Next(2048, 16384);
+				items.Add(($"{idx}", RandomNumberGenerator.GetBytes(size)));
+			}
+
+			using BlockCache blockCache = BlockCache.CreateInMemory(10, 64, 4096);
+			Parallel.For(0, 8, threadIdx =>
+			{
+				Random rng = new Random(threadIdx);
+				for (int idx = 0; idx < 10000; idx++)
+				{
+					int itemIdx = rng.Next(items.Count);
+					(string name, byte[] data) = items[itemIdx];
+
+					if ((threadIdx & 1) == 0)
+					{
+						// Writer
+						blockCache.Add(name, data);
+					}
+					else
+					{
+						// Reader
+						using IBlockCacheValue? cacheValue = blockCache.Get(name);
+						Assert.IsTrue(cacheValue == null || cacheValue.Data.ToArray().SequenceEqual(data));
+					}
+				}
+			});
+
+			foreach ((string name, byte[] data) in items)
+			{
+				using IBlockCacheValue? cacheValue = blockCache.Get(name);
+				Assert.IsTrue(cacheValue == null || cacheValue.Data.ToArray().SequenceEqual(data));
+			}
+		}
 	}
 }
