Horde: Downgrade log event about duplicate build products to information, since we don't have the list of duplicate build products to ignore in Horde.

diff --git a/Engine/Source/Programs/Horde/Drivers/JobDriver/Execution/JobExecutor.cs b/Engine/Source/Programs/Horde/Drivers/JobDriver/Execution/JobExecutor.cs
--- a/Engine/Source/Programs/Horde/Drivers/JobDriver/Execution/JobExecutor.cs
+++ b/Engine/Source/Programs/Horde/Drivers/JobDriver/Execution/JobExecutor.cs
@@ -1220,7 +1220,7 @@ namespace JobDriver.Execution
 					}
 					else if (!inputFile.Equals(source.File) && !TempStorage.IsDuplicateBuildProduct(inputFile.ToFileReference(workspaceDir)))
 					{
-						logger.LogError("File '{File}' was produced by {InputBlock} and {CurrentBlock}", inputFile.RelativePath, inputStorageBlock.ToString(), source.BlockRef.ToString());
+						logger.LogInformation("File '{File}' was produced by {InputBlock} and {CurrentBlock}", inputFile.RelativePath, inputStorageBlock.ToString(), source.BlockRef.ToString());
 					}
 				}
 			}
