Horde: Add an additional log message for ticking blobs.

#rnx

diff --git a/Engine/Source/Programs/Horde/Horde.Server/Storage/StorageService.cs b/Engine/Source/Programs/Horde/Horde.Server/Storage/StorageService.cs
--- a/Engine/Source/Programs/Horde/Horde.Server/Storage/StorageService.cs
+++ b/Engine/Source/Programs/Horde/Horde.Server/Storage/StorageService.cs
@@ -675,6 +675,8 @@ namespace Horde.Server.Storage
 					innerSpan.SetAttribute("Last", current[^1].Id.ToString());
 					innerSpan.SetAttribute("Count", current.Count);
 
+					_logger.LogDebug("Ticking {NumBlobs} blobs from {FirstId} to {LastId}", current.Count, current[0].Id, current[^1].Id);
+
 					foreach (BlobInfo blobInfo in current)
 					{
 						NamespaceInfo? namespaceInfo;
