Horde: Fix issue with artifact modal, make it possible to specify a browse artifact id for steps which create more than one artifact
#rnx

diff --git a/Engine/Source/Programs/Horde/HordeDashboard/src/components/artifacts/ArtifactsModal.tsx b/Engine/Source/Programs/Horde/HordeDashboard/src/components/artifacts/ArtifactsModal.tsx
--- a/Engine/Source/Programs/Horde/HordeDashboard/src/components/artifacts/ArtifactsModal.tsx
+++ b/Engine/Source/Programs/Horde/HordeDashboard/src/components/artifacts/ArtifactsModal.tsx
@@ -29,12 +29,13 @@ type BrowserItem = {
 
 class ArtifactsHandler {
 
-   constructor(jobId: string, stepId: string, contextType: ArtifactContextType, artifactPath?: string, artifacts?: GetArtifactResponseV2[]) {
+   constructor(jobId: string, stepId: string, contextType: ArtifactContextType, artifactPath?: string, artifacts?: GetArtifactResponseV2[], artifactId?: string) {
       makeObservable(this);
       this.jobId = jobId;
       this.stepId = stepId;
       this.context = contextType;
       this.artifacts = artifacts;
+      this.artifactId = artifactId;
       this.set(artifactPath);
 
       const params = new URLSearchParams(window.location.search);
@@ -89,7 +90,16 @@ class ArtifactsHandler {
          return;
       }
 
-      let a = artifacts.find(a => a.type === this.context)!;
+      let a: GetArtifactResponseV2 | undefined; 
+
+      if (this.artifactId) {
+         a = artifacts.find(a => a.id === this.artifactId)
+      }
+
+      if (!a) {
+         a = artifacts.find(a => a.type === this.context)!;
+      }
+      
 
       if (!a) {
          console.error("Unable to find artifact for context", this.context, artifacts);
@@ -255,6 +265,8 @@ class ArtifactsHandler {
    artifact?: GetArtifactResponseV2;
    artifacts?: GetArtifactResponseV2[];
 
+   artifactId?: string;
+
    readonly context: ArtifactContextType;
 
    jobId: string;
@@ -509,10 +521,10 @@ const DownloadButton: React.FC<{ handler: ArtifactsHandler }> = observer(({ hand
 
 let idcounter = 0;
 
-const JobDetailArtifactsInner: React.FC<{ jobId: string; stepId: string, artifacts?: GetArtifactResponseV2[], contextType: ArtifactContextType, artifactPath?: string }> = observer(({ jobId, stepId, artifacts, contextType, artifactPath }) => {
+const JobDetailArtifactsInner: React.FC<{ jobId: string; stepId: string, artifacts?: GetArtifactResponseV2[], contextType: ArtifactContextType, artifactPath?: string, artifactId? :string }> = observer(({ jobId, stepId, artifacts, contextType, artifactPath, artifactId }) => {
 
    // eslint-disable-next-line
-   const handler = ArtifactsHandler.current ?? new ArtifactsHandler(jobId, stepId, contextType, artifactPath, artifacts);
+   const handler = ArtifactsHandler.current ?? new ArtifactsHandler(jobId, stepId, contextType, artifactPath, artifacts, artifactId);
 
    const [infoKey, setInfoKey] = useState("");
 
@@ -762,7 +774,7 @@ const JobDetailArtifactsInner: React.FC<{ jobId: string; stepId: string, artifac
 
 
 
-export const JobArtifactsModal: React.FC<{ jobId: string; stepId: string, artifacts?: GetArtifactResponseV2[], contextType: ArtifactContextType, artifactPath?: string, onClose: () => void }> = ({ jobId, stepId, artifacts, contextType, artifactPath, onClose }) => {
+export const JobArtifactsModal: React.FC<{ jobId: string; stepId: string, artifacts?: GetArtifactResponseV2[], contextType: ArtifactContextType, artifactPath?: string, artifactId? :string, onClose: () => void }> = ({ jobId, stepId, artifacts, contextType, artifactPath, artifactId, onClose }) => {
 
    const { hordeClasses } = getHordeStyling();
 
@@ -784,7 +796,7 @@ export const JobArtifactsModal: React.FC<{ jobId: string; stepId: string, artifa
                      </Stack>
                   </Stack>
                   <Stack styles={{ root: { paddingLeft: 4, paddingRight: 0, paddingTop: 8, paddingBottom: 4 } }}>
-                     <JobDetailArtifactsInner stepId={stepId} jobId={jobId} contextType={contextType} artifactPath={artifactPath} artifacts={artifacts} />
+                     <JobDetailArtifactsInner stepId={stepId} jobId={jobId} contextType={contextType} artifactPath={artifactPath} artifacts={artifacts} artifactId={artifactId} />
                   </Stack>
                </Stack>
             </Stack>
diff --git a/Engine/Source/Programs/Horde/HordeDashboard/src/components/jobDetailsV2/JobDetailArtifacts.tsx b/Engine/Source/Programs/Horde/HordeDashboard/src/components/jobDetailsV2/JobDetailArtifacts.tsx
--- a/Engine/Source/Programs/Horde/HordeDashboard/src/components/jobDetailsV2/JobDetailArtifacts.tsx
+++ b/Engine/Source/Programs/Horde/HordeDashboard/src/components/jobDetailsV2/JobDetailArtifacts.tsx
@@ -290,7 +290,7 @@ export const JobArtifactsPanel: React.FC<{ jobDetails: JobDetailsV2 }> = observe
 
 
    return (<Stack id={sideRail.url} styles={{ root: { paddingTop: 18, paddingRight: 12 } }}>
-      {!!selected && <JobArtifactsModal jobId={jobData.id} stepId={selected.stepId} contextType={selected.type} onClose={() => setSelected(undefined)} />}
+      {!!selected && <JobArtifactsModal jobId={jobData.id} stepId={selected.stepId} contextType={selected.type} artifactId={selected.id} onClose={() => setSelected(undefined)} />}
       <Stack className={hordeClasses.raised} >
          <Stack tokens={{ childrenGap: 12 }} grow>
             <Stack horizontal>
