Horde: Include the blob id for any soft-deleted blobs.

#rnx

diff --git a/Engine/Source/Programs/Horde/Horde.Server/Storage/StorageService.cs b/Engine/Source/Programs/Horde/Horde.Server/Storage/StorageService.cs
--- a/Engine/Source/Programs/Horde/Horde.Server/Storage/StorageService.cs
+++ b/Engine/Source/Programs/Horde/Horde.Server/Storage/StorageService.cs
@@ -601,9 +601,11 @@ namespace Horde.Server.Storage
 				try
 				{
 					string path = locator.Path.ToString();
-					if (await _blobCollection.Find(x => x.NamespaceId == namespaceId && x.Path == path && x.GcVersion >= CurrentGcVersion).AnyAsync(cancellationToken))
+
+					BlobInfo? blobInfo = await _blobCollection.Find(x => x.NamespaceId == namespaceId && x.Path == path && x.GcVersion >= CurrentGcVersion).FirstAsync(cancellationToken);
+					if (blobInfo != null)
 					{
-						_logger.LogWarning("Blob {Locator} accessed after being garbage collected", locator);
+						_logger.LogWarning("Blob {BlobId} ({Locator}) accessed after being garbage collected", blobInfo.Id, locator);
 					}
 				}
 				catch (Exception ex)
