Horde: Disable AwsObjectStore debug output, and add StorageService debug output.

#rnx

diff --git a/Engine/Source/Programs/Horde/Horde.Server/appsettings.json b/Engine/Source/Programs/Horde/Horde.Server/appsettings.json
--- a/Engine/Source/Programs/Horde/Horde.Server/appsettings.json
+++ b/Engine/Source/Programs/Horde/Horde.Server/appsettings.json
@@ -18,7 +18,8 @@
 				"System.Net.Http.HttpClient": "Warning",
 				"Grpc": "Warning",
 				"Horde.Server.Issues.IssueService": "Debug",
-				"Horde.Server.Storage.ObjectStores.AwsObjectStore": "Debug",
+				"Horde.Server.Storage.ObjectStores.AwsObjectStore": "Information",
+				"Horde.Server.Storage.StorageService": "Debug",
 				"EpicGames.Horde.Storage": "Debug"
 			}
 		}
