Horde: Remove requirement to specify a stream or key to enumerate artifacts.

diff --git a/Engine/Source/Programs/Horde/Horde.Server/Artifacts/ArtifactsController.cs b/Engine/Source/Programs/Horde/Horde.Server/Artifacts/ArtifactsController.cs
--- a/Engine/Source/Programs/Horde/Horde.Server/Artifacts/ArtifactsController.cs
+++ b/Engine/Source/Programs/Horde/Horde.Server/Artifacts/ArtifactsController.cs
@@ -179,5 +179,34 @@
 		}
 
 		/// <summary>
+		/// Finds artifacts matching certain criteria
+		/// </summary>
+		/// <param name="streamId">Stream to search</param>
+		/// <param name="minChange">Minimum changelist number for artifacts to return</param>
+		/// <param name="maxChange">Maximum changelist number for artifacts to return</param>
+		/// <param name="name">Artifact name</param>
+		/// <param name="type">Type of the artifact</param>
+		/// <param name="keys">Keys to find</param>
+		/// <param name="maxResults">Maximum number of results to return</param>
+		/// <param name="filter">Filter for returned values</param>
+		/// <returns>Information about all the artifacts</returns>
+		[HttpGet]
+		[Route("/api/v2/artifacts")]
+		[ProducesResponseType(typeof(FindArtifactsResponse), 200)]
+		public async Task<ActionResult<object>> FindArtifactsAsync([FromQuery] StreamId? streamId = null, [FromQuery] int? minChange = null, [FromQuery] int? maxChange = null, [FromQuery(Name = "name")] ArtifactName? name = null, [FromQuery(Name = "type")] ArtifactType? type = null, [FromQuery(Name = "key")] IEnumerable<string>? keys = null, [FromQuery] int maxResults = 100, [FromQuery] PropertyFilter? filter = null)
+		{
+			FindArtifactsResponse response = new FindArtifactsResponse();
+			await foreach (IArtifact artifact in _artifactCollection.FindAsync(streamId, minChange, maxChange, name, type, keys, maxResults, HttpContext.RequestAborted))
+			{
+				if (_globalConfig.Authorize(artifact.AclScope, ArtifactAclAction.ReadArtifact, User))
+				{
+					response.Artifacts.Add(new GetArtifactResponse(artifact.Id, artifact.Name, artifact.Type, artifact.Description, artifact.StreamId, artifact.Change, artifact.Keys, artifact.Metadata));
+				}
+			}
+
+			return PropertyFilter.Apply(response, filter);
+		}
+
+		/// <summary>
 		/// Gets metadata about an artifact object
 		/// </summary>
@@ -887,39 +916,5 @@
 		{
 			return await DownloadInternalAsync(id, DownloadArtifactFormat.Zip, request.Filter, cancellationToken);
 		}
-
-		/// <summary>
-		/// Finds artifacts matching certain criteria
-		/// </summary>
-		/// <param name="streamId">Stream to search</param>
-		/// <param name="minChange">Minimum changelist number for artifacts to return</param>
-		/// <param name="maxChange">Maximum changelist number for artifacts to return</param>
-		/// <param name="name">Artifact name</param>
-		/// <param name="type">Type of the artifact</param>
-		/// <param name="keys">Keys to find</param>
-		/// <param name="maxResults">Maximum number of results to return</param>
-		/// <param name="filter">Filter for returned values</param>
-		/// <returns>Information about all the artifacts</returns>
-		[HttpGet]
-		[Route("/api/v2/artifacts")]
-		[ProducesResponseType(typeof(FindArtifactsResponse), 200)]
-		public async Task<ActionResult<object>> FindArtifactsAsync([FromQuery] StreamId? streamId = null, [FromQuery] int? minChange = null, [FromQuery] int? maxChange = null, [FromQuery(Name = "name")] ArtifactName? name = null, [FromQuery(Name = "type")] ArtifactType? type = null, [FromQuery(Name = "key")] IEnumerable<string>? keys = null, [FromQuery] int maxResults = 100, [FromQuery] PropertyFilter? filter = null)
-		{
-			if (streamId == null && keys == null)
-			{
-				return BadRequest("Missing streamId or key parameter");
-			}
-
-			FindArtifactsResponse response = new FindArtifactsResponse();
-			await foreach (IArtifact artifact in _artifactCollection.FindAsync(streamId, minChange, maxChange, name, type, keys, maxResults, HttpContext.RequestAborted))
-			{
-				if (_globalConfig.Authorize(artifact.AclScope, ArtifactAclAction.ReadArtifact, User))
-				{
-					response.Artifacts.Add(new GetArtifactResponse(artifact.Id, artifact.Name, artifact.Type, artifact.Description, artifact.StreamId, artifact.Change, artifact.Keys, artifact.Metadata));
-				}
-			}
-
-			return PropertyFilter.Apply(response, filter);
-		}
 	}
 }
