Horde: Bump the GC version number.

#rnx

diff --git a/Engine/Source/Programs/Horde/Horde.Server/Storage/StorageService.cs b/Engine/Source/Programs/Horde/Horde.Server/Storage/StorageService.cs
--- a/Engine/Source/Programs/Horde/Horde.Server/Storage/StorageService.cs
+++ b/Engine/Source/Programs/Horde/Horde.Server/Storage/StorageService.cs
@@ -245,7 +245,7 @@ namespace Horde.Server.Storage
 			}
 		}
 
-		const int CurrentGcVersion = 1;
+		const int CurrentGcVersion = 2;
 
 		internal class BlobInfo
 		{
