Horde: Compute a digest for each block stored in the block cache, and validate it before returning values.

diff --git a/Engine/Source/Programs/Horde/Horde.Server.Tests/Artifacts/BlockCacheTests.cs b/Engine/Source/Programs/Horde/Horde.Server.Tests/Artifacts/BlockCacheTests.cs
--- a/Engine/Source/Programs/Horde/Horde.Server.Tests/Artifacts/BlockCacheTests.cs
+++ b/Engine/Source/Programs/Horde/Horde.Server.Tests/Artifacts/BlockCacheTests.cs
@@ -8,6 +8,7 @@ using System.Security.Cryptography;
 using System.Threading.Tasks;
 using EpicGames.Core;
 using Horde.Server.Artifacts;
+using Microsoft.Extensions.Logging;
 using Microsoft.VisualStudio.TestTools.UnitTesting;
 
 namespace Horde.Server.Tests.Artifacts
@@ -21,10 +22,8 @@ namespace Horde.Server.Tests.Artifacts
 			using BlockCache blockCache = BlockCache.CreateInMemory(1, 4096, 4096);
 			blockCache.Add("hello", new byte[] { 1, 2, 3 });
 
-			IBlockCacheValue? value = blockCache.Get("hello");
-			Assert.IsNotNull(value);
-
-			Assert.IsTrue(value.Data.ToArray().SequenceEqual(new byte[] { 1, 2, 3 }));
+			using IBlockCacheValue? value = blockCache.Get("hello");
+			Assert.IsTrue(value != null && value.Data.ToArray().SequenceEqual(new byte[] { 1, 2, 3 }));
 		}
 
 		[TestMethod]
@@ -35,10 +34,8 @@ namespace Horde.Server.Tests.Artifacts
 			using BlockCache blockCache = BlockCache.CreateInMemory(1);
 			blockCache.Add("hello", data);
 
-			IBlockCacheValue? value = blockCache.Get("hello");
-			Assert.IsNotNull(value);
-
-			Assert.IsTrue(value.Data.ToArray().SequenceEqual(data));
+			using IBlockCacheValue? value = blockCache.Get("hello");
+			Assert.IsTrue(value != null && value.Data.ToArray().SequenceEqual(data));
 		}
 
 		[TestMethod]
@@ -80,14 +77,14 @@ namespace Horde.Server.Tests.Artifacts
 			Random sizeRng = new Random(0);
 
 			List<(string Name, byte[] Data)> items = new List<(string, byte[])>();
-			for (int idx = 0; idx < 512; idx++)
+			for (int idx = 0; idx < 128; idx++)
 			{
 				int size = sizeRng.Next(2048, 16384);
 				items.Add(($"{idx}", RandomNumberGenerator.GetBytes(size)));
 			}
 
-			using BlockCache blockCache = BlockCache.CreateInMemory(10, 64, 4096);
-			Parallel.For(0, 8, threadIdx =>
+			using BlockCache blockCache = BlockCache.CreateInMemory(5, 16, 4096);
+			Parallel.For(0, 4, threadIdx =>
 			{
 				Random rng = new Random(threadIdx);
 				for (int idx = 0; idx < 10000; idx++)
@@ -115,5 +112,42 @@ namespace Horde.Server.Tests.Artifacts
 				Assert.IsTrue(cacheValue == null || cacheValue.Data.ToArray().SequenceEqual(data));
 			}
 		}
+
+		class TestLogger : ILogger
+		{
+			public List<string> Messages { get; } = new List<string>();
+
+			public IDisposable? BeginScope<TState>(TState state) where TState : notnull
+				=> null;
+
+			public bool IsEnabled(LogLevel logLevel)
+				=> true;
+
+			public void Log<TState>(LogLevel logLevel, EventId eventId, TState state, Exception? exception, Func<TState, Exception?, string> formatter)
+				=> Messages.Add(formatter(state, exception));
+		}
+
+		[TestMethod]
+		public void CorruptionTest()
+		{
+			TestLogger logger = new TestLogger();
+
+			using BlockCache blockCache = BlockCache.CreateInMemory(1, 4096, 4096, logger);
+			blockCache.Add("hello", new byte[] { 1, 2, 3 });
+
+			{
+				using IBlockCacheValue? value = blockCache.Get("hello");
+				Assert.IsTrue(value != null && value.Data.ToArray().SequenceEqual(new byte[] { 1, 2, 3 }));
+				Assert.AreEqual(0, logger.Messages.Count);
+			}
+
+			blockCache.DangerousCorruptValue("hello");
+
+			{
+				using IBlockCacheValue? value = blockCache.Get("hello");
+				Assert.IsNull(value);
+				Assert.AreEqual(1, logger.Messages.Count);
+			}
+		}
 	}
 }
diff --git a/Engine/Source/Programs/Horde/Horde.Server/Artifacts/BlockCache.cs b/Engine/Source/Programs/Horde/Horde.Server/Artifacts/BlockCache.cs
--- a/Engine/Source/Programs/Horde/Horde.Server/Artifacts/BlockCache.cs
+++ b/Engine/Source/Programs/Horde/Horde.Server/Artifacts/BlockCache.cs
@@ -7,11 +7,14 @@ using System.Collections.Concurrent;
 using System.Collections.Generic;
 using System.Diagnostics;
 using System.IO;
+using System.IO.Hashing;
 using System.IO.MemoryMappedFiles;
 using System.Numerics;
 using System.Text;
 using System.Threading;
 using EpicGames.Core;
+using Microsoft.Extensions.Logging;
+using Microsoft.Extensions.Logging.Abstractions;
 
 namespace Horde.Server.Artifacts
 {
@@ -40,29 +43,31 @@ namespace Horde.Server.Artifacts
 		const int NumEvictionEntries = 4;
 
 		// Header information for a block
-		record struct BlockHeader(IoHash Hash, uint PackedData)
+		record struct BlockHeader(IoHash Key, uint PackedData, long Digest)
 		{
-			public const int NumBytes = 24;
+			public const int NumBytes = 32;
 
 			public int Index => (int)((PackedData >> 24) & 0xff);
 			public int Count => (int)((PackedData >> 16) & 0xff);
 			public int Length => ((int)(PackedData & 0xffff)) + 1;
 
-			public BlockHeader(IoHash hash, int index, int count, int length)
-				: this(hash, ((uint)index << 24) | ((uint)count << 16) | (uint)(length - 1))
+			public BlockHeader(IoHash key, int index, int count, int length, long digest)
+				: this(key, ((uint)index << 24) | ((uint)count << 16) | (uint)(length - 1), digest)
 			{ }
 
 			public static BlockHeader Read(ReadOnlySpan<byte> data)
 			{
 				IoHash hash = new IoHash(data);
 				uint packed = BinaryPrimitives.ReadUInt32LittleEndian(data[IoHash.NumBytes..]);
-				return new BlockHeader(hash, packed);
+				long digest = BinaryPrimitives.ReadInt64LittleEndian(data[(IoHash.NumBytes + 4)..]);
+				return new BlockHeader(hash, packed, digest);
 			}
 
 			public void Write(Span<byte> data)
 			{
-				Hash.CopyTo(data);
+				Key.CopyTo(data);
 				BinaryPrimitives.WriteUInt32LittleEndian(data[IoHash.NumBytes..], PackedData);
+				BinaryPrimitives.WriteInt64LittleEndian(data[(IoHash.NumBytes + 4)..], Digest);
 			}
 		}
 
@@ -189,6 +194,7 @@ namespace Horde.Server.Artifacts
 		readonly int _blockSizeLog2;
 		readonly ConcurrentDictionary<IoHash, BlockState> _lookup = new ConcurrentDictionary<IoHash, BlockState>();
 		readonly ConcurrentQueue<int> _freeBlocks = new ConcurrentQueue<int>();
+		readonly ILogger _logger;
 
 		int _numAllocatedBlocks = 0;
 
@@ -198,7 +204,8 @@ namespace Horde.Server.Artifacts
 		/// <param name="partitions">Allocated partition data</param>
 		/// <param name="numBlocksPerPartition">Number of blocks in each partition</param>
 		/// <param name="blockSize">The size of each block. Must be a power of two.</param>
-		BlockCache(Partition[] partitions, int numBlocksPerPartition, int blockSize)
+		/// <param name="logger">Logger for diagnostic messages</param>
+		BlockCache(Partition[] partitions, int numBlocksPerPartition, int blockSize, ILogger? logger)
 		{
 			if ((numBlocksPerPartition & (numBlocksPerPartition - 1)) != 0)
 			{
@@ -215,6 +222,7 @@ namespace Horde.Server.Artifacts
 			_numBlocksPerPartitionLog2 = BitOperations.Log2((uint)numBlocksPerPartition);
 			_blockSize = blockSize;
 			_blockSizeLog2 = BitOperations.Log2((uint)blockSize);
+			_logger = logger ?? NullLogger.Instance;
 		}
 
 		/// <summary>
@@ -223,7 +231,8 @@ namespace Horde.Server.Artifacts
 		/// <param name="numPartitions">Number of partitions in the cache</param>
 		/// <param name="numBlocksPerPartition">Number of blocks in each partition</param>
 		/// <param name="blockSize">Size of each block</param>
-		public static BlockCache CreateInMemory(int numPartitions, int numBlocksPerPartition = DefaultNumBlocksPerPartition, int blockSize = DefaultBlockSize)
+		/// <param name="logger">Logger for diagnostic messages</param>
+		public static BlockCache CreateInMemory(int numPartitions, int numBlocksPerPartition = DefaultNumBlocksPerPartition, int blockSize = DefaultBlockSize, ILogger? logger = null)
 		{
 			int partitionSize = GetPartitionSize(numBlocksPerPartition, blockSize);
 
@@ -234,7 +243,7 @@ namespace Horde.Server.Artifacts
 				partitions[idx] = new Partition(numBlocksPerPartition, data);
 			}
 
-			return new BlockCache(partitions, numBlocksPerPartition, blockSize);
+			return new BlockCache(partitions, numBlocksPerPartition, blockSize, logger);
 		}
 
 		/// <summary>
@@ -244,7 +253,8 @@ namespace Horde.Server.Artifacts
 		/// <param name="numPartitions">Number of partitions to create. Each partition is ~1gb.</param>
 		/// <param name="numBlocksPerPartition">Number of blocks in each partition</param>
 		/// <param name="blockSize">Size of a block</param>
-		public static BlockCache Create(DirectoryReference rootDir, int numPartitions, int numBlocksPerPartition = DefaultNumBlocksPerPartition, int blockSize = DefaultBlockSize)
+		/// <param name="logger">Logger for diagnostic messages</param>
+		public static BlockCache Create(DirectoryReference rootDir, int numPartitions, int numBlocksPerPartition = DefaultNumBlocksPerPartition, int blockSize = DefaultBlockSize, ILogger? logger = null)
 		{
 			MemoryMappedFilePartition[] partitions = new MemoryMappedFilePartition[numPartitions];
 			try
@@ -255,7 +265,7 @@ namespace Horde.Server.Artifacts
 					FileReference file = FileReference.Combine(rootDir, $"partition{idx:0000}.dat");
 					partitions[idx] = MemoryMappedFilePartition.Create(file, numBlocksPerPartition, blockSize);
 				}
-				return new BlockCache(partitions, numBlocksPerPartition, blockSize);
+				return new BlockCache(partitions, numBlocksPerPartition, blockSize, logger);
 			}
 			catch
 			{
@@ -294,7 +304,7 @@ namespace Horde.Server.Artifacts
 			for (int blockIdx = 0; blockIdx < _numAllocatedBlocks; blockIdx++)
 			{
 				BlockHeader header = GetBlockHeader(blockIdx);
-				if (header.Hash == IoHash.Zero)
+				if (header.Key == IoHash.Zero)
 				{
 					numFreeBlocks++;
 				}
@@ -331,29 +341,37 @@ namespace Horde.Server.Artifacts
 					blockIdx = _rng.Next(_numBlocks);
 
 					BlockHeader header = GetBlockHeader(blockIdx);
-					if (_lookup.TryGetValue(header.Hash, out BlockState? state))
+					if (_lookup.TryGetValue(header.Key, out BlockState? state))
 					{
-						priorityQueue.Enqueue((header.Hash, state), state.LastAccessTicks);
+						priorityQueue.Enqueue((header.Key, state), state.LastAccessTicks);
 					}
 				}
 
 				for (int idx = 0; idx < NumEvictionEntries && priorityQueue.Count > 0; idx++)
 				{
 					(IoHash hash, BlockState evictState) = priorityQueue.Dequeue();
-					if (evictState.TryAddWriteLock())
-					{
-						_lookup.TryRemove(hash, out _);
-						foreach (int evictBlockIdx in evictState.BlockIdxs)
-						{
-							SetBlockHeader(evictBlockIdx, default);
-							_freeBlocks.Enqueue(evictBlockIdx);
-						}
-					}
+					TryDelete(hash, evictState);
 				}
 			}
 			return blockIdx;
 		}
 
+		bool TryDelete(IoHash hash, BlockState blockState)
+		{
+			if (!blockState.TryAddWriteLock())
+			{
+				return false;
+			}
+
+			_lookup.TryRemove(hash, out _);
+			foreach (int evictBlockIdx in blockState.BlockIdxs)
+			{
+				SetBlockHeader(evictBlockIdx, default);
+				_freeBlocks.Enqueue(evictBlockIdx);
+			}
+			return true;
+		}
+
 		BlockHeader GetBlockHeader(int blockIdx)
 		{
 			Partition partition = _partitions[blockIdx >> _numBlocksPerPartitionLog2];
@@ -383,7 +401,7 @@ namespace Horde.Server.Artifacts
 				return false;
 			}
 
-			IoHash hash = IoHash.Compute(Encoding.UTF8.GetBytes(key));
+			IoHash keyHash = IoHash.Compute(Encoding.UTF8.GetBytes(key));
 
 			int numBlocks = (value.Length + (_blockSize - 1)) >> _blockSizeLog2;
 			int[] blockIdxs = new int[numBlocks];
@@ -396,15 +414,17 @@ namespace Horde.Server.Artifacts
 				int freeBlockIdx = GetNextFreeBlockIdx();
 				blockIdxs[idx] = freeBlockIdx;
 
-				Memory<byte> blockData = GetBlockData(freeBlockIdx, _blockSize);
+				Memory<byte> blockData = GetBlockData(freeBlockIdx, length);
 				value.Slice(offset, length).CopyTo(blockData);
 
-				BlockHeader header = new BlockHeader(hash, idx, numBlocks, length);
+				long digest = ComputeDigest(blockData.Span);
+
+				BlockHeader header = new BlockHeader(keyHash, idx, numBlocks, length, digest);
 				SetBlockHeader(freeBlockIdx, header);
 			}
 
 			BlockState state = new BlockState(blockIdxs, value.Length & (_blockSize - 1));
-			if (!_lookup.TryAdd(hash, state))
+			if (!_lookup.TryAdd(keyHash, state))
 			{
 				foreach (int blockIdx in blockIdxs)
 				{
@@ -415,6 +435,13 @@ namespace Horde.Server.Artifacts
 			return true;
 		}
 
+		static long ComputeDigest(ReadOnlySpan<byte> span)
+		{
+			Span<byte> result = stackalloc byte[sizeof(long)];
+			XxHash64.Hash(span, result);
+			return BinaryPrimitives.ReadInt64BigEndian(result);
+		}
+
 		/// <summary>
 		/// Gets the data associated with a given key
 		/// </summary>
@@ -433,13 +460,40 @@ namespace Horde.Server.Artifacts
 			blockState.Touch();
 
 			ReadOnlySequenceBuilder<byte> builder = new ReadOnlySequenceBuilder<byte>();
-			for (int idx = 0; idx + 1 < blockState.BlockIdxs.Length; idx++)
+			for(int idx = 0; idx < blockState.BlockIdxs.Length; idx++)
 			{
-				builder.Append(GetBlockData(blockState.BlockIdxs[idx], _blockSize));
+				int blockIdx = blockState.BlockIdxs[idx];
+
+				BlockHeader header = GetBlockHeader(blockIdx);
+				ReadOnlyMemory<byte> data = GetBlockData(blockIdx, header.Length);
+
+				long digest = ComputeDigest(data.Span);
+				if (digest != header.Digest)
+				{
+					_logger.LogWarning("Corrupt block in cache '{Key}' ({Index}/{Count}) (hash: {Hash}, block {BlockIdx}, expected digest: {Digest:x8}, actual digest: {ActualDigest:x8})", key, idx + 1, blockState.BlockIdxs.Length, hash, blockIdx, header.Digest, digest);
+					blockState.ReleaseReadLock();
+					TryDelete(hash, blockState);
+					return null;
+				}
+
+				builder.Append(data);
 			}
-			builder.Append(GetBlockData(blockState.BlockIdxs[^1], blockState.LastBlockLength));
 
 			return new BlockCacheValue(builder.Construct(), blockState);
 		}
+
+		/// <summary>
+		/// Helper method for testing corruption handling
+		/// </summary>
+		internal void DangerousCorruptValue(string key)
+		{
+			IoHash hash = IoHash.Compute(Encoding.UTF8.GetBytes(key));
+
+			BlockState? blockState;
+			if (_lookup.TryGetValue(hash, out blockState) && blockState.TryAddReadLock())
+			{
+				GetBlockData(blockState.BlockIdxs[0], 1).Span[0]++;
+			}
+		}
 	}
 }
