PR #12025: Horde: Fix infinite recursion for HasClaim methods for IAccount and IServiceAccount classes

#rnx

diff --git a/Engine/Source/Programs/Horde/Horde.Server/Accounts/IAccount.cs b/Engine/Source/Programs/Horde/Horde.Server/Accounts/IAccount.cs
--- a/Engine/Source/Programs/Horde/Horde.Server/Accounts/IAccount.cs
+++ b/Engine/Source/Programs/Horde/Horde.Server/Accounts/IAccount.cs
@@ -142,7 +142,7 @@ namespace Horde.Server.Accounts
 		/// <param name="claim">Claim to test for</param>
 		/// <returns>True if the user has the claim</returns>
 		public static bool HasClaim(this IAccount account, AclClaimConfig claim)
-			=> HasClaim(account, claim);
+			=> HasClaim(account, claim.Type, claim.Value);
 
 		/// <summary>
 		/// Update settings for the account, retrying if the account object has changed
diff --git a/Engine/Source/Programs/Horde/Horde.Server/ServiceAccounts/IServiceAccount.cs b/Engine/Source/Programs/Horde/Horde.Server/ServiceAccounts/IServiceAccount.cs
--- a/Engine/Source/Programs/Horde/Horde.Server/ServiceAccounts/IServiceAccount.cs
+++ b/Engine/Source/Programs/Horde/Horde.Server/ServiceAccounts/IServiceAccount.cs
@@ -98,7 +98,7 @@ namespace Horde.Server.ServiceAccounts
 		/// <param name="claim">Claim to test for</param>
 		/// <returns>True if the user has the claim</returns>
 		public static bool HasClaim(this IServiceAccount account, AclClaimConfig claim)
-			=> HasClaim(account, claim);
+			=> HasClaim(account, claim.Type, claim.Value);
 
 		/// <summary>
 		/// Update settings for the account, retrying if the account object has changed
diff --git a/Engine/Source/Programs/Horde/Horde.Server/ServiceAccounts/ServiceAccountCollection.cs b/Engine/Source/Programs/Horde/Horde.Server/ServiceAccounts/ServiceAccountCollection.cs
--- a/Engine/Source/Programs/Horde/Horde.Server/ServiceAccounts/ServiceAccountCollection.cs
+++ b/Engine/Source/Programs/Horde/Horde.Server/ServiceAccounts/ServiceAccountCollection.cs
@@ -1,5 +1,6 @@
 // Copyright Epic Games, Inc. All Rights Reserved.
 
+using System;
 using System.Collections.Generic;
 using System.Linq;
 using System.Security.Cryptography;
@@ -43,7 +44,7 @@ namespace Horde.Server.ServiceAccounts
 			public ServiceAccountId Id { get; set; }
 
 			/// <inheritdoc/>
-			public string SecretToken { get; set; } = string.Empty;
+			public string SecretToken { get; set; } = String.Empty;
 
 			[BsonElement("Claims2")]
 			public List<ClaimDocument> Claims { get; set; } = new List<ClaimDocument>();
