Horde: Fix static analysis warnings.

#rnx

diff --git a/Engine/Source/Programs/Horde/Horde.Server.Tests/Artifacts/BlockCacheTests.cs b/Engine/Source/Programs/Horde/Horde.Server.Tests/Artifacts/BlockCacheTests.cs
--- a/Engine/Source/Programs/Horde/Horde.Server.Tests/Artifacts/BlockCacheTests.cs
+++ b/Engine/Source/Programs/Horde/Horde.Server.Tests/Artifacts/BlockCacheTests.cs
@@ -6,7 +6,6 @@ using System.Collections.Generic;
 using System.Linq;
 using System.Security.Cryptography;
 using System.Threading.Tasks;
-using EpicGames.Core;
 using Horde.Server.Artifacts;
 using Microsoft.Extensions.Logging;
 using Microsoft.VisualStudio.TestTools.UnitTesting;
diff --git a/Engine/Source/Programs/Horde/Horde.Server.Tests/Devices/DeviceServiceTests.cs b/Engine/Source/Programs/Horde/Horde.Server.Tests/Devices/DeviceServiceTests.cs
--- a/Engine/Source/Programs/Horde/Horde.Server.Tests/Devices/DeviceServiceTests.cs
+++ b/Engine/Source/Programs/Horde/Horde.Server.Tests/Devices/DeviceServiceTests.cs
@@ -15,7 +15,6 @@ using Horde.Server.Devices;
 using Horde.Server.Jobs;
 using Horde.Server.Jobs.Graphs;
 using Horde.Server.Jobs.Templates;
-using Horde.Server.Server;
 using Horde.Server.Users;
 using Horde.Server.Utilities;
 using Microsoft.AspNetCore.Http;
diff --git a/Engine/Source/Programs/Horde/Horde.Server/Server/GlobalConfig.cs b/Engine/Source/Programs/Horde/Horde.Server/Server/GlobalConfig.cs
--- a/Engine/Source/Programs/Horde/Horde.Server/Server/GlobalConfig.cs
+++ b/Engine/Source/Programs/Horde/Horde.Server/Server/GlobalConfig.cs
@@ -15,7 +15,6 @@ using EpicGames.Horde.Agents.Pools;
 using EpicGames.Horde.Artifacts;
 using EpicGames.Horde.Common;
 using EpicGames.Horde.Compute;
-using EpicGames.Horde.Devices;
 using EpicGames.Horde.Jobs.Templates;
 using EpicGames.Horde.Projects;
 using EpicGames.Horde.Secrets;
