Horde: Update release notes.

#rnx

diff --git a/Engine/Source/Programs/Horde/Docs/ReleaseNotes.md b/Engine/Source/Programs/Horde/Docs/ReleaseNotes.md
--- a/Engine/Source/Programs/Horde/Docs/ReleaseNotes.md
+++ b/Engine/Source/Programs/Horde/Docs/ReleaseNotes.md
@@ -1,5 +1,16 @@
 # Release Notes
 
+## 2024-06-20
+
+* Compute a digest for each block stored in the block cache, and validate it before returning values. (34515388)
+* Remove requirement to specify a stream or key to enumerate artifacts. (34513227)
+* Downgrade log event about duplicate build products to information, since we don't have the list of duplicate build products to ignore in Horde. (34481396)
+* Exclude blocks for empty files from the Unsync manifest. (34464454)
+* Add IPoolSizeStrategyFactory interface, which is used for creating IPoolSizeStrategy instances. This allows us to remove the graph/job/stream collection interfaces from FleetService, allowing it to exist without any job handling code in the solution. (34457658)
+* Do not shutdown disabled agents by default. Prior to this change, this value defaulted to 8 hours which assumes you want this behavior in the first place. Making this nullable allow for opt-in instead. (34455677)
+* Remove tracing span from shared tickers. For long-running callbacks, this can create data which is difficult for Otel/DD to handle. (34444738)
+* Add a debug endpoint which streams a random block of data to the caller. Usage is: /api/v1/debug/randomdata?size=1mb (34426596)
+
 ## 2024-06-17
 
 * Add support for auto-assigning cluster ID during compute allocation requests (34423910)
