Horde: Catch cancellation exceptions thrown from MongoBufferedWriter rather than logging them as an error.

#rnx

diff --git a/Engine/Source/Programs/Horde/Horde.Server/Utilities/MongoBufferedWriter.cs b/Engine/Source/Programs/Horde/Horde.Server/Utilities/MongoBufferedWriter.cs
--- a/Engine/Source/Programs/Horde/Horde.Server/Utilities/MongoBufferedWriter.cs
+++ b/Engine/Source/Programs/Horde/Horde.Server/Utilities/MongoBufferedWriter.cs
@@ -61,6 +61,7 @@ namespace Horde.Server.Utilities
 		public async ValueTask StopAsync(CancellationToken cancellationToken)
 		{
 			await _backgroundTask.StopAsync(cancellationToken);
+			await FlushAsync(cancellationToken);
 		}
 
 		// Flushes the sink in the background
@@ -81,6 +82,10 @@ namespace Horde.Server.Utilities
 
 					await FlushAsync(cancellationToken);
 				}
+				catch (OperationCanceledException)
+				{
+					break;
+				}
 				catch (Exception ex)
 				{
 					_logger.LogError(ex, "Exception in buffered writer: {Message}", ex.Message);
