Horde: Fix error due to setting response code after request has already been written.

#rnx

diff --git a/Engine/Source/Programs/Horde/Horde.Server/Artifacts/ArtifactsController.cs b/Engine/Source/Programs/Horde/Horde.Server/Artifacts/ArtifactsController.cs
--- a/Engine/Source/Programs/Horde/Horde.Server/Artifacts/ArtifactsController.cs
+++ b/Engine/Source/Programs/Horde/Horde.Server/Artifacts/ArtifactsController.cs
@@ -691,7 +691,7 @@ namespace Horde.Server.Artifacts
 			}
 
 			await response.CompleteAsync();
-			return Ok();
+			return Empty;
 		}
 
 		/// <summary>
