Horde: Fix exception checking whether blobs exist.

#rnx

diff --git a/Engine/Source/Programs/Horde/Horde.Server/Storage/StorageService.cs b/Engine/Source/Programs/Horde/Horde.Server/Storage/StorageService.cs
--- a/Engine/Source/Programs/Horde/Horde.Server/Storage/StorageService.cs
+++ b/Engine/Source/Programs/Horde/Horde.Server/Storage/StorageService.cs
@@ -602,7 +602,7 @@ namespace Horde.Server.Storage
 				{
 					string path = locator.Path.ToString();
 
-					BlobInfo? blobInfo = await _blobCollection.Find(x => x.NamespaceId == namespaceId && x.Path == path && x.GcVersion >= CurrentGcVersion).FirstAsync(cancellationToken);
+					BlobInfo? blobInfo = await _blobCollection.Find(x => x.NamespaceId == namespaceId && x.Path == path && x.GcVersion >= CurrentGcVersion).FirstOrDefaultAsync(cancellationToken);
 					if (blobInfo != null)
 					{
 						_logger.LogWarning("Blob {BlobId} ({Locator}) accessed after being garbage collected", blobInfo.Id, locator);
