Horde: Tool deployment management feature
#rnx
#skipci

diff --git a/Engine/Source/Programs/Horde/HordeDashboard/src/backend/Api.ts b/Engine/Source/Programs/Horde/HordeDashboard/src/backend/Api.ts
--- a/Engine/Source/Programs/Horde/HordeDashboard/src/backend/Api.ts
+++ b/Engine/Source/Programs/Horde/HordeDashboard/src/backend/Api.ts
@@ -1334,7 +1334,7 @@ export type UpdateStepRequest = {
 	/**If the step has been requested to abort */
 	abortRequested?: boolean;
 
-	/** The reason the step was canceled */	
+	/** The reason the step was canceled */
 	cancellationReason?: string;
 
 	/**Specifies the log file id for this step */
@@ -5046,7 +5046,21 @@ export type GetTestDataRefResponse = {
 	/// The number of suite tests swith errors
 	suiteSuccessCount?: number;
 }
-	
+
+/** Current state of a tool's deployment */
+export enum ToolDeploymentState {
+	/** The deployment is ongoing */
+	Active = "Active",
+
+	/** The deployment should be paused at its current state */
+	Paused = "Paused",
+
+	/** Deployment of this version is complete */
+	Complete = "Complete",
+
+	/** The deployment has been cancelled. */
+	Cancelled = "Canceled"
+}
 
 /** Summary for a particular tool */
 export type GetToolSummaryResponse = {
@@ -5075,6 +5089,12 @@ export type GetToolSummaryResponse = {
 	/** Identifier for the current deployment */
 	deploymentId?: string;
 
+	/** State of the current deployment */
+	deploymentState?: ToolDeploymentState;
+
+	/** Progress of the current deployment */
+	deploymentProgress?: number;
+
 	/** Whether to show this tool for download inside UGS */
 	showInUGS: boolean;
 
@@ -5082,6 +5102,38 @@ export type GetToolSummaryResponse = {
 	showInDashboard: boolean;
 }
 
+export type GetToolDeploymentResponse = {
+	id: string;
+	version: string;
+	state: ToolDeploymentState;
+	progress: number;
+	startedAt: Date | string;
+	duration: any;
+	refName: string;
+	locator: any;
+}
+
+export type GetToolResponse = {
+	id: string;
+	name: string;
+	description: string;
+	category?: string;
+	group?: string;
+	platforms?: string[];
+	deployments: GetToolDeploymentResponse[];
+	public: boolean;
+	showInUgs: boolean;
+	showInDashboard: boolean;
+};
+
+/**  Update an existing deployment */
+export type UpdateDeploymentRequest = {
+
+	/** New state for the deployment */
+	state?: ToolDeploymentState;
+}
+
+
 /** Job Bisect */
 
 /// State of a bisect task
diff --git a/Engine/Source/Programs/Horde/HordeDashboard/src/backend/Backend.ts b/Engine/Source/Programs/Horde/HordeDashboard/src/backend/Backend.ts
--- a/Engine/Source/Programs/Horde/HordeDashboard/src/backend/Backend.ts
+++ b/Engine/Source/Programs/Horde/HordeDashboard/src/backend/Backend.ts
@@ -1,7 +1,7 @@
 // Copyright Epic Games, Inc. All Rights Reserved.
 
 import templateCache from '../backend/TemplateCache';
-import { AccountClaimMessage, AgentData, AgentQuery, ApproveAgentsRequest, ArtifactData, AuditLogEntry, AuditLogQuery, BatchUpdatePoolRequest, ChangeSummaryData, CreateAccountRequest, CreateBisectTaskRequest, CreateBisectTaskResponse, CreateDeviceRequest, CreateDeviceResponse, CreateExternalIssueRequest, CreateExternalIssueResponse, CreateJobRequest, CreateJobResponse, CreateNoticeRequest, CreatePoolRequest, CreateServiceAccountRequest, CreateServiceAccountResponse, CreateSoftwareResponse, CreateSubscriptionRequest, CreateSubscriptionResponse, CreateZipRequest, DashboardPreference, DevicePoolTelemetryQuery, DeviceTelemetryQuery, EventData, FindArtifactsResponse, FindIssueResponse, FindJobTimingsResponse, GetAccountResponse, GetAgentLeaseResponse, GetAgentSessionResponse, GetAgentTelemetrySampleResponse, GetArtifactDirectoryResponse, GetArtifactZipRequest, GetBisectTaskResponse, GetDashboardConfigResponse, GetDevicePlatformResponse, GetDevicePoolResponse, GetDevicePoolTelemetryResponse, GetDeviceReservationResponse, GetDeviceResponse, GetDeviceTelemetryResponse, GetExternalIssueProjectResponse, GetExternalIssueResponse, GetIssueStreamResponse, GetJobsTabResponse, GetJobStepRefResponse, GetJobStepTraceResponse, GetJobTimingResponse, GetLogEventResponse, GetLogFileResponse, GetNoticeResponse, GetNotificationResponse, GetPendingAgentsResponse, GetPerforceServerStatusResponse, GetPoolResponse, GetPoolSummaryResponse, GetServerInfoResponse, GetServerSettingsResponse, GetServiceAccountResponse, GetSoftwareResponse, GetSubscriptionResponse, GetTelemetryMetricsResponse, GetTestDataDetailsResponse, GetTestDataRefResponse, GetTestMetaResponse, GetTestResponse, GetTestsRequest, GetTestStreamResponse, GetToolSummaryResponse, GetUserResponse, GetUtilizationTelemetryResponse, GlobalConfig, IssueData, IssueQuery, IssueQueryV2, JobData, JobQuery, JobsTabColumnType, JobStepOutcome, JobStreamQuery, JobTimingsQuery, LeaseData, LogData, LogEventQuery, LogLineData, MetricsQuery, PoolData, PoolQuery, PreflightConfigResponse, ProjectData, ScheduleData, ScheduleQuery, SearchLogFileResponse, ServerStatusResponse, ServerUpdateResponse, SessionData, StreamData, TabType, TestData, UpdateAccountRequest, UpdateAgentRequest, UpdateBisectTaskRequest, UpdateCurrentAccountRequest, UpdateDeviceRequest, UpdateGlobalConfigRequest, UpdateIssueRequest, UpdateJobRequest, UpdateLeaseRequest, UpdateNoticeRequest, UpdateNotificationsRequest, UpdatePoolRequest, UpdateServerSettingsRequest, UpdateServiceAccountRequest, UpdateServiceAccountResponse, UpdateStepRequest, UpdateStepResponse, UpdateTemplateRefRequest, UpdateUserRequest, UsersQuery } from './Api';
+import { AccountClaimMessage, AgentData, AgentQuery, ApproveAgentsRequest, ArtifactData, AuditLogEntry, AuditLogQuery, BatchUpdatePoolRequest, ChangeSummaryData, CreateAccountRequest, CreateBisectTaskRequest, CreateBisectTaskResponse, CreateDeviceRequest, CreateDeviceResponse, CreateExternalIssueRequest, CreateExternalIssueResponse, CreateJobRequest, CreateJobResponse, CreateNoticeRequest, CreatePoolRequest, CreateServiceAccountRequest, CreateServiceAccountResponse, CreateSoftwareResponse, CreateSubscriptionRequest, CreateSubscriptionResponse, CreateZipRequest, DashboardPreference, DevicePoolTelemetryQuery, DeviceTelemetryQuery, EventData, FindArtifactsResponse, FindIssueResponse, FindJobTimingsResponse, GetAccountResponse, GetAgentLeaseResponse, GetAgentSessionResponse, GetAgentTelemetrySampleResponse, GetArtifactDirectoryResponse, GetArtifactZipRequest, GetBisectTaskResponse, GetDashboardConfigResponse, GetDevicePlatformResponse, GetDevicePoolResponse, GetDevicePoolTelemetryResponse, GetDeviceReservationResponse, GetDeviceResponse, GetDeviceTelemetryResponse, GetExternalIssueProjectResponse, GetExternalIssueResponse, GetIssueStreamResponse, GetJobsTabResponse, GetJobStepRefResponse, GetJobStepTraceResponse, GetJobTimingResponse, GetLogEventResponse, GetLogFileResponse, GetNoticeResponse, GetNotificationResponse, GetPendingAgentsResponse, GetPerforceServerStatusResponse, GetPoolResponse, GetPoolSummaryResponse, GetServerInfoResponse, GetServerSettingsResponse, GetServiceAccountResponse, GetSoftwareResponse, GetSubscriptionResponse, GetTelemetryMetricsResponse, GetTestDataDetailsResponse, GetTestDataRefResponse, GetTestMetaResponse, GetTestResponse, GetTestsRequest, GetTestStreamResponse, GetToolResponse, GetToolSummaryResponse, GetUserResponse, GetUtilizationTelemetryResponse, GlobalConfig, IssueData, IssueQuery, IssueQueryV2, JobData, JobQuery, JobsTabColumnType, JobStepOutcome, JobStreamQuery, JobTimingsQuery, LeaseData, LogData, LogEventQuery, LogLineData, MetricsQuery, PoolData, PoolQuery, PreflightConfigResponse, ProjectData, ScheduleData, ScheduleQuery, SearchLogFileResponse, ServerStatusResponse, ServerUpdateResponse, SessionData, StreamData, TabType, TestData, UpdateAccountRequest, UpdateAgentRequest, UpdateBisectTaskRequest, UpdateCurrentAccountRequest, UpdateDeploymentRequest, UpdateDeviceRequest, UpdateGlobalConfigRequest, UpdateIssueRequest, UpdateJobRequest, UpdateLeaseRequest, UpdateNoticeRequest, UpdateNotificationsRequest, UpdatePoolRequest, UpdateServerSettingsRequest, UpdateServiceAccountRequest, UpdateServiceAccountResponse, UpdateStepRequest, UpdateStepResponse, UpdateTemplateRefRequest, UpdateUserRequest, UsersQuery } from './Api';
 import dashboard, { Dashboard } from './Dashboard';
 import { ChallengeStatus, Fetch } from './Fetch';
 import { projectStore } from './ProjectStore';
@@ -1612,6 +1612,26 @@ export class Backend {
         });
     }
 
+    getTool(toolId: string): Promise<GetToolResponse> {
+        return new Promise<GetToolResponse>((resolve, reject) => {
+            this.backend.get(`/api/v1/tools/${toolId}`).then((value) => {
+                resolve(value.data as GetToolResponse);
+            }).catch(reason => {
+                reject(reason);
+            });
+        });
+    }
+
+    /// Updates the state of an active deployment.
+    updateDeployment(toolId: string, deploymentId: string, update: UpdateDeploymentRequest): Promise<void> {
+        return new Promise<void>((resolve, reject) => {
+            this.backend.patch(`/api/v1/tools/${toolId}/deployments/${deploymentId}`, update).then(() => {
+                resolve();
+            }).catch(reason => {
+                reject(reason);
+            });
+        });
+    }
     // update a notice
     updateNotice(request: UpdateNoticeRequest): Promise<void> {
         return new Promise<void>((resolve, reject) => {
diff --git a/Engine/Source/Programs/Horde/HordeDashboard/src/components/ToolView.tsx b/Engine/Source/Programs/Horde/HordeDashboard/src/components/ToolView.tsx
--- a/Engine/Source/Programs/Horde/HordeDashboard/src/components/ToolView.tsx
+++ b/Engine/Source/Programs/Horde/HordeDashboard/src/components/ToolView.tsx
@@ -1,6 +1,6 @@
 // Copyright Epic Games, Inc. All Rights Reserved.
 
-import { DetailsList, DetailsListLayoutMode, IColumn, IDetailsGroupDividerProps, IDetailsGroupRenderProps, IDetailsListProps, IGroup, IconButton, Pivot, PivotItem, PrimaryButton, SelectionMode, Stack, Text } from "@fluentui/react";
+import { DetailsList, DetailsListLayoutMode, DirectionalHint, IColumn, IContextualMenuProps, IDetailsGroupDividerProps, IDetailsGroupRenderProps, IDetailsListProps, IGroup, IconButton, Pivot, PivotItem, PrimaryButton, SelectionMode, Stack, Text } from "@fluentui/react";
 import { observer } from "mobx-react-lite";
 import React, { useEffect, useState } from "react";
 import backend from "../backend";
@@ -10,6 +10,7 @@ import { useWindowSize } from "../base/utilities/hooks";
 import { getHordeStyling } from "../styles/Styles";
 import { Breadcrumbs } from "./Breadcrumbs";
 import { TopNav } from "./TopNav";
+import { ToolManagementModal } from "./tools/ToolManagement";
 
 const defaultCategory = "General";
 
@@ -135,6 +136,8 @@ const handler = new ToolHandler();
 
 const ToolPanel: React.FC<{ selectedKey: string }> = observer(({ selectedKey }) => {
 
+   const [manageTool, setManageTool] = useState("");
+
    useEffect(() => {
 
       handler.start();
@@ -204,6 +207,20 @@ const ToolPanel: React.FC<{ selectedKey: string }> = observer(({ selectedKey })
          width = 366;
       }
 
+      const downloadProps: IContextualMenuProps = {
+         items: [
+            {
+               key: 'managetool',
+               text: 'Manage',
+               onClick: () => {
+                  setManageTool(tool.id);
+               }
+            },
+         ],
+         directionalHint: DirectionalHint.bottomRightEdge
+      };
+
+
       return <Stack horizontal verticalAlign="center" verticalFill={true} style={{ padding: 12 }}>
          {!!groupProps && groupProps.group!.count > 1 && <IconButton style={{ color: modeColors.text, fontSize: 18, marginRight: 4 }} iconProps={{ iconName: groupProps.group?.isCollapsed ? 'ChevronRight' : 'ChevronDown' }} onClick={(event: any) => {
             event?.stopPropagation();
@@ -211,17 +228,17 @@ const ToolPanel: React.FC<{ selectedKey: string }> = observer(({ selectedKey })
          }}
          />}
 
-         <Stack style={{ width: width, paddingLeft: pad}}>
+         <Stack style={{ width: width, paddingLeft: pad }}>
             <Text style={{ fontFamily: (groupProps || !tool.group) ? "Horde Open Sans SemiBold" : undefined, color: modeColors.text }}>{tool.name}</Text>
          </Stack>
          <Stack style={{ width: 580 }}>
             <Text style={{ color: modeColors.text }}>{tool.description ?? ""}</Text>
          </Stack>
-         <Stack style={{ width: 260 }}>
+         <Stack style={{ width: 220 }}>
             <Text style={{ color: modeColors.text }}>{tool.version}</Text>
          </Stack>
          <Stack >
-            <PrimaryButton style={{ width: 120, color: "#FFFFFF" }} text="Download" href={`/api/v1/tools/${tool.id}?action=download`} />
+            <PrimaryButton split menuProps={downloadProps} style={{ width: 140, color: "#FFFFFF" }} text="Download" href={`/api/v1/tools/${tool.id}?action=download`} />
          </Stack>
       </Stack>
    }
@@ -256,6 +273,7 @@ const ToolPanel: React.FC<{ selectedKey: string }> = observer(({ selectedKey })
 
 
    return <Stack>
+      {!!manageTool && <ToolManagementModal toolId={manageTool} toolName={handler.tools.find(t => t.id === manageTool)?.name} onClose={() => setManageTool("")} />}
       {!tools.length && handler.loaded && <Stack style={{ paddingBottom: 12 }}>
          <Stack verticalAlign="center">
             <Stack horizontalAlign="center">
diff --git a/Engine/Source/Programs/Horde/HordeDashboard/src/components/tools/ToolManagement.tsx b/Engine/Source/Programs/Horde/HordeDashboard/src/components/tools/ToolManagement.tsx
new file mode 100644
--- /dev/null
+++ b/Engine/Source/Programs/Horde/HordeDashboard/src/components/tools/ToolManagement.tsx
@@ -0,0 +1,304 @@
+import { DefaultButton, DetailsList, DetailsListLayoutMode, IContextualMenuProps, IDetailsListProps, IconButton, Modal, PrimaryButton, ProgressIndicator, SelectionMode, Spinner, SpinnerSize, Stack, Text } from "@fluentui/react";
+import { useConst } from '@fluentui/react-hooks';
+import Markdown from "markdown-to-jsx";
+import { observer } from "mobx-react-lite";
+import moment from "moment";
+import { useState } from "react";
+import backend from "../../backend";
+import { GetToolDeploymentResponse, GetToolResponse, GetToolSummaryResponse, ToolDeploymentState, UpdateDeploymentRequest } from "../../backend/Api";
+import dashboard from "../../backend/Dashboard";
+import { PollBase } from "../../backend/PollBase";
+import { getHordeStyling } from "../../styles/Styles";
+import ErrorHandler from "../ErrorHandler";
+
+class ToolHandler extends PollBase {
+
+   constructor(toolId: string, pollTime = 10000) {
+      super(pollTime);
+      this.toolId = toolId;
+
+      this.poll();
+   }
+
+   async poll(): Promise<void> {
+
+      const requests = [backend.getTools(), backend.getTool(this.toolId)]
+
+      const results = await Promise.all(requests);
+
+      const tools = results[0] as GetToolSummaryResponse[];
+
+      this.data = results[1] as GetToolResponse;
+
+      this.currentDeployment = undefined;
+
+      const tool = tools.find(t => t.id === this.data?.id);
+
+      if (tool) {
+         this.currentDeployment = this.data!.deployments?.find(d => d.id === tool.deploymentId);
+
+      }
+
+      this.setUpdated();
+   }
+
+   data?: GetToolResponse
+   currentDeployment?: GetToolDeploymentResponse;
+
+   toolId: string;
+}
+
+type DeploymentItem = {
+   deploy: GetToolDeploymentResponse;
+}
+
+const ActionConfirmation: React.FC<{ handler: ToolHandler, request: UpdateDeploymentRequest, text: string, actionName: string, deployment: GetToolDeploymentResponse, onClose: () => void }> = observer(({ handler, request, text, actionName, deployment, onClose }) => {
+
+   const [submitting, setSubmitting] = useState(false);
+
+   const { hordeClasses, modeColors } = getHordeStyling();
+   if (submitting) {
+      return <Modal isOpen={true} isBlocking={true} topOffsetFixed={true} styles={{ main: { padding: 8, width: 800, backgroundColor: modeColors.background, height: 180, hasBeenOpened: false, top: "24px", position: "absolute" } }} className={hordeClasses.modal}>
+         <Stack tokens={{ childrenGap: 40 }} styles={{ root: { padding: 8, paddingBottom: 24 } }}>
+            <Stack grow verticalAlign="center" horizontalAlign="center" style={{ paddingTop: 12 }}>
+               <Text variant="mediumPlus" styles={{ root: { fontWeight: "unset", fontFamily: "Horde Open Sans SemiBold" } }}>{`Updating Deployment, One Moment Please...`}</Text>
+            </Stack>
+            <Stack verticalAlign="center">
+               <Spinner size={SpinnerSize.large} />
+            </Stack>
+         </Stack>
+      </Modal>
+   }
+
+   const height = 120;
+
+   const onSubmit = async () => {
+
+      setSubmitting(true);
+
+      let errorReason: any = undefined;
+
+      try {
+
+         await backend.updateDeployment(handler.data!.id, deployment.id, request);
+
+         handler.poll();
+
+      } catch (reason) {
+
+         errorReason = reason ? reason : "Unknown Error";
+
+      } finally {
+         if (errorReason) {
+
+            ErrorHandler.set({
+
+               reason: `${errorReason}`,
+               title: `Error Updating Deployment`,
+               message: `There was an issue updating the tool deployment.\n\nReason: ${errorReason}\n\nTime: ${moment.utc().format("MMM Do, HH:mm z")}`
+
+            }, true);
+
+         }
+
+      }
+   }
+
+   return <Stack>
+      <Modal isModeless={false} key="action_confirm_key" isOpen={true} isBlocking={true} topOffsetFixed={true} styles={{ main: { padding: 8, width: 480, hasBeenOpened: false, top: "80px", position: "absolute" } }} className={hordeClasses.modal} onDismiss={() => { }}>
+         <Stack horizontal styles={{ root: { padding: 8 } }}>
+            <Stack.Item grow={2} style={{ paddingLeft: 8 }}>
+               <Text variant="mediumPlus">{`Update ${handler.data?.name}`}</Text>
+            </Stack.Item>
+            <Stack.Item grow={0}>
+               <IconButton
+                  iconProps={{ iconName: 'Cancel' }}
+                  ariaLabel="Close popup modal"
+                  onClick={() => { onClose(); }}
+               />
+            </Stack.Item>
+         </Stack>
+
+         <Stack style={{ marginRight: 18, marginTop: 18 }}>
+            <Stack styles={{ root: { height: height } }}>
+               <Stack style={{ paddingLeft: 24, paddingRight: 24, height: height - 12 }} tokens={{ childrenGap: 8 }}>
+                  <Markdown>{text}</Markdown>
+               </Stack>
+            </Stack>
+         </Stack>
+
+         <Stack styles={{ root: { padding: 8 } }}>
+            <Stack horizontal tokens={{ childrenGap: 16 }} styles={{ root: { paddingTop: 24, paddingLeft: 8, paddingBottom: 8, paddingRight: 12 } }}>
+               <Stack grow />
+               <PrimaryButton text={actionName} disabled={submitting} onClick={() => { onSubmit() }} />
+               <DefaultButton text="Cancel" disabled={submitting} onClick={() => { onClose(); }} />
+            </Stack>
+         </Stack>
+      </Modal>
+   </Stack>;
+});
+
+
+const ToolManagementInner: React.FC<{ handler: ToolHandler }> = observer(({ handler }) => {
+
+   const [actionState, setActionState] = useState<{ update?: UpdateDeploymentRequest, text?: string, actionName?: string, deployment?: GetToolDeploymentResponse }>({});
+
+   handler.subscribe();
+
+   if (!handler.data) {
+      return null;
+   }
+
+   const items: DeploymentItem[] = handler.data.deployments.map(d => {
+      return {
+         deploy: d
+      };
+   }).sort((a, b) => {
+      if (a.deploy.id === handler.currentDeployment?.id) {
+         return -1;
+      }
+
+      if (b.deploy.id === handler.currentDeployment?.id) {
+         return 1;
+      }
+
+      return -a.deploy.version.localeCompare(b.deploy.version)
+   })
+
+   const renderProgress = (item: DeploymentItem) => {
+
+      const active = item.deploy.state === ToolDeploymentState.Active;
+      if (active || item.deploy.progress >= 1) {
+         return null;
+      }
+
+      return <Stack>
+         <ProgressIndicator percentComplete={item.deploy.progress} barHeight={3} styles={{ root: { paddingLeft: 12, width: 208 } }} />
+      </Stack>
+
+   }
+
+   const renderActions = (item: DeploymentItem) => {
+
+      const current = item.deploy.id === handler.currentDeployment?.id;
+      const active = item.deploy.state === ToolDeploymentState.Active;
+      const disabled = !dashboard.development;
+
+      const actionProps: IContextualMenuProps = {
+         items: [
+            {
+               key: 'manage_current',
+               text: 'Set Current',
+               disabled: current || active || disabled,
+               onClick: () => {
+                  setActionState({ update: { state: ToolDeploymentState.Active }, text: `Set **${item.deploy.version}** Current?`, actionName: "Set Current", deployment: item.deploy })
+               }
+            },
+            {
+               key: 'manage_pause',
+               text: 'Pause',
+               disabled: !active || disabled,
+               onClick: () => {
+                  setActionState({ update: { state: ToolDeploymentState.Paused }, text: `Pause **${item.deploy.version}** Deployment?`, actionName: "Pause", deployment: item.deploy })
+               }
+            },
+            {
+               key: 'manage_delete',
+               text: 'Delete',
+               disabled: current || disabled,
+               onClick: () => {
+                  setActionState({ update: { state: ToolDeploymentState.Cancelled }, text: `Delete **${item.deploy.version}** Deployment?`, actionName: "Delete", deployment: item.deploy })
+               }
+            },
+         ]
+      };
+
+
+      const r = item.deploy;
+
+      return <Stack>
+         <DefaultButton text="Actions" menuProps={actionProps} />
+      </Stack>
+
+   }
+
+   const renderState = (item: DeploymentItem) => {
+      let state: string = item.deploy.state;
+      if (item.deploy.id === handler.currentDeployment?.id) {
+         state = "Current";
+      }
+      return <Stack style={{ width: 96 }} horizontalAlign="center"><Text style={{ fontWeight: state === "Current" ? 600 : undefined }}>{state}</Text></Stack>
+   }
+
+   const renderRow: IDetailsListProps['onRenderRow'] = (props) => {
+
+      if (props) {
+
+         const item = props!.item as DeploymentItem;
+
+         let background: string | undefined;
+         if (props.itemIndex % 2 === 0) {
+            background = dashboard.darktheme ? "#1D2021" : "#EAE9E9";
+         }
+
+         return <Stack horizontal verticalAlign="center" verticalFill tokens={{ childrenGap: 24 }} styles={{ root: { backgroundColor: background, paddingLeft: 12, paddingRight: 12, paddingTop: 8, paddingBottom: 8 } }}>
+            <Stack><Text style={{ fontWeight: 600 }}>{item.deploy.version}</Text></Stack>
+            {renderState(item)}
+            {renderProgress(item)}
+            <Stack grow />
+            {renderActions(item)}
+         </Stack>
+      }
+      return null;
+   };
+
+
+   return <Stack style={{ padding: 24 }}>
+      {!!actionState.update && <ActionConfirmation handler={handler} request={actionState.update} text={actionState.text ?? ""} actionName={actionState.actionName ?? ""} deployment={actionState.deployment!} onClose={() => setActionState({})} />}
+      <DetailsList
+         isHeaderVisible={false}
+         styles={{
+            headerWrapper: {
+               paddingTop: 0
+            }
+         }}
+         items={items}
+         selectionMode={SelectionMode.single}
+         layoutMode={DetailsListLayoutMode.justified}
+         compact={false}
+         onRenderRow={renderRow}
+      />
+
+   </Stack>
+
+})
+
+export const ToolManagementModal: React.FC<{ toolId: string; toolName?: string, onClose: () => void }> = ({ toolId, toolName, onClose }) => {
+
+   const handler = useConst(new ToolHandler(toolId))
+
+   const { hordeClasses } = getHordeStyling();
+
+   let headerText = toolName ?? "Tool Management";
+
+
+   return <Stack>
+      <Modal isModeless={false} key="tool_management_key" isOpen={true} isBlocking={true} topOffsetFixed={true} styles={{ main: { padding: 8, width: 700, hasBeenOpened: false, top: "80px", position: "absolute" } }} className={hordeClasses.modal} onDismiss={() => { onClose() }}>
+         <Stack horizontal styles={{ root: { padding: 8 } }}>
+            <Stack>
+               <Text variant="mediumPlus">{headerText}</Text>
+            </Stack>
+            <Stack grow />
+            <Stack>
+               <IconButton
+                  iconProps={{ iconName: 'Cancel' }}
+                  ariaLabel="Close popup modal"
+                  onClick={() => { onClose(); }}
+               />
+            </Stack>
+         </Stack>
+         <ToolManagementInner handler={handler} />
+      </Modal>
+   </Stack>
+}
+
