Horde: Remove unused include context functionality from config system.

#rnx

diff --git a/Engine/Source/Programs/Horde/Horde.Server/Configuration/ConfigType.cs b/Engine/Source/Programs/Horde/Horde.Server/Configuration/ConfigType.cs
--- a/Engine/Source/Programs/Horde/Horde.Server/Configuration/ConfigType.cs
+++ b/Engine/Source/Programs/Horde/Horde.Server/Configuration/ConfigType.cs
@@ -33,14 +33,6 @@ namespace Horde.Server.Configuration
 	{
 	}
 
-	/// <summary>
-	/// Captures the current scope of outer objects in the current variable
-	/// </summary>
-	[AttributeUsage(AttributeTargets.Class)]
-	public sealed class ConfigIncludeContextAttribute : Attribute
-	{
-	}
-
 	/// <summary>
 	/// Attribute used to mark <see cref="Uri"/> properties that are relative to their containing file
 	/// </summary>
@@ -738,7 +730,6 @@ namespace Horde.Server.Configuration
 
 		readonly Type _type;
 		readonly bool _isIncludeRoot;
-		readonly bool _isIncludeContext;
 		readonly bool _isMacroScope;
 		readonly Dictionary<string, Property> _nameToProperty = new Dictionary<string, Property>(StringComparer.OrdinalIgnoreCase);
 		readonly Dictionary<string, Property> _nameToMacroProperty = new Dictionary<string, Property>(StringComparer.OrdinalIgnoreCase);
@@ -753,7 +744,6 @@ namespace Horde.Server.Configuration
 
 			_type = type;
 			_isIncludeRoot = type.GetCustomAttribute<ConfigIncludeRootAttribute>() != null;
-			_isIncludeContext = type.GetCustomAttribute<ConfigIncludeContextAttribute>() != null;
 			_isMacroScope = type.GetCustomAttribute<ConfigMacroScopeAttribute>() != null;
 
 			// Find all the direct include properties
@@ -984,11 +974,6 @@ namespace Horde.Server.Configuration
 
 		async Task ParseIncludesAsync(JsonObject jsonObject, object targetObject, ClassConfigType targetType, ConfigContext context, CancellationToken cancellationToken)
 		{
-			if (_isIncludeContext)
-			{
-				context.IncludeContextStack.Push(JsonSerializer.Deserialize(jsonObject, _type, context.JsonOptions)!);
-			}
-
 			foreach ((string name, JsonNode? node) in jsonObject)
 			{
 				if (_nameToIncludeProperty.TryGetValue(name, out Property? property) && node != null)
@@ -996,11 +981,6 @@ namespace Horde.Server.Configuration
 					await property.ParseIncludesAsync(node, targetObject, targetType, context, cancellationToken);
 				}
 			}
-
-			if (_isIncludeContext)
-			{
-				context.IncludeContextStack.Pop();
-			}
 		}
 	}
 }
diff --git a/Engine/Source/Programs/Horde/Horde.Server/Server/GlobalConfig.cs b/Engine/Source/Programs/Horde/Horde.Server/Server/GlobalConfig.cs
--- a/Engine/Source/Programs/Horde/Horde.Server/Server/GlobalConfig.cs
+++ b/Engine/Source/Programs/Horde/Horde.Server/Server/GlobalConfig.cs
@@ -82,7 +82,6 @@ namespace Horde.Server.Server
 	/// <summary>
 	/// Directive to merge config data from another source
 	/// </summary>
-	[ConfigIncludeContext]
 	public class ConfigInclude
 	{
 		/// <summary>
