Horde: Add a separate property for parsing include directives in config files.

#rnx

diff --git a/Engine/Source/Programs/Horde/Horde.Server/Configuration/ConfigType.cs b/Engine/Source/Programs/Horde/Horde.Server/Configuration/ConfigType.cs
--- a/Engine/Source/Programs/Horde/Horde.Server/Configuration/ConfigType.cs
+++ b/Engine/Source/Programs/Horde/Horde.Server/Configuration/ConfigType.cs
@@ -410,16 +410,19 @@ namespace Horde.Server.Configuration
 
 		class ScalarProperty : Property
 		{
+			readonly bool _relativePath;
+
 			public ScalarProperty(string name, PropertyInfo propertyInfo)
 				: base(name, propertyInfo)
 			{
+				_relativePath = PropertyInfo.GetCustomAttribute<ConfigRelativePathAttribute>() != null;
 			}
 
 			public override bool HasMacros() => false;
 
 			public override void ParseMacros(JsonNode jsonNode, ConfigContext context, Dictionary<string, string> macros) { }
 
-			public override bool HasIncludes() => PropertyInfo.GetCustomAttribute<ConfigIncludeAttribute>() != null;
+			public override bool HasIncludes() => false;
 
 			public override Task MergeAsync(object target, JsonNode? node, ConfigContext context, CancellationToken cancellationToken)
 			{
@@ -430,7 +433,7 @@ namespace Horde.Server.Configuration
 				{
 					value = null;
 				}
-				else if (PropertyInfo.GetCustomAttribute<ConfigRelativePathAttribute>() != null)
+				else if (_relativePath)
 				{
 					value = CombinePaths(context.CurrentFile, JsonSerializer.Deserialize<string>(node, context.JsonOptions) ?? String.Empty).AbsoluteUri;
 				}
@@ -449,6 +452,18 @@ namespace Horde.Server.Configuration
 				return Task.CompletedTask;
 			}
 
+			public override Task ParseIncludesAsync(JsonNode jsonNode, object targetObject, ClassConfigType targetType, ConfigContext context, CancellationToken cancellationToken)
+				=> Task.CompletedTask;
+		}
+
+		class IncludeProperty : ScalarProperty
+		{
+			public IncludeProperty(string name, PropertyInfo propertyInfo)
+				: base(name, propertyInfo)
+			{ }
+
+			public override bool HasIncludes() => true;
+
 			public override async Task ParseIncludesAsync(JsonNode jsonNode, object targetObject, ClassConfigType targetType, ConfigContext context, CancellationToken cancellationToken)
 			{
 				string? path = (string?)jsonNode;
@@ -811,29 +826,39 @@ namespace Horde.Server.Configuration
 			Type propertyType = propertyInfo.PropertyType;
 			if (!propertyType.IsClass || propertyType == typeof(string))
 			{
-				return new ScalarProperty(name, propertyInfo);
-			}
-			else if (propertyType.IsAssignableTo(typeof(ConfigResource)))
-			{
-				return new ResourceProperty(name, propertyInfo);
-			}
-			else if (propertyType.IsGenericType && propertyType.GetGenericTypeDefinition() == typeof(List<>))
-			{
-				Type elementType = propertyType.GetGenericArguments()[0];
-				return new ListProperty(name, propertyInfo, FindOrAddValueType(elementType));
-			}
-			else if (propertyType.IsGenericType && propertyType.GetGenericTypeDefinition() == typeof(Dictionary<,>))
-			{
-				Type elementType = propertyType.GetGenericArguments()[1];
-				return new DictionaryProperty(name, propertyInfo, FindOrAddValueType(elementType));
-			}
-			else if (propertyType.IsAssignableTo(typeof(JsonNode)))
-			{
-				return new JsonNodeProperty(name, propertyInfo);
+				if (propertyInfo.GetCustomAttribute<ConfigIncludeAttribute>() != null)
+				{
+					return new IncludeProperty(name, propertyInfo);
+				}
+				else
+				{
+					return new ScalarProperty(name, propertyInfo);
+				}
 			}
 			else
 			{
-				return new ObjectProperty(name, propertyInfo, FindOrAdd(propertyType));
+				if (propertyType.IsAssignableTo(typeof(ConfigResource)))
+				{
+					return new ResourceProperty(name, propertyInfo);
+				}
+				else if (propertyType.IsGenericType && propertyType.GetGenericTypeDefinition() == typeof(List<>))
+				{
+					Type elementType = propertyType.GetGenericArguments()[0];
+					return new ListProperty(name, propertyInfo, FindOrAddValueType(elementType));
+				}
+				else if (propertyType.IsGenericType && propertyType.GetGenericTypeDefinition() == typeof(Dictionary<,>))
+				{
+					Type elementType = propertyType.GetGenericArguments()[1];
+					return new DictionaryProperty(name, propertyInfo, FindOrAddValueType(elementType));
+				}
+				else if (propertyType.IsAssignableTo(typeof(JsonNode)))
+				{
+					return new JsonNodeProperty(name, propertyInfo);
+				}
+				else
+				{
+					return new ObjectProperty(name, propertyInfo, FindOrAdd(propertyType));
+				}
 			}
 		}
 
