Horde: Tool management fixes
#rnx
#skipci

diff --git a/Engine/Source/Programs/Horde/HordeDashboard/src/backend/Backend.ts b/Engine/Source/Programs/Horde/HordeDashboard/src/backend/Backend.ts
--- a/Engine/Source/Programs/Horde/HordeDashboard/src/backend/Backend.ts
+++ b/Engine/Source/Programs/Horde/HordeDashboard/src/backend/Backend.ts
@@ -1625,7 +1625,7 @@ export class Backend {
     /// Updates the state of an active deployment.
     updateDeployment(toolId: string, deploymentId: string, update: UpdateDeploymentRequest): Promise<void> {
         return new Promise<void>((resolve, reject) => {
-            this.backend.patch(`/api/v1/tools/${toolId}/deployments/${deploymentId}`, update).then(() => {
+            this.backend.patch(`/api/v1/tools/${toolId}/deployments/${deploymentId}`, { request: update }).then(() => {
                 resolve();
             }).catch(reason => {
                 reject(reason);
diff --git a/Engine/Source/Programs/Horde/HordeDashboard/src/components/ToolView.tsx b/Engine/Source/Programs/Horde/HordeDashboard/src/components/ToolView.tsx
--- a/Engine/Source/Programs/Horde/HordeDashboard/src/components/ToolView.tsx
+++ b/Engine/Source/Programs/Horde/HordeDashboard/src/components/ToolView.tsx
@@ -238,7 +238,7 @@ const ToolPanel: React.FC<{ selectedKey: string }> = observer(({ selectedKey })
             <Text style={{ color: modeColors.text }}>{tool.version}</Text>
          </Stack>
          <Stack >
-            <PrimaryButton split menuProps={downloadProps} style={{ width: 140, color: "#FFFFFF" }} text="Download" href={`/api/v1/tools/${tool.id}?action=download`} />
+            <PrimaryButton split menuProps={downloadProps} style={{ width: 140, color: "#FFFFFF" }} text="Download" onClick={() => window.location.assign(`/api/v1/tools/${tool.id}?action=download`)} />
          </Stack>
       </Stack>
    }
