Horde: Tool management fixes
#rnx
#skipci

diff --git a/Engine/Source/Programs/Horde/HordeDashboard/src/components/ToolView.tsx b/Engine/Source/Programs/Horde/HordeDashboard/src/components/ToolView.tsx
--- a/Engine/Source/Programs/Horde/HordeDashboard/src/components/ToolView.tsx
+++ b/Engine/Source/Programs/Horde/HordeDashboard/src/components/ToolView.tsx
@@ -238,7 +238,7 @@ const ToolPanel: React.FC<{ selectedKey: string }> = observer(({ selectedKey })
             <Text style={{ color: modeColors.text }}>{tool.version}</Text>
          </Stack>
          <Stack >
-            <PrimaryButton split menuProps={downloadProps} style={{ width: 140, color: "#FFFFFF" }} text="Download" onClick={() => window.location.assign(`/api/v1/tools/${tool.id}?action=download`)} />
+            <PrimaryButton split menuProps={downloadProps} style={{ width: 130, color: "#FFFFFF" }} text="Download" onClick={() => window.location.assign(`/api/v1/tools/${tool.id}?action=download`)} />
          </Stack>
       </Stack>
    }
diff --git a/Engine/Source/Programs/Horde/HordeDashboard/src/components/tools/ToolManagement.tsx b/Engine/Source/Programs/Horde/HordeDashboard/src/components/tools/ToolManagement.tsx
--- a/Engine/Source/Programs/Horde/HordeDashboard/src/components/tools/ToolManagement.tsx
+++ b/Engine/Source/Programs/Horde/HordeDashboard/src/components/tools/ToolManagement.tsx
@@ -17,7 +17,7 @@ class ToolHandler extends PollBase {
       super(pollTime);
       this.toolId = toolId;
 
-      this.poll();
+      this.start();
    }
 
    async poll(): Promise<void> {
@@ -82,7 +82,7 @@ const ActionConfirmation: React.FC<{ handler: ToolHandler, request: UpdateDeploy
 
          await backend.updateDeployment(handler.data!.id, deployment.id, request);
 
-         handler.poll();
+         handler.poll();         
 
       } catch (reason) {
 
@@ -97,10 +97,13 @@ const ActionConfirmation: React.FC<{ handler: ToolHandler, request: UpdateDeploy
                title: `Error Updating Deployment`,
                message: `There was an issue updating the tool deployment.\n\nReason: ${errorReason}\n\nTime: ${moment.utc().format("MMM Do, HH:mm z")}`
 
-            }, true);
+            }, true);            
 
          }
 
+         setSubmitting(false);
+         onClose();
+
       }
    }
 
