Horde: Change config reader to operate on JsonNode instances rather than writing directly into target objects.

diff --git a/Engine/Source/Programs/Horde/Horde.Server/Configuration/ConfigType.cs b/Engine/Source/Programs/Horde/Horde.Server/Configuration/ConfigType.cs
--- a/Engine/Source/Programs/Horde/Horde.Server/Configuration/ConfigType.cs
+++ b/Engine/Source/Programs/Horde/Horde.Server/Configuration/ConfigType.cs
@@ -143,7 +143,7 @@ namespace Horde.Server.Configuration
 	/// </summary>
 	abstract class ConfigType
 	{
-		public abstract ValueTask<object?> ReadAsync(JsonNode? node, ConfigContext context, CancellationToken cancellationToken);
+		public abstract ValueTask<JsonNode?> PreprocessAsync(JsonNode? node, ConfigContext context, CancellationToken cancellationToken);
 
 		static readonly ConcurrentDictionary<Type, ConfigType> s_typeToValueType = new ConcurrentDictionary<Type, ConfigType>();
 
@@ -309,9 +309,10 @@ namespace Horde.Server.Configuration
 		{
 			ClassConfigType type = ClassConfigType.FindOrAdd(typeof(T));
 
-			T target = new T();
+			JsonObject target = new JsonObject();
 			await type.MergeObjectAsync(target, uri, context, cancellationToken);
-			return target;
+
+			return JsonSerializer.Deserialize<T>(target, context.JsonOptions) ?? new T();
 		}
 
 		/// <summary>
@@ -335,6 +336,42 @@ namespace Horde.Server.Configuration
 			}
 			return new Uri(baseUri, path);
 		}
+
+		protected static JsonNode? ExpandMacros(JsonNode? node, ConfigContext context)
+		{
+			if (node == null)
+			{
+				return node;
+			}
+			else if (node is JsonObject obj)
+			{
+				JsonObject result = new JsonObject();
+				foreach ((string propertyName, JsonNode? propertyNode) in obj)
+				{
+					result[propertyName] = ExpandMacros(propertyNode, context);
+				}
+				return result;
+			}
+			else if (node is JsonArray arr)
+			{
+				JsonArray result = new JsonArray();
+				foreach (JsonNode? elementNode in arr)
+				{
+					result.Add(ExpandMacros(elementNode, context));
+				}
+				return result;
+			}
+			else if (node is JsonValue val && val.GetValueKind() == JsonValueKind.String)
+			{
+				string strValue = ((string?)val)!;
+				string expandedStrValue = context.ExpandMacros(strValue);
+				return JsonValue.Create(expandedStrValue);
+			}
+			else
+			{
+				return node.DeepClone();
+			}
+		}
 	}
 
 	/// <summary>
@@ -353,31 +390,19 @@ namespace Horde.Server.Configuration
 			_type = type;
 		}
 
-		public override ValueTask<object?> ReadAsync(JsonNode? node, ConfigContext context, CancellationToken cancellationToken)
+		public override ValueTask<JsonNode?> PreprocessAsync(JsonNode? node, ConfigContext context, CancellationToken cancellationToken)
 		{
+			JsonNode? result;
 			if (node == null)
 			{
-				return new ValueTask<object?>((object?)null);
+				result = null;
 			}
 			else
 			{
-				return new ValueTask<object?>(Deserialize(node, _type, context));
+				result = ExpandMacros(node, context);
 			}
-		}
 
-		public static object? Deserialize(JsonNode node, Type propertyType, ConfigContext context)
-		{
-			JsonElement element = node.GetValue<JsonElement>();
-			if (element.ValueKind == JsonValueKind.String)
-			{
-				string strValue = element.GetString()!;
-				string expandedStrValue = context.ExpandMacros(strValue);
-				if (!ReferenceEquals(strValue, expandedStrValue))
-				{
-					node = JsonValue.Create(expandedStrValue);
-				}
-			}
-			return JsonSerializer.Deserialize(node, propertyType, context.JsonOptions);
+			return new ValueTask<JsonNode?>(result);
 		}
 	}
 
@@ -389,13 +414,9 @@ namespace Horde.Server.Configuration
 		abstract class Property
 		{
 			public string Name { get; }
-			public PropertyInfo PropertyInfo { get; }
 
-			protected Property(string name, PropertyInfo propertyInfo)
-			{
-				Name = name;
-				PropertyInfo = propertyInfo;
-			}
+			protected Property(string name)
+				=> Name = name;
 
 			public abstract bool HasMacros();
 
@@ -403,9 +424,11 @@ namespace Horde.Server.Configuration
 
 			public abstract bool HasIncludes();
 
-			public abstract Task MergeAsync(object target, JsonNode? node, ConfigContext context, CancellationToken cancellationToken);
+			// Reads a node into a target object (possibly merging with an existing property)
+			public abstract Task MergeIntoObjectAsync(JsonNode? node, JsonObject target, ConfigContext context, CancellationToken cancellationToken);
 
-			public abstract Task ParseIncludesAsync(JsonNode jsonNode, object targetObject, ClassConfigType targetType, ConfigContext context, CancellationToken cancellationToken);
+			// Traverse the property tree starting with 'node' and process any include directives, merging the results into the given target object.
+			public abstract Task ParseIncludesAsync(JsonNode node, JsonObject target, ClassConfigType targetType, ConfigContext context, CancellationToken cancellationToken);
 		}
 
 		class ScalarProperty : Property
@@ -413,9 +436,9 @@ namespace Horde.Server.Configuration
 			readonly bool _relativePath;
 
 			public ScalarProperty(string name, PropertyInfo propertyInfo)
-				: base(name, propertyInfo)
+				: base(name)
 			{
-				_relativePath = PropertyInfo.GetCustomAttribute<ConfigRelativePathAttribute>() != null;
+				_relativePath = propertyInfo.GetCustomAttribute<ConfigRelativePathAttribute>() != null;
 			}
 
 			public override bool HasMacros() => false;
@@ -424,35 +447,29 @@ namespace Horde.Server.Configuration
 
 			public override bool HasIncludes() => false;
 
-			public override Task MergeAsync(object target, JsonNode? node, ConfigContext context, CancellationToken cancellationToken)
+			public override Task MergeIntoObjectAsync(JsonNode? node, JsonObject target, ConfigContext context, CancellationToken cancellationToken)
 			{
 				context.AddProperty(Name);
 
-				object? value;
+				JsonNode? result;
 				if (node == null)
 				{
-					value = null;
+					result = null;
 				}
-				else if (_relativePath)
+				else if (_relativePath && node is JsonValue value && value.GetValueKind() == JsonValueKind.String)
 				{
-					value = CombinePaths(context.CurrentFile, JsonSerializer.Deserialize<string>(node, context.JsonOptions) ?? String.Empty).AbsoluteUri;
+					result = JsonValue.Create(CombinePaths(context.CurrentFile, value.ToString()).AbsoluteUri);
 				}
 				else
 				{
-					value = ScalarConfigType.Deserialize(node, PropertyInfo.PropertyType, context);
-				}
-
-				if (!PropertyInfo.CanWrite)
-				{
-					throw new ConfigException(context, $"Property {context.CurrentScope}.{Name} does not have a setter.");
+					result = ExpandMacros(node, context);
 				}
-
-				PropertyInfo.SetValue(target, value);
+				target[Name] = result;
 
 				return Task.CompletedTask;
 			}
 
-			public override Task ParseIncludesAsync(JsonNode jsonNode, object targetObject, ClassConfigType targetType, ConfigContext context, CancellationToken cancellationToken)
+			public override Task ParseIncludesAsync(JsonNode node, JsonObject targetObject, ClassConfigType targetType, ConfigContext context, CancellationToken cancellationToken)
 				=> Task.CompletedTask;
 		}
 
@@ -464,7 +481,7 @@ namespace Horde.Server.Configuration
 
 			public override bool HasIncludes() => true;
 
-			public override async Task ParseIncludesAsync(JsonNode jsonNode, object targetObject, ClassConfigType targetType, ConfigContext context, CancellationToken cancellationToken)
+			public override async Task ParseIncludesAsync(JsonNode jsonNode, JsonObject targetObject, ClassConfigType targetType, ConfigContext context, CancellationToken cancellationToken)
 			{
 				string? path = (string?)jsonNode;
 
@@ -474,7 +491,7 @@ namespace Horde.Server.Configuration
 				context.IncludeStack.Push(file);
 
 				JsonObject includedJsonObject = await ParseFileAsync(file, context, cancellationToken);
-				await targetType.MergeObjectAsync(targetObject, includedJsonObject, context, cancellationToken);
+				await targetType.MergeIntoObjectAsync(includedJsonObject, targetObject, context, cancellationToken);
 
 				context.IncludeStack.Pop();
 			}
@@ -487,22 +504,18 @@ namespace Horde.Server.Configuration
 			{
 			}
 
-			public override async Task MergeAsync(object target, JsonNode? node, ConfigContext context, CancellationToken cancellationToken)
+			public override async Task MergeIntoObjectAsync(JsonNode? node, JsonObject target, ConfigContext context, CancellationToken cancellationToken)
 			{
 				context.AddProperty(Name);
 
-				if (!PropertyInfo.CanWrite)
-				{
-					throw new ConfigException(context, $"Property {context.CurrentScope}{Name} does not have a setter.");
-				}
-
 				Uri uri = CombinePaths(context.CurrentFile, JsonSerializer.Deserialize<string>(node, context.JsonOptions) ?? String.Empty);
 				IConfigFile file = await ReadFileAsync(uri, context, cancellationToken);
 
 				ConfigResource resource = new ConfigResource();
 				resource.Path = uri.AbsoluteUri;
 				resource.Data = await file.ReadAsync(cancellationToken);
-				PropertyInfo.SetValue(target, resource);
+
+				target[Name] = JsonSerializer.SerializeToNode(resource, context.JsonOptions);
 			}
 		}
 
@@ -510,8 +523,8 @@ namespace Horde.Server.Configuration
 		{
 			readonly ConfigType _elementType;
 
-			public ListProperty(string name, PropertyInfo propertyInfo, ConfigType elementType)
-				: base(name, propertyInfo)
+			public ListProperty(string name, ConfigType elementType)
+				: base(name)
 			{
 				_elementType = elementType;
 			}
@@ -532,34 +545,32 @@ namespace Horde.Server.Configuration
 
 			public override bool HasIncludes() => _elementType is ClassConfigType elementType && elementType.HasIncludes();
 
-			public override async Task MergeAsync(object target, JsonNode? node, ConfigContext context, CancellationToken cancellationToken)
+			public override async Task MergeIntoObjectAsync(JsonNode? node, JsonObject target, ConfigContext context, CancellationToken cancellationToken)
 			{
-				IList? list = (IList?)PropertyInfo.GetValue(target);
-				if (list == null)
+				JsonArray? targetArray = target[Name]?.AsArray();
+				if (targetArray == null)
 				{
-					object value = Activator.CreateInstance(PropertyInfo.PropertyType)!;
-					PropertyInfo.SetValue(target, value);
-					list = (IList)value;
+					targetArray = new JsonArray();
+					target[Name] = targetArray;
 				}
 
-				JsonArray array = (JsonArray)node!;
-				foreach (JsonNode? element in array)
+				foreach (JsonNode? element in (JsonArray)node!)
 				{
-					context.EnterScope($"{Name}[{list.Count}]");
+					context.EnterScope($"{Name}[{targetArray.Count}]");
 
-					object? elementValue = await _elementType.ReadAsync(element, context, cancellationToken);
-					list.Add(elementValue);
+					JsonNode? elementValue = await _elementType.PreprocessAsync(element, context, cancellationToken);
+					targetArray.Add(elementValue);
 
 					context.LeaveScope();
 				}
 			}
 
-			public override async Task ParseIncludesAsync(JsonNode jsonNode, object targetObject, ClassConfigType targetType, ConfigContext context, CancellationToken cancellationToken)
+			public override async Task ParseIncludesAsync(JsonNode node, JsonObject targetObject, ClassConfigType targetType, ConfigContext context, CancellationToken cancellationToken)
 			{
-				if (jsonNode is JsonArray jsonArrayValue)
+				if (node is JsonArray arrayNode)
 				{
 					ClassConfigType classElementType = (ClassConfigType)_elementType;
-					foreach (JsonObject jsonObjectElement in jsonArrayValue.OfType<JsonObject>())
+					foreach (JsonObject jsonObjectElement in arrayNode.OfType<JsonObject>())
 					{
 						await classElementType.ParseIncludesAsync(jsonObjectElement, targetObject, targetType, context, cancellationToken);
 					}
@@ -571,8 +582,8 @@ namespace Horde.Server.Configuration
 		{
 			readonly ConfigType _elementType;
 
-			public DictionaryProperty(string name, PropertyInfo propertyInfo, ConfigType elementType)
-				: base(name, propertyInfo)
+			public DictionaryProperty(string name, ConfigType elementType)
+				: base(name)
 			{
 				_elementType = elementType;
 			}
@@ -583,36 +594,33 @@ namespace Horde.Server.Configuration
 
 			public override bool HasIncludes() => _elementType is ClassConfigType elementType && elementType.HasIncludes();
 
-			public override async Task MergeAsync(object target, JsonNode? node, ConfigContext context, CancellationToken cancellationToken)
+			public override async Task MergeIntoObjectAsync(JsonNode? node, JsonObject target, ConfigContext context, CancellationToken cancellationToken)
 			{
-				IDictionary? dictionary = (IDictionary?)PropertyInfo.GetValue(target);
-				if (dictionary == null)
+				JsonObject? targetObject = target[Name]?.AsObject();
+				if (targetObject == null)
 				{
-					object value = Activator.CreateInstance(PropertyInfo.PropertyType)!;
-					PropertyInfo.SetValue(target, value);
-					dictionary = (IDictionary)value;
+					targetObject = new JsonObject();
+					target[Name] = targetObject;
 				}
 
-				JsonObject obj = (JsonObject)node!;
-				foreach ((string key, JsonNode? element) in obj)
+				foreach ((string key, JsonNode? element) in (JsonObject)node!)
 				{
 					context.EnterScope($"{Name}[{key}]");
 
-					object? elementValue = await _elementType.ReadAsync(element, context, cancellationToken);
-					dictionary.Add(key, elementValue);
+					JsonNode? elementValue = await _elementType.PreprocessAsync(element, context, cancellationToken);
+					targetObject[key] = elementValue;
 
 					context.LeaveScope();
 				}
 			}
 
-			public override async Task ParseIncludesAsync(JsonNode jsonNode, object targetObject, ClassConfigType targetType, ConfigContext context, CancellationToken cancellationToken)
+			public override async Task ParseIncludesAsync(JsonNode node, JsonObject targetObj, ClassConfigType targetType, ConfigContext context, CancellationToken cancellationToken)
 			{
-				if (jsonNode is JsonObject jsonObjectValue)
+				if (node is JsonObject obj && _elementType is ClassConfigType classElementType)
 				{
-					ClassConfigType classElementType = (ClassConfigType)_elementType;
-					foreach (JsonObject jsonObjectElement in jsonObjectValue.Select(x => x.Value).OfType<JsonObject>())
+					foreach (JsonObject jsonObjectElement in obj.Select(x => x.Value).OfType<JsonObject>())
 					{
-						await classElementType.ParseIncludesAsync(jsonObjectElement, targetObject, targetType, context, cancellationToken);
+						await classElementType.ParseIncludesAsync(jsonObjectElement, targetObj, targetType, context, cancellationToken);
 					}
 				}
 			}
@@ -620,8 +628,8 @@ namespace Horde.Server.Configuration
 
 		class JsonNodeProperty : Property
 		{
-			public JsonNodeProperty(string name, PropertyInfo propertyInfo)
-				: base(name, propertyInfo)
+			public JsonNodeProperty(string name)
+				: base(name)
 			{
 			}
 
@@ -633,10 +641,9 @@ namespace Horde.Server.Configuration
 
 			public override bool HasIncludes() => false;
 
-			public override Task MergeAsync(object target, JsonNode? node, ConfigContext context, CancellationToken cancellationToken)
+			public override Task MergeIntoObjectAsync(JsonNode? node, JsonObject target, ConfigContext context, CancellationToken cancellationToken)
 			{
-				JsonNode? current = (JsonNode?)PropertyInfo.GetValue(target)!;
-				PropertyInfo.SetValue(target, MergeNodes(current, node));
+				target[Name] = node?.DeepClone();
 				return Task.CompletedTask;
 			}
 
@@ -661,18 +668,16 @@ namespace Horde.Server.Configuration
 				}
 			}
 
-			public override Task ParseIncludesAsync(JsonNode jsonNode, object targetObject, ClassConfigType targetType, ConfigContext context, CancellationToken cancellationToken)
-			{
-				return Task.CompletedTask;
-			}
+			public override Task ParseIncludesAsync(JsonNode jsonNode, JsonObject targetObj, ClassConfigType targetType, ConfigContext context, CancellationToken cancellationToken)
+				=> Task.CompletedTask;
 		}
 
 		class ObjectProperty : Property
 		{
 			readonly ClassConfigType _classConfigType;
 
-			public ObjectProperty(string name, PropertyInfo propertyInfo, ClassConfigType classConfigType)
-				: base(name, propertyInfo)
+			public ObjectProperty(string name, ClassConfigType classConfigType)
+				: base(name)
 			{
 				_classConfigType = classConfigType;
 			}
@@ -689,7 +694,7 @@ namespace Horde.Server.Configuration
 
 			public override bool HasIncludes() => _classConfigType.HasIncludes();
 
-			public override async Task MergeAsync(object target, JsonNode? node, ConfigContext context, CancellationToken cancellationToken)
+			public override async Task MergeIntoObjectAsync(JsonNode? node, JsonObject target, ConfigContext context, CancellationToken cancellationToken)
 			{
 				if (node is JsonObject obj)
 				{
@@ -698,7 +703,7 @@ namespace Horde.Server.Configuration
 
 					context.EnterScope(Name);
 
-					object? childTarget = PropertyInfo.GetValue(target);
+					JsonObject? childTarget = target[Name] as JsonObject;
 					if (childTarget == null)
 					{
 						if (otherFile != null)
@@ -706,12 +711,11 @@ namespace Horde.Server.Configuration
 							throw new ConfigException(context, $"Property {context.CurrentScope}.{Name} conflicts with value in {otherFile}.");
 						}
 
-						childTarget = await _classConfigType.ReadAsync(node, context, cancellationToken);
-						PropertyInfo.SetValue(target, childTarget);
+						target[Name] = await _classConfigType.ReadAsync(obj, context, cancellationToken);
 					}
 					else
 					{
-						await _classConfigType.MergeObjectAsync(childTarget, obj, context, cancellationToken);
+						await _classConfigType.MergeIntoObjectAsync(obj, childTarget, context, cancellationToken);
 					}
 
 					context.LeaveScope();
@@ -720,31 +724,21 @@ namespace Horde.Server.Configuration
 				{
 					context.AddProperty(Name);
 
-					object? value;
-					try
-					{
-						value = JsonSerializer.Deserialize(node, PropertyInfo.PropertyType, context.JsonOptions);
-					}
-					catch (Exception ex)
-					{
-						throw new ConfigException(context, $"Unable to parse property {context.CurrentScope}.{Name} from '{node?.ToJsonString(context.JsonOptions)}': {ex.Message}", ex);
-					}
-
-					PropertyInfo.SetValue(target, value);
+					target[Name] = node?.DeepClone();
 				}
 			}
 
-			public override async Task ParseIncludesAsync(JsonNode jsonNode, object targetObject, ClassConfigType targetType, ConfigContext context, CancellationToken cancellationToken)
+			public override async Task ParseIncludesAsync(JsonNode node, JsonObject targetObj, ClassConfigType targetType, ConfigContext context, CancellationToken cancellationToken)
 			{
-				if (jsonNode is JsonObject jsonObjectValue)
+				if (node is JsonObject obj)
 				{
-					await _classConfigType.ParseIncludesAsync(jsonObjectValue, targetObject, targetType, context, cancellationToken);
+					await _classConfigType.ParseIncludesAsync(obj, targetObj, targetType, context, cancellationToken);
 				}
 			}
 		}
 
-		readonly Type _type;
 		readonly bool _isIncludeRoot;
+		readonly bool _isMacro;
 		readonly bool _isMacroScope;
 		readonly Dictionary<string, Property> _nameToProperty = new Dictionary<string, Property>(StringComparer.OrdinalIgnoreCase);
 		readonly Dictionary<string, Property> _nameToMacroProperty = new Dictionary<string, Property>(StringComparer.OrdinalIgnoreCase);
@@ -757,8 +751,8 @@ namespace Horde.Server.Configuration
 		{
 			s_typeToObjectValueType.TryAdd(type, this);
 
-			_type = type;
 			_isIncludeRoot = type.GetCustomAttribute<ConfigIncludeRootAttribute>() != null;
+			_isMacro = type == typeof(ConfigMacro);
 			_isMacroScope = type.GetCustomAttribute<ConfigMacroScopeAttribute>() != null;
 
 			// Find all the direct include properties
@@ -786,7 +780,7 @@ namespace Horde.Server.Configuration
 			}
 
 			// Build up a list of possible types for this object
-			JsonKnownTypesAttribute? knownTypes = _type.GetCustomAttribute<JsonKnownTypesAttribute>();
+			JsonKnownTypesAttribute? knownTypes = type.GetCustomAttribute<JsonKnownTypesAttribute>();
 			if (knownTypes != null)
 			{
 				_knownTypes = new Dictionary<string, ClassConfigType>(StringComparer.Ordinal);
@@ -801,7 +795,7 @@ namespace Horde.Server.Configuration
 			}
 		}
 
-		bool HasMacros() => !_isMacroScope && (_type == typeof(ConfigMacro) || _nameToMacroProperty.Count > 0);
+		bool HasMacros() => !_isMacroScope && (_isMacro || _nameToMacroProperty.Count > 0);
 
 		bool HasIncludes() => !_isIncludeRoot && _nameToIncludeProperty.Count > 0;
 
@@ -844,25 +838,25 @@ namespace Horde.Server.Configuration
 				else if (propertyType.IsGenericType && propertyType.GetGenericTypeDefinition() == typeof(List<>))
 				{
 					Type elementType = propertyType.GetGenericArguments()[0];
-					return new ListProperty(name, propertyInfo, FindOrAddValueType(elementType));
+					return new ListProperty(name, FindOrAddValueType(elementType));
 				}
 				else if (propertyType.IsGenericType && propertyType.GetGenericTypeDefinition() == typeof(Dictionary<,>))
 				{
 					Type elementType = propertyType.GetGenericArguments()[1];
-					return new DictionaryProperty(name, propertyInfo, FindOrAddValueType(elementType));
+					return new DictionaryProperty(name, FindOrAddValueType(elementType));
 				}
 				else if (propertyType.IsAssignableTo(typeof(JsonNode)))
 				{
-					return new JsonNodeProperty(name, propertyInfo);
+					return new JsonNodeProperty(name);
 				}
 				else
 				{
-					return new ObjectProperty(name, propertyInfo, FindOrAdd(propertyType));
+					return new ObjectProperty(name, FindOrAdd(propertyType));
 				}
 			}
 		}
 
-		public override async ValueTask<object?> ReadAsync(JsonNode? node, ConfigContext context, CancellationToken cancellationToken)
+		public override async ValueTask<JsonNode?> PreprocessAsync(JsonNode? node, ConfigContext context, CancellationToken cancellationToken)
 		{
 			if (node == null)
 			{
@@ -874,24 +868,21 @@ namespace Horde.Server.Configuration
 			}
 			else
 			{
-				return JsonSerializer.Deserialize(node, _type, context.JsonOptions);
+				return node;
 			}
 		}
 
-		public ValueTask<object> ReadAsync(JsonObject obj, ConfigContext context, CancellationToken cancellationToken)
+		public async ValueTask<JsonObject> ReadAsync(JsonObject obj, ConfigContext context, CancellationToken cancellationToken)
 		{
 			ClassConfigType targetType = this;
 			if (_knownTypes != null && obj.TryGetPropertyValue("Type", out JsonNode? knownTypeNode) && knownTypeNode != null)
 			{
 				targetType = _knownTypes[knownTypeNode.ToString()];
 			}
-			return targetType.ReadConcreteTypeAsync(obj, context, cancellationToken);
-		}
 
-		async ValueTask<object> ReadConcreteTypeAsync(JsonObject obj, ConfigContext context, CancellationToken cancellationToken)
-		{
-			object result = Activator.CreateInstance(_type)!;
-			await MergeObjectAsync(result, obj, context, cancellationToken);
+			JsonObject result = new JsonObject();
+			await targetType.MergeIntoObjectAsync(obj, result, context, cancellationToken);
+
 			return result;
 		}
 
@@ -926,7 +917,7 @@ namespace Horde.Server.Configuration
 			return obj;
 		}
 
-		public async Task MergeObjectAsync(object target, Uri uri, ConfigContext context, CancellationToken cancellationToken)
+		public async Task MergeObjectAsync(JsonObject target, Uri uri, ConfigContext context, CancellationToken cancellationToken)
 		{
 			if (context.IncludeStack.Any(x => x.Uri == uri))
 			{
@@ -938,33 +929,37 @@ namespace Horde.Server.Configuration
 			context.IncludeStack.Push(file);
 
 			JsonObject obj = await ParseFileAsync(file, context, cancellationToken);
-			await MergeObjectAsync(target, obj, context, cancellationToken);
+			await MergeIntoObjectAsync(obj, target, context, cancellationToken);
 
 			context.IncludeStack.Pop();
 		}
 
-		async Task MergeObjectAsync(object target, JsonObject obj, ConfigContext context, CancellationToken cancellationToken)
+		async Task MergeIntoObjectAsync(JsonObject newObject, JsonObject targetObject, ConfigContext context, CancellationToken cancellationToken)
 		{
 			// Before parsing properties into this object, read all the includes recursively
 			if (_isIncludeRoot)
 			{
-				await ParseIncludesAsync(obj, target, this, context, cancellationToken);
+				await ParseIncludesAsync(newObject, targetObject, this, context, cancellationToken);
 			}
 
 			// Parse all the macros for this scope
 			if (_isMacroScope)
 			{
 				Dictionary<string, string> macros = new Dictionary<string, string>(StringComparer.OrdinalIgnoreCase);
-				ParseMacros(obj, context, macros);
+				ParseMacros(newObject, context, macros);
 				context.MacroScopes.Add(macros);
 			}
 
 			// Parse all the properties into this object
-			foreach ((string name, JsonNode? node) in obj)
+			foreach ((string name, JsonNode? newNode) in newObject)
 			{
 				if (_nameToProperty.TryGetValue(name, out Property? property))
 				{
-					await property.MergeAsync(target, node, context, cancellationToken);
+					await property.MergeIntoObjectAsync(newNode, targetObject, context, cancellationToken);
+				}
+				else
+				{
+					targetObject[name] = ExpandMacros(newNode, context);
 				}
 			}
 
@@ -977,7 +972,7 @@ namespace Horde.Server.Configuration
 
 		void ParseMacros(JsonObject jsonObject, ConfigContext context, Dictionary<string, string> macros)
 		{
-			if (_type == typeof(ConfigMacro))
+			if (_isMacro)
 			{
 				ConfigMacro? macro = JsonSerializer.Deserialize<ConfigMacro>(jsonObject, context.JsonOptions);
 				if (macro != null)
@@ -997,13 +992,13 @@ namespace Horde.Server.Configuration
 			}
 		}
 
-		async Task ParseIncludesAsync(JsonObject jsonObject, object targetObject, ClassConfigType targetType, ConfigContext context, CancellationToken cancellationToken)
+		async Task ParseIncludesAsync(JsonObject obj, JsonObject targetObj, ClassConfigType targetType, ConfigContext context, CancellationToken cancellationToken)
 		{
-			foreach ((string name, JsonNode? node) in jsonObject)
+			foreach ((string name, JsonNode? node) in obj)
 			{
 				if (_nameToIncludeProperty.TryGetValue(name, out Property? property) && node != null)
 				{
-					await property.ParseIncludesAsync(node, targetObject, targetType, context, cancellationToken);
+					await property.ParseIncludesAsync(node, targetObj, targetType, context, cancellationToken);
 				}
 			}
 		}
diff --git a/Engine/Source/Programs/Horde/Horde.Server/Streams/StreamConfig.cs b/Engine/Source/Programs/Horde/Horde.Server/Streams/StreamConfig.cs
--- a/Engine/Source/Programs/Horde/Horde.Server/Streams/StreamConfig.cs
+++ b/Engine/Source/Programs/Horde/Horde.Server/Streams/StreamConfig.cs
@@ -81,6 +81,7 @@ namespace Horde.Server.Streams
 	[JsonSchemaCatalog("Horde Stream", "Horde stream configuration file", new[] { "*.stream.json", "Streams/*.json" })]
 	[ConfigIncludeRoot]
 	[ConfigMacroScope]
+	[DebuggerDisplay("{Id}")]
 	public class StreamConfig
 	{
 		/// <summary>
