Horde: Move logic for merging object configurations together into a separate file to the preprocessor.

#rnx

diff --git a/Engine/Source/Programs/Horde/Horde.Server/Configuration/ConfigObject.cs b/Engine/Source/Programs/Horde/Horde.Server/Configuration/ConfigObject.cs
new file mode 100644
--- /dev/null
+++ b/Engine/Source/Programs/Horde/Horde.Server/Configuration/ConfigObject.cs
@@ -0,0 +1,183 @@
+// Copyright Epic Games, Inc. All Rights Reserved.
+
+using System;
+using System.Collections;
+using System.Collections.Generic;
+using System.Linq;
+using System.Reflection;
+using EpicGames.Core;
+
+namespace Horde.Server.Configuration
+{
+	/// <summary>
+	/// Possible methods for merging config values
+	/// </summary>
+	public enum ConfigMergeStrategy
+	{
+		/// <summary>
+		/// Default strategy; replace with the base value if the current value is null
+		/// </summary>
+		Default,
+
+		/// <summary>
+		/// Append the contents of this list to the base list
+		/// </summary>
+		Append,
+
+		/// <summary>
+		/// Recursively merge object properties
+		/// </summary>
+		Recursive,
+	}
+
+	/// <summary>
+	/// Attribute used to mark object properties whose child properties should be merged individually
+	/// </summary>
+	[AttributeUsage(AttributeTargets.Property)]
+	public sealed class ConfigMergeStrategyAttribute : Attribute
+	{
+		/// <summary>
+		/// Strategy for merging this property
+		/// </summary>
+		public ConfigMergeStrategy Strategy { get; }
+
+		/// <summary>
+		/// Constructor
+		/// </summary>
+		public ConfigMergeStrategyAttribute(ConfigMergeStrategy strategy) => Strategy = strategy;
+	}
+
+	/// <summary>
+	/// Helper methods for merging default values
+	/// </summary>
+	public static class ConfigObject
+	{
+		/// <summary>
+		/// Merges any unassigned properties from the source object to the target
+		/// </summary>
+		/// <param name="objects">Set of objects to use for merging</param>
+		public static void MergeDefaults<TKey, TValue>(IEnumerable<(TKey Key, TKey? BaseKey, TValue Value)> objects)
+			where TKey : notnull
+			where TValue : class
+		{
+			// Convert the objects to a dictionary
+			List<(TKey Key, TKey BaseKey, TValue Value)> remainingObjects = new List<(TKey Key, TKey BaseKey, TValue Value)>();
+
+			// Find all the objects with no base
+			Dictionary<TKey, TValue> handledValues = new Dictionary<TKey, TValue>();
+			foreach ((TKey key, TKey? baseKey, TValue value) in objects)
+			{
+#pragma warning disable CA1508 // Avoid dead conditional code (false positive due to generics)
+				if (baseKey == null || Equals(baseKey, default(TKey)))
+				{
+					handledValues.Add(key, value);
+				}
+				else
+				{
+					remainingObjects.Add((key, baseKey, value));
+				}
+#pragma warning restore CA1508 // Avoid dead conditional code
+			}
+
+			// Iteratively merge objects with their base
+			for (int lastRemainingObjectCount = 0; remainingObjects.Count != lastRemainingObjectCount;)
+			{
+				lastRemainingObjectCount = remainingObjects.Count;
+				for (int idx = remainingObjects.Count - 1; idx >= 0; idx--)
+				{
+					(TKey key, TKey baseKey, TValue value) = remainingObjects[idx];
+					if (handledValues.TryGetValue(baseKey, out TValue? baseValue))
+					{
+						MergeDefaults(value, baseValue);
+						handledValues.Add(key, value);
+						remainingObjects.RemoveAt(idx);
+					}
+				}
+			}
+
+			// Check we were able to merge everything
+			if (remainingObjects.Count > 0)
+			{
+				HashSet<TKey> validKeys = new HashSet<TKey>(objects.Select(x => x.Key));
+				foreach ((TKey key, TKey baseKey, TValue value) in remainingObjects)
+				{
+					if (!validKeys.Contains(baseKey!))
+					{
+						throw new Exception($"{key} has invalid/missing base {baseKey}");
+					}
+				}
+
+				List<TKey> circularKeys = objects.Where(x => !handledValues.ContainsKey(x.Key)).Select(x => x.Key).ToList();
+				if (circularKeys.Count == 1)
+				{
+					throw new Exception($"{circularKeys[0]} has a circular dependency on itself");
+				}
+				else
+				{
+					throw new Exception($"{StringUtils.FormatList(circularKeys.Select(x => x.ToString() ?? "(unknown)"))} have circular dependencies.");
+				}
+			}
+		}
+
+		/// <summary>
+		/// Merges any unassigned properties from the source object to the target
+		/// </summary>
+		/// <param name="target">Target object to merge into</param>
+		/// <param name="source">Source object to merge from</param>
+		public static void MergeDefaults<T>(T target, T source) where T : class
+		{
+			MergeDefaults(typeof(T), target, source);
+		}
+
+		/// <summary>
+		/// Merges any unassigned properties from the source object to the target
+		/// </summary>
+		/// <param name="type">Type of the referenced object</param>
+		/// <param name="target">Target object to merge into</param>
+		/// <param name="source">Source object to merge from</param>
+		public static void MergeDefaults(Type type, object target, object source)
+		{
+			foreach (PropertyInfo propertyInfo in type.GetProperties(BindingFlags.Instance | BindingFlags.Public))
+			{
+				object? targetValue = propertyInfo.GetValue(target);
+				object? sourceValue = propertyInfo.GetValue(source);
+
+				ConfigMergeStrategy strategy = propertyInfo.GetCustomAttribute<ConfigMergeStrategyAttribute>()?.Strategy ?? ConfigMergeStrategy.Default;
+				switch (strategy)
+				{
+					case ConfigMergeStrategy.Default:
+						if (targetValue == null)
+						{
+							propertyInfo.SetValue(target, sourceValue);
+						}
+						break;
+					case ConfigMergeStrategy.Append:
+						if (targetValue == null)
+						{
+							propertyInfo.SetValue(target, sourceValue);
+						}
+						else if (sourceValue != null)
+						{
+							IList sourceList = (IList)sourceValue;
+							IList targetList = (IList)targetValue;
+							for (int idx = 0; idx < sourceList.Count; idx++)
+							{
+								targetList.Insert(idx, sourceList[idx]);
+							}
+						}
+						break;
+					case ConfigMergeStrategy.Recursive:
+						if (targetValue == null)
+						{
+							propertyInfo.SetValue(target, sourceValue);
+						}
+						else if (sourceValue != null)
+						{
+							MergeDefaults(propertyInfo.PropertyType, targetValue, sourceValue);
+						}
+						break;
+				}
+			}
+		}
+	}
+}
diff --git a/Engine/Source/Programs/Horde/Horde.Server/Configuration/ConfigType.cs b/Engine/Source/Programs/Horde/Horde.Server/Configuration/ConfigType.cs
--- a/Engine/Source/Programs/Horde/Horde.Server/Configuration/ConfigType.cs
+++ b/Engine/Source/Programs/Horde/Horde.Server/Configuration/ConfigType.cs
@@ -65,44 +65,6 @@ namespace Horde.Server.Configuration
 		public string Value { get; set; } = String.Empty;
 	}
 
-	/// <summary>
-	/// Possible methods for merging config values
-	/// </summary>
-	public enum ConfigMergeStrategy
-	{
-		/// <summary>
-		/// Default strategy; replace with the base value if the current value is null
-		/// </summary>
-		Default,
-
-		/// <summary>
-		/// Append the contents of this list to the base list
-		/// </summary>
-		Append,
-
-		/// <summary>
-		/// Recursively merge object properties
-		/// </summary>
-		Recursive,
-	}
-
-	/// <summary>
-	/// Attribute used to mark object properties whose child properties should be merged individually
-	/// </summary>
-	[AttributeUsage(AttributeTargets.Property)]
-	public sealed class ConfigMergeStrategyAttribute : Attribute
-	{
-		/// <summary>
-		/// Strategy for merging this property
-		/// </summary>
-		public ConfigMergeStrategy Strategy { get; }
-
-		/// <summary>
-		/// Constructor
-		/// </summary>
-		public ConfigMergeStrategyAttribute(ConfigMergeStrategy strategy) => Strategy = strategy;
-	}
-
 	/// <summary>
 	/// Exception thrown when reading config files
 	/// </summary>
@@ -169,134 +131,6 @@ namespace Horde.Server.Configuration
 			return value;
 		}
 
-		/// <summary>
-		/// Merges any unassigned properties from the source object to the target
-		/// </summary>
-		/// <param name="objects">Set of objects to use for merging</param>
-		public static void MergeDefaults<TKey, TValue>(IEnumerable<(TKey Key, TKey? BaseKey, TValue Value)> objects)
-			where TKey : notnull
-			where TValue : class
-		{
-			// Convert the objects to a dictionary
-			List<(TKey Key, TKey BaseKey, TValue Value)> remainingObjects = new List<(TKey Key, TKey BaseKey, TValue Value)>();
-
-			// Find all the objects with no base
-			Dictionary<TKey, TValue> handledValues = new Dictionary<TKey, TValue>();
-			foreach ((TKey key, TKey? baseKey, TValue value) in objects)
-			{
-#pragma warning disable CA1508 // Avoid dead conditional code (false positive due to generics)
-				if (baseKey == null || Object.Equals(baseKey, default(TKey)))
-				{
-					handledValues.Add(key, value);
-				}
-				else
-				{
-					remainingObjects.Add((key, baseKey, value));
-				}
-#pragma warning restore CA1508 // Avoid dead conditional code
-			}
-
-			// Iteratively merge objects with their base
-			for (int lastRemainingObjectCount = 0; remainingObjects.Count != lastRemainingObjectCount;)
-			{
-				lastRemainingObjectCount = remainingObjects.Count;
-				for (int idx = remainingObjects.Count - 1; idx >= 0; idx--)
-				{
-					(TKey key, TKey baseKey, TValue value) = remainingObjects[idx];
-					if (handledValues.TryGetValue(baseKey, out TValue? baseValue))
-					{
-						MergeDefaults(value, baseValue);
-						handledValues.Add(key, value);
-						remainingObjects.RemoveAt(idx);
-					}
-				}
-			}
-
-			// Check we were able to merge everything
-			if (remainingObjects.Count > 0)
-			{
-				HashSet<TKey> validKeys = new HashSet<TKey>(objects.Select(x => x.Key));
-				foreach ((TKey key, TKey baseKey, TValue value) in remainingObjects)
-				{
-					if (!validKeys.Contains(baseKey!))
-					{
-						throw new Exception($"{key} has invalid/missing base {baseKey}");
-					}
-				}
-
-				List<TKey> circularKeys = objects.Where(x => !handledValues.ContainsKey(x.Key)).Select(x => x.Key).ToList();
-				if (circularKeys.Count == 1)
-				{
-					throw new Exception($"{circularKeys[0]} has a circular dependency on itself");
-				}
-				else
-				{
-					throw new Exception($"{StringUtils.FormatList(circularKeys.Select(x => x.ToString() ?? "(unknown)"))} have circular dependencies.");
-				}
-			}
-		}
-
-		/// <summary>
-		/// Merges any unassigned properties from the source object to the target
-		/// </summary>
-		/// <param name="target">Target object to merge into</param>
-		/// <param name="source">Source object to merge from</param>
-		public static void MergeDefaults<T>(T target, T source) where T : class
-		{
-			MergeDefaults(typeof(T), target, source);
-		}
-
-		/// <summary>
-		/// Merges any unassigned properties from the source object to the target
-		/// </summary>
-		/// <param name="type">Type of the referenced object</param>
-		/// <param name="target">Target object to merge into</param>
-		/// <param name="source">Source object to merge from</param>
-		public static void MergeDefaults(Type type, object target, object source)
-		{
-			foreach (PropertyInfo propertyInfo in type.GetProperties(BindingFlags.Instance | BindingFlags.Public))
-			{
-				object? targetValue = propertyInfo.GetValue(target);
-				object? sourceValue = propertyInfo.GetValue(source);
-
-				ConfigMergeStrategy strategy = propertyInfo.GetCustomAttribute<ConfigMergeStrategyAttribute>()?.Strategy ?? ConfigMergeStrategy.Default;
-				switch (strategy)
-				{
-					case ConfigMergeStrategy.Default:
-						if (targetValue == null)
-						{
-							propertyInfo.SetValue(target, sourceValue);
-						}
-						break;
-					case ConfigMergeStrategy.Append:
-						if (targetValue == null)
-						{
-							propertyInfo.SetValue(target, sourceValue);
-						}
-						else if (sourceValue != null)
-						{
-							IList sourceList = (IList)sourceValue;
-							IList targetList = (IList)targetValue;
-							for (int idx = 0; idx < sourceList.Count; idx++)
-							{
-								targetList.Insert(idx, sourceList[idx]);
-							}
-						}
-						break;
-					case ConfigMergeStrategy.Recursive:
-						if (targetValue == null)
-						{
-							propertyInfo.SetValue(target, sourceValue);
-						}
-						else if (sourceValue != null)
-						{
-							MergeDefaults(propertyInfo.PropertyType, targetValue, sourceValue);
-						}
-						break;
-				}
-			}
-		}
-
 		/// <summary>
 		/// Reads an object from a particular URL
 		/// </summary>
@@ -435,10 +269,10 @@ namespace Horde.Server.Configuration
 		{
 			readonly bool _relativePath;
 
-			public ScalarProperty(string name, PropertyInfo propertyInfo)
+			public ScalarProperty(string name, bool relativePath)
 				: base(name)
 			{
-				_relativePath = propertyInfo.GetCustomAttribute<ConfigRelativePathAttribute>() != null;
+				_relativePath = relativePath;
 			}
 
 			public override bool HasMacros() => false;
@@ -475,8 +309,8 @@ namespace Horde.Server.Configuration
 
 		class IncludeProperty : ScalarProperty
 		{
-			public IncludeProperty(string name, PropertyInfo propertyInfo)
-				: base(name, propertyInfo)
+			public IncludeProperty(string name, bool relativePath)
+				: base(name, relativePath)
 			{ }
 
 			public override bool HasIncludes() => true;
@@ -499,8 +333,8 @@ namespace Horde.Server.Configuration
 
 		class ResourceProperty : ScalarProperty
 		{
-			public ResourceProperty(string name, PropertyInfo propertyInfo)
-				: base(name, propertyInfo)
+			public ResourceProperty(string name)
+				: base(name, false)
 			{
 			}
 
@@ -820,20 +654,21 @@ namespace Horde.Server.Configuration
 			Type propertyType = propertyInfo.PropertyType;
 			if (!propertyType.IsClass || propertyType == typeof(string))
 			{
+				bool relativePath = propertyInfo.GetCustomAttribute<ConfigRelativePathAttribute>() != null;
 				if (propertyInfo.GetCustomAttribute<ConfigIncludeAttribute>() != null)
 				{
-					return new IncludeProperty(name, propertyInfo);
+					return new IncludeProperty(name, relativePath);
 				}
 				else
 				{
-					return new ScalarProperty(name, propertyInfo);
+					return new ScalarProperty(name, relativePath);
 				}
 			}
 			else
 			{
 				if (propertyType.IsAssignableTo(typeof(ConfigResource)))
 				{
-					return new ResourceProperty(name, propertyInfo);
+					return new ResourceProperty(name);
 				}
 				else if (propertyType.IsGenericType && propertyType.GetGenericTypeDefinition() == typeof(List<>))
 				{
diff --git a/Engine/Source/Programs/Horde/Horde.Server/Server/GlobalConfig.cs b/Engine/Source/Programs/Horde/Horde.Server/Server/GlobalConfig.cs
--- a/Engine/Source/Programs/Horde/Horde.Server/Server/GlobalConfig.cs
+++ b/Engine/Source/Programs/Horde/Horde.Server/Server/GlobalConfig.cs
@@ -348,7 +348,7 @@ namespace Horde.Server.Server
 			{
 				_poolLookup.Add(pool.Id, pool);
 			}
-			ConfigType.MergeDefaults<string, PoolConfig>(Pools.Select(x => (x.Id.ToString(), x.Base?.ToString(), x)));
+			ConfigObject.MergeDefaults<string, PoolConfig>(Pools.Select(x => (x.Id.ToString(), x.Base?.ToString(), x)));
 			UpdateWorkspacesForPools();
 
 			_telemetryStoreLookup.Clear();
diff --git a/Engine/Source/Programs/Horde/Horde.Server/Streams/StreamConfig.cs b/Engine/Source/Programs/Horde/Horde.Server/Streams/StreamConfig.cs
--- a/Engine/Source/Programs/Horde/Horde.Server/Streams/StreamConfig.cs
+++ b/Engine/Source/Programs/Horde/Horde.Server/Streams/StreamConfig.cs
@@ -275,9 +275,9 @@ namespace Horde.Server.Streams
 				WorkspaceTypes.TryAdd(wid, wc);
 			}
 
-			ConfigType.MergeDefaults(AgentTypes.Select(x => (x.Key, x.Value.Base, x.Value)));
-			ConfigType.MergeDefaults(WorkspaceTypes.Select(x => (x.Key, x.Value.Base, x.Value)));
-			ConfigType.MergeDefaults(Templates.Select(x => (x.Id, x.Base, x)));
+			ConfigObject.MergeDefaults(AgentTypes.Select(x => (x.Key, x.Value.Base, x.Value)));
+			ConfigObject.MergeDefaults(WorkspaceTypes.Select(x => (x.Key, x.Value.Base, x.Value)));
+			ConfigObject.MergeDefaults(Templates.Select(x => (x.Id, x.Base, x)));
 
 			foreach (TemplateRefConfig template in Templates)
 			{
