Horde: Allow constructing config types without using attributes, so we can create schemas for previous config versions without having the backing classes.

#rnx

diff --git a/Engine/Source/Programs/Horde/Horde.Server/Configuration/ConfigContext.cs b/Engine/Source/Programs/Horde/Horde.Server/Configuration/ConfigContext.cs
--- a/Engine/Source/Programs/Horde/Horde.Server/Configuration/ConfigContext.cs
+++ b/Engine/Source/Programs/Horde/Horde.Server/Configuration/ConfigContext.cs
@@ -5,6 +5,8 @@ using System.Collections.Generic;
 using System.Diagnostics;
 using System.Diagnostics.CodeAnalysis;
 using System.Text.Json;
+using System.Threading;
+using System.Threading.Tasks;
 using EpicGames.Core;
 using Microsoft.Extensions.Logging;
 
@@ -157,5 +159,29 @@ namespace Horde.Server.Configuration
 			}
 			return null;
 		}
+
+		/// <summary>
+		/// Reads the contents of a file using the appropriate source
+		/// </summary>
+		/// <param name="uri">Uri of the file to read. The scheme indicates the source to read from.</param>
+		/// <param name="cancellationToken">Cancellation token for the operation</param>
+		/// <returns>Information about the file</returns>
+		public async ValueTask<IConfigFile> ReadFileAsync(Uri uri, CancellationToken cancellationToken)
+		{
+			IConfigSource? source = Sources[uri.Scheme];
+			if (source == null)
+			{
+				throw new ConfigException(this, $"Invalid/unknown scheme for config file {uri}");
+			}
+
+			IConfigFile? file;
+			if (!Files.TryGetValue(uri, out file))
+			{
+				file = await source.GetAsync(uri, cancellationToken);
+				Files.Add(uri, file);
+			}
+
+			return file;
+		}
 	}
 }
diff --git a/Engine/Source/Programs/Horde/Horde.Server/Configuration/ConfigFile.cs b/Engine/Source/Programs/Horde/Horde.Server/Configuration/ConfigFile.cs
--- a/Engine/Source/Programs/Horde/Horde.Server/Configuration/ConfigFile.cs
+++ b/Engine/Source/Programs/Horde/Horde.Server/Configuration/ConfigFile.cs
@@ -1,6 +1,8 @@
 // Copyright Epic Games, Inc. All Rights Reserved.
 
 using System;
+using System.Text.Json;
+using System.Text.Json.Nodes;
 using System.Threading;
 using System.Threading.Tasks;
 using Horde.Server.Users;
@@ -59,5 +61,25 @@ namespace Horde.Server.Configuration
 			}
 			return path;
 		}
+
+		/// <summary>
+		/// Parses a config file as a json object
+		/// </summary>
+		/// <param name="file">File to parse</param>
+		/// <param name="context">Context for serialization</param>
+		/// <param name="cancellationToken">Cancellation token for the operation</param>
+		/// <returns>The parsed file</returns>
+		internal static async ValueTask<JsonObject> ParseFileAsync(this IConfigFile file, ConfigContext context, CancellationToken cancellationToken = default)
+		{
+			ReadOnlyMemory<byte> data = await file.ReadAsync(cancellationToken);
+
+			JsonObject? obj = JsonSerializer.Deserialize<JsonObject>(data.Span, context.JsonOptions);
+			if (obj == null)
+			{
+				throw new ConfigException(context, $"Config file {file.Uri} contains a null object.");
+			}
+
+			return obj;
+		}
 	}
 }
diff --git a/Engine/Source/Programs/Horde/Horde.Server/Configuration/ConfigType.cs b/Engine/Source/Programs/Horde/Horde.Server/Configuration/ConfigType.cs
--- a/Engine/Source/Programs/Horde/Horde.Server/Configuration/ConfigType.cs
+++ b/Engine/Source/Programs/Horde/Horde.Server/Configuration/ConfigType.cs
@@ -1,7 +1,6 @@
 // Copyright Epic Games, Inc. All Rights Reserved.
 
 using System;
-using System.Collections;
 using System.Collections.Concurrent;
 using System.Collections.Generic;
 using System.Linq;
@@ -105,10 +104,25 @@ namespace Horde.Server.Configuration
 	/// </summary>
 	abstract class ConfigType
 	{
-		public abstract ValueTask<JsonNode?> PreprocessAsync(JsonNode? node, ConfigContext context, CancellationToken cancellationToken);
+		/// <summary>
+		/// Preprocess a JSON node
+		/// </summary>
+		/// <param name="node">Node to preprocess</param>
+		/// <param name="context">Context for the preprocessor</param>
+		/// <param name="cancellationToken">Cancellation token for the operation</param>
+		/// <returns>Preprocessed node</returns>
+		public async ValueTask<JsonNode?> PreprocessAsync(JsonNode? node, ConfigContext context, CancellationToken cancellationToken)
+		{
+			return await PreprocessAndMergeAsync(node, null, context, cancellationToken);
+		}
 
 		static readonly ConcurrentDictionary<Type, ConfigType> s_typeToValueType = new ConcurrentDictionary<Type, ConfigType>();
 
+		/// <summary>
+		/// 
+		/// </summary>
+		/// <param name="type"></param>
+		/// <returns></returns>
 		public static ConfigType FindOrAddValueType(Type type)
 		{
 			ConfigType? value;
@@ -116,11 +130,11 @@ namespace Horde.Server.Configuration
 			{
 				if (!type.IsClass || type == typeof(string))
 				{
-					value = new ScalarConfigType(type);
+					value = new ScalarConfigType(false);
 				}
 				else
 				{
-					value = ClassConfigType.FindOrAdd(type);
+					value = ObjectConfigType.FindOrAdd(type);
 				}
 
 				lock (s_typeToValueType)
@@ -141,12 +155,16 @@ namespace Horde.Server.Configuration
 		/// <returns></returns>
 		public static async Task<T> ReadAsync<T>(Uri uri, ConfigContext context, CancellationToken cancellationToken) where T : class, new()
 		{
-			ClassConfigType type = ClassConfigType.FindOrAdd(typeof(T));
+			IConfigFile file = await context.ReadFileAsync(uri, cancellationToken);
+			context.IncludeStack.Push(file);
 
-			JsonObject target = new JsonObject();
-			await type.MergeObjectAsync(target, uri, context, cancellationToken);
+			ObjectConfigType type = ObjectConfigType.FindOrAdd(typeof(T));
 
-			return JsonSerializer.Deserialize<T>(target, context.JsonOptions) ?? new T();
+			JsonObject obj = await file.ParseFileAsync(context, cancellationToken);
+			obj = (JsonObject)(await type.PreprocessAndMergeAsync(obj, null, context, cancellationToken))!;
+
+			context.IncludeStack.Pop();
+			return JsonSerializer.Deserialize<T>(obj, context.JsonOptions) ?? new T();
 		}
 
 		/// <summary>
@@ -171,6 +189,9 @@ namespace Horde.Server.Configuration
 			return new Uri(baseUri, path);
 		}
 
+		/// <summary>
+		/// Helper method to expand all macros within a node without performing any other processing on it
+		/// </summary>
 		protected static JsonNode? ExpandMacros(JsonNode? node, ConfigContext context)
 		{
 			if (node == null)
@@ -206,6 +227,45 @@ namespace Horde.Server.Configuration
 				return node.DeepClone();
 			}
 		}
+
+		/// <summary>
+		/// Determines whether this type contains macro definitions
+		/// </summary>
+		internal abstract bool HasMacroDefinitions();
+
+		/// <summary>
+		/// Parses macro definitions from this object
+		/// </summary>
+		/// <param name="node">Node to parse macros from</param>
+		/// <param name="context">Context for the preprocessor</param>
+		/// <param name="macros">Macros parsed from the boject</param>
+		internal abstract void ParseMacroDefinitions(JsonNode node, ConfigContext context, Dictionary<string, string> macros);
+
+		/// <summary>
+		/// Whether this type contains include directives
+		/// </summary>
+		internal abstract bool HasIncludes();
+
+		/// <summary>
+		/// Reads a node into a target object (possibly merging with an existing property)
+		/// </summary>
+		/// <param name="node">Node to preprocess</param>
+		/// <param name="existingNode">Optional existing node to merge with. Can be modified.</param>
+		/// <param name="context">Context for the preprocessor</param>
+		/// <param name="cancellationToken">Cancellation token for the operation</param>
+		/// <returns>The preprocessed node. May be existingNode if changes are merged with it.</returns>
+		internal abstract Task<JsonNode?> PreprocessAndMergeAsync(JsonNode? node, JsonNode? existingNode, ConfigContext context, CancellationToken cancellationToken);
+
+		/// <summary>
+		/// Traverse the property tree starting with 'node' and process any include directives, merging the results into the given target object.
+		/// </summary>
+		/// <param name="node"></param>
+		/// <param name="target"></param>
+		/// <param name="targetType"></param>
+		/// <param name="context"></param>
+		/// <param name="cancellationToken"></param>
+		/// <returns></returns>
+		internal abstract Task<JsonObject> ParseIncludesAsync(JsonNode node, JsonObject target, ObjectConfigType targetType, ConfigContext context, CancellationToken cancellationToken);
 	}
 
 	/// <summary>
@@ -213,443 +273,373 @@ namespace Horde.Server.Configuration
 	/// </summary>
 	class ScalarConfigType : ConfigType
 	{
-		readonly Type _type;
+		readonly bool _relativePath;
 
-		public ScalarConfigType(Type type)
-		{
-			if (type.IsGenericType && type.GetGenericTypeDefinition() == typeof(Nullable<>))
-			{
-				type = type.GetGenericArguments()[0];
-			}
-			_type = type;
-		}
+		public ScalarConfigType(bool relativePath)
+			=> _relativePath = relativePath;
+
+		/// <inheritdoc/>
+		internal override bool HasMacroDefinitions() => false;
 
-		public override ValueTask<JsonNode?> PreprocessAsync(JsonNode? node, ConfigContext context, CancellationToken cancellationToken)
+		/// <inheritdoc/>
+		internal override void ParseMacroDefinitions(JsonNode jsonNode, ConfigContext context, Dictionary<string, string> macros) { }
+
+		/// <inheritdoc/>
+		internal override bool HasIncludes() => false;
+
+		/// <inheritdoc/>
+		internal override Task<JsonNode?> PreprocessAndMergeAsync(JsonNode? node, JsonNode? existingNode, ConfigContext context, CancellationToken cancellationToken)
 		{
 			JsonNode? result;
 			if (node == null)
 			{
 				result = null;
 			}
+			else if (_relativePath && node is JsonValue value && value.GetValueKind() == JsonValueKind.String)
+			{
+				result = JsonValue.Create(CombinePaths(context.CurrentFile, value.ToString()).AbsoluteUri);
+			}
 			else
 			{
 				result = ExpandMacros(node, context);
 			}
 
-			return new ValueTask<JsonNode?>(result);
+			return Task.FromResult(result);
 		}
+
+		/// <inheritdoc/>
+		internal override Task<JsonObject> ParseIncludesAsync(JsonNode node, JsonObject targetObject, ObjectConfigType targetType, ConfigContext context, CancellationToken cancellationToken)
+			=> Task.FromResult(targetObject);
 	}
 
 	/// <summary>
-	/// Implementation of <see cref="ConfigType"/> to handle class types
+	/// Property which specifies the path of another config file to include 
 	/// </summary>
-	class ClassConfigType : ConfigType
+	class IncludeConfigType : ScalarConfigType
 	{
-		abstract class Property
-		{
-			public string Name { get; }
-
-			protected Property(string name)
-				=> Name = name;
-
-			public abstract bool HasMacros();
+		public IncludeConfigType(bool relativePath)
+			: base(relativePath)
+		{ }
 
-			public abstract void ParseMacros(JsonNode jsonNode, ConfigContext context, Dictionary<string, string> macros);
+		/// <inheritdoc/>
+		internal override bool HasIncludes() => true;
 
-			public abstract bool HasIncludes();
-
-			// Reads a node into a target object (possibly merging with an existing property)
-			public abstract Task MergeIntoObjectAsync(JsonNode? node, JsonObject target, ConfigContext context, CancellationToken cancellationToken);
-
-			// Traverse the property tree starting with 'node' and process any include directives, merging the results into the given target object.
-			public abstract Task ParseIncludesAsync(JsonNode node, JsonObject target, ClassConfigType targetType, ConfigContext context, CancellationToken cancellationToken);
-		}
-
-		class ScalarProperty : Property
+		/// <inheritdoc/>
+		internal override async Task<JsonObject> ParseIncludesAsync(JsonNode jsonNode, JsonObject targetObject, ObjectConfigType targetType, ConfigContext context, CancellationToken cancellationToken)
 		{
-			readonly bool _relativePath;
-
-			public ScalarProperty(string name, bool relativePath)
-				: base(name)
-			{
-				_relativePath = relativePath;
-			}
-
-			public override bool HasMacros() => false;
+			string? path = (string?)jsonNode;
 
-			public override void ParseMacros(JsonNode jsonNode, ConfigContext context, Dictionary<string, string> macros) { }
+			Uri uri = ConfigType.CombinePaths(context.CurrentFile, context.ExpandMacros(path!));
+			IConfigFile file = await context.ReadFileAsync(uri, cancellationToken);
 
-			public override bool HasIncludes() => false;
+			context.IncludeStack.Push(file);
 
-			public override Task MergeIntoObjectAsync(JsonNode? node, JsonObject target, ConfigContext context, CancellationToken cancellationToken)
-			{
-				context.AddProperty(Name);
+			JsonObject includedJsonObject = await file.ParseFileAsync(context, cancellationToken);
+			JsonNode? result = await targetType.PreprocessAndMergeAsync(includedJsonObject, targetObject, context, cancellationToken);
 
-				JsonNode? result;
-				if (node == null)
-				{
-					result = null;
-				}
-				else if (_relativePath && node is JsonValue value && value.GetValueKind() == JsonValueKind.String)
-				{
-					result = JsonValue.Create(CombinePaths(context.CurrentFile, value.ToString()).AbsoluteUri);
-				}
-				else
-				{
-					result = ExpandMacros(node, context);
-				}
-				target[Name] = result;
-
-				return Task.CompletedTask;
-			}
+			context.IncludeStack.Pop();
 
-			public override Task ParseIncludesAsync(JsonNode node, JsonObject targetObject, ClassConfigType targetType, ConfigContext context, CancellationToken cancellationToken)
-				=> Task.CompletedTask;
+			return (JsonObject)result!;
 		}
+	}
 
-		class IncludeProperty : ScalarProperty
-		{
-			public IncludeProperty(string name, bool relativePath)
-				: base(name, relativePath)
-			{ }
-
-			public override bool HasIncludes() => true;
-
-			public override async Task ParseIncludesAsync(JsonNode jsonNode, JsonObject targetObject, ClassConfigType targetType, ConfigContext context, CancellationToken cancellationToken)
-			{
-				string? path = (string?)jsonNode;
-
-				Uri uri = ConfigType.CombinePaths(context.CurrentFile, context.ExpandMacros(path!));
-				IConfigFile file = await ReadFileAsync(uri, context, cancellationToken);
-
-				context.IncludeStack.Push(file);
-
-				JsonObject includedJsonObject = await ParseFileAsync(file, context, cancellationToken);
-				await targetType.MergeIntoObjectAsync(includedJsonObject, targetObject, context, cancellationToken);
-
-				context.IncludeStack.Pop();
-			}
-		}
+	/// <summary>
+	/// Property containing a binary resource
+	/// </summary>
+	class ResourceConfigType : ScalarConfigType
+	{
+		public ResourceConfigType()
+			: base(false)
+		{ }
 
-		class ResourceProperty : ScalarProperty
+		/// <inheritdoc/>
+		internal override async Task<JsonNode?> PreprocessAndMergeAsync(JsonNode? node, JsonNode? existingNode, ConfigContext context, CancellationToken cancellationToken)
 		{
-			public ResourceProperty(string name)
-				: base(name, false)
-			{
-			}
-
-			public override async Task MergeIntoObjectAsync(JsonNode? node, JsonObject target, ConfigContext context, CancellationToken cancellationToken)
-			{
-				context.AddProperty(Name);
-
-				Uri uri = CombinePaths(context.CurrentFile, JsonSerializer.Deserialize<string>(node, context.JsonOptions) ?? String.Empty);
-				IConfigFile file = await ReadFileAsync(uri, context, cancellationToken);
+			Uri uri = CombinePaths(context.CurrentFile, JsonSerializer.Deserialize<string>(node, context.JsonOptions) ?? String.Empty);
+			IConfigFile file = await context.ReadFileAsync(uri, cancellationToken);
 
-				ConfigResource resource = new ConfigResource();
-				resource.Path = uri.AbsoluteUri;
-				resource.Data = await file.ReadAsync(cancellationToken);
+			ConfigResource resource = new ConfigResource();
+			resource.Path = uri.AbsoluteUri;
+			resource.Data = await file.ReadAsync(cancellationToken);
 
-				target[Name] = JsonSerializer.SerializeToNode(resource, context.JsonOptions);
-			}
+			return JsonSerializer.SerializeToNode(resource, context.JsonOptions);
 		}
+	}
 
-		class ListProperty : Property
-		{
-			readonly ConfigType _elementType;
+	/// <summary>
+	/// Array of Json values
+	/// </summary>
+	class ArrayConfigType : ConfigType
+	{
+		readonly ConfigType _elementType;
 
-			public ListProperty(string name, ConfigType elementType)
-				: base(name)
-			{
-				_elementType = elementType;
-			}
+		/// <summary>
+		/// Constructor
+		/// </summary>
+		public ArrayConfigType(ConfigType elementType)
+		{
+			_elementType = elementType;
+		}
 
-			public override bool HasMacros() => _elementType is ClassConfigType elementType && elementType.HasMacros();
+		/// <inheritdoc/>
+		internal override bool HasMacroDefinitions() => _elementType is ObjectConfigType elementType && elementType.HasMacroDefinitions();
 
-			public override void ParseMacros(JsonNode jsonNode, ConfigContext context, Dictionary<string, string> macros)
+		/// <inheritdoc/>
+		internal override void ParseMacroDefinitions(JsonNode jsonNode, ConfigContext context, Dictionary<string, string> macros)
+		{
+			if (jsonNode is JsonArray jsonArrayValue)
 			{
-				if (jsonNode is JsonArray jsonArrayValue)
+				ObjectConfigType classElementType = (ObjectConfigType)_elementType;
+				foreach (JsonObject jsonObjectElement in jsonArrayValue.OfType<JsonObject>())
 				{
-					ClassConfigType classElementType = (ClassConfigType)_elementType;
-					foreach (JsonObject jsonObjectElement in jsonArrayValue.OfType<JsonObject>())
-					{
-						classElementType.ParseMacros(jsonObjectElement, context, macros);
-					}
+					classElementType.ParseMacroDefinitions(jsonObjectElement, context, macros);
 				}
 			}
+		}
 
-			public override bool HasIncludes() => _elementType is ClassConfigType elementType && elementType.HasIncludes();
+		/// <inheritdoc/>
+		internal override bool HasIncludes() => _elementType is ObjectConfigType elementType && elementType.HasIncludes();
 
-			public override async Task MergeIntoObjectAsync(JsonNode? node, JsonObject target, ConfigContext context, CancellationToken cancellationToken)
+		/// <inheritdoc/>
+		internal override async Task<JsonNode?> PreprocessAndMergeAsync(JsonNode? node, JsonNode? existingNode, ConfigContext context, CancellationToken cancellationToken)
+		{
+			JsonArray targetArray = ((JsonArray?)existingNode) ?? new JsonArray();
+			foreach (JsonNode? element in (JsonArray)node!)
 			{
-				JsonArray? targetArray = target[Name]?.AsArray();
-				if (targetArray == null)
-				{
-					targetArray = new JsonArray();
-					target[Name] = targetArray;
-				}
+				context.EnterScope($"[{targetArray.Count}]");
 
-				foreach (JsonNode? element in (JsonArray)node!)
-				{
-					context.EnterScope($"{Name}[{targetArray.Count}]");
+				JsonNode? elementValue = await _elementType.PreprocessAsync(element, context, cancellationToken);
+				targetArray.Add(elementValue);
 
-					JsonNode? elementValue = await _elementType.PreprocessAsync(element, context, cancellationToken);
-					targetArray.Add(elementValue);
-
-					context.LeaveScope();
-				}
+				context.LeaveScope();
 			}
+			return targetArray;
+		}
 
-			public override async Task ParseIncludesAsync(JsonNode node, JsonObject targetObject, ClassConfigType targetType, ConfigContext context, CancellationToken cancellationToken)
+		/// <inheritdoc/>
+		internal override async Task<JsonObject> ParseIncludesAsync(JsonNode node, JsonObject targetObj, ObjectConfigType targetType, ConfigContext context, CancellationToken cancellationToken)
+		{
+			if (node is JsonArray arrayNode)
 			{
-				if (node is JsonArray arrayNode)
+				ObjectConfigType classElementType = (ObjectConfigType)_elementType;
+				foreach (JsonObject jsonObjectElement in arrayNode.OfType<JsonObject>())
 				{
-					ClassConfigType classElementType = (ClassConfigType)_elementType;
-					foreach (JsonObject jsonObjectElement in arrayNode.OfType<JsonObject>())
-					{
-						await classElementType.ParseIncludesAsync(jsonObjectElement, targetObject, targetType, context, cancellationToken);
-					}
+					targetObj = await classElementType.ParseIncludesAsync(jsonObjectElement, targetObj, targetType, context, cancellationToken);
 				}
 			}
+			return targetObj;
 		}
+	}
 
-		class DictionaryProperty : Property
-		{
-			readonly ConfigType _elementType;
+	/// <summary>
+	/// Arbitary mapping of string values to keys
+	/// </summary>
+	class DictionaryConfigType : ConfigType
+	{
+		readonly ConfigType _elementType;
 
-			public DictionaryProperty(string name, ConfigType elementType)
-				: base(name)
-			{
-				_elementType = elementType;
-			}
+		public DictionaryConfigType(ConfigType elementType)
+		{
+			_elementType = elementType;
+		}
 
-			public override bool HasMacros() => false;
+		/// <inheritdoc/>
+		internal override bool HasMacroDefinitions() => false;
 
-			public override void ParseMacros(JsonNode jsonNode, ConfigContext context, Dictionary<string, string> macros) { }
+		/// <inheritdoc/>
+		internal override void ParseMacroDefinitions(JsonNode jsonNode, ConfigContext context, Dictionary<string, string> macros) { }
 
-			public override bool HasIncludes() => _elementType is ClassConfigType elementType && elementType.HasIncludes();
+		/// <inheritdoc/>
+		internal override bool HasIncludes() => _elementType is ObjectConfigType elementType && elementType.HasIncludes();
 
-			public override async Task MergeIntoObjectAsync(JsonNode? node, JsonObject target, ConfigContext context, CancellationToken cancellationToken)
+		/// <inheritdoc/>
+		internal override async Task<JsonNode?> PreprocessAndMergeAsync(JsonNode? node, JsonNode? existingNode, ConfigContext context, CancellationToken cancellationToken)
+		{
+			JsonObject? targetObject = ((JsonObject?)existingNode) ?? new JsonObject();
+			foreach ((string key, JsonNode? element) in (JsonObject)node!)
 			{
-				JsonObject? targetObject = target[Name]?.AsObject();
-				if (targetObject == null)
-				{
-					targetObject = new JsonObject();
-					target[Name] = targetObject;
-				}
+				context.EnterScope($"[{key}]");
 
-				foreach ((string key, JsonNode? element) in (JsonObject)node!)
-				{
-					context.EnterScope($"{Name}[{key}]");
+				JsonNode? elementValue = await _elementType.PreprocessAsync(element, context, cancellationToken);
+				targetObject[key] = elementValue;
 
-					JsonNode? elementValue = await _elementType.PreprocessAsync(element, context, cancellationToken);
-					targetObject[key] = elementValue;
-
-					context.LeaveScope();
-				}
+				context.LeaveScope();
 			}
+			return targetObject;
+		}
 
-			public override async Task ParseIncludesAsync(JsonNode node, JsonObject targetObj, ClassConfigType targetType, ConfigContext context, CancellationToken cancellationToken)
+		/// <inheritdoc/>
+		internal override async Task<JsonObject> ParseIncludesAsync(JsonNode node, JsonObject targetObj, ObjectConfigType targetType, ConfigContext context, CancellationToken cancellationToken)
+		{
+			if (node is JsonObject obj && _elementType is ObjectConfigType classElementType)
 			{
-				if (node is JsonObject obj && _elementType is ClassConfigType classElementType)
+				foreach (JsonObject jsonObjectElement in obj.Select(x => x.Value).OfType<JsonObject>())
 				{
-					foreach (JsonObject jsonObjectElement in obj.Select(x => x.Value).OfType<JsonObject>())
-					{
-						await classElementType.ParseIncludesAsync(jsonObjectElement, targetObj, targetType, context, cancellationToken);
-					}
+					targetObj = await classElementType.ParseIncludesAsync(jsonObjectElement, targetObj, targetType, context, cancellationToken);
 				}
 			}
+			return targetObj;
 		}
+	}
 
-		class JsonNodeProperty : Property
-		{
-			public JsonNodeProperty(string name)
-				: base(name)
-			{
-			}
+	/// <summary>
+	/// Special config type for storing a json node
+	/// </summary>
+	class JsonNodeConfigType : ConfigType
+	{
+		/// <inheritdoc/>
+		internal override bool HasMacroDefinitions() => false;
 
-			public override bool HasMacros() => false;
+		/// <inheritdoc/>
+		internal override void ParseMacroDefinitions(JsonNode jsonNode, ConfigContext context, Dictionary<string, string> macros)
+		{
+		}
 
-			public override void ParseMacros(JsonNode jsonNode, ConfigContext context, Dictionary<string, string> macros)
-			{
-			}
+		/// <inheritdoc/>
+		internal override bool HasIncludes() => false;
 
-			public override bool HasIncludes() => false;
+		/// <inheritdoc/>
+		internal override Task<JsonNode?> PreprocessAndMergeAsync(JsonNode? node, JsonNode? existingNode, ConfigContext context, CancellationToken cancellationToken)
+		{
+			return Task.FromResult(node?.DeepClone());
+		}
 
-			public override Task MergeIntoObjectAsync(JsonNode? node, JsonObject target, ConfigContext context, CancellationToken cancellationToken)
+		static JsonNode? MergeNodes(JsonNode? first, JsonNode? second)
+		{
+			if (second == null)
 			{
-				target[Name] = node?.DeepClone();
-				return Task.CompletedTask;
+				return first?.DeepClone();
 			}
-
-			static JsonNode? MergeNodes(JsonNode? first, JsonNode? second)
+			else if (first is JsonObject firstObj && second is JsonObject secondObj)
 			{
-				if (second == null)
-				{
-					return first?.DeepClone();
-				}
-				else if (first is JsonObject firstObj && second is JsonObject secondObj)
+				JsonObject mergedObj = new JsonObject(firstObj);
+				foreach ((string childName, JsonNode? childNode) in secondObj)
 				{
-					JsonObject mergedObj = new JsonObject(firstObj);
-					foreach ((string childName, JsonNode? childNode) in secondObj)
-					{
-						mergedObj[childName] = MergeNodes(firstObj[childName], childNode);
-					}
-					return mergedObj;
-				}
-				else
-				{
-					return second.DeepClone();
+					mergedObj[childName] = MergeNodes(firstObj[childName], childNode);
 				}
+				return mergedObj;
+			}
+			else
+			{
+				return second.DeepClone();
 			}
-
-			public override Task ParseIncludesAsync(JsonNode jsonNode, JsonObject targetObj, ClassConfigType targetType, ConfigContext context, CancellationToken cancellationToken)
-				=> Task.CompletedTask;
 		}
 
-		class ObjectProperty : Property
-		{
-			readonly ClassConfigType _classConfigType;
+		/// <inheritdoc/>
+		internal override Task<JsonObject> ParseIncludesAsync(JsonNode jsonNode, JsonObject targetObj, ObjectConfigType targetType, ConfigContext context, CancellationToken cancellationToken)
+			=> Task.FromResult(targetObj);
+	}
 
-			public ObjectProperty(string name, ClassConfigType classConfigType)
-				: base(name)
-			{
-				_classConfigType = classConfigType;
-			}
+	/// <summary>
+	/// Implementation of <see cref="ConfigType"/> to handle class types
+	/// </summary>
+	class ObjectConfigType : ConfigType
+	{
+		readonly bool _isIncludeRoot;
+		readonly bool _isMacro;
+		readonly bool _isMacroScope;
+		readonly Dictionary<string, ConfigType> _nameToProperty;
+		readonly Dictionary<string, ConfigType> _nameToMacroProperty = new Dictionary<string, ConfigType>(StringComparer.OrdinalIgnoreCase);
+		readonly Dictionary<string, ConfigType> _nameToIncludeProperty = new Dictionary<string, ConfigType>(StringComparer.OrdinalIgnoreCase);
+		readonly Dictionary<string, ObjectConfigType>? _knownTypes;
 
-			public override bool HasMacros() => _classConfigType.HasMacros();
+		static readonly ConcurrentDictionary<Type, ObjectConfigType> s_typeToObjectValueType = new ConcurrentDictionary<Type, ObjectConfigType>();
 
-			public override void ParseMacros(JsonNode jsonNode, ConfigContext context, Dictionary<string, string> macros)
-			{
-				if (jsonNode is JsonObject obj)
-				{
-					_classConfigType.ParseMacros(obj, context, macros);
-				}
-			}
+		/// <summary>
+		/// Constructor
+		/// </summary>
+		public ObjectConfigType(bool isIncludeRoot, bool isMacro, bool isMacroScope, IEnumerable<KeyValuePair<string, ConfigType>> properties, IEnumerable<KeyValuePair<string, ObjectConfigType>>? knownTypes)
+		{
+			_isIncludeRoot = isIncludeRoot;
+			_isMacro = isMacro;
+			_isMacroScope = isMacroScope;
 
-			public override bool HasIncludes() => _classConfigType.HasIncludes();
+			_nameToProperty = new Dictionary<string, ConfigType>(properties, StringComparer.OrdinalIgnoreCase);
 
-			public override async Task MergeIntoObjectAsync(JsonNode? node, JsonObject target, ConfigContext context, CancellationToken cancellationToken)
+			// Build a map of all the properties which can contain macros or include other files
+			foreach ((string name, ConfigType property) in _nameToProperty)
 			{
-				if (node is JsonObject obj)
+				if (property.HasMacroDefinitions())
 				{
-					Uri? otherFile;
-					context.TryAddProperty(Name, out otherFile);
-
-					context.EnterScope(Name);
-
-					JsonObject? childTarget = target[Name] as JsonObject;
-					if (childTarget == null)
-					{
-						if (otherFile != null)
-						{
-							throw new ConfigException(context, $"Property {context.CurrentScope}.{Name} conflicts with value in {otherFile}.");
-						}
-
-						target[Name] = await _classConfigType.ReadAsync(obj, context, cancellationToken);
-					}
-					else
-					{
-						await _classConfigType.MergeIntoObjectAsync(obj, childTarget, context, cancellationToken);
-					}
-
-					context.LeaveScope();
+					_nameToMacroProperty.Add(name, property);
 				}
-				else
+				if (property.HasIncludes())
 				{
-					context.AddProperty(Name);
-
-					target[Name] = node?.DeepClone();
+					_nameToIncludeProperty.Add(name, property);
 				}
 			}
 
-			public override async Task ParseIncludesAsync(JsonNode node, JsonObject targetObj, ClassConfigType targetType, ConfigContext context, CancellationToken cancellationToken)
+			// Copy the known types
+			if (knownTypes != null)
 			{
-				if (node is JsonObject obj)
-				{
-					await _classConfigType.ParseIncludesAsync(obj, targetObj, targetType, context, cancellationToken);
-				}
+				_knownTypes = new Dictionary<string, ObjectConfigType>(knownTypes, StringComparer.Ordinal);
 			}
 		}
 
-		readonly bool _isIncludeRoot;
-		readonly bool _isMacro;
-		readonly bool _isMacroScope;
-		readonly Dictionary<string, Property> _nameToProperty = new Dictionary<string, Property>(StringComparer.OrdinalIgnoreCase);
-		readonly Dictionary<string, Property> _nameToMacroProperty = new Dictionary<string, Property>(StringComparer.OrdinalIgnoreCase);
-		readonly Dictionary<string, Property> _nameToIncludeProperty = new Dictionary<string, Property>(StringComparer.OrdinalIgnoreCase);
-		readonly Dictionary<string, ClassConfigType>? _knownTypes;
-
-		static readonly ConcurrentDictionary<Type, ClassConfigType> s_typeToObjectValueType = new ConcurrentDictionary<Type, ClassConfigType>();
-
-		public ClassConfigType(Type type)
+		static ObjectConfigType FromType(Type type)
 		{
-			s_typeToObjectValueType.TryAdd(type, this);
-
-			_isIncludeRoot = type.GetCustomAttribute<ConfigIncludeRootAttribute>() != null;
-			_isMacro = type == typeof(ConfigMacro);
-			_isMacroScope = type.GetCustomAttribute<ConfigMacroScopeAttribute>() != null;
+			bool isIncludeRoot = type.GetCustomAttribute<ConfigIncludeRootAttribute>() != null;
+			bool isMacro = type == typeof(ConfigMacro);
+			bool isMacroScope = type.GetCustomAttribute<ConfigMacroScopeAttribute>() != null;
 
 			// Find all the direct include properties
+			Dictionary<string, ConfigType> nameToProperty = new Dictionary<string, ConfigType>(StringComparer.OrdinalIgnoreCase);
+
 			PropertyInfo[] propertyInfos = type.GetProperties(BindingFlags.Instance | BindingFlags.Public | BindingFlags.GetProperty);
 			foreach (PropertyInfo propertyInfo in propertyInfos)
 			{
 				if (propertyInfo.GetCustomAttribute<JsonIgnoreAttribute>() == null)
 				{
 					string name = propertyInfo.GetCustomAttribute<JsonPropertyNameAttribute>()?.Name ?? propertyInfo.Name;
-					_nameToProperty.Add(name, CreateProperty(name, propertyInfo));
-				}
-			}
-
-			// Build a map of all the properties which can contain macros or include other files
-			foreach (Property property in _nameToProperty.Values)
-			{
-				if (property.HasMacros())
-				{
-					_nameToMacroProperty.Add(property.Name, property);
-				}
-				if (property.HasIncludes())
-				{
-					_nameToIncludeProperty.Add(property.Name, property);
+					nameToProperty.Add(name, CreateProperty(propertyInfo));
 				}
 			}
 
 			// Build up a list of possible types for this object
-			JsonKnownTypesAttribute? knownTypes = type.GetCustomAttribute<JsonKnownTypesAttribute>();
+			Dictionary<string, ObjectConfigType>? discriminatorToKnownType = null;
+
+			JsonKnownTypesAttribute? knownTypes = type.GetCustomAttribute<JsonKnownTypesAttribute>(false);
 			if (knownTypes != null)
 			{
-				_knownTypes = new Dictionary<string, ClassConfigType>(StringComparer.Ordinal);
+				discriminatorToKnownType = new Dictionary<string, ObjectConfigType>(StringComparer.Ordinal);
 				foreach (Type knownType in knownTypes.Types)
 				{
-					ClassConfigType knownConfigType = FindOrAdd(knownType);
+					ObjectConfigType knownConfigType = FindOrAdd(knownType);
 					foreach (JsonDiscriminatorAttribute discriminatorAttribute in knownType.GetCustomAttributes(typeof(JsonDiscriminatorAttribute), true))
 					{
-						_knownTypes.Add(discriminatorAttribute.Name, knownConfigType);
+						discriminatorToKnownType.Add(discriminatorAttribute.Name, knownConfigType);
 					}
 				}
 			}
+
+			// Create the type
+			return new ObjectConfigType(isIncludeRoot, isMacro, isMacroScope, nameToProperty, discriminatorToKnownType);
 		}
 
-		bool HasMacros() => !_isMacroScope && (_isMacro || _nameToMacroProperty.Count > 0);
+		/// <inheritdoc/>
+		internal override bool HasMacroDefinitions() => !_isMacroScope && (_isMacro || _nameToMacroProperty.Count > 0);
 
-		bool HasIncludes() => !_isIncludeRoot && _nameToIncludeProperty.Count > 0;
+		/// <inheritdoc/>
+		internal override bool HasIncludes() => !_isIncludeRoot && _nameToIncludeProperty.Count > 0;
 
-		public static ClassConfigType FindOrAdd(Type type)
+		public static ObjectConfigType FindOrAdd(Type type)
 		{
-			ClassConfigType? value;
+			ObjectConfigType? value;
 			if (!s_typeToObjectValueType.TryGetValue(type, out value))
 			{
 				lock (s_typeToObjectValueType)
 				{
 					if (!s_typeToObjectValueType.TryGetValue(type, out value))
 					{
-						value = new ClassConfigType(type);
+						value = ObjectConfigType.FromType(type);
+						s_typeToObjectValueType.TryAdd(type, value);
 					}
 				}
 			}
 			return value;
 		}
 
-		static Property CreateProperty(string name, PropertyInfo propertyInfo)
+		static ConfigType CreateProperty(PropertyInfo propertyInfo)
 		{
 			Type propertyType = propertyInfo.PropertyType;
 			if (!propertyType.IsClass || propertyType == typeof(string))
@@ -657,144 +647,115 @@ namespace Horde.Server.Configuration
 				bool relativePath = propertyInfo.GetCustomAttribute<ConfigRelativePathAttribute>() != null;
 				if (propertyInfo.GetCustomAttribute<ConfigIncludeAttribute>() != null)
 				{
-					return new IncludeProperty(name, relativePath);
+					return new IncludeConfigType(relativePath);
 				}
 				else
 				{
-					return new ScalarProperty(name, relativePath);
+					return new ScalarConfigType(relativePath);
 				}
 			}
 			else
 			{
 				if (propertyType.IsAssignableTo(typeof(ConfigResource)))
 				{
-					return new ResourceProperty(name);
+					return new ResourceConfigType();
 				}
 				else if (propertyType.IsGenericType && propertyType.GetGenericTypeDefinition() == typeof(List<>))
 				{
 					Type elementType = propertyType.GetGenericArguments()[0];
-					return new ListProperty(name, FindOrAddValueType(elementType));
+					return new ArrayConfigType(FindOrAddValueType(elementType));
 				}
 				else if (propertyType.IsGenericType && propertyType.GetGenericTypeDefinition() == typeof(Dictionary<,>))
 				{
 					Type elementType = propertyType.GetGenericArguments()[1];
-					return new DictionaryProperty(name, FindOrAddValueType(elementType));
+					return new DictionaryConfigType(FindOrAddValueType(elementType));
 				}
 				else if (propertyType.IsAssignableTo(typeof(JsonNode)))
 				{
-					return new JsonNodeProperty(name);
+					return new JsonNodeConfigType();
 				}
 				else
 				{
-					return new ObjectProperty(name, FindOrAdd(propertyType));
+					return FindOrAdd(propertyType);
 				}
 			}
 		}
 
-		public override async ValueTask<JsonNode?> PreprocessAsync(JsonNode? node, ConfigContext context, CancellationToken cancellationToken)
-		{
-			if (node == null)
-			{
-				throw new ConfigException(context, "Unable to deserialize object from null value");
-			}
-			else if (node is JsonObject obj)
-			{
-				return await ReadAsync(obj, context, cancellationToken);
-			}
-			else
-			{
-				return node;
-			}
-		}
-
 		public async ValueTask<JsonObject> ReadAsync(JsonObject obj, ConfigContext context, CancellationToken cancellationToken)
 		{
-			ClassConfigType targetType = this;
+			ObjectConfigType targetType = this;
 			if (_knownTypes != null && obj.TryGetPropertyValue("Type", out JsonNode? knownTypeNode) && knownTypeNode != null)
 			{
 				targetType = _knownTypes[knownTypeNode.ToString()];
 			}
-
-			JsonObject result = new JsonObject();
-			await targetType.MergeIntoObjectAsync(obj, result, context, cancellationToken);
-
-			return result;
+			return (JsonObject)(await targetType.PreprocessAndMergeAsync(obj, null, context, cancellationToken))!;
 		}
 
-		static async ValueTask<IConfigFile> ReadFileAsync(Uri uri, ConfigContext context, CancellationToken cancellationToken)
+		/// <inheritdoc/>
+		internal override void ParseMacroDefinitions(JsonNode jsonNode, ConfigContext context, Dictionary<string, string> macros)
 		{
-			IConfigSource? source = context.Sources[uri.Scheme];
-			if (source == null)
-			{
-				throw new ConfigException(context, $"Invalid/unknown scheme for config file {uri}");
-			}
-
-			IConfigFile? file;
-			if (!context.Files.TryGetValue(uri, out file))
+			JsonObject jsonObject = (JsonObject)jsonNode;
+			if (_isMacro)
 			{
-				file = await source.GetAsync(uri, cancellationToken);
-				context.Files.Add(uri, file);
+				ConfigMacro? macro = JsonSerializer.Deserialize<ConfigMacro>(jsonObject, context.JsonOptions);
+				if (macro != null)
+				{
+					macros.Add(macro.Name, macro.Value);
+				}
 			}
-
-			return file;
-		}
-
-		static async ValueTask<JsonObject> ParseFileAsync(IConfigFile file, ConfigContext context, CancellationToken cancellationToken)
-		{
-			ReadOnlyMemory<byte> data = await file.ReadAsync(cancellationToken);
-
-			JsonObject? obj = JsonSerializer.Deserialize<JsonObject>(data.Span, context.JsonOptions);
-			if (obj == null)
+			else
 			{
-				throw new ConfigException(context, $"Config file {file.Uri} contains a null object.");
+				foreach ((string name, JsonNode? node) in jsonObject)
+				{
+					if (node != null && _nameToMacroProperty.TryGetValue(name, out ConfigType? property))
+					{
+						property.ParseMacroDefinitions(node, context, macros);
+					}
+				}
 			}
-
-			return obj;
 		}
 
-		public async Task MergeObjectAsync(JsonObject target, Uri uri, ConfigContext context, CancellationToken cancellationToken)
+		/// <inheritdoc/>
+		internal override async Task<JsonNode?> PreprocessAndMergeAsync(JsonNode? node, JsonNode? existingNode, ConfigContext context, CancellationToken cancellationToken)
 		{
-			if (context.IncludeStack.Any(x => x.Uri == uri))
+			if (node is not JsonObject newObject)
 			{
-				throw new ConfigException(context, $"Recursive include of file {uri}");
+				return node?.DeepClone();
 			}
 
-			IConfigFile file = await ReadFileAsync(uri, context, cancellationToken);
+			JsonObject? target = (JsonObject?)existingNode ?? new JsonObject();
 
-			context.IncludeStack.Push(file);
-
-			JsonObject obj = await ParseFileAsync(file, context, cancellationToken);
-			await MergeIntoObjectAsync(obj, target, context, cancellationToken);
-
-			context.IncludeStack.Pop();
-		}
-
-		async Task MergeIntoObjectAsync(JsonObject newObject, JsonObject targetObject, ConfigContext context, CancellationToken cancellationToken)
-		{
 			// Before parsing properties into this object, read all the includes recursively
 			if (_isIncludeRoot)
 			{
-				await ParseIncludesAsync(newObject, targetObject, this, context, cancellationToken);
+				await ParseIncludesAsync(newObject, target, this, context, cancellationToken);
 			}
 
 			// Parse all the macros for this scope
 			if (_isMacroScope)
 			{
 				Dictionary<string, string> macros = new Dictionary<string, string>(StringComparer.OrdinalIgnoreCase);
-				ParseMacros(newObject, context, macros);
+				ParseMacroDefinitions(newObject, context, macros);
 				context.MacroScopes.Add(macros);
 			}
 
 			// Parse all the properties into this object
 			foreach ((string name, JsonNode? newNode) in newObject)
 			{
-				if (_nameToProperty.TryGetValue(name, out Property? property))
+				if (newNode is JsonValue)
+				{
+					context.AddProperty(name);
+				}
+
+				if (_nameToProperty.TryGetValue(name, out ConfigType? property))
 				{
-					await property.MergeIntoObjectAsync(newNode, targetObject, context, cancellationToken);
+					context.EnterScope(name);
+					target[name] = await property.PreprocessAndMergeAsync(newNode, target[name], context, cancellationToken);
+					context.LeaveScope();
 				}
 				else
 				{
-					targetObject[name] = ExpandMacros(newNode, context);
+					target[name] = ExpandMacros(newNode, context);
 				}
 			}
 
@@ -803,39 +764,24 @@ namespace Horde.Server.Configuration
 			{
 				context.MacroScopes.RemoveAt(context.MacroScopes.Count - 1);
 			}
+
+			return target;
 		}
 
-		void ParseMacros(JsonObject jsonObject, ConfigContext context, Dictionary<string, string> macros)
+		/// <inheritdoc/>
+		internal override async Task<JsonObject> ParseIncludesAsync(JsonNode node, JsonObject targetObj, ObjectConfigType targetType, ConfigContext context, CancellationToken cancellationToken)
 		{
-			if (_isMacro)
+			if (node is JsonObject obj)
 			{
-				ConfigMacro? macro = JsonSerializer.Deserialize<ConfigMacro>(jsonObject, context.JsonOptions);
-				if (macro != null)
+				foreach ((string name, JsonNode? propertyNode) in obj)
 				{
-					macros.Add(macro.Name, macro.Value);
-				}
-			}
-			else
-			{
-				foreach ((string name, JsonNode? node) in jsonObject)
-				{
-					if (node != null && _nameToMacroProperty.TryGetValue(name, out Property? property))
+					if (_nameToIncludeProperty.TryGetValue(name, out ConfigType? property) && propertyNode != null)
 					{
-						property.ParseMacros(node, context, macros);
+						targetObj = await property.ParseIncludesAsync(propertyNode, targetObj, targetType, context, cancellationToken);
 					}
 				}
 			}
-		}
-
-		async Task ParseIncludesAsync(JsonObject obj, JsonObject targetObj, ClassConfigType targetType, ConfigContext context, CancellationToken cancellationToken)
-		{
-			foreach ((string name, JsonNode? node) in obj)
-			{
-				if (_nameToIncludeProperty.TryGetValue(name, out Property? property) && node != null)
-				{
-					await property.ParseIncludesAsync(node, targetObj, targetType, context, cancellationToken);
-				}
-			}
+			return targetObj;
 		}
 	}
 }
