Horde: Remove unused config preprocessor support for polymorphic types.

#rnx

diff --git a/Engine/Source/Programs/Horde/Horde.Server/Configuration/ConfigType.cs b/Engine/Source/Programs/Horde/Horde.Server/Configuration/ConfigType.cs
--- a/Engine/Source/Programs/Horde/Horde.Server/Configuration/ConfigType.cs
+++ b/Engine/Source/Programs/Horde/Horde.Server/Configuration/ConfigType.cs
@@ -541,14 +541,13 @@ namespace Horde.Server.Configuration
 		readonly Dictionary<string, ConfigType> _nameToProperty;
 		readonly Dictionary<string, ConfigType> _nameToMacroProperty = new Dictionary<string, ConfigType>(StringComparer.OrdinalIgnoreCase);
 		readonly Dictionary<string, ConfigType> _nameToIncludeProperty = new Dictionary<string, ConfigType>(StringComparer.OrdinalIgnoreCase);
-		readonly Dictionary<string, ObjectConfigType>? _knownTypes;
 
 		static readonly ConcurrentDictionary<Type, ObjectConfigType> s_typeToObjectValueType = new ConcurrentDictionary<Type, ObjectConfigType>();
 
 		/// <summary>
 		/// Constructor
 		/// </summary>
-		public ObjectConfigType(bool isIncludeRoot, bool isMacro, bool isMacroScope, IEnumerable<KeyValuePair<string, ConfigType>> properties, IEnumerable<KeyValuePair<string, ObjectConfigType>>? knownTypes)
+		public ObjectConfigType(bool isIncludeRoot, bool isMacro, bool isMacroScope, IEnumerable<KeyValuePair<string, ConfigType>> properties)
 		{
 			_isIncludeRoot = isIncludeRoot;
 			_isMacro = isMacro;
@@ -568,12 +567,6 @@ namespace Horde.Server.Configuration
 					_nameToIncludeProperty.Add(name, property);
 				}
 			}
-
-			// Copy the known types
-			if (knownTypes != null)
-			{
-				_knownTypes = new Dictionary<string, ObjectConfigType>(knownTypes, StringComparer.Ordinal);
-			}
 		}
 
 		static ObjectConfigType FromType(Type type)
@@ -595,25 +588,8 @@ namespace Horde.Server.Configuration
 				}
 			}
 
-			// Build up a list of possible types for this object
-			Dictionary<string, ObjectConfigType>? discriminatorToKnownType = null;
-
-			JsonKnownTypesAttribute? knownTypes = type.GetCustomAttribute<JsonKnownTypesAttribute>(false);
-			if (knownTypes != null)
-			{
-				discriminatorToKnownType = new Dictionary<string, ObjectConfigType>(StringComparer.Ordinal);
-				foreach (Type knownType in knownTypes.Types)
-				{
-					ObjectConfigType knownConfigType = FindOrAdd(knownType);
-					foreach (JsonDiscriminatorAttribute discriminatorAttribute in knownType.GetCustomAttributes(typeof(JsonDiscriminatorAttribute), true))
-					{
-						discriminatorToKnownType.Add(discriminatorAttribute.Name, knownConfigType);
-					}
-				}
-			}
-
 			// Create the type
-			return new ObjectConfigType(isIncludeRoot, isMacro, isMacroScope, nameToProperty, discriminatorToKnownType);
+			return new ObjectConfigType(isIncludeRoot, isMacro, isMacroScope, nameToProperty);
 		}
 
 		/// <inheritdoc/>
@@ -683,12 +659,7 @@ namespace Horde.Server.Configuration
 
 		public async ValueTask<JsonObject> ReadAsync(JsonObject obj, ConfigContext context, CancellationToken cancellationToken)
 		{
-			ObjectConfigType targetType = this;
-			if (_knownTypes != null && obj.TryGetPropertyValue("Type", out JsonNode? knownTypeNode) && knownTypeNode != null)
-			{
-				targetType = _knownTypes[knownTypeNode.ToString()];
-			}
-			return (JsonObject)(await targetType.PreprocessAndMergeAsync(obj, null, context, cancellationToken))!;
+			return (JsonObject)(await PreprocessAndMergeAsync(obj, null, context, cancellationToken))!;
 		}
 
 		/// <inheritdoc/>
