Horde: Allow constructing a preprocessing tree for config reading without using reflection.

#rnx

diff --git a/Engine/Source/Programs/Horde/Horde.Server/Configuration/ConfigType.cs b/Engine/Source/Programs/Horde/Horde.Server/Configuration/ConfigType.cs
--- a/Engine/Source/Programs/Horde/Horde.Server/Configuration/ConfigType.cs
+++ b/Engine/Source/Programs/Horde/Horde.Server/Configuration/ConfigType.cs
@@ -102,47 +102,19 @@ namespace Horde.Server.Configuration
 	/// <summary>
 	/// Base class for types that can be read from config files
 	/// </summary>
-	abstract class ConfigType
+	static class ConfigType
 	{
 		/// <summary>
 		/// Preprocess a JSON node
 		/// </summary>
+		/// <param name="configType"></param>
 		/// <param name="node">Node to preprocess</param>
 		/// <param name="context">Context for the preprocessor</param>
 		/// <param name="cancellationToken">Cancellation token for the operation</param>
 		/// <returns>Preprocessed node</returns>
-		public async ValueTask<JsonNode?> PreprocessAsync(JsonNode? node, ConfigContext context, CancellationToken cancellationToken)
+		public static async ValueTask<JsonNode?> PreprocessAsync(this ConfigNode configType, JsonNode? node, ConfigContext context, CancellationToken cancellationToken)
 		{
-			return await PreprocessAndMergeAsync(node, null, context, cancellationToken);
-		}
-
-		static readonly ConcurrentDictionary<Type, ConfigType> s_typeToValueType = new ConcurrentDictionary<Type, ConfigType>();
-
-		/// <summary>
-		/// 
-		/// </summary>
-		/// <param name="type"></param>
-		/// <returns></returns>
-		public static ConfigType FindOrAddValueType(Type type)
-		{
-			ConfigType? value;
-			if (!s_typeToValueType.TryGetValue(type, out value))
-			{
-				if (!type.IsClass || type == typeof(string))
-				{
-					value = new ScalarConfigType(false);
-				}
-				else
-				{
-					value = ObjectConfigType.FindOrAdd(type);
-				}
-
-				lock (s_typeToValueType)
-				{
-					value = s_typeToValueType.GetOrAdd(type, value);
-				}
-			}
-			return value;
+			return await configType.PreprocessAsync(node, null, context, cancellationToken);
 		}
 
 		/// <summary>
@@ -156,14 +128,15 @@ namespace Horde.Server.Configuration
 		public static async Task<T> ReadAsync<T>(Uri uri, ConfigContext context, CancellationToken cancellationToken) where T : class, new()
 		{
 			IConfigFile file = await context.ReadFileAsync(uri, cancellationToken);
-			context.IncludeStack.Push(file);
+			JsonObject obj = await file.ParseFileAsync(context, cancellationToken);
 
-			ObjectConfigType type = ObjectConfigType.FindOrAdd(typeof(T));
+			context.IncludeStack.Push(file);
 
-			JsonObject obj = await file.ParseFileAsync(context, cancellationToken);
-			obj = (JsonObject)(await type.PreprocessAndMergeAsync(obj, null, context, cancellationToken))!;
+			ObjectConfigNode type = new ObjectConfigNode(typeof(T));
+			obj = (JsonObject)(await type.PreprocessAsync(obj, null, context, cancellationToken))!;
 
 			context.IncludeStack.Pop();
+
 			return JsonSerializer.Deserialize<T>(obj, context.JsonOptions) ?? new T();
 		}
 
@@ -192,7 +165,7 @@ namespace Horde.Server.Configuration
 		/// <summary>
 		/// Helper method to expand all macros within a node without performing any other processing on it
 		/// </summary>
-		protected static JsonNode? ExpandMacros(JsonNode? node, ConfigContext context)
+		internal static JsonNode? ExpandMacros(JsonNode? node, ConfigContext context)
 		{
 			if (node == null)
 			{
@@ -227,24 +200,20 @@ namespace Horde.Server.Configuration
 				return node.DeepClone();
 			}
 		}
+	}
 
-		/// <summary>
-		/// Determines whether this type contains macro definitions
-		/// </summary>
-		internal abstract bool HasMacroDefinitions();
-
+	/// <summary>
+	/// Node in the preprocessor parse tree
+	/// </summary>
+	abstract class ConfigNode
+	{
 		/// <summary>
 		/// Parses macro definitions from this object
 		/// </summary>
 		/// <param name="node">Node to parse macros from</param>
 		/// <param name="context">Context for the preprocessor</param>
 		/// <param name="macros">Macros parsed from the boject</param>
-		internal abstract void ParseMacroDefinitions(JsonNode node, ConfigContext context, Dictionary<string, string> macros);
-
-		/// <summary>
-		/// Whether this type contains include directives
-		/// </summary>
-		internal abstract bool HasIncludes();
+		public abstract void ParseMacros(JsonNode? node, ConfigContext context, Dictionary<string, string> macros);
 
 		/// <summary>
 		/// Reads a node into a target object (possibly merging with an existing property)
@@ -254,108 +223,89 @@ namespace Horde.Server.Configuration
 		/// <param name="context">Context for the preprocessor</param>
 		/// <param name="cancellationToken">Cancellation token for the operation</param>
 		/// <returns>The preprocessed node. May be existingNode if changes are merged with it.</returns>
-		internal abstract Task<JsonNode?> PreprocessAndMergeAsync(JsonNode? node, JsonNode? existingNode, ConfigContext context, CancellationToken cancellationToken);
+		public abstract Task<JsonNode?> PreprocessAsync(JsonNode? node, JsonNode? existingNode, ConfigContext context, CancellationToken cancellationToken);
 
 		/// <summary>
 		/// Traverse the property tree starting with 'node' and process any include directives, merging the results into the given target object.
 		/// </summary>
 		/// <param name="node"></param>
-		/// <param name="target"></param>
-		/// <param name="targetType"></param>
+		/// <param name="includes">Receives the included files</param>
 		/// <param name="context"></param>
 		/// <param name="cancellationToken"></param>
 		/// <returns></returns>
-		internal abstract Task<JsonObject> ParseIncludesAsync(JsonNode node, JsonObject target, ObjectConfigType targetType, ConfigContext context, CancellationToken cancellationToken);
+		public abstract Task ParseIncludesAsync(JsonNode? node, List<IConfigFile> includes, ConfigContext context, CancellationToken cancellationToken);
 	}
 
 	/// <summary>
 	/// Implementation of <see cref="ConfigType"/> for scalar types
 	/// </summary>
-	class ScalarConfigType : ConfigType
+	class ScalarConfigNode : ConfigNode
 	{
-		readonly bool _relativePath;
+		public bool Include { get; set; }
+		public bool RelativeToContainingFile { get; set; }
 
-		public ScalarConfigType(bool relativePath)
-			=> _relativePath = relativePath;
-
-		/// <inheritdoc/>
-		internal override bool HasMacroDefinitions() => false;
-
-		/// <inheritdoc/>
-		internal override void ParseMacroDefinitions(JsonNode jsonNode, ConfigContext context, Dictionary<string, string> macros) { }
+		public ScalarConfigNode(bool include = false, bool relativePath = false)
+		{
+			Include = include;
+			RelativeToContainingFile = relativePath;
+		}
 
 		/// <inheritdoc/>
-		internal override bool HasIncludes() => false;
+		public override void ParseMacros(JsonNode? jsonNode, ConfigContext context, Dictionary<string, string> macros)
+		{ }
 
 		/// <inheritdoc/>
-		internal override Task<JsonNode?> PreprocessAndMergeAsync(JsonNode? node, JsonNode? existingNode, ConfigContext context, CancellationToken cancellationToken)
+		public override Task<JsonNode?> PreprocessAsync(JsonNode? node, JsonNode? existingNode, ConfigContext context, CancellationToken cancellationToken)
 		{
 			JsonNode? result;
 			if (node == null)
 			{
 				result = null;
 			}
-			else if (_relativePath && node is JsonValue value && value.GetValueKind() == JsonValueKind.String)
+			else if (RelativeToContainingFile && node is JsonValue value && value.GetValueKind() == JsonValueKind.String)
 			{
-				result = JsonValue.Create(CombinePaths(context.CurrentFile, value.ToString()).AbsoluteUri);
+				result = JsonValue.Create(ConfigType.CombinePaths(context.CurrentFile, value.ToString()).AbsoluteUri);
 			}
 			else
 			{
-				result = ExpandMacros(node, context);
+				result = ConfigType.ExpandMacros(node, context);
 			}
 
 			return Task.FromResult(result);
 		}
 
 		/// <inheritdoc/>
-		internal override Task<JsonObject> ParseIncludesAsync(JsonNode node, JsonObject targetObject, ObjectConfigType targetType, ConfigContext context, CancellationToken cancellationToken)
-			=> Task.FromResult(targetObject);
-	}
-
-	/// <summary>
-	/// Property which specifies the path of another config file to include 
-	/// </summary>
-	class IncludeConfigType : ScalarConfigType
-	{
-		public IncludeConfigType(bool relativePath)
-			: base(relativePath)
-		{ }
-
-		/// <inheritdoc/>
-		internal override bool HasIncludes() => true;
-
-		/// <inheritdoc/>
-		internal override async Task<JsonObject> ParseIncludesAsync(JsonNode jsonNode, JsonObject targetObject, ObjectConfigType targetType, ConfigContext context, CancellationToken cancellationToken)
+		public override async Task ParseIncludesAsync(JsonNode? node, List<IConfigFile> includes, ConfigContext context, CancellationToken cancellationToken)
 		{
-			string? path = (string?)jsonNode;
-
-			Uri uri = ConfigType.CombinePaths(context.CurrentFile, context.ExpandMacros(path!));
-			IConfigFile file = await context.ReadFileAsync(uri, cancellationToken);
-
-			context.IncludeStack.Push(file);
-
-			JsonObject includedJsonObject = await file.ParseFileAsync(context, cancellationToken);
-			JsonNode? result = await targetType.PreprocessAndMergeAsync(includedJsonObject, targetObject, context, cancellationToken);
+			if (Include)
+			{
+				string? path = (string?)node;
 
-			context.IncludeStack.Pop();
+				Uri uri = ConfigType.CombinePaths(context.CurrentFile, context.ExpandMacros(path!));
 
-			return (JsonObject)result!;
+				IConfigFile file = await context.ReadFileAsync(uri, cancellationToken);
+				includes.Add(file);
+			}
 		}
 	}
 
 	/// <summary>
 	/// Property containing a binary resource
 	/// </summary>
-	class ResourceConfigType : ScalarConfigType
+	class ResourceConfigNode : ConfigNode
 	{
-		public ResourceConfigType()
-			: base(false)
+		/// <inheritdoc/>
+		public override Task ParseIncludesAsync(JsonNode? node, List<IConfigFile> includes, ConfigContext context, CancellationToken cancellationToken)
+			=> Task.CompletedTask;
+
+		/// <inheritdoc/>
+		public override void ParseMacros(JsonNode? node, ConfigContext context, Dictionary<string, string> macros)
 		{ }
 
 		/// <inheritdoc/>
-		internal override async Task<JsonNode?> PreprocessAndMergeAsync(JsonNode? node, JsonNode? existingNode, ConfigContext context, CancellationToken cancellationToken)
+		public override async Task<JsonNode?> PreprocessAsync(JsonNode? node, JsonNode? existingNode, ConfigContext context, CancellationToken cancellationToken)
 		{
-			Uri uri = CombinePaths(context.CurrentFile, JsonSerializer.Deserialize<string>(node, context.JsonOptions) ?? String.Empty);
+			Uri uri = ConfigType.CombinePaths(context.CurrentFile, JsonSerializer.Deserialize<string>(node, context.JsonOptions) ?? String.Empty);
 			IConfigFile file = await context.ReadFileAsync(uri, cancellationToken);
 
 			ConfigResource resource = new ConfigResource();
@@ -369,46 +319,39 @@ namespace Horde.Server.Configuration
 	/// <summary>
 	/// Array of Json values
 	/// </summary>
-	class ArrayConfigType : ConfigType
+	class ArrayConfigNode : ConfigNode
 	{
-		readonly ConfigType _elementType;
+		public ConfigNode ElementType { get; }
 
 		/// <summary>
 		/// Constructor
 		/// </summary>
-		public ArrayConfigType(ConfigType elementType)
+		public ArrayConfigNode(ConfigNode elementType)
 		{
-			_elementType = elementType;
+			ElementType = elementType;
 		}
 
 		/// <inheritdoc/>
-		internal override bool HasMacroDefinitions() => _elementType is ObjectConfigType elementType && elementType.HasMacroDefinitions();
-
-		/// <inheritdoc/>
-		internal override void ParseMacroDefinitions(JsonNode jsonNode, ConfigContext context, Dictionary<string, string> macros)
+		public override void ParseMacros(JsonNode? node, ConfigContext context, Dictionary<string, string> macros)
 		{
-			if (jsonNode is JsonArray jsonArrayValue)
+			if (node is JsonArray arrayNode)
 			{
-				ObjectConfigType classElementType = (ObjectConfigType)_elementType;
-				foreach (JsonObject jsonObjectElement in jsonArrayValue.OfType<JsonObject>())
+				foreach (JsonNode? elementNode in arrayNode)
 				{
-					classElementType.ParseMacroDefinitions(jsonObjectElement, context, macros);
+					ElementType.ParseMacros(elementNode, context, macros);
 				}
 			}
 		}
 
 		/// <inheritdoc/>
-		internal override bool HasIncludes() => _elementType is ObjectConfigType elementType && elementType.HasIncludes();
-
-		/// <inheritdoc/>
-		internal override async Task<JsonNode?> PreprocessAndMergeAsync(JsonNode? node, JsonNode? existingNode, ConfigContext context, CancellationToken cancellationToken)
+		public override async Task<JsonNode?> PreprocessAsync(JsonNode? node, JsonNode? existingNode, ConfigContext context, CancellationToken cancellationToken)
 		{
 			JsonArray targetArray = ((JsonArray?)existingNode) ?? new JsonArray();
 			foreach (JsonNode? element in (JsonArray)node!)
 			{
 				context.EnterScope($"[{targetArray.Count}]");
 
-				JsonNode? elementValue = await _elementType.PreprocessAsync(element, context, cancellationToken);
+				JsonNode? elementValue = await ElementType.PreprocessAsync(element, context, cancellationToken);
 				targetArray.Add(elementValue);
 
 				context.LeaveScope();
@@ -417,50 +360,44 @@ namespace Horde.Server.Configuration
 		}
 
 		/// <inheritdoc/>
-		internal override async Task<JsonObject> ParseIncludesAsync(JsonNode node, JsonObject targetObj, ObjectConfigType targetType, ConfigContext context, CancellationToken cancellationToken)
+		public override async Task ParseIncludesAsync(JsonNode? node, List<IConfigFile> includes, ConfigContext context, CancellationToken cancellationToken)
 		{
 			if (node is JsonArray arrayNode)
 			{
-				ObjectConfigType classElementType = (ObjectConfigType)_elementType;
-				foreach (JsonObject jsonObjectElement in arrayNode.OfType<JsonObject>())
+				ObjectConfigNode classElementType = (ObjectConfigNode)ElementType;
+				foreach (JsonObject? element in arrayNode)
 				{
-					targetObj = await classElementType.ParseIncludesAsync(jsonObjectElement, targetObj, targetType, context, cancellationToken);
+					await classElementType.ParseIncludesAsync(element, includes, context, cancellationToken);
 				}
 			}
-			return targetObj;
 		}
 	}
 
 	/// <summary>
 	/// Arbitary mapping of string values to keys
 	/// </summary>
-	class DictionaryConfigType : ConfigType
+	class DictionaryConfigNode : ConfigNode
 	{
-		readonly ConfigType _elementType;
+		public ConfigNode ElementType { get; }
 
-		public DictionaryConfigType(ConfigType elementType)
+		public DictionaryConfigNode(ConfigNode elementType)
 		{
-			_elementType = elementType;
+			ElementType = elementType;
 		}
 
 		/// <inheritdoc/>
-		internal override bool HasMacroDefinitions() => false;
-
-		/// <inheritdoc/>
-		internal override void ParseMacroDefinitions(JsonNode jsonNode, ConfigContext context, Dictionary<string, string> macros) { }
-
-		/// <inheritdoc/>
-		internal override bool HasIncludes() => _elementType is ObjectConfigType elementType && elementType.HasIncludes();
+		public override void ParseMacros(JsonNode? jsonNode, ConfigContext context, Dictionary<string, string> macros) 
+		{ }
 
 		/// <inheritdoc/>
-		internal override async Task<JsonNode?> PreprocessAndMergeAsync(JsonNode? node, JsonNode? existingNode, ConfigContext context, CancellationToken cancellationToken)
+		public override async Task<JsonNode?> PreprocessAsync(JsonNode? node, JsonNode? existingNode, ConfigContext context, CancellationToken cancellationToken)
 		{
 			JsonObject? targetObject = ((JsonObject?)existingNode) ?? new JsonObject();
 			foreach ((string key, JsonNode? element) in (JsonObject)node!)
 			{
 				context.EnterScope($"[{key}]");
 
-				JsonNode? elementValue = await _elementType.PreprocessAsync(element, context, cancellationToken);
+				JsonNode? elementValue = await ElementType.PreprocessAsync(element, context, cancellationToken);
 				targetObject[key] = elementValue;
 
 				context.LeaveScope();
@@ -469,269 +406,221 @@ namespace Horde.Server.Configuration
 		}
 
 		/// <inheritdoc/>
-		internal override async Task<JsonObject> ParseIncludesAsync(JsonNode node, JsonObject targetObj, ObjectConfigType targetType, ConfigContext context, CancellationToken cancellationToken)
+		public override async Task ParseIncludesAsync(JsonNode? node, List<IConfigFile> includes, ConfigContext context, CancellationToken cancellationToken)
 		{
-			if (node is JsonObject obj && _elementType is ObjectConfigType classElementType)
+			if (node is JsonObject obj && ElementType is ObjectConfigNode classElementType)
 			{
-				foreach (JsonObject jsonObjectElement in obj.Select(x => x.Value).OfType<JsonObject>())
+				foreach ((_, JsonNode? value) in obj)
 				{
-					targetObj = await classElementType.ParseIncludesAsync(jsonObjectElement, targetObj, targetType, context, cancellationToken);
+					await classElementType.ParseIncludesAsync(value, includes, context, cancellationToken);
 				}
 			}
-			return targetObj;
 		}
 	}
 
 	/// <summary>
-	/// Special config type for storing a json node
+	/// Handles macro objects
 	/// </summary>
-	class JsonNodeConfigType : ConfigType
+	class MacroConfigNode : ConfigNode
 	{
-		/// <inheritdoc/>
-		internal override bool HasMacroDefinitions() => false;
+		public override Task ParseIncludesAsync(JsonNode? node, List<IConfigFile> includes, ConfigContext context, CancellationToken cancellationToken)
+			=> Task.CompletedTask;
 
-		/// <inheritdoc/>
-		internal override void ParseMacroDefinitions(JsonNode jsonNode, ConfigContext context, Dictionary<string, string> macros)
+		public override void ParseMacros(JsonNode? node, ConfigContext context, Dictionary<string, string> macros)
 		{
-		}
-
-		/// <inheritdoc/>
-		internal override bool HasIncludes() => false;
-
-		/// <inheritdoc/>
-		internal override Task<JsonNode?> PreprocessAndMergeAsync(JsonNode? node, JsonNode? existingNode, ConfigContext context, CancellationToken cancellationToken)
-		{
-			return Task.FromResult(node?.DeepClone());
-		}
-
-		static JsonNode? MergeNodes(JsonNode? first, JsonNode? second)
-		{
-			if (second == null)
-			{
-				return first?.DeepClone();
-			}
-			else if (first is JsonObject firstObj && second is JsonObject secondObj)
-			{
-				JsonObject mergedObj = new JsonObject(firstObj);
-				foreach ((string childName, JsonNode? childNode) in secondObj)
-				{
-					mergedObj[childName] = MergeNodes(firstObj[childName], childNode);
-				}
-				return mergedObj;
-			}
-			else
+			ConfigMacro? macro = JsonSerializer.Deserialize<ConfigMacro>(node, context.JsonOptions);
+			if (macro != null)
 			{
-				return second.DeepClone();
+				macros.Add(macro.Name, macro.Value);
 			}
 		}
 
-		/// <inheritdoc/>
-		internal override Task<JsonObject> ParseIncludesAsync(JsonNode jsonNode, JsonObject targetObj, ObjectConfigType targetType, ConfigContext context, CancellationToken cancellationToken)
-			=> Task.FromResult(targetObj);
+		public override Task<JsonNode?> PreprocessAsync(JsonNode? node, JsonNode? existingNode, ConfigContext context, CancellationToken cancellationToken)
+			=> Task.FromResult<JsonNode?>(node?.DeepClone());
 	}
 
 	/// <summary>
 	/// Implementation of <see cref="ConfigType"/> to handle class types
 	/// </summary>
-	class ObjectConfigType : ConfigType
+	class ObjectConfigNode : ConfigNode
 	{
-		readonly bool _isIncludeRoot;
-		readonly bool _isMacro;
-		readonly bool _isMacroScope;
-		readonly Dictionary<string, ConfigType> _nameToProperty;
-		readonly Dictionary<string, ConfigType> _nameToMacroProperty = new Dictionary<string, ConfigType>(StringComparer.OrdinalIgnoreCase);
-		readonly Dictionary<string, ConfigType> _nameToIncludeProperty = new Dictionary<string, ConfigType>(StringComparer.OrdinalIgnoreCase);
-
-		static readonly ConcurrentDictionary<Type, ObjectConfigType> s_typeToObjectValueType = new ConcurrentDictionary<Type, ObjectConfigType>();
+		public bool IncludeRoot { get; set; }
+		public bool MacroScope { get; set; }
+		public Dictionary<string, ConfigNode> Properties { get; }
 
 		/// <summary>
 		/// Constructor
 		/// </summary>
-		public ObjectConfigType(bool isIncludeRoot, bool isMacro, bool isMacroScope, IEnumerable<KeyValuePair<string, ConfigType>> properties)
+		public ObjectConfigNode(bool isIncludeRoot, bool isMacroScope, IEnumerable<KeyValuePair<string, ConfigNode>> properties)
 		{
-			_isIncludeRoot = isIncludeRoot;
-			_isMacro = isMacro;
-			_isMacroScope = isMacroScope;
-
-			_nameToProperty = new Dictionary<string, ConfigType>(properties, StringComparer.OrdinalIgnoreCase);
-
-			// Build a map of all the properties which can contain macros or include other files
-			foreach ((string name, ConfigType property) in _nameToProperty)
-			{
-				if (property.HasMacroDefinitions())
-				{
-					_nameToMacroProperty.Add(name, property);
-				}
-				if (property.HasIncludes())
-				{
-					_nameToIncludeProperty.Add(name, property);
-				}
-			}
+			IncludeRoot = isIncludeRoot;
+			MacroScope = isMacroScope;
+			Properties = new Dictionary<string, ConfigNode>(properties, StringComparer.OrdinalIgnoreCase);
 		}
 
-		static ObjectConfigType FromType(Type type)
+		/// <summary>
+		/// Constructor
+		/// </summary>
+		/// <param name="type">Type to construct from</param>
+		public ObjectConfigNode(Type type)
 		{
-			bool isIncludeRoot = type.GetCustomAttribute<ConfigIncludeRootAttribute>() != null;
-			bool isMacro = type == typeof(ConfigMacro);
-			bool isMacroScope = type.GetCustomAttribute<ConfigMacroScopeAttribute>() != null;
+			IncludeRoot = type.GetCustomAttribute<ConfigIncludeRootAttribute>() != null;
+			MacroScope = type.GetCustomAttribute<ConfigMacroScopeAttribute>() != null;
+			Properties = new Dictionary<string, ConfigNode>(StringComparer.OrdinalIgnoreCase);
 
 			// Find all the direct include properties
-			Dictionary<string, ConfigType> nameToProperty = new Dictionary<string, ConfigType>(StringComparer.OrdinalIgnoreCase);
-
 			PropertyInfo[] propertyInfos = type.GetProperties(BindingFlags.Instance | BindingFlags.Public | BindingFlags.GetProperty);
 			foreach (PropertyInfo propertyInfo in propertyInfos)
 			{
 				if (propertyInfo.GetCustomAttribute<JsonIgnoreAttribute>() == null)
 				{
-					string name = propertyInfo.GetCustomAttribute<JsonPropertyNameAttribute>()?.Name ?? propertyInfo.Name;
-					nameToProperty.Add(name, CreateProperty(propertyInfo));
-				}
-			}
-
-			// Create the type
-			return new ObjectConfigType(isIncludeRoot, isMacro, isMacroScope, nameToProperty);
-		}
-
-		/// <inheritdoc/>
-		internal override bool HasMacroDefinitions() => !_isMacroScope && (_isMacro || _nameToMacroProperty.Count > 0);
-
-		/// <inheritdoc/>
-		internal override bool HasIncludes() => !_isIncludeRoot && _nameToIncludeProperty.Count > 0;
-
-		public static ObjectConfigType FindOrAdd(Type type)
-		{
-			ObjectConfigType? value;
-			if (!s_typeToObjectValueType.TryGetValue(type, out value))
-			{
-				lock (s_typeToObjectValueType)
-				{
-					if (!s_typeToObjectValueType.TryGetValue(type, out value))
+					ConfigNode? propertyType = CreateTypeForProperty(propertyInfo);
+					if (propertyType != null)
 					{
-						value = ObjectConfigType.FromType(type);
-						s_typeToObjectValueType.TryAdd(type, value);
+						string name = propertyInfo.GetCustomAttribute<JsonPropertyNameAttribute>()?.Name ?? propertyInfo.Name;
+						Properties.Add(name, propertyType);
 					}
 				}
 			}
-			return value;
 		}
 
-		static ConfigType CreateProperty(PropertyInfo propertyInfo)
+		bool IsDefault()
+			=> !IncludeRoot && !MacroScope && Properties.Count == 0;
+
+		static ConfigNode? CreateTypeForProperty(PropertyInfo propertyInfo)
 		{
 			Type propertyType = propertyInfo.PropertyType;
 			if (!propertyType.IsClass || propertyType == typeof(string))
 			{
+				bool include = propertyInfo.GetCustomAttribute<ConfigIncludeAttribute>() != null;
 				bool relativePath = propertyInfo.GetCustomAttribute<ConfigRelativePathAttribute>() != null;
-				if (propertyInfo.GetCustomAttribute<ConfigIncludeAttribute>() != null)
+				if (include || relativePath)
 				{
-					return new IncludeConfigType(relativePath);
+					return new ScalarConfigNode(include, relativePath);
 				}
 				else
 				{
-					return new ScalarConfigType(relativePath);
+					return null;
 				}
 			}
+			return CreateType(propertyType);
+		}
+
+		public static ConfigNode? CreateType(Type type)
+		{
+			ConfigNode? value;
+			if (!type.IsClass || type == typeof(string) || type == typeof(JsonNode))
+			{
+				value = null;
+			}
+			else if (type == typeof(ConfigMacro))
+			{
+				value = new MacroConfigNode();
+			}
+			else if (type.IsAssignableTo(typeof(ConfigResource)))
+			{
+				value = new ResourceConfigNode();
+			}
+			else if (type.IsGenericType && type.GetGenericTypeDefinition() == typeof(List<>))
+			{
+				ConfigNode? elementType = CreateType(type.GetGenericArguments()[0]);
+				value = (elementType == null) ? null : new ArrayConfigNode(elementType);
+			}
+			else if (type.IsGenericType && type.GetGenericTypeDefinition() == typeof(Dictionary<,>))
+			{
+				ConfigNode? elementType = CreateType(type.GetGenericArguments()[1]);
+				value = (elementType == null) ? null : new DictionaryConfigNode(elementType);
+			}
 			else
 			{
-				if (propertyType.IsAssignableTo(typeof(ConfigResource)))
-				{
-					return new ResourceConfigType();
-				}
-				else if (propertyType.IsGenericType && propertyType.GetGenericTypeDefinition() == typeof(List<>))
-				{
-					Type elementType = propertyType.GetGenericArguments()[0];
-					return new ArrayConfigType(FindOrAddValueType(elementType));
-				}
-				else if (propertyType.IsGenericType && propertyType.GetGenericTypeDefinition() == typeof(Dictionary<,>))
-				{
-					Type elementType = propertyType.GetGenericArguments()[1];
-					return new DictionaryConfigType(FindOrAddValueType(elementType));
-				}
-				else if (propertyType.IsAssignableTo(typeof(JsonNode)))
-				{
-					return new JsonNodeConfigType();
-				}
-				else
-				{
-					return FindOrAdd(propertyType);
-				}
+				ObjectConfigNode objValue = new ObjectConfigNode(type);
+				value = objValue.IsDefault() ? null : objValue;
 			}
+			return value;
 		}
 
 		public async ValueTask<JsonObject> ReadAsync(JsonObject obj, ConfigContext context, CancellationToken cancellationToken)
 		{
-			return (JsonObject)(await PreprocessAndMergeAsync(obj, null, context, cancellationToken))!;
+			return (JsonObject)(await PreprocessAsync(obj, null, context, cancellationToken))!;
 		}
 
 		/// <inheritdoc/>
-		internal override void ParseMacroDefinitions(JsonNode jsonNode, ConfigContext context, Dictionary<string, string> macros)
+		public override void ParseMacros(JsonNode? jsonNode, ConfigContext context, Dictionary<string, string> macros)
 		{
-			JsonObject jsonObject = (JsonObject)jsonNode;
-			if (_isMacro)
-			{
-				ConfigMacro? macro = JsonSerializer.Deserialize<ConfigMacro>(jsonObject, context.JsonOptions);
-				if (macro != null)
-				{
-					macros.Add(macro.Name, macro.Value);
-				}
-			}
-			else
+			if (jsonNode is JsonObject jsonObject)
 			{
 				foreach ((string name, JsonNode? node) in jsonObject)
 				{
-					if (node != null && _nameToMacroProperty.TryGetValue(name, out ConfigType? property))
+					if (node != null && Properties.TryGetValue(name, out ConfigNode? property))
 					{
-						property.ParseMacroDefinitions(node, context, macros);
+						property.ParseMacros(node, context, macros);
 					}
 				}
 			}
 		}
 
 		/// <inheritdoc/>
-		internal override async Task<JsonNode?> PreprocessAndMergeAsync(JsonNode? node, JsonNode? existingNode, ConfigContext context, CancellationToken cancellationToken)
+		public override async Task<JsonNode?> PreprocessAsync(JsonNode? node, JsonNode? existingNode, ConfigContext context, CancellationToken cancellationToken)
 		{
-			if (node is not JsonObject newObject)
+			JsonObject? obj = (JsonObject?)node;
+			JsonObject? target = (JsonObject?)existingNode;
+			if (obj == null)
 			{
-				return node?.DeepClone();
+				return target;
 			}
 
-			JsonObject? target = (JsonObject?)existingNode ?? new JsonObject();
-
 			// Before parsing properties into this object, read all the includes recursively
-			if (_isIncludeRoot)
+			if (IncludeRoot)
 			{
-				await ParseIncludesAsync(newObject, target, this, context, cancellationToken);
+				// Find all the includes
+				List<IConfigFile> includes = new List<IConfigFile>();
+				await ParseIncludesInternalAsync(obj, includes, context, cancellationToken);
+
+				// Find all the files, merge them into target
+				foreach (IConfigFile include in includes)
+				{
+					context.IncludeStack.Push(include);
+
+					JsonObject includedJsonObject = await include.ParseFileAsync(context, cancellationToken);
+					target = (JsonObject?)await PreprocessAsync(includedJsonObject, target, context, cancellationToken);
+
+					context.IncludeStack.Pop();
+				}
 			}
 
+			// Ensure that the target object is valid so we can write properties into it
+			target ??= new JsonObject();
+
 			// Parse all the macros for this scope
-			if (_isMacroScope)
+			if (MacroScope)
 			{
 				Dictionary<string, string> macros = new Dictionary<string, string>(StringComparer.OrdinalIgnoreCase);
-				ParseMacroDefinitions(newObject, context, macros);
+				ParseMacros(obj, context, macros);
 				context.MacroScopes.Add(macros);
 			}
 
 			// Parse all the properties into this object
-			foreach ((string name, JsonNode? newNode) in newObject)
+			foreach ((string name, JsonNode? newNode) in obj)
 			{
 				if (newNode is JsonValue)
 				{
 					context.AddProperty(name);
 				}
 
-				if (_nameToProperty.TryGetValue(name, out ConfigType? property))
+				if (Properties.TryGetValue(name, out ConfigNode? property))
 				{
 					context.EnterScope(name);
-					target[name] = await property.PreprocessAndMergeAsync(newNode, target[name], context, cancellationToken);
+					target[name] = await property.PreprocessAsync(newNode, target[name], context, cancellationToken);
 					context.LeaveScope();
 				}
 				else
 				{
-					target[name] = ExpandMacros(newNode, context);
+					target[name] = Merge(ConfigType.ExpandMacros(newNode, context), target[name]);
 				}
 			}
 
 			// Parse all the macros for this scope
-			if (_isMacroScope)
+			if (MacroScope)
 			{
 				context.MacroScopes.RemoveAt(context.MacroScopes.Count - 1);
 			}
@@ -739,20 +628,45 @@ namespace Horde.Server.Configuration
 			return target;
 		}
 
+		static JsonNode? Merge(JsonNode? source, JsonNode? target)
+		{
+			if (source is JsonObject sourceObj && target is JsonObject targetObj)
+			{
+				foreach ((string name, JsonNode? node) in sourceObj)
+				{
+					targetObj[name] = Merge(node, targetObj[name]);
+				}
+				return target;
+			}
+			else if (source is JsonArray sourceArr && target is JsonArray targetArr)
+			{
+				foreach (JsonNode? node in sourceArr)
+				{
+					targetArr.Add(node?.DeepClone());
+				}
+				return target;
+			}
+			return source;
+		}
+
 		/// <inheritdoc/>
-		internal override async Task<JsonObject> ParseIncludesAsync(JsonNode node, JsonObject targetObj, ObjectConfigType targetType, ConfigContext context, CancellationToken cancellationToken)
+		public override async Task ParseIncludesAsync(JsonNode? node, List<IConfigFile> includes, ConfigContext context, CancellationToken cancellationToken)
+		{
+			if (!IncludeRoot && node is JsonObject obj)
+			{
+				await ParseIncludesInternalAsync(obj, includes, context, cancellationToken);
+			}
+		}
+
+		async Task ParseIncludesInternalAsync(JsonObject obj, List<IConfigFile> includes, ConfigContext context, CancellationToken cancellationToken)
 		{
-			if (node is JsonObject obj)
+			foreach ((string name, JsonNode? propertyNode) in obj)
 			{
-				foreach ((string name, JsonNode? propertyNode) in obj)
+				if (Properties.TryGetValue(name, out ConfigNode? property))
 				{
-					if (_nameToIncludeProperty.TryGetValue(name, out ConfigType? property) && propertyNode != null)
-					{
-						targetObj = await property.ParseIncludesAsync(propertyNode, targetObj, targetType, context, cancellationToken);
-					}
+					await property.ParseIncludesAsync(propertyNode, includes, context, cancellationToken);
 				}
 			}
-			return targetObj;
 		}
 	}
 }
