Horde: Move config classes into separate files.

#rnx

diff --git a/Engine/Source/Programs/Horde/Horde.Server/Configuration/ConfigContext.cs b/Engine/Source/Programs/Horde/Horde.Server/Configuration/ConfigContext.cs
--- a/Engine/Source/Programs/Horde/Horde.Server/Configuration/ConfigContext.cs
+++ b/Engine/Source/Programs/Horde/Horde.Server/Configuration/ConfigContext.cs
@@ -12,11 +12,13 @@ using Microsoft.Extensions.Logging;
 
 namespace Horde.Server.Configuration
 {
+	using JsonObject = System.Text.Json.Nodes.JsonObject;
+
 	/// <summary>
 	/// Context for reading a tree of config files
 	/// </summary>
 	[DebuggerDisplay("{CurrentFile}")]
-	class ConfigContext
+	public class ConfigContext
 	{
 		/// <summary>
 		/// Options for serializing config files
@@ -33,11 +35,6 @@ namespace Horde.Server.Configuration
 		/// </summary>
 		public Stack<string> ScopeStack { get; } = new Stack<string>();
 
-		/// <summary>
-		/// Stack of objects
-		/// </summary>
-		public Stack<object> IncludeContextStack { get; } = new Stack<object>();
-
 		/// <summary>
 		/// Map of property path to the file declaring a value for it
 		/// </summary>
@@ -183,5 +180,46 @@ namespace Horde.Server.Configuration
 
 			return file;
 		}
+
+		/// <summary>
+		/// Reads an object from a particular URL
+		/// </summary>
+		/// <typeparam name="T">Type of object to read</typeparam>
+		/// <param name="uri">Location of the file to read</param>
+		/// <param name="cancellationToken">Cancellation token for the operation</param>
+		/// <returns></returns>
+		public async Task<T> ReadAsync<T>(Uri uri, CancellationToken cancellationToken) where T : class, new()
+		{
+			IConfigFile file = await ReadFileAsync(uri, cancellationToken);
+			JsonObject obj = await ParseFileAsync(file, cancellationToken);
+
+			IncludeStack.Push(file);
+
+			ObjectConfigNode type = new ObjectConfigNode(typeof(T));
+			obj = (JsonObject)(await type.PreprocessAsync(obj, null, this, cancellationToken))!;
+
+			IncludeStack.Pop();
+
+			return JsonSerializer.Deserialize<T>(obj, JsonOptions) ?? new T();
+		}
+
+		/// <summary>
+		/// Parses a config file as a json object
+		/// </summary>
+		/// <param name="file">File to parse</param>
+		/// <param name="cancellationToken">Cancellation token for the operation</param>
+		/// <returns>The parsed file</returns>
+		public async ValueTask<JsonObject> ParseFileAsync(IConfigFile file, CancellationToken cancellationToken = default)
+		{
+			ReadOnlyMemory<byte> data = await file.ReadAsync(cancellationToken);
+
+			JsonObject? obj = JsonSerializer.Deserialize<JsonObject>(data.Span, JsonOptions);
+			if (obj == null)
+			{
+				throw new ConfigException(this, $"Config file {file.Uri} contains a null object.");
+			}
+
+			return obj;
+		}
 	}
 }
diff --git a/Engine/Source/Programs/Horde/Horde.Server/Configuration/ConfigException.cs b/Engine/Source/Programs/Horde/Horde.Server/Configuration/ConfigException.cs
new file mode 100644
--- /dev/null
+++ b/Engine/Source/Programs/Horde/Horde.Server/Configuration/ConfigException.cs
@@ -0,0 +1,37 @@
+// Copyright Epic Games, Inc. All Rights Reserved.
+
+using System;
+using System.Collections.Generic;
+
+namespace Horde.Server.Configuration
+{
+	/// <summary>
+	/// Exception thrown when reading config files
+	/// </summary>
+	public sealed class ConfigException : Exception
+	{
+		readonly ConfigContext _context;
+
+		/// <summary>
+		/// Stack of properties
+		/// </summary>
+		public IEnumerable<string> ScopeStack => _context.ScopeStack;
+
+		/// <summary>
+		/// Constructor
+		/// </summary>
+		/// <param name="context">Current parse context for the error</param>
+		/// <param name="message">Description of the error</param>
+		/// <param name="innerException">Inner exception details</param>
+		internal ConfigException(ConfigContext context, string message, Exception? innerException = null)
+			: base(message, innerException)
+		{
+			_context = context;
+		}
+
+		/// <summary>
+		/// Gets the parser context when this exception was thrown. This is not exposed as a public property to avoid serializing the whole thing to Serilog.
+		/// </summary>
+		internal ConfigContext GetContext() => _context;
+	}
+}
diff --git a/Engine/Source/Programs/Horde/Horde.Server/Configuration/ConfigFile.cs b/Engine/Source/Programs/Horde/Horde.Server/Configuration/ConfigFile.cs
deleted file mode 100644
diff --git a/Engine/Source/Programs/Horde/Horde.Server/Configuration/ConfigMacro.cs b/Engine/Source/Programs/Horde/Horde.Server/Configuration/ConfigMacro.cs
new file mode 100644
--- /dev/null
+++ b/Engine/Source/Programs/Horde/Horde.Server/Configuration/ConfigMacro.cs
@@ -0,0 +1,22 @@
+// Copyright Epic Games, Inc. All Rights Reserved.
+
+using System;
+
+namespace Horde.Server.Configuration
+{
+	/// <summary>
+	/// Declares a config macro
+	/// </summary>
+	public class ConfigMacro
+	{
+		/// <summary>
+		/// Name of the macro property
+		/// </summary>
+		public string Name { get; set; } = String.Empty;
+
+		/// <summary>
+		/// Value for the macro property
+		/// </summary>
+		public string Value { get; set; } = String.Empty;
+	}
+}
diff --git a/Engine/Source/Programs/Horde/Horde.Server/Configuration/ConfigNode.cs b/Engine/Source/Programs/Horde/Horde.Server/Configuration/ConfigNode.cs
new file mode 100644
--- /dev/null
+++ b/Engine/Source/Programs/Horde/Horde.Server/Configuration/ConfigNode.cs
@@ -0,0 +1,619 @@
+// Copyright Epic Games, Inc. All Rights Reserved.
+
+using System;
+using System.Collections.Generic;
+using System.Reflection;
+using System.Text.Json;
+using System.Text.Json.Nodes;
+using System.Text.Json.Serialization;
+using System.Threading;
+using System.Threading.Tasks;
+
+namespace Horde.Server.Configuration
+{
+	/// <summary>
+	/// Attribute used to mark <see cref="Uri"/> properties that include other config files
+	/// </summary>
+	[AttributeUsage(AttributeTargets.Property)]
+	public sealed class ConfigIncludeAttribute : Attribute
+	{
+	}
+
+	/// <summary>
+	/// Specifies that a class is the root for including other files
+	/// </summary>
+	[AttributeUsage(AttributeTargets.Class)]
+	public sealed class ConfigIncludeRootAttribute : Attribute
+	{
+	}
+
+	/// <summary>
+	/// Attribute used to mark <see cref="Uri"/> properties that are relative to their containing file
+	/// </summary>
+	[AttributeUsage(AttributeTargets.Property)]
+	public sealed class ConfigRelativePathAttribute : Attribute
+	{
+	}
+
+	/// <summary>
+	/// Attribute used to mark <see cref="Uri"/> properties that are relative to their containing file
+	/// </summary>
+	[AttributeUsage(AttributeTargets.Class)]
+	public sealed class ConfigMacroScopeAttribute : Attribute
+	{
+	}
+
+	/// <summary>
+	/// Node in the preprocessor parse tree
+	/// </summary>
+	public abstract class ConfigNode
+	{
+		/// <summary>
+		/// Parses macro definitions from this object
+		/// </summary>
+		/// <param name="node">Node to parse macros from</param>
+		/// <param name="context">Context for the preprocessor</param>
+		/// <param name="macros">Macros parsed from the boject</param>
+		public abstract void ParseMacros(JsonNode? node, ConfigContext context, Dictionary<string, string> macros);
+
+		/// <summary>
+		/// Preprocesses a json document
+		/// </summary>
+		/// <param name="node">Node to preprocess</param>
+		/// <param name="context">Context for the preprocessor</param>
+		/// <param name="cancellationToken">Cancellation token for the operation</param>
+		/// <returns>Preprocessed node</returns>
+		public Task<JsonNode?> PreprocessAsync(JsonNode? node, ConfigContext context, CancellationToken cancellationToken)
+			=> PreprocessAsync(node, null, context, cancellationToken);
+
+		/// <summary>
+		/// Preprocesses a json document (possibly merging with an existing property)
+		/// </summary>
+		/// <param name="node">Node to preprocess</param>
+		/// <param name="existingNode">Optional existing node to merge with. Can be modified.</param>
+		/// <param name="context">Context for the preprocessor</param>
+		/// <param name="cancellationToken">Cancellation token for the operation</param>
+		/// <returns>The preprocessed node. May be existingNode if changes are merged with it.</returns>
+		public abstract Task<JsonNode?> PreprocessAsync(JsonNode? node, JsonNode? existingNode, ConfigContext context, CancellationToken cancellationToken);
+
+		/// <summary>
+		/// Traverse the property tree starting with 'node' and process any include directives, merging the results into the given target object.
+		/// </summary>
+		/// <param name="node"></param>
+		/// <param name="includes">Receives the included files</param>
+		/// <param name="context"></param>
+		/// <param name="cancellationToken"></param>
+		/// <returns></returns>
+		public abstract Task ParseIncludesAsync(JsonNode? node, List<IConfigFile> includes, ConfigContext context, CancellationToken cancellationToken);
+
+		/// <summary>
+		/// Combine a relative path with a base URI to produce a new URI
+		/// </summary>
+		/// <param name="baseUri">Base uri to rebase relative to</param>
+		/// <param name="path">Relative path</param>
+		/// <returns>Absolute URI</returns>
+		public static Uri CombinePaths(Uri baseUri, string path)
+		{
+			if (path.StartsWith("//", StringComparison.Ordinal))
+			{
+				if (baseUri.Scheme == PerforceConfigSource.Scheme)
+				{
+					return new Uri($"{PerforceConfigSource.Scheme}://{baseUri.Host}{path}");
+				}
+				else
+				{
+					return new Uri($"{PerforceConfigSource.Scheme}://{PerforceConnectionSettings.Default}{path}");
+				}
+			}
+			return new Uri(baseUri, path);
+		}
+
+		/// <summary>
+		/// Helper method to expand all macros within a node without performing any other processing on it
+		/// </summary>
+		protected static JsonNode? ExpandMacros(JsonNode? node, ConfigContext context)
+		{
+			if (node == null)
+			{
+				return node;
+			}
+			else if (node is JsonObject obj)
+			{
+				JsonObject result = new JsonObject();
+				foreach ((string propertyName, JsonNode? propertyNode) in obj)
+				{
+					result[propertyName] = ExpandMacros(propertyNode, context);
+				}
+				return result;
+			}
+			else if (node is JsonArray arr)
+			{
+				JsonArray result = new JsonArray();
+				foreach (JsonNode? elementNode in arr)
+				{
+					result.Add(ExpandMacros(elementNode, context));
+				}
+				return result;
+			}
+			else if (node is JsonValue val && val.GetValueKind() == JsonValueKind.String)
+			{
+				string strValue = ((string?)val)!;
+				string expandedStrValue = context.ExpandMacros(strValue);
+				return JsonValue.Create(expandedStrValue);
+			}
+			else
+			{
+				return node.DeepClone();
+			}
+		}
+	}
+
+	/// <summary>
+	/// Implementation of <see cref="ConfigNode"/> for scalar types
+	/// </summary>
+	public class ScalarConfigNode : ConfigNode
+	{
+		/// <summary>
+		/// Whether this node represents an include directive
+		/// </summary>
+		public bool Include { get; set; }
+
+		/// <summary>
+		/// Specifies that this node is a string which should be made into an absolute path using the path of the current file.
+		/// </summary>
+		public bool RelativeToContainingFile { get; set; }
+
+		/// <summary>
+		/// Constructor
+		/// </summary>
+		public ScalarConfigNode(bool include = false, bool relativePath = false)
+		{
+			Include = include;
+			RelativeToContainingFile = relativePath;
+		}
+
+		/// <inheritdoc/>
+		public override void ParseMacros(JsonNode? jsonNode, ConfigContext context, Dictionary<string, string> macros)
+		{ }
+
+		/// <inheritdoc/>
+		public override Task<JsonNode?> PreprocessAsync(JsonNode? node, JsonNode? existingNode, ConfigContext context, CancellationToken cancellationToken)
+		{
+			JsonNode? result;
+			if (node == null)
+			{
+				result = null;
+			}
+			else if (RelativeToContainingFile && node is JsonValue value && value.GetValueKind() == JsonValueKind.String)
+			{
+				result = JsonValue.Create(CombinePaths(context.CurrentFile, value.ToString()).AbsoluteUri);
+			}
+			else
+			{
+				result = ExpandMacros(node, context);
+			}
+
+			return Task.FromResult(result);
+		}
+
+		/// <inheritdoc/>
+		public override async Task ParseIncludesAsync(JsonNode? node, List<IConfigFile> includes, ConfigContext context, CancellationToken cancellationToken)
+		{
+			if (Include)
+			{
+				string? path = (string?)node;
+
+				Uri uri = CombinePaths(context.CurrentFile, context.ExpandMacros(path!));
+
+				IConfigFile file = await context.ReadFileAsync(uri, cancellationToken);
+				includes.Add(file);
+			}
+		}
+	}
+
+	/// <summary>
+	/// Property containing a binary resource
+	/// </summary>
+	public class ResourceConfigNode : ConfigNode
+	{
+		/// <inheritdoc/>
+		public override Task ParseIncludesAsync(JsonNode? node, List<IConfigFile> includes, ConfigContext context, CancellationToken cancellationToken)
+			=> Task.CompletedTask;
+
+		/// <inheritdoc/>
+		public override void ParseMacros(JsonNode? node, ConfigContext context, Dictionary<string, string> macros)
+		{ }
+
+		/// <inheritdoc/>
+		public override async Task<JsonNode?> PreprocessAsync(JsonNode? node, JsonNode? existingNode, ConfigContext context, CancellationToken cancellationToken)
+		{
+			Uri uri = CombinePaths(context.CurrentFile, JsonSerializer.Deserialize<string>(node, context.JsonOptions) ?? String.Empty);
+			IConfigFile file = await context.ReadFileAsync(uri, cancellationToken);
+
+			ConfigResource resource = new ConfigResource();
+			resource.Path = uri.AbsoluteUri;
+			resource.Data = await file.ReadAsync(cancellationToken);
+
+			return JsonSerializer.SerializeToNode(resource, context.JsonOptions);
+		}
+	}
+
+	/// <summary>
+	/// Array of Json values
+	/// </summary>
+	public class ArrayConfigNode : ConfigNode
+	{
+		/// <summary>
+		/// Type of each element of the array
+		/// </summary>
+		public ConfigNode ElementType { get; }
+
+		/// <summary>
+		/// Constructor
+		/// </summary>
+		public ArrayConfigNode(ConfigNode elementType)
+		{
+			ElementType = elementType;
+		}
+
+		/// <inheritdoc/>
+		public override void ParseMacros(JsonNode? node, ConfigContext context, Dictionary<string, string> macros)
+		{
+			if (node is JsonArray arrayNode)
+			{
+				foreach (JsonNode? elementNode in arrayNode)
+				{
+					ElementType.ParseMacros(elementNode, context, macros);
+				}
+			}
+		}
+
+		/// <inheritdoc/>
+		public override async Task<JsonNode?> PreprocessAsync(JsonNode? node, JsonNode? existingNode, ConfigContext context, CancellationToken cancellationToken)
+		{
+			JsonArray targetArray = ((JsonArray?)existingNode) ?? new JsonArray();
+			foreach (JsonNode? element in (JsonArray)node!)
+			{
+				context.EnterScope($"[{targetArray.Count}]");
+
+				JsonNode? elementValue = await ElementType.PreprocessAsync(element, context, cancellationToken);
+				targetArray.Add(elementValue);
+
+				context.LeaveScope();
+			}
+			return targetArray;
+		}
+
+		/// <inheritdoc/>
+		public override async Task ParseIncludesAsync(JsonNode? node, List<IConfigFile> includes, ConfigContext context, CancellationToken cancellationToken)
+		{
+			if (node is JsonArray arrayNode)
+			{
+				ObjectConfigNode classElementType = (ObjectConfigNode)ElementType;
+				foreach (JsonObject? element in arrayNode)
+				{
+					await classElementType.ParseIncludesAsync(element, includes, context, cancellationToken);
+				}
+			}
+		}
+	}
+
+	/// <summary>
+	/// Arbitary mapping of string values to keys
+	/// </summary>
+	public class DictionaryConfigNode : ConfigNode
+	{
+		/// <summary>
+		/// Type of values in the dictionary
+		/// </summary>
+		public ConfigNode ValueType { get; }
+
+		/// <summary>
+		/// Constructor
+		/// </summary>
+		public DictionaryConfigNode(ConfigNode valueType)
+			=> ValueType = valueType;
+
+		/// <inheritdoc/>
+		public override void ParseMacros(JsonNode? jsonNode, ConfigContext context, Dictionary<string, string> macros)
+		{ }
+
+		/// <inheritdoc/>
+		public override async Task<JsonNode?> PreprocessAsync(JsonNode? node, JsonNode? existingNode, ConfigContext context, CancellationToken cancellationToken)
+		{
+			JsonObject? targetObject = ((JsonObject?)existingNode) ?? new JsonObject();
+			foreach ((string key, JsonNode? element) in (JsonObject)node!)
+			{
+				context.EnterScope($"[{key}]");
+
+				JsonNode? elementValue = await ValueType.PreprocessAsync(element, context, cancellationToken);
+				targetObject[key] = elementValue;
+
+				context.LeaveScope();
+			}
+			return targetObject;
+		}
+
+		/// <inheritdoc/>
+		public override async Task ParseIncludesAsync(JsonNode? node, List<IConfigFile> includes, ConfigContext context, CancellationToken cancellationToken)
+		{
+			if (node is JsonObject obj && ValueType is ObjectConfigNode classElementType)
+			{
+				foreach ((_, JsonNode? value) in obj)
+				{
+					await classElementType.ParseIncludesAsync(value, includes, context, cancellationToken);
+				}
+			}
+		}
+	}
+
+	/// <summary>
+	/// Handles macro objects
+	/// </summary>
+	public class MacroConfigNode : ConfigNode
+	{
+		/// <inheritdoc/>
+		public override Task ParseIncludesAsync(JsonNode? node, List<IConfigFile> includes, ConfigContext context, CancellationToken cancellationToken)
+			=> Task.CompletedTask;
+
+		/// <inheritdoc/>
+		public override void ParseMacros(JsonNode? node, ConfigContext context, Dictionary<string, string> macros)
+		{
+			ConfigMacro? macro = JsonSerializer.Deserialize<ConfigMacro>(node, context.JsonOptions);
+			if (macro != null)
+			{
+				macros.Add(macro.Name, macro.Value);
+			}
+		}
+
+		/// <inheritdoc/>
+		public override Task<JsonNode?> PreprocessAsync(JsonNode? node, JsonNode? existingNode, ConfigContext context, CancellationToken cancellationToken)
+			=> Task.FromResult<JsonNode?>(node?.DeepClone());
+	}
+
+	/// <summary>
+	/// Implementation of <see cref="ConfigNode"/> to handle class types
+	/// </summary>
+	public class ObjectConfigNode : ConfigNode
+	{
+		/// <summary>
+		/// Whether this object should be treated as a root for nested include directies
+		/// </summary>
+		public bool IncludeRoot { get; set; }
+
+		/// <summary>
+		/// Declares a new macro scope
+		/// </summary>
+		public bool MacroScope { get; set; }
+
+		/// <summary>
+		/// Properties within this object
+		/// </summary>
+		public Dictionary<string, ConfigNode> Properties { get; }
+
+		/// <summary>
+		/// Constructor
+		/// </summary>
+		public ObjectConfigNode(bool isIncludeRoot, bool isMacroScope, IEnumerable<KeyValuePair<string, ConfigNode>> properties)
+		{
+			IncludeRoot = isIncludeRoot;
+			MacroScope = isMacroScope;
+			Properties = new Dictionary<string, ConfigNode>(properties, StringComparer.OrdinalIgnoreCase);
+		}
+
+		/// <summary>
+		/// Constructor
+		/// </summary>
+		/// <param name="type">Type to construct from</param>
+		public ObjectConfigNode(Type type)
+		{
+			IncludeRoot = type.GetCustomAttribute<ConfigIncludeRootAttribute>() != null;
+			MacroScope = type.GetCustomAttribute<ConfigMacroScopeAttribute>() != null;
+			Properties = new Dictionary<string, ConfigNode>(StringComparer.OrdinalIgnoreCase);
+
+			// Find all the direct include properties
+			PropertyInfo[] propertyInfos = type.GetProperties(BindingFlags.Instance | BindingFlags.Public | BindingFlags.GetProperty);
+			foreach (PropertyInfo propertyInfo in propertyInfos)
+			{
+				if (propertyInfo.GetCustomAttribute<JsonIgnoreAttribute>() == null)
+				{
+					ConfigNode? propertyType = CreateTypeForProperty(propertyInfo);
+					if (propertyType != null)
+					{
+						string name = propertyInfo.GetCustomAttribute<JsonPropertyNameAttribute>()?.Name ?? propertyInfo.Name;
+						Properties.Add(name, propertyType);
+					}
+				}
+			}
+		}
+
+		bool IsDefault()
+			=> !IncludeRoot && !MacroScope && Properties.Count == 0;
+
+		static ConfigNode? CreateTypeForProperty(PropertyInfo propertyInfo)
+		{
+			Type propertyType = propertyInfo.PropertyType;
+			if (!propertyType.IsClass || propertyType == typeof(string))
+			{
+				bool include = propertyInfo.GetCustomAttribute<ConfigIncludeAttribute>() != null;
+				bool relativePath = propertyInfo.GetCustomAttribute<ConfigRelativePathAttribute>() != null;
+				if (include || relativePath)
+				{
+					return new ScalarConfigNode(include, relativePath);
+				}
+				else
+				{
+					return null;
+				}
+			}
+			return CreateType(propertyType);
+		}
+
+		static ConfigNode? CreateType(Type type)
+		{
+			ConfigNode? value;
+			if (!type.IsClass || type == typeof(string) || type == typeof(JsonNode))
+			{
+				value = null;
+			}
+			else if (type == typeof(ConfigMacro))
+			{
+				value = new MacroConfigNode();
+			}
+			else if (type.IsAssignableTo(typeof(ConfigResource)))
+			{
+				value = new ResourceConfigNode();
+			}
+			else if (type.IsGenericType && type.GetGenericTypeDefinition() == typeof(List<>))
+			{
+				ConfigNode? elementType = CreateType(type.GetGenericArguments()[0]);
+				value = (elementType == null) ? null : new ArrayConfigNode(elementType);
+			}
+			else if (type.IsGenericType && type.GetGenericTypeDefinition() == typeof(Dictionary<,>))
+			{
+				ConfigNode? elementType = CreateType(type.GetGenericArguments()[1]);
+				value = (elementType == null) ? null : new DictionaryConfigNode(elementType);
+			}
+			else
+			{
+				ObjectConfigNode objValue = new ObjectConfigNode(type);
+				value = objValue.IsDefault() ? null : objValue;
+			}
+			return value;
+		}
+
+		/// <summary>
+		/// Preprocesses an object
+		/// </summary>
+		/// <param name="obj">Object to process</param>
+		/// <param name="context">Preprocessor context</param>
+		/// <param name="cancellationToken">Cancellation token for the operation</param>
+		public async Task<JsonObject> PreprocessAsync(JsonObject obj, ConfigContext context, CancellationToken cancellationToken)
+			=> (JsonObject)(await PreprocessAsync(obj, null, context, cancellationToken))!;
+
+		/// <inheritdoc/>
+		public override void ParseMacros(JsonNode? jsonNode, ConfigContext context, Dictionary<string, string> macros)
+		{
+			if (jsonNode is JsonObject jsonObject)
+			{
+				foreach ((string name, JsonNode? node) in jsonObject)
+				{
+					if (node != null && Properties.TryGetValue(name, out ConfigNode? property))
+					{
+						property.ParseMacros(node, context, macros);
+					}
+				}
+			}
+		}
+
+		/// <inheritdoc/>
+		public override async Task<JsonNode?> PreprocessAsync(JsonNode? node, JsonNode? existingNode, ConfigContext context, CancellationToken cancellationToken)
+		{
+			JsonObject? obj = (JsonObject?)node;
+			JsonObject? target = (JsonObject?)existingNode;
+			if (obj == null)
+			{
+				return target;
+			}
+
+			// Before parsing properties into this object, read all the includes recursively
+			if (IncludeRoot)
+			{
+				// Find all the includes
+				List<IConfigFile> includes = new List<IConfigFile>();
+				await ParseIncludesInternalAsync(obj, includes, context, cancellationToken);
+
+				// Find all the files, merge them into target
+				foreach (IConfigFile include in includes)
+				{
+					context.IncludeStack.Push(include);
+
+					JsonObject includedJsonObject = await context.ParseFileAsync(include, cancellationToken);
+					target = (JsonObject?)await PreprocessAsync(includedJsonObject, target, context, cancellationToken);
+
+					context.IncludeStack.Pop();
+				}
+			}
+
+			// Ensure that the target object is valid so we can write properties into it
+			target ??= new JsonObject();
+
+			// Parse all the macros for this scope
+			if (MacroScope)
+			{
+				Dictionary<string, string> macros = new Dictionary<string, string>(StringComparer.OrdinalIgnoreCase);
+				ParseMacros(obj, context, macros);
+				context.MacroScopes.Add(macros);
+			}
+
+			// Parse all the properties into this object
+			foreach ((string name, JsonNode? newNode) in obj)
+			{
+				if (newNode is JsonValue)
+				{
+					context.AddProperty(name);
+				}
+
+				if (Properties.TryGetValue(name, out ConfigNode? property))
+				{
+					context.EnterScope(name);
+					target[name] = await property.PreprocessAsync(newNode, target[name], context, cancellationToken);
+					context.LeaveScope();
+				}
+				else
+				{
+					target[name] = Merge(ExpandMacros(newNode, context), target[name]);
+				}
+			}
+
+			// Parse all the macros for this scope
+			if (MacroScope)
+			{
+				context.MacroScopes.RemoveAt(context.MacroScopes.Count - 1);
+			}
+
+			return target;
+		}
+
+		static JsonNode? Merge(JsonNode? source, JsonNode? target)
+		{
+			if (source is JsonObject sourceObj && target is JsonObject targetObj)
+			{
+				foreach ((string name, JsonNode? node) in sourceObj)
+				{
+					targetObj[name] = Merge(node, targetObj[name]);
+				}
+				return target;
+			}
+			else if (source is JsonArray sourceArr && target is JsonArray targetArr)
+			{
+				foreach (JsonNode? node in sourceArr)
+				{
+					targetArr.Add(node?.DeepClone());
+				}
+				return target;
+			}
+			return source;
+		}
+
+		/// <inheritdoc/>
+		public override async Task ParseIncludesAsync(JsonNode? node, List<IConfigFile> includes, ConfigContext context, CancellationToken cancellationToken)
+		{
+			if (!IncludeRoot && node is JsonObject obj)
+			{
+				await ParseIncludesInternalAsync(obj, includes, context, cancellationToken);
+			}
+		}
+
+		async Task ParseIncludesInternalAsync(JsonObject obj, List<IConfigFile> includes, ConfigContext context, CancellationToken cancellationToken)
+		{
+			foreach ((string name, JsonNode? propertyNode) in obj)
+			{
+				if (Properties.TryGetValue(name, out ConfigNode? property))
+				{
+					await property.ParseIncludesAsync(propertyNode, includes, context, cancellationToken);
+				}
+			}
+		}
+	}
+}
diff --git a/Engine/Source/Programs/Horde/Horde.Server/Configuration/ConfigService.cs b/Engine/Source/Programs/Horde/Horde.Server/Configuration/ConfigService.cs
--- a/Engine/Source/Programs/Horde/Horde.Server/Configuration/ConfigService.cs
+++ b/Engine/Source/Programs/Horde/Horde.Server/Configuration/ConfigService.cs
@@ -260,7 +260,7 @@ namespace Horde.Server.Configuration
 			try
 			{
 				Uri globalConfigUri = GetGlobalConfigUri();
-				GlobalConfig globalConfig = await ConfigType.ReadAsync<GlobalConfig>(globalConfigUri, context, cancellationToken);
+				GlobalConfig globalConfig = await context.ReadAsync<GlobalConfig>(globalConfigUri, cancellationToken);
 				globalConfig.PostLoad(_serverSettings);
 
 				foreach (OverrideConfigFile file in overrideFiles.Values)
@@ -500,7 +500,7 @@ namespace Horde.Server.Configuration
 
 			if (isPerforcePath)
 			{
-				return ConfigType.CombinePaths(new Uri(FileReference.Combine(new DirectoryReference(defaultConfigDir), "_").FullName), configPath);
+				return ConfigNode.CombinePaths(new Uri(FileReference.Combine(new DirectoryReference(defaultConfigDir), "_").FullName), configPath);
 			}
 
 			if (isAbsPath)
@@ -565,7 +565,10 @@ namespace Horde.Server.Configuration
 				// Read the new config in
 				Uri globalConfigUri = GetGlobalConfigUri();
 
-				GlobalConfig globalConfig = await ConfigType.ReadAsync<GlobalConfig>(globalConfigUri, context, cancellationToken);
+				GlobalConfig globalConfig = await context.ReadAsync<GlobalConfig>(globalConfigUri, cancellationToken);
+
+				File.WriteAllBytes("D:\\fixed.json", JsonSerializer.SerializeToUtf8Bytes(globalConfig, new JsonSerializerOptions { WriteIndented = true }));
+
 				if (globalConfig.VersionEnum < GlobalVersion.Latest)
 				{
 					List<string> message = new List<string>();
diff --git a/Engine/Source/Programs/Horde/Horde.Server/Configuration/ConfigSource.cs b/Engine/Source/Programs/Horde/Horde.Server/Configuration/ConfigSource.cs
deleted file mode 100644
diff --git a/Engine/Source/Programs/Horde/Horde.Server/Configuration/ConfigType.cs b/Engine/Source/Programs/Horde/Horde.Server/Configuration/ConfigType.cs
deleted file mode 100644
diff --git a/Engine/Source/Programs/Horde/Horde.Server/Configuration/IConfigFile.cs b/Engine/Source/Programs/Horde/Horde.Server/Configuration/IConfigFile.cs
new file mode 100644
--- /dev/null
+++ b/Engine/Source/Programs/Horde/Horde.Server/Configuration/IConfigFile.cs
@@ -0,0 +1,63 @@
+// Copyright Epic Games, Inc. All Rights Reserved.
+
+using System;
+using System.Threading;
+using System.Threading.Tasks;
+using Horde.Server.Users;
+
+namespace Horde.Server.Configuration
+{
+	/// <summary>
+	/// Accessor for a specific revision of a config file. Provides metadata about the current revision, and allows reading its data.
+	/// </summary>
+	public interface IConfigFile
+	{
+		/// <summary>
+		/// URI of the config file
+		/// </summary>
+		Uri Uri { get; }
+
+		/// <summary>
+		/// String used to identify a specific version of the config data. Used for ordinal comparisons, otherwise opaque.
+		/// </summary>
+		string Revision { get; }
+
+		/// <summary>
+		/// Author of this revision, if known
+		/// </summary>
+		IUser? Author { get; }
+
+		/// <summary>
+		/// Reads data for the config file
+		/// </summary>
+		/// <param name="cancellationToken">Cancellation token for the operation</param>
+		/// <returns>UTF-8 encoded data for the config file</returns>
+		ValueTask<ReadOnlyMemory<byte>> ReadAsync(CancellationToken cancellationToken);
+	}
+
+	/// <summary>
+	/// Extension methods for config files
+	/// </summary>
+	public static class ConfigFileExtensions
+	{
+		/// <summary>
+		/// Formats a the URI for a config file in a format that is more readable for users
+		/// </summary>
+		/// <param name="file">Config file to get a path for</param>
+		public static string GetUserFormattedPath(this IConfigFile file)
+		{
+			const string DefaultPerforcePrefix = "perforce://default//";
+
+			string path = file.Uri.ToString();
+			if (path.StartsWith(DefaultPerforcePrefix, StringComparison.OrdinalIgnoreCase))
+			{
+				path = path.Substring(DefaultPerforcePrefix.Length - 2);
+			}
+			if (!String.IsNullOrEmpty(file.Revision))
+			{
+				path = $"{path}@{file.Revision}";
+			}
+			return path;
+		}
+	}
+}
diff --git a/Engine/Source/Programs/Horde/Horde.Server/Configuration/IConfigSource.cs b/Engine/Source/Programs/Horde/Horde.Server/Configuration/IConfigSource.cs
new file mode 100644
--- /dev/null
+++ b/Engine/Source/Programs/Horde/Horde.Server/Configuration/IConfigSource.cs
@@ -0,0 +1,446 @@
+// Copyright Epic Games, Inc. All Rights Reserved.
+
+using System;
+using System.Collections.Concurrent;
+using System.Collections.Generic;
+using System.IO;
+using System.Linq;
+using System.Threading;
+using System.Threading.Tasks;
+using EpicGames.Core;
+using EpicGames.Horde.Users;
+using EpicGames.Perforce;
+using Horde.Server.Users;
+using Microsoft.Extensions.Caching.Memory;
+using Microsoft.Extensions.Logging;
+using Microsoft.Extensions.Options;
+
+namespace Horde.Server.Configuration
+{
+	/// <summary>
+	/// Source for reading config files
+	/// </summary>
+	public interface IConfigSource
+	{
+		/// <summary>
+		/// URI scheme for this config source
+		/// </summary>
+		string Scheme { get; }
+
+		/// <summary>
+		/// Update interval for this source
+		/// </summary>
+		TimeSpan UpdateInterval { get; }
+
+		/// <summary>
+		/// Reads a config file from this source
+		/// </summary>
+		/// <param name="uris">Locations of the config files to query</param>
+		/// <param name="cancellationToken">Cancellation token for the operation</param>
+		/// <returns>Config file data</returns>
+		Task<IConfigFile[]> GetFilesAsync(Uri[] uris, CancellationToken cancellationToken);
+
+		/// <summary>
+		/// Gets summary infomration for sending notifications
+		/// </summary>
+		/// <param name="files">New files for the configuration</param>
+		/// <param name="prevFiles">Previous set of files for the configuration</param>
+		/// <param name="updateInfo">Information about the update</param>
+		/// <param name="cancellationToken">Cancellation token for the operation</param>
+		Task GetUpdateInfoAsync(IReadOnlyDictionary<Uri, string> files, IReadOnlyDictionary<Uri, string>? prevFiles, ConfigUpdateInfo updateInfo, CancellationToken cancellationToken);
+	}
+
+	/// <summary>
+	/// Extension methods for <see cref="IConfigSource"/>
+	/// </summary>
+	public static class ConfigSource
+	{
+		/// <summary>
+		/// Gets a single config file from a source
+		/// </summary>
+		/// <param name="source">Source to query</param>
+		/// <param name="uri">Location of the config file to query</param>
+		/// <param name="cancellationToken">Cancellation token for the operation</param>
+		/// <returns>Config file data</returns>
+		public static async Task<IConfigFile> GetAsync(this IConfigSource source, Uri uri, CancellationToken cancellationToken)
+		{
+			IConfigFile[] result = await source.GetFilesAsync(new[] { uri }, cancellationToken);
+			return result[0];
+		}
+	}
+
+	/// <summary>
+	/// In-memory config file source
+	/// </summary>
+	public sealed class InMemoryConfigSource : IConfigSource
+	{
+		class ConfigFileRevisionImpl : IConfigFile
+		{
+			public Uri Uri { get; }
+			public string Revision { get; }
+			public ReadOnlyMemory<byte> Data { get; }
+			public IUser? Author => null;
+
+			public ConfigFileRevisionImpl(Uri uri, string version, ReadOnlyMemory<byte> data)
+			{
+				Uri = uri;
+				Revision = version;
+				Data = data;
+			}
+
+			public ValueTask<ReadOnlyMemory<byte>> ReadAsync(CancellationToken cancellationToken) => new ValueTask<ReadOnlyMemory<byte>>(Data);
+		}
+
+		readonly Dictionary<Uri, ConfigFileRevisionImpl> _files = new Dictionary<Uri, ConfigFileRevisionImpl>();
+
+		/// <summary>
+		/// Name of the scheme for this source
+		/// </summary>
+		public const string Scheme = "memory";
+
+		/// <inheritdoc/>
+		string IConfigSource.Scheme => Scheme;
+
+		/// <inheritdoc/>
+		public TimeSpan UpdateInterval => TimeSpan.FromSeconds(1.0);
+
+		/// <summary>
+		/// Manually adds a new config file
+		/// </summary>
+		/// <param name="path">Path to the config file</param>
+		/// <param name="data">Config file data</param>
+		public void Add(Uri path, ReadOnlyMemory<byte> data)
+		{
+			_files.Add(path, new ConfigFileRevisionImpl(path, "v1", data));
+		}
+
+		/// <inheritdoc/>
+		public Task<IConfigFile[]> GetFilesAsync(Uri[] uris, CancellationToken cancellationToken)
+		{
+			IConfigFile[] result = new IConfigFile[uris.Length];
+			for (int idx = 0; idx < uris.Length; idx++)
+			{
+				ConfigFileRevisionImpl? configFile;
+				if (!_files.TryGetValue(uris[idx], out configFile))
+				{
+					throw new FileNotFoundException($"Config file {uris[idx]} not found.");
+				}
+				result[idx] = configFile;
+			}
+			return Task.FromResult(result);
+		}
+
+		/// <inheritdoc/>
+		public Task GetUpdateInfoAsync(IReadOnlyDictionary<Uri, string> files, IReadOnlyDictionary<Uri, string>? prevFiles, ConfigUpdateInfo updateInfo, CancellationToken cancellationToken)
+			=> Task.CompletedTask;
+	}
+
+	/// <summary>
+	/// Config file source which reads from the filesystem
+	/// </summary>
+	public sealed class FileConfigSource : IConfigSource
+	{
+		class ConfigFileImpl : IConfigFile
+		{
+			public Uri Uri { get; }
+			public string Revision { get; }
+			public DateTime LastWriteTimeUtc { get; }
+			public ReadOnlyMemory<byte> Data { get; }
+			public IUser? Author => null;
+
+			public ConfigFileImpl(Uri uri, DateTime lastWriteTimeUtc, ReadOnlyMemory<byte> data)
+			{
+				Uri = uri;
+				Revision = $"timestamp={lastWriteTimeUtc.Ticks}";
+				LastWriteTimeUtc = lastWriteTimeUtc;
+				Data = data;
+			}
+
+			public ValueTask<ReadOnlyMemory<byte>> ReadAsync(CancellationToken cancellationToken) => new ValueTask<ReadOnlyMemory<byte>>(Data);
+		}
+
+		/// <summary>
+		/// Name of the scheme for this source
+		/// </summary>
+		public const string Scheme = "file";
+
+		/// <inheritdoc/>
+		string IConfigSource.Scheme => Scheme;
+
+		/// <inheritdoc/>
+		public TimeSpan UpdateInterval => TimeSpan.FromSeconds(5.0);
+
+		readonly DirectoryReference _baseDir;
+		readonly ConcurrentDictionary<FileReference, ConfigFileImpl> _files = new ConcurrentDictionary<FileReference, ConfigFileImpl>();
+
+		/// <summary>
+		/// Constructor
+		/// </summary>
+		public FileConfigSource()
+			: this(DirectoryReference.GetCurrentDirectory())
+		{
+		}
+
+		/// <summary>
+		/// Constructor
+		/// </summary>
+		/// <param name="baseDir">Base directory for resolving relative paths</param>
+		public FileConfigSource(DirectoryReference baseDir)
+		{
+			_baseDir = baseDir;
+		}
+
+		/// <inheritdoc/>
+		public async Task<IConfigFile[]> GetFilesAsync(Uri[] uris, CancellationToken cancellationToken)
+		{
+			IConfigFile[] files = new IConfigFile[uris.Length];
+			for (int idx = 0; idx < uris.Length; idx++)
+			{
+				Uri uri = uris[idx];
+				FileReference localPath = FileReference.Combine(_baseDir, uri.LocalPath);
+
+				ConfigFileImpl? file;
+				for (; ; )
+				{
+					if (_files.TryGetValue(localPath, out file))
+					{
+						if (FileReference.GetLastWriteTimeUtc(localPath) == file.LastWriteTimeUtc)
+						{
+							break;
+						}
+						else
+						{
+							_files.TryRemove(new KeyValuePair<FileReference, ConfigFileImpl>(localPath, file));
+						}
+					}
+
+					using (FileStream stream = FileReference.Open(localPath, FileMode.Open, FileAccess.Read, FileShare.Read))
+					{
+						using MemoryStream memoryStream = new MemoryStream();
+						await stream.CopyToAsync(memoryStream, cancellationToken);
+						DateTime lastWriteTime = FileReference.GetLastWriteTimeUtc(localPath);
+						file = new ConfigFileImpl(uri, lastWriteTime, memoryStream.ToArray());
+					}
+
+					if (_files.TryAdd(localPath, file))
+					{
+						break;
+					}
+				}
+
+				files[idx] = file;
+			}
+			return files;
+		}
+
+		/// <inheritdoc/>
+		public Task GetUpdateInfoAsync(IReadOnlyDictionary<Uri, string> files, IReadOnlyDictionary<Uri, string>? prevFiles, ConfigUpdateInfo updateInfo, CancellationToken cancellationToken)
+			=> Task.CompletedTask;
+	}
+
+	/// <summary>
+	/// Perforce cluster config file source
+	/// </summary>
+	public sealed class PerforceConfigSource : IConfigSource
+	{
+		class ConfigFileImpl : IConfigFile
+		{
+			public Uri Uri { get; }
+			public int Change { get; }
+			public string Revision { get; }
+			public IUser? Author { get; }
+
+			readonly PerforceConfigSource _owner;
+
+			public ConfigFileImpl(Uri uri, int change, IUser? author, PerforceConfigSource owner)
+			{
+				Uri = uri;
+				Change = change;
+				Revision = $"{change}";
+				Author = author;
+				_owner = owner;
+			}
+
+			public ValueTask<ReadOnlyMemory<byte>> ReadAsync(CancellationToken cancellationToken) => _owner.ReadAsync(Uri, Change, cancellationToken);
+		}
+
+		/// <summary>
+		/// Name of the scheme for this source
+		/// </summary>
+		public const string Scheme = "perforce";
+
+		/// <inheritdoc/>
+		string IConfigSource.Scheme => Scheme;
+
+		/// <inheritdoc/>
+		public TimeSpan UpdateInterval => TimeSpan.FromMinutes(1.0);
+
+		readonly IOptionsMonitor<ServerSettings> _settings;
+		readonly IUserCollection _userCollection;
+		readonly IMemoryCache _cache;
+		readonly ILogger _logger;
+
+		/// <summary>
+		/// Constructor
+		/// </summary>
+		public PerforceConfigSource(IOptionsMonitor<ServerSettings> settings, IUserCollection userCollection, IMemoryCache cache, ILogger<PerforceConfigSource> logger)
+		{
+			_settings = settings;
+			_userCollection = userCollection;
+			_cache = cache;
+			_logger = logger;
+		}
+
+		/// <inheritdoc/>
+		public async Task<IConfigFile[]> GetFilesAsync(Uri[] uris, CancellationToken cancellationToken)
+		{
+			Dictionary<Uri, IConfigFile> results = new Dictionary<Uri, IConfigFile>();
+			foreach (IGrouping<string, Uri> group in uris.GroupBy(x => x.Host))
+			{
+				using (IPerforceConnection perforce = await ConnectAsync(group.Key, cancellationToken))
+				{
+					try
+					{
+						FileSpecList fileSpec = group.Select(x => x.AbsolutePath).Distinct(StringComparer.OrdinalIgnoreCase).ToList();
+
+						List<FStatRecord> records = await perforce.FStatAsync(FStatOptions.ShortenOutput, fileSpec, cancellationToken).ToListAsync(cancellationToken);
+						records.RemoveAll(x => x.HeadAction == FileAction.Delete || x.HeadAction == FileAction.MoveDelete);
+
+						Dictionary<string, FStatRecord> absolutePathToRecord = records.ToDictionary(x => x.DepotFile ?? String.Empty, x => x, StringComparer.OrdinalIgnoreCase);
+						foreach (Uri uri in group)
+						{
+							FStatRecord? record;
+							if (!absolutePathToRecord.TryGetValue(uri.AbsolutePath, out record))
+							{
+								throw new FileNotFoundException($"Unable to read {uri}. No matching files found.");
+							}
+
+							IUser? author = await GetAuthorAsync(perforce, group.Key, record.HeadChange, cancellationToken);
+							results[uri] = new ConfigFileImpl(uri, record.HeadChange, author, this);
+						}
+					}
+					catch (PerforceException ex) when (ex.Error != null && ex.Error.Generic == PerforceGenericCode.Config)
+					{
+						throw new PerforceException($"{ex.Message} [server: {perforce.Settings.ServerAndPort}, user: {perforce.Settings.UserName}]", ex);
+					}
+				}
+			}
+			return uris.ConvertAll(x => results[x]);
+		}
+
+		async ValueTask<IUser?> GetAuthorAsync(IPerforceConnection perforce, string host, int change, CancellationToken cancellationToken)
+		{
+			string cacheKey = $"{nameof(PerforceConfigSource)}:author:{host}@{change}";
+			if (!_cache.TryGetValue(cacheKey, out string? author))
+			{
+				ChangeRecord record = await perforce.GetChangeAsync(GetChangeOptions.None, change, cancellationToken);
+				using (ICacheEntry entry = _cache.CreateEntry(cacheKey))
+				{
+					entry.SetSlidingExpiration(TimeSpan.FromHours(1.0));
+					entry.SetSize(256);
+					entry.SetValue(record.User);
+				}
+			}
+			return (author != null) ? await _userCollection.FindUserByLoginAsync(author, cancellationToken) : null;
+		}
+
+		async ValueTask<ReadOnlyMemory<byte>> ReadAsync(Uri uri, int change, CancellationToken cancellationToken)
+		{
+			string cacheKey = $"{nameof(PerforceConfigSource)}:data:{uri}@{change}";
+			if (_cache.TryGetValue(cacheKey, out ReadOnlyMemory<byte> data))
+			{
+				_logger.LogInformation("Read {Uri}@{Change} from cache ({Key})", uri, change, cacheKey);
+			}
+			else
+			{
+				_logger.LogInformation("Reading {Uri} at CL {Change} from Perforce", uri, change);
+				using (IPerforceConnection perforce = await ConnectAsync(uri.Host, cancellationToken))
+				{
+					PerforceResponse<PrintRecord<byte[]>> response = await perforce.TryPrintDataAsync($"{uri.AbsolutePath}@{change}", cancellationToken);
+					response.EnsureSuccess();
+					data = response.Data.Contents!;
+				}
+				using (ICacheEntry entry = _cache.CreateEntry(cacheKey))
+				{
+					entry.SetSlidingExpiration(TimeSpan.FromHours(1.0));
+					entry.SetSize(data.Length);
+					entry.SetValue(data);
+				}
+			}
+			return data;
+		}
+
+		async Task<IPerforceConnection> ConnectAsync(string host, CancellationToken cancellationToken)
+		{
+			_ = cancellationToken;
+
+			ServerSettings settings = _settings.CurrentValue;
+
+			PerforceConnectionId connectionId = new PerforceConnectionId();
+			if (!String.IsNullOrEmpty(host))
+			{
+				connectionId = new PerforceConnectionId(host);
+			}
+
+			PerforceConnectionSettings? connectionSettings = settings.Perforce.FirstOrDefault(x => x.Id == connectionId);
+			if (connectionSettings == null)
+			{
+				if (connectionId == PerforceConnectionSettings.Default)
+				{
+					connectionSettings = new PerforceConnectionSettings();
+				}
+				else
+				{
+					throw new InvalidOperationException($"No Perforce connection settings defined for '{connectionId}'.");
+				}
+			}
+
+			IPerforceConnection connection = await PerforceConnection.CreateAsync(connectionSettings.ToPerforceSettings(), _logger);
+			if (connectionSettings.Credentials != null && !String.IsNullOrEmpty(connectionSettings.Credentials.Password) && String.IsNullOrEmpty(connectionSettings.Credentials.Ticket))
+			{
+				await connection.LoginAsync(connectionSettings.Credentials.Password, cancellationToken);
+			}
+			return connection;
+		}
+
+		/// <inheritdoc/>
+		public async Task GetUpdateInfoAsync(IReadOnlyDictionary<Uri, string> files, IReadOnlyDictionary<Uri, string>? prevFiles, ConfigUpdateInfo updateInfo, CancellationToken cancellationToken)
+		{
+			foreach (IGrouping<string, KeyValuePair<Uri, string>> group in files.Where(x => x.Key.Scheme == PerforceConfigSource.Scheme).GroupBy(x => x.Key.Host))
+			{
+				// Figure out the most recent changelist
+				int change = group.Select(x => Int32.Parse(x.Value)).Max();
+				updateInfo.Status.Add($"Perforce changelist ({group.Key}): {change}");
+
+				// Find all the users that have committed changes
+				if (prevFiles != null)
+				{
+					using IPerforceConnection perforce = await ConnectAsync(group.Key, cancellationToken);
+					foreach ((Uri uri, string revision) in group)
+					{
+						string? prevRevision;
+						if (prevFiles.TryGetValue(uri, out prevRevision) && prevRevision != revision)
+						{
+							int prevRevisionNum = Int32.Parse(prevRevision);
+							string fileSpec = $"{uri.LocalPath}@{prevRevisionNum + 1},{revision}";
+							await FindAuthorsAsync(perforce, fileSpec, updateInfo.Authors, cancellationToken);
+						}
+					}
+				}
+			}
+		}
+
+		async Task FindAuthorsAsync(IPerforceConnection perforce, string fileSpec, HashSet<UserId> authors, CancellationToken cancellationToken)
+		{
+			List<ChangesRecord> records = await perforce.GetChangesAsync(ChangesOptions.None, -1, ChangeStatus.Submitted, fileSpec, cancellationToken);
+			foreach (ChangesRecord record in records)
+			{
+				IUser? user = await _userCollection.FindUserByLoginAsync(record.User, cancellationToken);
+				if (user != null)
+				{
+					authors.Add(user.Id);
+				}
+			}
+		}
+	}
+}
