Horde: Fix support for recursively defined config types.

#rnx

diff --git a/Engine/Source/Programs/Horde/Horde.Server.Tests/Configuration/ConfigTests.cs b/Engine/Source/Programs/Horde/Horde.Server.Tests/Configuration/ConfigTests.cs
--- a/Engine/Source/Programs/Horde/Horde.Server.Tests/Configuration/ConfigTests.cs
+++ b/Engine/Source/Programs/Horde/Horde.Server.Tests/Configuration/ConfigTests.cs
@@ -72,9 +72,9 @@ namespace Horde.Server.Tests.Configuration
 			Dictionary<string, IConfigSource> sources = new Dictionary<string, IConfigSource>();
 			sources["memory"] = source;
 
-			ConfigContext options = new ConfigContext(_jsonOptions, sources, NullLogger.Instance);
+			ConfigContext context = new ConfigContext(_jsonOptions, sources, NullLogger.Instance);
 
-			ConfigObject result = await ConfigType.ReadAsync<ConfigObject>(fooUri, options, cancellationToken);
+			ConfigObject result = await context.ReadAsync<ConfigObject>(fooUri, cancellationToken);
 			Assert.AreEqual(result.TestString, "hello");
 			Assert.IsTrue(result.TestList.SequenceEqual(new[] { "secondobj", "there", "world" }));
 			Assert.AreEqual(result.TestObject!.ValueA, "hi");
@@ -108,7 +108,7 @@ namespace Horde.Server.Tests.Configuration
 
 			ConfigContext context = new ConfigContext(_jsonOptions, sources, NullLogger.Instance);
 
-			ConfigObject result = await ConfigType.ReadAsync<ConfigObject>(fooUri, context, cancellationToken);
+			ConfigObject result = await context.ReadAsync<ConfigObject>(fooUri, cancellationToken);
 			Assert.AreEqual(result.TestString, "hello");
 
 			// Check it returns the same object if the timestamp hasn't changed
@@ -266,9 +266,9 @@ namespace Horde.Server.Tests.Configuration
 			Dictionary<string, IConfigSource> sources = new Dictionary<string, IConfigSource>();
 			sources["memory"] = source;
 
-			ConfigContext options = new ConfigContext(_jsonOptions, sources, NullLogger.Instance);
+			ConfigContext context = new ConfigContext(_jsonOptions, sources, NullLogger.Instance);
 
-			BaseMacroScope result = await ConfigType.ReadAsync<BaseMacroScope>(fooUri, options, cancellationToken);
+			BaseMacroScope result = await context.ReadAsync<BaseMacroScope>(fooUri, cancellationToken);
 			Assert.AreEqual("This is a macro MacroValue", result.Value);
 			Assert.AreEqual("List element macro MacroValue", result.ListValue[0]);
 			Assert.AreEqual("Object macro MacroValue", result.ObjectValue.Value);
@@ -299,9 +299,9 @@ namespace Horde.Server.Tests.Configuration
 			Dictionary<string, IConfigSource> sources = new Dictionary<string, IConfigSource>();
 			sources["memory"] = source;
 
-			ConfigContext options = new ConfigContext(_jsonOptions, sources, NullLogger.Instance);
+			ConfigContext context = new ConfigContext(_jsonOptions, sources, NullLogger.Instance);
 
-			BaseMacroScope result = await ConfigType.ReadAsync<BaseMacroScope>(fooUri, options, cancellationToken);
+			BaseMacroScope result = await context.ReadAsync<BaseMacroScope>(fooUri, cancellationToken);
 			Assert.AreEqual("This is a macro MacroValue $(MacroName2)", result.Value);
 			Assert.AreEqual("This is a macro MacroValue MacroValue2", result.ChildScope!.Value);
 		}
diff --git a/Engine/Source/Programs/Horde/Horde.Server/Configuration/ConfigContext.cs b/Engine/Source/Programs/Horde/Horde.Server/Configuration/ConfigContext.cs
--- a/Engine/Source/Programs/Horde/Horde.Server/Configuration/ConfigContext.cs
+++ b/Engine/Source/Programs/Horde/Horde.Server/Configuration/ConfigContext.cs
@@ -196,7 +196,7 @@ namespace Horde.Server.Configuration
 			IncludeStack.Push(file);
 
 			ObjectConfigNode type = new ObjectConfigNode(typeof(T));
-			obj = (JsonObject)(await type.PreprocessAsync(obj, null, this, cancellationToken))!;
+			obj = await type.PreprocessAsync(obj, this, cancellationToken);
 
 			IncludeStack.Pop();
 
diff --git a/Engine/Source/Programs/Horde/Horde.Server/Configuration/ConfigNode.cs b/Engine/Source/Programs/Horde/Horde.Server/Configuration/ConfigNode.cs
--- a/Engine/Source/Programs/Horde/Horde.Server/Configuration/ConfigNode.cs
+++ b/Engine/Source/Programs/Horde/Horde.Server/Configuration/ConfigNode.cs
@@ -406,7 +406,13 @@ namespace Horde.Server.Configuration
 		/// </summary>
 		/// <param name="type">Type to construct from</param>
 		public ObjectConfigNode(Type type)
+			: this(type, new Dictionary<Type, ConfigNode>())
+		{ }
+
+		ObjectConfigNode(Type type, Dictionary<Type, ConfigNode> recursiveTypes)
 		{
+			recursiveTypes.Add(type, this);
+
 			IncludeRoot = type.GetCustomAttribute<ConfigIncludeRootAttribute>() != null;
 			MacroScope = type.GetCustomAttribute<ConfigMacroScopeAttribute>() != null;
 			Properties = new Dictionary<string, ConfigNode>(StringComparer.OrdinalIgnoreCase);
@@ -417,7 +423,7 @@ namespace Horde.Server.Configuration
 			{
 				if (propertyInfo.GetCustomAttribute<JsonIgnoreAttribute>() == null)
 				{
-					ConfigNode? propertyType = CreateTypeForProperty(propertyInfo);
+					ConfigNode? propertyType = CreateTypeForProperty(propertyInfo, recursiveTypes);
 					if (propertyType != null)
 					{
 						string name = propertyInfo.GetCustomAttribute<JsonPropertyNameAttribute>()?.Name ?? propertyInfo.Name;
@@ -430,7 +436,7 @@ namespace Horde.Server.Configuration
 		bool IsDefault()
 			=> !IncludeRoot && !MacroScope && Properties.Count == 0;
 
-		static ConfigNode? CreateTypeForProperty(PropertyInfo propertyInfo)
+		static ConfigNode? CreateTypeForProperty(PropertyInfo propertyInfo, Dictionary<Type, ConfigNode> recursiveTypes)
 		{
 			Type propertyType = propertyInfo.PropertyType;
 			if (!propertyType.IsClass || propertyType == typeof(string))
@@ -446,10 +452,10 @@ namespace Horde.Server.Configuration
 					return null;
 				}
 			}
-			return CreateType(propertyType);
+			return CreateType(propertyType, recursiveTypes);
 		}
 
-		static ConfigNode? CreateType(Type type)
+		static ConfigNode? CreateType(Type type, Dictionary<Type, ConfigNode> recursiveTypes)
 		{
 			ConfigNode? value;
 			if (!type.IsClass || type == typeof(string) || type == typeof(JsonNode))
@@ -466,15 +472,15 @@ namespace Horde.Server.Configuration
 			}
 			else if (type.IsGenericType && type.GetGenericTypeDefinition() == typeof(List<>))
 			{
-				ConfigNode? elementType = CreateType(type.GetGenericArguments()[0]);
+				ConfigNode? elementType = CreateType(type.GetGenericArguments()[0], recursiveTypes);
 				value = (elementType == null) ? null : new ArrayConfigNode(elementType);
 			}
 			else if (type.IsGenericType && type.GetGenericTypeDefinition() == typeof(Dictionary<,>))
 			{
-				ConfigNode? elementType = CreateType(type.GetGenericArguments()[1]);
+				ConfigNode? elementType = CreateType(type.GetGenericArguments()[1], recursiveTypes);
 				value = (elementType == null) ? null : new DictionaryConfigNode(elementType);
 			}
-			else
+			else if (!recursiveTypes.TryGetValue(type, out value))
 			{
 				ObjectConfigNode objValue = new ObjectConfigNode(type);
 				value = objValue.IsDefault() ? null : objValue;
@@ -500,7 +506,10 @@ namespace Horde.Server.Configuration
 				{
 					if (node != null && Properties.TryGetValue(name, out ConfigNode? property))
 					{
-						property.ParseMacros(node, context, macros);
+						if (property is not ObjectConfigNode propertyObj || !propertyObj.MacroScope)
+						{
+							property.ParseMacros(node, context, macros);
+						}
 					}
 				}
 			}
