Horde: Add creation time to artifact responses.

diff --git a/Engine/Source/Programs/Horde/Horde.Server/Artifacts/ArtifactsController.cs b/Engine/Source/Programs/Horde/Horde.Server/Artifacts/ArtifactsController.cs
--- a/Engine/Source/Programs/Horde/Horde.Server/Artifacts/ArtifactsController.cs
+++ b/Engine/Source/Programs/Horde/Horde.Server/Artifacts/ArtifactsController.cs
@@ -200,7 +200,7 @@ namespace Horde.Server.Artifacts
 			{
 				if (_globalConfig.Authorize(artifact.AclScope, ArtifactAclAction.ReadArtifact, User))
 				{
-					response.Artifacts.Add(new GetArtifactResponse(artifact.Id, artifact.Name, artifact.Type, artifact.Description, artifact.StreamId, artifact.Change, artifact.Keys, artifact.Metadata));
+					response.Artifacts.Add(new GetArtifactResponse(artifact.Id, artifact.Name, artifact.Type, artifact.Description, artifact.StreamId, artifact.Change, artifact.Keys, artifact.Metadata, artifact.CreatedAtUtc));
 				}
 			}
 
@@ -228,7 +228,7 @@ namespace Horde.Server.Artifacts
 				return Forbid(ArtifactAclAction.ReadArtifact, artifact.AclScope);
 			}
 
-			return PropertyFilter.Apply(new GetArtifactResponse(artifact.Id, artifact.Name, artifact.Type, artifact.Description, artifact.StreamId, artifact.Change, artifact.Keys, artifact.Metadata), filter);
+			return PropertyFilter.Apply(new GetArtifactResponse(artifact.Id, artifact.Name, artifact.Type, artifact.Description, artifact.StreamId, artifact.Change, artifact.Keys, artifact.Metadata, artifact.CreatedAtUtc), filter);
 		}
 
 		/// <summary>
diff --git a/Engine/Source/Programs/Shared/EpicGames.Horde/Artifacts/ArtifactMessages.cs b/Engine/Source/Programs/Shared/EpicGames.Horde/Artifacts/ArtifactMessages.cs
--- a/Engine/Source/Programs/Shared/EpicGames.Horde/Artifacts/ArtifactMessages.cs
+++ b/Engine/Source/Programs/Shared/EpicGames.Horde/Artifacts/ArtifactMessages.cs
@@ -1,5 +1,6 @@
 ﻿// Copyright Epic Games, Inc. All Rights Reserved.
 
+using System;
 using System.Collections.Generic;
 using EpicGames.Core;
 using EpicGames.Horde.Storage;
@@ -92,10 +93,15 @@ namespace EpicGames.Horde.Artifacts
 		/// </summary>
 		public IReadOnlyList<string> Metadata { get; }
 
+		/// <summary>
+		/// Time that the artifact was created
+		/// </summary>
+		public DateTime CreatedAtUtc { get; }
+
 		/// <summary>
 		/// Constructor
 		/// </summary>
-		public GetArtifactResponse(ArtifactId id, ArtifactName name, ArtifactType type, string? description, StreamId streamId, int change, IReadOnlyList<string> keys, IReadOnlyList<string> metadata)
+		public GetArtifactResponse(ArtifactId id, ArtifactName name, ArtifactType type, string? description, StreamId streamId, int change, IReadOnlyList<string> keys, IReadOnlyList<string> metadata, DateTime createdAtUtc)
 		{
 			Id = id;
 			Name = name;
@@ -105,6 +111,7 @@ namespace EpicGames.Horde.Artifacts
 			Change = change;
 			Keys = keys;
 			Metadata = metadata;
+			CreatedAtUtc = createdAtUtc;
 		}
 	}
 
