Horde: Remove debug code.

#rnx

diff --git a/Engine/Source/Programs/Horde/Horde.Server/Configuration/ConfigService.cs b/Engine/Source/Programs/Horde/Horde.Server/Configuration/ConfigService.cs
--- a/Engine/Source/Programs/Horde/Horde.Server/Configuration/ConfigService.cs
+++ b/Engine/Source/Programs/Horde/Horde.Server/Configuration/ConfigService.cs
@@ -566,9 +566,6 @@ namespace Horde.Server.Configuration
 				Uri globalConfigUri = GetGlobalConfigUri();
 
 				GlobalConfig globalConfig = await context.ReadAsync<GlobalConfig>(globalConfigUri, cancellationToken);
-
-				File.WriteAllBytes("D:\\fixed.json", JsonSerializer.SerializeToUtf8Bytes(globalConfig, new JsonSerializerOptions { WriteIndented = true }));
-
 				if (globalConfig.VersionEnum < GlobalVersion.Latest)
 				{
 					List<string> message = new List<string>();
