Horde: Allow non-object types with include directives in config files.

#rnx

diff --git a/Engine/Source/Programs/Horde/Horde.Server/Configuration/ConfigNode.cs b/Engine/Source/Programs/Horde/Horde.Server/Configuration/ConfigNode.cs
--- a/Engine/Source/Programs/Horde/Horde.Server/Configuration/ConfigNode.cs
+++ b/Engine/Source/Programs/Horde/Horde.Server/Configuration/ConfigNode.cs
@@ -289,10 +289,9 @@ namespace Horde.Server.Configuration
 		{
 			if (node is JsonArray arrayNode)
 			{
-				ObjectConfigNode classElementType = (ObjectConfigNode)ElementType;
 				foreach (JsonObject? element in arrayNode)
 				{
-					await classElementType.ParseIncludesAsync(element, includes, context, cancellationToken);
+					await ElementType.ParseIncludesAsync(element, includes, context, cancellationToken);
 				}
 			}
 		}
