Horde: Compress unsync manifest responses if the client specifies an 'application/x-unsync-manifest-zstd' accept header.

diff --git a/Engine/Source/Programs/Horde/Horde.Server/Artifacts/ArtifactsController.cs b/Engine/Source/Programs/Horde/Horde.Server/Artifacts/ArtifactsController.cs
--- a/Engine/Source/Programs/Horde/Horde.Server/Artifacts/ArtifactsController.cs
+++ b/Engine/Source/Programs/Horde/Horde.Server/Artifacts/ArtifactsController.cs
@@ -627,8 +627,11 @@ namespace Horde.Server.Artifacts
 				return Forbid(ArtifactAclAction.ReadArtifact, artifact.AclScope);
 			}
 
-			UnsyncManifest? manifest = await _unsyncCache.GetManifestAsync(artifact, cancellationToken);
-			if (manifest == null)
+			const string ZstdManifestContentType = "application/x-unsync-manifest-zstd";
+			bool compressed = Request.Headers.Accept.Contains(ZstdManifestContentType);
+
+			ReadOnlyMemory<byte> manifestData = await _unsyncCache.GetManifestDataAsync(artifact, compressed, cancellationToken);
+			if (manifestData.IsEmpty)
 			{
 				return NotFound(id);
 			}
@@ -639,58 +642,12 @@ namespace Horde.Server.Artifacts
 
 			// Write the response directly to the writer
 			HttpResponse response = HttpContext.Response;
-			response.ContentType = "application/json";
+			response.ContentType = compressed? ZstdManifestContentType : "application/json";
 			response.StatusCode = (int)HttpStatusCode.OK;
 			await response.StartAsync(cancellationToken);
-
-			using (Utf8JsonWriter writer = new Utf8JsonWriter(HttpContext.Response.BodyWriter, new JsonWriterOptions { Indented = true }))
-			{
-				writer.WriteStartObject();
-				writer.WriteString("type", "unsync_manifest");
-				writer.WriteString("hash_strong", "Blake3.160");
-				writer.WriteString("chunking", "RollingBuzHash");
-				writer.WriteNumber("chunking_block_size_min", LeafChunkedDataNodeOptions.Default.MinSize);
-				writer.WriteNumber("chunking_block_size_max", LeafChunkedDataNodeOptions.Default.MaxSize);
-				writer.WriteNumber("chunking_block_size_target", LeafChunkedDataNodeOptions.Default.TargetSize);
-
-				writer.WriteStartArray("files");
-				foreach (UnsyncFile file in manifest.Files)
-				{
-					writer.WriteStartObject();
-					writer.WriteString("name", file.Name);
-					if (file.Executable)
-					{
-						writer.WriteBoolean("exec", file.Executable);
-					}
-					if (file.ReadOnly)
-					{
-						writer.WriteBoolean("read_only", file.ReadOnly);
-					}
-					if (file.ModTime != default)
-					{
-						writer.WriteNumber("mtime", file.ModTime.ToFileTimeUtc());
-					}
-					writer.WriteNumber("size", file.Length);
-					writer.WriteNumber("block_size", LeafChunkedDataNodeOptions.Default.TargetSize);
-					writer.WriteStartArray("blocks");
-					foreach (UnsyncBlock block in file.Blocks)
-					{
-						writer.WriteStartObject();
-						writer.WriteNumber("offset", block.Offset);
-						writer.WriteNumber("size", block.Length);
-						writer.WriteNumber("hash_weak", block.RollingHash);
-						writer.WriteString("hash_strong", block.Blob.Hash.ToString());
-						writer.WriteEndObject();
-					}
-					writer.WriteEndArray();
-					writer.WriteEndObject();
-				}
-				writer.WriteEndArray();
-
-				writer.WriteEndObject();
-			}
-
+			await response.BodyWriter.WriteAsync(manifestData, cancellationToken);
 			await response.CompleteAsync();
+
 			return Empty;
 		}
 
diff --git a/Engine/Source/Programs/Horde/Horde.Server/Artifacts/UnsyncCache.cs b/Engine/Source/Programs/Horde/Horde.Server/Artifacts/UnsyncCache.cs
--- a/Engine/Source/Programs/Horde/Horde.Server/Artifacts/UnsyncCache.cs
+++ b/Engine/Source/Programs/Horde/Horde.Server/Artifacts/UnsyncCache.cs
@@ -5,10 +5,12 @@ using System.Collections.Frozen;
 using System.Collections.Generic;
 using System.Diagnostics;
 using System.Linq;
+using System.Text.Json;
 using System.Threading;
 using System.Threading.Tasks;
 using EpicGames.Core;
 using EpicGames.Horde.Storage;
+using EpicGames.Horde.Storage.Bundles;
 using EpicGames.Horde.Storage.Nodes;
 using Horde.Server.Storage;
 using Microsoft.Extensions.Caching.Memory;
@@ -51,12 +53,16 @@ namespace Horde.Server.Artifacts
 		{
 			public IStorageClient StorageClient { get; }
 			public UnsyncManifest Manifest { get; }
+			public ReadOnlyMemory<byte> ManifestData { get; }
+			public ReadOnlyMemory<byte> ZstdManifestData { get; }
 			public FrozenDictionary<IoHash, IBlobRef<LeafChunkedDataNode>> Blobs { get; }
 
-			public ArtifactInfo(IStorageClient storageClient, UnsyncManifest manifest, FrozenDictionary<IoHash, IBlobRef<LeafChunkedDataNode>> blobs)
+			public ArtifactInfo(IStorageClient storageClient, UnsyncManifest manifest, ReadOnlyMemory<byte> manifestData, ReadOnlyMemory<byte> zstdManifestData, FrozenDictionary<IoHash, IBlobRef<LeafChunkedDataNode>> blobs)
 			{
 				StorageClient = storageClient;
 				Manifest = manifest;
+				ManifestData = manifestData;
+				ZstdManifestData = zstdManifestData;
 				Blobs = blobs;
 			}
 
@@ -84,17 +90,28 @@ namespace Horde.Server.Artifacts
 			=> _cache.Dispose();
 
 		/// <summary>
-		/// 
+		/// Gets the manifest for an artifact
 		/// </summary>
-		/// <param name="artifact"></param>
-		/// <param name="cancellationToken"></param>
-		/// <returns></returns>
 		public async ValueTask<UnsyncManifest?> GetManifestAsync(IArtifact artifact, CancellationToken cancellationToken = default)
 		{
 			ArtifactInfo? artifactInfo = await GetArtifactInfoAsync(artifact, cancellationToken);
 			return artifactInfo?.Manifest;
 		}
 
+		/// <summary>
+		/// Gets the json manifest data for an artifact
+		/// </summary>
+		public async ValueTask<ReadOnlyMemory<byte>> GetManifestDataAsync(IArtifact artifact, bool compressed, CancellationToken cancellationToken = default)
+		{
+			ArtifactInfo? artifactInfo = await GetArtifactInfoAsync(artifact, cancellationToken);
+			if (artifactInfo == null)
+			{
+				return default;
+			}
+
+			return compressed ? artifactInfo.ZstdManifestData : artifactInfo.ManifestData;
+		}
+
 		async Task<ArtifactInfo?> GetArtifactInfoAsync(IArtifact artifact, CancellationToken cancellationToken = default)
 		{
 			// Get the artifact info
@@ -121,6 +138,60 @@ namespace Horde.Server.Artifacts
 			return await task.WaitAsync(cancellationToken);
 		}
 
+		static byte[] SerializeManifest(UnsyncManifest manifest)
+		{
+			using (ChunkedMemoryWriter chunkedMemoryWriter = new ChunkedMemoryWriter(4096))
+			{
+				using (Utf8JsonWriter writer = new Utf8JsonWriter(chunkedMemoryWriter))
+				{
+					writer.WriteStartObject();
+					writer.WriteString("type", "unsync_manifest");
+					writer.WriteString("hash_strong", "Blake3.160");
+					writer.WriteString("chunking", "RollingBuzHash");
+					writer.WriteNumber("chunking_block_size_min", LeafChunkedDataNodeOptions.Default.MinSize);
+					writer.WriteNumber("chunking_block_size_max", LeafChunkedDataNodeOptions.Default.MaxSize);
+					writer.WriteNumber("chunking_block_size_target", LeafChunkedDataNodeOptions.Default.TargetSize);
+
+					writer.WriteStartArray("files");
+					foreach (UnsyncFile file in manifest.Files)
+					{
+						writer.WriteStartObject();
+						writer.WriteString("name", file.Name);
+						if (file.Executable)
+						{
+							writer.WriteBoolean("exec", file.Executable);
+						}
+						if (file.ReadOnly)
+						{
+							writer.WriteBoolean("read_only", file.ReadOnly);
+						}
+						if (file.ModTime != default)
+						{
+							writer.WriteNumber("mtime", file.ModTime.ToFileTimeUtc());
+						}
+						writer.WriteNumber("size", file.Length);
+						writer.WriteNumber("block_size", LeafChunkedDataNodeOptions.Default.TargetSize);
+						writer.WriteStartArray("blocks");
+						foreach (UnsyncBlock block in file.Blocks)
+						{
+							writer.WriteStartObject();
+							writer.WriteNumber("offset", block.Offset);
+							writer.WriteNumber("size", block.Length);
+							writer.WriteNumber("hash_weak", block.RollingHash);
+							writer.WriteString("hash_strong", block.Blob.Hash.ToString());
+							writer.WriteEndObject();
+						}
+						writer.WriteEndArray();
+						writer.WriteEndObject();
+					}
+					writer.WriteEndArray();
+
+					writer.WriteEndObject();
+				}
+				return chunkedMemoryWriter.ToByteArray();
+			}
+		}
+
 		/// <summary>
 		/// Reads a blob from an artifact
 		/// </summary>
@@ -170,7 +241,11 @@ namespace Horde.Server.Artifacts
 
 				_logger.LogDebug("Generated Unsync manifest for artifact {ArtifactId} in {Time:n1}ms", artifact.Id, timer.ElapsedMilliseconds);
 
-				ArtifactInfo artifactInfo = new ArtifactInfo(storageClient, new UnsyncManifest(files), blocks.ToFrozenDictionary());
+				UnsyncManifest manifest = new UnsyncManifest(files);
+				ReadOnlyMemory<byte> manifestData = SerializeManifest(manifest);
+				ReadOnlyMemory<byte> zstdManifestData = BundleData.Compress(BundleCompressionFormat.Zstd, manifestData);
+
+				ArtifactInfo artifactInfo = new ArtifactInfo(storageClient, manifest, manifestData, zstdManifestData, blocks.ToFrozenDictionary());
 				storageClient = null;
 				return artifactInfo;
 			}
diff --git a/Engine/Source/Programs/Shared/EpicGames.Horde/Storage/Bundles/BundleCompression.cs b/Engine/Source/Programs/Shared/EpicGames.Horde/Storage/Bundles/BundleCompression.cs
--- a/Engine/Source/Programs/Shared/EpicGames.Horde/Storage/Bundles/BundleCompression.cs
+++ b/Engine/Source/Programs/Shared/EpicGames.Horde/Storage/Bundles/BundleCompression.cs
@@ -53,13 +53,26 @@ namespace EpicGames.Horde.Storage.Bundles
 	/// </summary>
 	public static class BundleData
 	{
+		/// <summary>
+		/// Compress a block of data into a newly allocated block of memory
+		/// </summary>
+		/// <param name="format">Format for the compressed data</param>
+		/// <param name="input">The data to compress</param>
+		/// <returns>The compressed data</returns>
+		public static byte[] Compress(BundleCompressionFormat format, ReadOnlyMemory<byte> input)
+		{
+			ArrayMemoryWriter writer = new ArrayMemoryWriter(0);
+			int length = Compress(format, input, writer);
+			return writer.WrittenMemory.ToArray();
+		}
+
 		/// <summary>
 		/// Compress a data packet
 		/// </summary>
 		/// <param name="format">Format for the compressed data</param>
 		/// <param name="input">The data to compress</param>
 		/// <param name="writer">Writer for output data</param>
-		/// <returns>The compressed data</returns>
+		/// <returns>Length of the compressed data</returns>
 		public static int Compress(BundleCompressionFormat format, ReadOnlyMemory<byte> input, IMemoryWriter writer)
 		{
 			switch (format)
