Horde: Break out of the blob tick loop if the reset flag is set.

#rnx

diff --git a/Engine/Source/Programs/Horde/Horde.Server/Storage/StorageService.cs b/Engine/Source/Programs/Horde/Horde.Server/Storage/StorageService.cs
--- a/Engine/Source/Programs/Horde/Horde.Server/Storage/StorageService.cs
+++ b/Engine/Source/Programs/Horde/Horde.Server/Storage/StorageService.cs
@@ -662,7 +662,7 @@ namespace Horde.Server.Storage
 
 				// Compute missing import info, by searching for blobs with an ObjectId timestamp after the last import compute cycle
 				ObjectId latestInfoId = ObjectId.GenerateNewId(ingestTimeUtc);
-				for (; ; )
+				while (!gcState.Reset)
 				{
 					// Fetch the next batch of blobs
 					List<BlobInfo> current = await _blobCollection.Find(x => x.Id > gcState.LastImportBlobInfoId && x.Id < latestInfoId).SortBy(x => x.Id).Limit(500).ToListAsync(cancellationToken);
