Horde: Add support for compression of all responses using gzip, brotli, and zstd.

diff --git a/Engine/Source/Programs/Horde/Horde.Server/Artifacts/ArtifactsController.cs b/Engine/Source/Programs/Horde/Horde.Server/Artifacts/ArtifactsController.cs
--- a/Engine/Source/Programs/Horde/Horde.Server/Artifacts/ArtifactsController.cs
+++ b/Engine/Source/Programs/Horde/Horde.Server/Artifacts/ArtifactsController.cs
@@ -9,7 +9,6 @@ using System.IO;
 using System.Linq;
 using System.Net;
 using System.Net.Mime;
-using System.Text.Json;
 using System.Text.RegularExpressions;
 using System.Threading;
 using System.Threading.Tasks;
@@ -607,6 +606,20 @@ namespace Horde.Server.Artifacts
 			return new FileStreamResult(new MemoryStream(data), "application/x-horde-artifact") { FileDownloadName = $"{artifact.Name}.uartifact" };
 		}
 
+		class ManifestStreamResult : FileStreamResult
+		{
+			public ManifestStreamResult(System.IO.Stream stream, string mimeType)
+				: base(stream, mimeType)
+			{
+			}
+
+			/// <inheritdoc/>
+			public override Task ExecuteResultAsync(ActionContext context)
+			{
+				context.HttpContext.Response.Headers["Content-Disposition"] = new ContentDisposition { Inline = true }.ToString();
+				return base.ExecuteResultAsync(context);
+			}
+		}
 		/// <summary>
 		/// Creates an Unsync manifest for an artifact
 		/// </summary>
@@ -627,28 +640,15 @@ namespace Horde.Server.Artifacts
 				return Forbid(ArtifactAclAction.ReadArtifact, artifact.AclScope);
 			}
 
-			const string ZstdManifestContentType = "application/x-unsync-manifest-zstd";
-			bool compressed = Request.Headers.Accept.Contains(ZstdManifestContentType);
-
-			ReadOnlyMemory<byte> manifestData = await _unsyncCache.GetManifestDataAsync(artifact, compressed, cancellationToken);
+			ReadOnlyMemory<byte> manifestData = await _unsyncCache.GetManifestDataAsync(artifact, cancellationToken);
 			if (manifestData.IsEmpty)
 			{
 				return NotFound(id);
 			}
 
-			// Disable buffering for the response
-			IHttpResponseBodyFeature? responseBodyFeature = HttpContext.Features.Get<IHttpResponseBodyFeature>();
-			responseBodyFeature?.DisableBuffering();
-
-			// Write the response directly to the writer
-			HttpResponse response = HttpContext.Response;
-			response.ContentType = compressed? ZstdManifestContentType : "application/json";
-			response.StatusCode = (int)HttpStatusCode.OK;
-			await response.StartAsync(cancellationToken);
-			await response.BodyWriter.WriteAsync(manifestData, cancellationToken);
-			await response.CompleteAsync();
-
-			return Empty;
+#pragma warning disable CA2000
+			return new ManifestStreamResult(new ReadOnlyMemoryStream(manifestData), "application/json");
+#pragma warning restore CA2000
 		}
 
 		/// <summary>
diff --git a/Engine/Source/Programs/Horde/Horde.Server/Artifacts/UnsyncCache.cs b/Engine/Source/Programs/Horde/Horde.Server/Artifacts/UnsyncCache.cs
--- a/Engine/Source/Programs/Horde/Horde.Server/Artifacts/UnsyncCache.cs
+++ b/Engine/Source/Programs/Horde/Horde.Server/Artifacts/UnsyncCache.cs
@@ -10,7 +10,6 @@ using System.Threading;
 using System.Threading.Tasks;
 using EpicGames.Core;
 using EpicGames.Horde.Storage;
-using EpicGames.Horde.Storage.Bundles;
 using EpicGames.Horde.Storage.Nodes;
 using Horde.Server.Storage;
 using Microsoft.Extensions.Caching.Memory;
@@ -54,15 +53,13 @@ namespace Horde.Server.Artifacts
 			public IStorageClient StorageClient { get; }
 			public UnsyncManifest Manifest { get; }
 			public ReadOnlyMemory<byte> ManifestData { get; }
-			public ReadOnlyMemory<byte> ZstdManifestData { get; }
 			public FrozenDictionary<IoHash, IBlobRef<LeafChunkedDataNode>> Blobs { get; }
 
-			public ArtifactInfo(IStorageClient storageClient, UnsyncManifest manifest, ReadOnlyMemory<byte> manifestData, ReadOnlyMemory<byte> zstdManifestData, FrozenDictionary<IoHash, IBlobRef<LeafChunkedDataNode>> blobs)
+			public ArtifactInfo(IStorageClient storageClient, UnsyncManifest manifest, ReadOnlyMemory<byte> manifestData, FrozenDictionary<IoHash, IBlobRef<LeafChunkedDataNode>> blobs)
 			{
 				StorageClient = storageClient;
 				Manifest = manifest;
 				ManifestData = manifestData;
-				ZstdManifestData = zstdManifestData;
 				Blobs = blobs;
 			}
 
@@ -101,7 +98,7 @@ namespace Horde.Server.Artifacts
 		/// <summary>
 		/// Gets the json manifest data for an artifact
 		/// </summary>
-		public async ValueTask<ReadOnlyMemory<byte>> GetManifestDataAsync(IArtifact artifact, bool compressed, CancellationToken cancellationToken = default)
+		public async ValueTask<ReadOnlyMemory<byte>> GetManifestDataAsync(IArtifact artifact, CancellationToken cancellationToken = default)
 		{
 			ArtifactInfo? artifactInfo = await GetArtifactInfoAsync(artifact, cancellationToken);
 			if (artifactInfo == null)
@@ -109,7 +106,7 @@ namespace Horde.Server.Artifacts
 				return default;
 			}
 
-			return compressed ? artifactInfo.ZstdManifestData : artifactInfo.ManifestData;
+			return artifactInfo.ManifestData;
 		}
 
 		async Task<ArtifactInfo?> GetArtifactInfoAsync(IArtifact artifact, CancellationToken cancellationToken = default)
@@ -243,9 +240,8 @@ namespace Horde.Server.Artifacts
 
 				UnsyncManifest manifest = new UnsyncManifest(files);
 				ReadOnlyMemory<byte> manifestData = SerializeManifest(manifest);
-				ReadOnlyMemory<byte> zstdManifestData = BundleData.Compress(BundleCompressionFormat.Zstd, manifestData);
 
-				ArtifactInfo artifactInfo = new ArtifactInfo(storageClient, manifest, manifestData, zstdManifestData, blocks.ToFrozenDictionary());
+				ArtifactInfo artifactInfo = new ArtifactInfo(storageClient, manifest, manifestData, blocks.ToFrozenDictionary());
 				storageClient = null;
 				return artifactInfo;
 			}
diff --git a/Engine/Source/Programs/Horde/Horde.Server/Horde.Server.csproj b/Engine/Source/Programs/Horde/Horde.Server/Horde.Server.csproj
--- a/Engine/Source/Programs/Horde/Horde.Server/Horde.Server.csproj
+++ b/Engine/Source/Programs/Horde/Horde.Server/Horde.Server.csproj
@@ -81,6 +81,7 @@
     <PackageReference Include="System.Linq.Async" Version="6.0.1" />
     <PackageReference Include="T-Digest.NET" Version="1.0.3" />
     <PackageReference Include="TimeZoneConverter" Version="6.1.0" />
+    <PackageReference Include="ZstdSharp.Port" Version="0.8.1" />
   </ItemGroup>
 
   <ItemGroup>
diff --git a/Engine/Source/Programs/Horde/Horde.Server/Startup.cs b/Engine/Source/Programs/Horde/Horde.Server/Startup.cs
--- a/Engine/Source/Programs/Horde/Horde.Server/Startup.cs
+++ b/Engine/Source/Programs/Horde/Horde.Server/Startup.cs
@@ -108,6 +108,7 @@ using Microsoft.AspNetCore.Mvc;
 using Microsoft.AspNetCore.Mvc.Controllers;
 using Microsoft.AspNetCore.Mvc.Filters;
 using Microsoft.AspNetCore.Mvc.ModelBinding;
+using Microsoft.AspNetCore.ResponseCompression;
 using Microsoft.Extensions.Configuration;
 using Microsoft.Extensions.DependencyInjection;
 using Microsoft.Extensions.Diagnostics.HealthChecks;
@@ -468,6 +469,9 @@ namespace Horde.Server
 			services.AddResponseCompression(options =>
 			{
 				options.EnableForHttps = true;
+				options.Providers.Add<GzipCompressionProvider>();
+				options.Providers.Add<BrotliCompressionProvider>();
+				options.Providers.Add<ZstdCompressionProvider>();
 			});
 			
 			if (settings.CorsEnabled)
diff --git a/Engine/Source/Programs/Horde/Horde.Server/Utilities/ZstdCompressionProvider.cs b/Engine/Source/Programs/Horde/Horde.Server/Utilities/ZstdCompressionProvider.cs
new file mode 100644
--- /dev/null
+++ b/Engine/Source/Programs/Horde/Horde.Server/Utilities/ZstdCompressionProvider.cs
@@ -0,0 +1,25 @@
+// Copyright Epic Games, Inc. All Rights Reserved.
+
+using System.IO;
+using Microsoft.AspNetCore.ResponseCompression;
+
+namespace Horde.Server.Utilities
+{
+	/// <summary>
+	/// Compression provider for zstd responses
+	/// </summary>
+	class ZstdCompressionProvider : ICompressionProvider
+	{
+		/// <inheritdoc/>
+		public string EncodingName => "zstd";
+
+		/// <inheritdoc/>
+		public bool SupportsFlush => true;
+
+		/// <inheritdoc/>
+		public Stream CreateStream(Stream outputStream)
+		{
+			return new ZstdSharp.CompressionStream(outputStream, leaveOpen: false);
+		}
+	}
+}
