Horde: New Build V2 fix issue with dialog refreshing erractically due to bad observer
#rnx

diff --git a/Engine/Source/Programs/Horde/HordeDashboard/src/components/build/NewBuildV2.tsx b/Engine/Source/Programs/Horde/HordeDashboard/src/components/build/NewBuildV2.tsx
--- a/Engine/Source/Programs/Horde/HordeDashboard/src/components/build/NewBuildV2.tsx
+++ b/Engine/Source/Programs/Horde/HordeDashboard/src/components/build/NewBuildV2.tsx
@@ -2044,7 +2044,7 @@ const NewBuildV2Inner: React.FC<{ setUseLegacyDialog: (value: boolean) => void }
 
 })
 
-export const NewBuildV2: React.FC<{ streamId: string; show: boolean; onClose: (newJobId: string | undefined) => void, jobKey?: string; jobDetails?: JobDetailsV2, readOnly?: boolean }> = observer(({ streamId, jobKey, show, onClose, jobDetails, readOnly }) => {
+export const NewBuildV2: React.FC<{ streamId: string; show: boolean; onClose: (newJobId: string | undefined) => void, jobKey?: string; jobDetails?: JobDetailsV2, readOnly?: boolean }> = ({ streamId, jobKey, show, onClose, jobDetails, readOnly }) => {
 
    const { projectStore } = useBackend();
    const options = useConst(new BuildOptions(streamId, projectStore, onClose, jobDetails, jobKey, readOnly));
@@ -2063,7 +2063,7 @@ export const NewBuildV2: React.FC<{ streamId: string; show: boolean; onClose: (n
    return <Stack>
       <NewBuildV2Inner setUseLegacyDialog={(value: boolean) => { setNewBuildV1(value) }} />
    </Stack>
-})
+}
 
 export const ValidationErrorModal: React.FC<{ errors: ValidationError[], onClose: () => void }> = ({ errors, onClose }) => {
 
