Horde: Add documentation for setting up a secret store.

diff --git a/Engine/Source/Programs/Horde/Docs/Config.md b/Engine/Source/Programs/Horde/Docs/Config.md
--- a/Engine/Source/Programs/Horde/Docs/Config.md
+++ b/Engine/Source/Programs/Horde/Docs/Config.md
@@ -14,6 +14,7 @@ This section targets operating and managing Horde installations and shares some
 
 * [Agents](Config/Agents.md): Configure worker machines to connect to the Horde server for CI and remote execution workspaces.
 * [Devices](Config/Devices.md): Adding mobile devices and console development kits as shared resources.
+* [Secrets](Config/Secrets.md): Allowing Horde to manage access to confidential values.
 * [Storage](Config/Storage.md): Setting up different storage backends for Horde.
 
 ## Features
diff --git a/Engine/Source/Programs/Horde/Docs/Config/Secrets.md b/Engine/Source/Programs/Horde/Docs/Config/Secrets.md
new file mode 100644
--- /dev/null
+++ b/Engine/Source/Programs/Horde/Docs/Config/Secrets.md
@@ -0,0 +1,123 @@
+[Horde](../../README.md) > [Configuration](../Config.md) > Secrets
+
+# Secrets
+
+Horde implements an API for retrieving secrets that can be stored in its own configuration file or obtained from
+an external source. Marshalling data through Horde allows access to be controlled using Horde's permissions model,
+and for automated processes to impersonate the user that requested them.
+
+At the moment, Horde only supports the AWS Parameter Store as an external secret provider out of the box, though 
+other implementations can be added through the `ISecretProvider` interface.
+
+## Configuring secrets
+
+Secrets are defined in the `secrets` list of the [`globals.json`](Schema/Globals.md) file. Each 
+[entry](Schema/Globals.md#secretconfig) includes an identifier for the secret (`id`), a set of key value
+pairs (`data`) and an [ACL](Permissions.md) controlling who can access it.
+
+Additional keys and values may be merged from external providers by adding entries to the `sources` array.
+Each entry here contains the `name` of the provider to obtain the secret from, and a `path` used to locate
+the secret in a provider-specific syntax.
+
+Secrets from external providers may take two forms, determined by the `format` property.
+
+* `text` secrets are simple string values which are added to the secret's set of key/value pairs using
+  the specified `key` property.
+* `json` secrets are parsed as JSON objects and merged into the secret's key/value pairs using property
+  names as keys.
+
+An example secret may be configured as follows:
+
+   ```json
+    "secrets": [
+        {
+            "id": "horde-secrets",
+
+            // Some hard-coded property values        
+            "data": {
+                "aws-region": "us-east-1"
+            },
+
+            // Some values read from the AWS parameter store
+            "sources": [
+
+                // Read a single secret from the AWS parameter store and assign it to "aws-secret-access-key"
+                {
+                    "provider": "AwsParameterStore",
+                    "key": "aws-secret-access-key",
+                    "path": "name-of-secret-in-parameter-store"
+                },
+
+                // Read a JSON object from the AWS parameter store and merge all the key/value pairs into this secret.
+                {
+                    "provider": "AwsParameterStore",
+                    "format": "json",
+                    "path": "name-of-secret-in-parameter-store"
+                },
+            ],
+
+            // Only allow Horde agents to access this
+  			"acl": {
+				"entries": [
+					{
+						"claim": {
+							"type": "http://epicgames.com/ue/horde/role",
+							"value": "agent"
+						},
+						"actions": [
+							"ViewSecret"
+						]
+					}
+				]
+			}
+        }
+    ]
+   ```
+
+Secrets are queried from the external provider when requested by a user, and are not cached by Horde.
+
+## Using secrets
+
+The most common use case for secrets is during build automation pipelines. In this scenario, the Horde Server
+URL and credentials are taken from environment variables injected automatically by the Horde Agent, and allow
+the pipeline to request secrets on behalf of the user starting the job with little additional configuration.
+
+There are three common ways to use secrets:
+
+### 1. Using the **Horde-GetSecrets** BuildGraph task
+
+This task takes a file as a parameter, and will read it in, expand any secrets in the form {{ secret-name.property }}
+with their values from Horde, and write it back out again. Rather than updating an existing file, you can also put the
+template in a BuildGraph property and expand that instead, as follows:
+
+   ```xml
+   <Property Name="AwsEnvironmentText" Multiline="true">
+      AWS_REGION={{horde-secrets.aws-region}}
+      AWS_ACCESS_KEY_ID={{horde-secrets.aws-access-key-id}}
+      AWS_SECRET_ACCESS_KEY={{horde-secrets.aws-secret-access-key}}
+   </Property>
+   <Horde-GetSecrets File="credentials.txt" Text="$(AwsEnvironmentText)"/>
+   ```
+
+### 2. Using the **Horde-SetSecretEnvVar** BuildGraph task
+
+This task sets an environment variable to the value of a secret at runtime. Environment variables are inherited 
+by child processes but not set at the system level, so the environment variable will contain that secret until
+the end of the current step.
+
+   ```xml
+   <Horde-SetSecretEnvVar Name="AWS_SECRET_ACCESS_KEY" Secret="horde-secrets.aws-secret-access-key"/>
+   ```
+
+### 3. Using the Horde API
+
+The HTTP secrets endpoint is listed in Horde's API documentation, and `AutomationTool` includes the following
+utility methods for querying them at runtime:
+
+   ```csharp
+   IReadOnlyDictionary<string, string> secret = await CommandUtils.GetHordeSecretAsync(new SecretId("my-secret-name"));
+   ```
+
+   ```csharp
+   string propertyValue = await CommandUtils.GetHordeSecretAsync(new SecretId("my-secret-name"), "propertyName")
+   ```
