Horde: Rename projects before splitting into plugins (Horde.Agent -> HordeAgent, Horde.Server -> HordeServer, Horde -> HordeCmd).

#rnx

diff --git a/Engine/Source/Programs/Horde/BuildHorde.xml b/Engine/Source/Programs/Horde/BuildHorde.xml
--- a/Engine/Source/Programs/Horde/BuildHorde.xml
+++ b/Engine/Source/Programs/Horde/BuildHorde.xml
@@ -53,8 +53,8 @@
 		</Node>
 
 		<Node Name="Analyze HordeServer" Requires="Check Documentation">
-			<DotNet Arguments="build &quot;$(RootDir)/Engine/Source/Programs/Horde/Horde.Server/Horde.Server.csproj&quot; /p:Configuration=Analyze $(VersionArguments)" DotNetPath="$(DotNetPath)"/>
-			<DotNet Arguments="build &quot;$(RootDir)/Engine/Source/Programs/Horde/Horde.Server.Tests/Horde.Server.Tests.csproj&quot; /p:Configuration=Analyze $(VersionArguments)" DotNetPath="$(DotNetPath)"/>
+			<DotNet Arguments="build &quot;$(RootDir)/Engine/Source/Programs/Horde/HordeServer/HordeServer.csproj&quot; /p:Configuration=Analyze $(VersionArguments)" DotNetPath="$(DotNetPath)"/>
+			<DotNet Arguments="build &quot;$(RootDir)/Engine/Source/Programs/Horde/HordeServer.Tests/HordeServer.Tests.csproj&quot; /p:Configuration=Analyze $(VersionArguments)" DotNetPath="$(DotNetPath)"/>
 		</Node>
 
 		<Node Name="Check HordeServer Licenses" Requires="Analyze HordeServer">
@@ -64,24 +64,24 @@
 				[REDACTED_EMAIL]
 				[REDACTED_EMAIL]
 			</Property>
-			<NuGet-LicenseCheck BaseDir="Engine/Source/Programs/Horde/Horde.Server" IgnorePackages="$(IgnorePackages)" LicenseDir="Engine/Source/Programs/Horde/ThirdParty/Licenses" DotNetPath="$(DotNetPath)"/>
+			<NuGet-LicenseCheck BaseDir="Engine/Source/Programs/Horde/HordeServer" IgnorePackages="$(IgnorePackages)" LicenseDir="Engine/Source/Programs/Horde/ThirdParty/Licenses" DotNetPath="$(DotNetPath)"/>
 		</Node>
 	</Agent>
 	
 	<Agent Name="HordeServer" Type="Linux;Win64_Docker">
-		<Property Name="StagingDir" Value="$(RootDir)/Engine/Saved/Horde.Server"/>
+		<Property Name="StagingDir" Value="$(RootDir)/Engine/Saved/HordeServer"/>
 
 		<Node Name="Build HordeServer">
 			<!-- Tag all the files that need to be staged to build -->
 			<Tag Files="$(StagePaths)" Except=".../.vs/...;.../.git/...;.../bin/...;.../obj/..." With="#InputFiles"/>
 
 			<!-- Build the Docker image and publish it -->
-			<Docker-Build BaseDir="Engine" Files="#InputFiles" Tag="horde-server-bare" DockerFile="Engine/Source/Programs/Horde/Horde.Server/Dockerfile" Arguments="--build-arg msbuild_args=&quot;$(VersionArguments)&quot;"/>
+			<Docker-Build BaseDir="Engine" Files="#InputFiles" Tag="horde-server-bare" DockerFile="Engine/Source/Programs/Horde/HordeServer/Dockerfile" Arguments="--build-arg msbuild_args=&quot;$(VersionArguments)&quot;"/>
 		</Node>
 
 		<Node Name="Test HordeServer Code Coverage">
 			<Tag Files="$(StagePaths)" Except=".../.vs/...;.../.git/...;.../bin/...;.../obj/..." With="#InputFiles"/>
-			<Docker-Build BaseDir="Engine" Files="#InputFiles" Tag="hordeserver-codecov" UseBuildKit="true" Target="build-env-output" Output="type=local,dest=/tmp/docker-output" DockerFile="Engine/Source/Programs/Horde/Horde.Server/Dockerfile" Arguments="--build-arg msbuild_args=&quot;$(VersionArguments)&quot; --build-arg code_coverage=&quot;true&quot;" />
+			<Docker-Build BaseDir="Engine" Files="#InputFiles" Tag="hordeserver-codecov" UseBuildKit="true" Target="build-env-output" Output="type=local,dest=/tmp/docker-output" DockerFile="Engine/Source/Programs/Horde/HordeServer/Dockerfile" Arguments="--build-arg msbuild_args=&quot;$(VersionArguments)&quot; --build-arg code_coverage=&quot;true&quot;" />
 			<Copy From="/tmp/docker-output/dotcover-report.zip" To="$(RootDir)/Engine/Programs/AutomationTool/Saved/Logs/dotcover-report.zip" />
 		</Node>
 
@@ -97,7 +97,7 @@
 		
 		<Node Name="Build HordeServer with Dashboard" Requires="Build HordeServer;Build HordeDashboard">
 			<!-- Create a new image by combining server and dashboard image into one -->
-			<Docker-Build BaseDir="Engine/Source/Programs/Horde/Horde.Server" Files="Dockerfile*" Tag="horde-server" DockerFile="Engine/Source/Programs/Horde/Horde.Server/Dockerfile.dashboard" />
+			<Docker-Build BaseDir="Engine/Source/Programs/Horde/HordeServer" Files="Dockerfile*" Tag="horde-server" DockerFile="Engine/Source/Programs/Horde/HordeServer/Dockerfile.dashboard" />
 		</Node>
 	</Agent>
 	
@@ -129,7 +129,7 @@
 		</Node>
 		<Node Name="Build HordeAgent" Produces="#HordeAgentFiles">
 			<Delete Files="$(ServerToolsStagingDir_HordeAgent)/..."/>
-			<DotNet Arguments="publish &quot;$(RootDir)/Engine/Source/Programs/Horde/Horde.Agent/Horde.Agent.csproj&quot; --output &quot;$(ServerToolsStagingDir_HordeAgent)&quot; $(VersionArguments)"/>
+			<DotNet Arguments="publish &quot;$(RootDir)/Engine/Source/Programs/Horde/HordeAgent/HordeAgent.csproj&quot; --output &quot;$(ServerToolsStagingDir_HordeAgent)&quot; $(VersionArguments)"/>
 			<DotNet Arguments="publish &quot;$(RootDir)/Engine/Source/Programs/Horde/Drivers/JobDriver/JobDriver.csproj&quot; --output &quot;$(ServerToolsStagingDir_HordeAgent)/JobDriver&quot; $(VersionArguments)"/>
 			<Delete Files="$(ServerToolsStagingDir_HordeAgent)/appsettings.Local.json"/>
 			<Sign Files="$(ServerToolsStagingDir_HordeAgent)/HordeAgent.exe;$(ServerToolsStagingDir_HordeAgent)/HordeAgent.dll;$(ServerToolsStagingDir_HordeAgent)/EpicGames.*.dll" If="'$(SignOutputs)' == 'true'"/>
@@ -140,7 +140,7 @@
 			<Delete Files="$(ServerToolsStagingDir_HordeAgentWin64)/..."/>
 			
 			<!-- Build the agent -->
-			<DotNet Arguments="publish &quot;$(RootDir)/Engine/Source/Programs/Horde/Horde.Agent/Horde.Agent.csproj&quot; --output &quot;$(ServerToolsStagingDir_HordeAgentWin64)&quot; --runtime win7-x64 --self-contained $(VersionArguments)"/>
+			<DotNet Arguments="publish &quot;$(RootDir)/Engine/Source/Programs/Horde/HordeAgent/HordeAgent.csproj&quot; --output &quot;$(ServerToolsStagingDir_HordeAgentWin64)&quot; --runtime win7-x64 --self-contained $(VersionArguments)"/>
 			<Sign Files="$(ServerToolsStagingDir_HordeAgentWin64)/HordeAgent.exe;$(ServerToolsStagingDir_HordeAgentWin64)/HordeAgent.dll;$(ServerToolsStagingDir_HordeAgentWin64)/EpicGames.*.dll" If="'$(SignOutputs)' == 'true'"/>
 
 			<!-- Build the driver -->
@@ -148,8 +148,8 @@
 			<Sign Files="$(ServerToolsStagingDir_HordeAgentWin64)/JobDriver/JobDriver.exe;$(ServerToolsStagingDir_HordeAgentWin64)/JobDriver/JobDriver.dll;$(ServerToolsStagingDir_HordeAgentWin64)/JobDriver/EpicGames.*.dll" If="'$(SignOutputs)' == 'true'"/>
 
 			<!-- Build the tray app -->
-			<DotNet Arguments="publish &quot;$(RootDir)/Engine/Source/Programs/Horde/Horde.Agent.TrayApp/Horde.Agent.TrayApp.csproj&quot; --output &quot;$(ServerToolsStagingDir_HordeAgentWin64)/TrayApp&quot; --runtime win7-x64 --self-contained $(VersionArguments)"/>
-			<Sign Files="$(ServerToolsStagingDir_HordeAgentWin64)/TrayApp/Horde.Agent.TrayApp.exe;$(ServerToolsStagingDir_HordeAgentWin64)/TrayApp/EpicGames.*.dll" If="'$(SignOutputs)' == 'true'"/>
+			<DotNet Arguments="publish &quot;$(RootDir)/Engine/Source/Programs/Horde/HordeAgent.TrayApp/HordeAgent.TrayApp.csproj&quot; --output &quot;$(ServerToolsStagingDir_HordeAgentWin64)/TrayApp&quot; --runtime win7-x64 --self-contained $(VersionArguments)"/>
+			<Sign Files="$(ServerToolsStagingDir_HordeAgentWin64)/TrayApp/HordeAgent.TrayApp.exe;$(ServerToolsStagingDir_HordeAgentWin64)/TrayApp/EpicGames.*.dll" If="'$(SignOutputs)' == 'true'"/>
 
 			<Delete Files="$(ServerToolsStagingDir_HordeAgentWin64)/appsettings.Local.json"/>
 			<Tag Files="$(ServerToolsStagingDir_HordeAgentWin64)/..." With="#HordeAgentWin64Files"/>
@@ -226,7 +226,7 @@
 
 			<!-- Create the agent file list -->
 			<Delete Files="$(StagingDir)/Agent-Bulk/..."/>
-			<Tag Files="#HordeAgentWin64Files" Except="HordeAgent.exe;Horde.Agent.TrayApp.exe" With="#AgentFiles"/>
+			<Tag Files="#HordeAgentWin64Files" Except="HordeAgent.exe;HordeAgent.TrayApp.exe" With="#AgentFiles"/>
 			<Copy Files="#AgentFiles" From="$(ServerToolsStagingDir_HordeAgentWin64)" To="$(StagingDir)/Agent-Bulk/Agent"/>
 
 			<Spawn Exe="$(WixDir)/heat.exe" Arguments="dir &quot;$(StagingDir)/Agent-Bulk&quot; -cg HordeAgentFiles -scom -sreg -gg -dr InstallDir -srd -var var.SourceDir -out &quot;$(InstallerDir)/HordeAgentFiles.wxs&quot;"/>
@@ -250,20 +250,20 @@
 		
 				<!-- Stage the server -->
 				<Delete Files="$(StagingDir)/Server/$(ServerArch)/..."/>
-				<DotNet Arguments="publish &quot;$(RootDir)/Engine/Source/Programs/Horde/Horde.Server/Horde.Server.csproj&quot; --output &quot;$(StagingDir)/Server/$(ServerArch)&quot; --runtime $(ServerArch) --self-contained $(VersionArguments)" DotNetPath="$(DotNetPath)" />
+				<DotNet Arguments="publish &quot;$(RootDir)/Engine/Source/Programs/Horde/HordeServer/HordeServer.csproj&quot; --output &quot;$(StagingDir)/Server/$(ServerArch)&quot; --runtime $(ServerArch) --self-contained $(VersionArguments)" DotNetPath="$(DotNetPath)" />
 
 				<!-- Remove the local config file. This is only used for local iteration. -->
 				<Delete Files="$(StagingDir)/Server/$(ServerArch)/appsettings.Local.json"/>
 
 				<!-- Add the bundled tools into the installer -->
-				<DotNet Arguments="run --project &quot;$(RootDir)/Engine/Source/Programs/Horde/Horde.Server/Horde.Server.csproj&quot; -- generate tooldata -id=horde-cmd -name=&quot;Horde (Command-Line)&quot; -description=&quot;Update package for the Horde command-line tool.&quot; -version=&quot;$(InformationalVersion)&quot; -inputdir=&quot;$(ServerToolsStagingDir_HordeCmd)&quot; -serverdir=&quot;$(StagingDir)/Server/$(ServerArch)&quot; -public -showinugs -Category=Internal" DotNetPath="$(DotNetPath)" />
-				<DotNet Arguments="run --project &quot;$(RootDir)/Engine/Source/Programs/Horde/Horde.Server/Horde.Server.csproj&quot; -- generate tooldata -id=horde-agent -name=&quot;Horde Agent (Cross-Platform)&quot; -description=&quot;Cross-platform Horde Agent distribution. Allows hosts to take on work from the Horde Server.&quot; -version=&quot;$(InformationalVersion)&quot; -inputdir=&quot;$(ServerToolsStagingDir_HordeAgent)&quot; -serverdir=&quot;$(StagingDir)/Server/$(ServerArch)&quot; -public" DotNetPath="$(DotNetPath)" />
-				<DotNet Arguments="run --project &quot;$(RootDir)/Engine/Source/Programs/Horde/Horde.Server/Horde.Server.csproj&quot; -- generate tooldata -id=horde-agent-win64 -name=&quot;Horde Agent (Windows)&quot; -description=&quot;Update package for older installations of Horde Agent through the MSI installer.&quot; -version=&quot;$(InformationalVersion)&quot; -inputdir=&quot;$(ServerToolsStagingDir_HordeAgentWin64)&quot; -serverdir=&quot;$(StagingDir)/Server/$(ServerArch)&quot; -public -showindashboard=false -category=Internal" DotNetPath="$(DotNetPath)" />
-				<DotNet Arguments="run --project &quot;$(RootDir)/Engine/Source/Programs/Horde/Horde.Server/Horde.Server.csproj&quot; -- generate tooldata -id=horde-agent-win64-installer -name=&quot;Horde Agent (Windows Installer)&quot; -description=&quot;Installer for the Horde Agent on Windows, with tray application.&quot; -version=&quot;$(InformationalVersion)&quot; -inputdir=&quot;$(ServerToolsStagingDir_HordeAgentMsi)&quot; -serverdir=&quot;$(StagingDir)/Server/$(ServerArch)&quot; -public" DotNetPath="$(DotNetPath)" />
-				<DotNet Arguments="run --project &quot;$(RootDir)/Engine/Source/Programs/Horde/Horde.Server/Horde.Server.csproj&quot; -- generate tooldata -id=p4v-utils -name=&quot;P4V Extensions&quot; -description=&quot;Update package for the P4VUtils tool.&quot; -version=&quot;$(InformationalVersion)&quot; -inputdir=&quot;$(ServerToolsStagingDir_P4VUtils)&quot; -serverdir=&quot;$(StagingDir)/Server/$(ServerArch)&quot; -public -showinugs -category=Internal" DotNetPath="$(DotNetPath)" />
-				<DotNet Arguments="run --project &quot;$(RootDir)/Engine/Source/Programs/Horde/Horde.Server/Horde.Server.csproj&quot; -- generate tooldata -id=ugs-cmd -name=&quot;UnrealGameSync (Command-Line)&quot; -description=&quot;Command-line version of UnrealGameSync. Supports Windows, MacOS and Linux.&quot; -version=&quot;$(InformationalVersion)&quot; -inputdir=&quot;$(ServerToolsStagingDir_UnrealGameSyncCmd)&quot; -serverdir=&quot;$(StagingDir)/Server/$(ServerArch)&quot; -public" DotNetPath="$(DotNetPath)" />
-				<DotNet Arguments="run --project &quot;$(RootDir)/Engine/Source/Programs/Horde/Horde.Server/Horde.Server.csproj&quot; -- generate tooldata -id=ugs-msi -name=&quot;UnrealGameSync (Windows Installer)&quot; -description=&quot;Installer for UnrealGameSync.&quot; -version=&quot;$(InformationalVersion)&quot; -inputdir=&quot;$(ServerToolsStagingDir_UnrealGameSyncMsi)&quot; -serverdir=&quot;$(StagingDir)/Server/$(ServerArch)&quot; -public" DotNetPath="$(DotNetPath)" />
-				<DotNet Arguments="run --project &quot;$(RootDir)/Engine/Source/Programs/Horde/Horde.Server/Horde.Server.csproj&quot; -- generate tooldata -id=ugs-win -name=&quot;UnrealGameSync (Windows Update Package)&quot; -description=&quot;Update package for UGS on Windows.&quot; -version=&quot;$(InformationalVersion)&quot; -inputdir=&quot;$(ServerToolsStagingDir_UnrealGameSync)&quot; -serverdir=&quot;$(StagingDir)/Server/$(ServerArch)&quot; -public -showindashboard=false -category=Internal" DotNetPath="$(DotNetPath)" />
+				<DotNet Arguments="run --project &quot;$(RootDir)/Engine/Source/Programs/Horde/HordeServer/HordeServer.csproj&quot; -- generate tooldata -id=horde-cmd -name=&quot;Horde (Command-Line)&quot; -description=&quot;Update package for the Horde command-line tool.&quot; -version=&quot;$(InformationalVersion)&quot; -inputdir=&quot;$(ServerToolsStagingDir_HordeCmd)&quot; -serverdir=&quot;$(StagingDir)/Server/$(ServerArch)&quot; -public -showinugs -Category=Internal" DotNetPath="$(DotNetPath)" />
+				<DotNet Arguments="run --project &quot;$(RootDir)/Engine/Source/Programs/Horde/HordeServer/HordeServer.csproj&quot; -- generate tooldata -id=horde-agent -name=&quot;Horde Agent (Cross-Platform)&quot; -description=&quot;Cross-platform Horde Agent distribution. Allows hosts to take on work from the Horde Server.&quot; -version=&quot;$(InformationalVersion)&quot; -inputdir=&quot;$(ServerToolsStagingDir_HordeAgent)&quot; -serverdir=&quot;$(StagingDir)/Server/$(ServerArch)&quot; -public" DotNetPath="$(DotNetPath)" />
+				<DotNet Arguments="run --project &quot;$(RootDir)/Engine/Source/Programs/Horde/HordeServer/HordeServer.csproj&quot; -- generate tooldata -id=horde-agent-win64 -name=&quot;Horde Agent (Windows)&quot; -description=&quot;Update package for older installations of Horde Agent through the MSI installer.&quot; -version=&quot;$(InformationalVersion)&quot; -inputdir=&quot;$(ServerToolsStagingDir_HordeAgentWin64)&quot; -serverdir=&quot;$(StagingDir)/Server/$(ServerArch)&quot; -public -showindashboard=false -category=Internal" DotNetPath="$(DotNetPath)" />
+				<DotNet Arguments="run --project &quot;$(RootDir)/Engine/Source/Programs/Horde/HordeServer/HordeServer.csproj&quot; -- generate tooldata -id=horde-agent-win64-installer -name=&quot;Horde Agent (Windows Installer)&quot; -description=&quot;Installer for the Horde Agent on Windows, with tray application.&quot; -version=&quot;$(InformationalVersion)&quot; -inputdir=&quot;$(ServerToolsStagingDir_HordeAgentMsi)&quot; -serverdir=&quot;$(StagingDir)/Server/$(ServerArch)&quot; -public" DotNetPath="$(DotNetPath)" />
+				<DotNet Arguments="run --project &quot;$(RootDir)/Engine/Source/Programs/Horde/HordeServer/HordeServer.csproj&quot; -- generate tooldata -id=p4v-utils -name=&quot;P4V Extensions&quot; -description=&quot;Update package for the P4VUtils tool.&quot; -version=&quot;$(InformationalVersion)&quot; -inputdir=&quot;$(ServerToolsStagingDir_P4VUtils)&quot; -serverdir=&quot;$(StagingDir)/Server/$(ServerArch)&quot; -public -showinugs -category=Internal" DotNetPath="$(DotNetPath)" />
+				<DotNet Arguments="run --project &quot;$(RootDir)/Engine/Source/Programs/Horde/HordeServer/HordeServer.csproj&quot; -- generate tooldata -id=ugs-cmd -name=&quot;UnrealGameSync (Command-Line)&quot; -description=&quot;Command-line version of UnrealGameSync. Supports Windows, MacOS and Linux.&quot; -version=&quot;$(InformationalVersion)&quot; -inputdir=&quot;$(ServerToolsStagingDir_UnrealGameSyncCmd)&quot; -serverdir=&quot;$(StagingDir)/Server/$(ServerArch)&quot; -public" DotNetPath="$(DotNetPath)" />
+				<DotNet Arguments="run --project &quot;$(RootDir)/Engine/Source/Programs/Horde/HordeServer/HordeServer.csproj&quot; -- generate tooldata -id=ugs-msi -name=&quot;UnrealGameSync (Windows Installer)&quot; -description=&quot;Installer for UnrealGameSync.&quot; -version=&quot;$(InformationalVersion)&quot; -inputdir=&quot;$(ServerToolsStagingDir_UnrealGameSyncMsi)&quot; -serverdir=&quot;$(StagingDir)/Server/$(ServerArch)&quot; -public" DotNetPath="$(DotNetPath)" />
+				<DotNet Arguments="run --project &quot;$(RootDir)/Engine/Source/Programs/Horde/HordeServer/HordeServer.csproj&quot; -- generate tooldata -id=ugs-win -name=&quot;UnrealGameSync (Windows Update Package)&quot; -description=&quot;Update package for UGS on Windows.&quot; -version=&quot;$(InformationalVersion)&quot; -inputdir=&quot;$(ServerToolsStagingDir_UnrealGameSync)&quot; -serverdir=&quot;$(StagingDir)/Server/$(ServerArch)&quot; -public -showindashboard=false -category=Internal" DotNetPath="$(DotNetPath)" />
 
 				<!-- Separate copy of the command line tool for installation with the server on Windows -->
 				<Delete Files="$(StagingDir)/Tools/..."/>
@@ -313,10 +313,10 @@
 		
 		<Node Name="Build Bundled Docker Image" Requires="Stage Horde Server Installer (win-x64);#ServerWithToolsFiles_win-x64">
 			<Tag Files="$(StagePaths)" Except=".../.vs/...;.../.git/...;.../bin/...;.../obj/..." With="#InputFiles"/>
-			<Docker-Build BaseDir="Engine" Files="#InputFiles" Tag="horde-server-bare" DockerFile="Engine/Source/Programs/Horde/Horde.Server/Dockerfile" Arguments="--build-arg msbuild_args=&quot;$(VersionArguments)&quot;"/>
+			<Docker-Build BaseDir="Engine" Files="#InputFiles" Tag="horde-server-bare" DockerFile="Engine/Source/Programs/Horde/HordeServer/Dockerfile" Arguments="--build-arg msbuild_args=&quot;$(VersionArguments)&quot;"/>
 
 			<!-- A Horde server image has already been built, and is referenced inside Dockerfile.bundled -->
-			<Copy From="Engine/Source/Programs/Horde/Horde.Server/Dockerfile.bundled" To="$(StagingDir)/Dockerfile.bundled" />
+			<Copy From="Engine/Source/Programs/Horde/HordeServer/Dockerfile.bundled" To="$(StagingDir)/Dockerfile.bundled" />
 			<Tag Files="$(StagingDir)/Dockerfile*;$(StagingDir)/Dashboard/...;$(StagingDir)/Server/Tools/...;$(StagingDir)/Server/*.json" With="#DockerFiles"/>
 			<Docker-Build BaseDir="$(StagingDir)" Files="#DockerFiles" Tag="horde-server:bundled" DockerFile="$(StagingDir)/Dockerfile.bundled" />
 		</Node>
diff --git a/Engine/Source/Programs/Horde/Drivers/JobDriver/JobDriver.csproj b/Engine/Source/Programs/Horde/Drivers/JobDriver/JobDriver.csproj
--- a/Engine/Source/Programs/Horde/Drivers/JobDriver/JobDriver.csproj
+++ b/Engine/Source/Programs/Horde/Drivers/JobDriver/JobDriver.csproj
@@ -37,16 +37,16 @@
     <Compile Include="..\..\..\UnrealBuildTool\Matchers\*.cs">
       <Link>Parser\Matchers\%(FileName)%(Extension)</Link>
     </Compile>
-    <Content Include="..\..\Horde.Server\Protos\horde\*.proto">
+    <Content Include="..\..\HordeServer\Protos\horde\*.proto">
       <Link>Protos\%(FileName)%(Extension)</Link>
     </Content>
-    <Content Include="..\..\Horde.Server\Protos\horde\common\**\*.proto">
+    <Content Include="..\..\HordeServer\Protos\horde\common\**\*.proto">
       <Link>Protos\horde\common\%(RecursiveDir)%(FileName)%(Extension)</Link>
     </Content>
-    <Content Include="..\..\Horde.Server\Protos\horde\messages\**\*.proto">
+    <Content Include="..\..\HordeServer\Protos\horde\messages\**\*.proto">
       <Link>Protos\horde\messages\%(RecursiveDir)%(FileName)%(Extension)</Link>
     </Content>
-    <Content Include="..\..\Horde.Server\Protos\horde\tasks\**\*.proto">
+    <Content Include="..\..\HordeServer\Protos\horde\tasks\**\*.proto">
       <Link>Protos\horde\tasks\%(RecursiveDir)%(FileName)%(Extension)</Link>
     </Content>
   </ItemGroup>
@@ -91,12 +91,12 @@
   </ItemGroup>
   
   <ItemGroup>
-    <Protobuf Include="..\..\Horde.Server\Protos\horde\job_rpc.proto" ProtoRoot="..\..\Horde.Server\Protos" GrpcServices="Client" Access="Public" AdditionalImportDirs="..\..\..\Shared\EpicGames.Horde\Protos"/>
-    <Protobuf Include="..\..\Horde.Server\Protos\horde\job_rpc_messages.proto" ProtoRoot="..\..\Horde.Server\Protos" Access="Public" />
-    <Protobuf Include="..\..\Horde.Server\Protos\horde\common\label_change.proto" ProtoRoot="..\..\Horde.Server\Protos" Access="Public" />
-    <Protobuf Include="..\..\Horde.Server\Protos\horde\messages\agent_workspace.proto" ProtoRoot="..\..\Horde.Server\Protos" Access="Public" />
-    <Protobuf Include="..\..\Horde.Server\Protos\horde\tasks\conform_task.proto" ProtoRoot="..\..\Horde.Server\Protos" Access="Public" />
-    <Protobuf Include="..\..\Horde.Server\Protos\horde\tasks\job_task.proto" ProtoRoot="..\..\Horde.Server\Protos" Access="Public" />
+    <Protobuf Include="..\..\HordeServer\Protos\horde\job_rpc.proto" ProtoRoot="..\..\HordeServer\Protos" GrpcServices="Client" Access="Public" AdditionalImportDirs="..\..\..\Shared\EpicGames.Horde\Protos"/>
+    <Protobuf Include="..\..\HordeServer\Protos\horde\job_rpc_messages.proto" ProtoRoot="..\..\HordeServer\Protos" Access="Public" />
+    <Protobuf Include="..\..\HordeServer\Protos\horde\common\label_change.proto" ProtoRoot="..\..\HordeServer\Protos" Access="Public" />
+    <Protobuf Include="..\..\HordeServer\Protos\horde\messages\agent_workspace.proto" ProtoRoot="..\..\HordeServer\Protos" Access="Public" />
+    <Protobuf Include="..\..\HordeServer\Protos\horde\tasks\conform_task.proto" ProtoRoot="..\..\HordeServer\Protos" Access="Public" />
+    <Protobuf Include="..\..\HordeServer\Protos\horde\tasks\job_task.proto" ProtoRoot="..\..\HordeServer\Protos" Access="Public" />
   </ItemGroup>
 
   <ItemGroup>
diff --git a/Engine/Source/Programs/Horde/Horde.Agent.TrayApp/Horde.Agent.TrayApp.csproj b/Engine/Source/Programs/Horde/Horde.Agent.TrayApp/Horde.Agent.TrayApp.csproj
--- a/Engine/Source/Programs/Horde/Horde.Agent.TrayApp/Horde.Agent.TrayApp.csproj
+++ b/Engine/Source/Programs/Horde/Horde.Agent.TrayApp/Horde.Agent.TrayApp.csproj
@@ -19,7 +19,7 @@
   </PropertyGroup>
 
   <ItemGroup>
-    <Compile Include="..\Horde.Agent\IpcMessages.cs" Link="IpcMessages.cs" />
+    <Compile Include="..\HordeAgent\IpcMessages.cs" Link="IpcMessages.cs" />
   </ItemGroup>
 
   <ItemGroup>
diff --git a/Engine/Source/Programs/Horde/HordeAgent.Tests/HordeAgent.Tests.csproj b/Engine/Source/Programs/Horde/HordeAgent.Tests/HordeAgent.Tests.csproj
new file mode 100644
--- /dev/null
+++ b/Engine/Source/Programs/Horde/HordeAgent.Tests/HordeAgent.Tests.csproj
@@ -0,0 +1,43 @@
+<Project Sdk="Microsoft.NET.Sdk">
+
+  <PropertyGroup>
+    <TargetFramework>net8.0</TargetFramework>
+    <Nullable>enable</Nullable>
+    <IsPackable>false</IsPackable>
+    <DebugType>full</DebugType>
+    <DebugType>pdbonly</DebugType>
+    <Configurations>Debug;Release;Development;Analyze</Configurations>
+    <NoWarn>1701;1702;1591</NoWarn>
+  </PropertyGroup>
+
+  <PropertyGroup Condition=" '$(Configuration)|$(Platform)' == 'Analyze|AnyCPU' ">
+    <RunAnalyzersDuringBuild>True</RunAnalyzersDuringBuild>
+    <TreatWarningsAsErrors>false</TreatWarningsAsErrors>
+    <EnforceCodeStyleInBuild>True</EnforceCodeStyleInBuild>
+    <GenerateDocumentationFile>True</GenerateDocumentationFile>
+    <DocumentationFile></DocumentationFile>
+  </PropertyGroup>
+
+  <ItemGroup>
+    <PackageReference Include="Microsoft.Extensions.Logging.Abstractions" Version="8.0.0" />
+    <PackageReference Include="Microsoft.NET.Test.Sdk" Version="17.8.0" />
+    <PackageReference Include="Moq" Version="4.20.70" />
+    <PackageReference Include="MSTest.TestAdapter" Version="3.1.1" />
+    <PackageReference Include="MSTest.TestFramework" Version="3.1.1" />
+    <PackageReference Include="coverlet.collector" Version="6.0.0">
+      <PrivateAssets>all</PrivateAssets>
+      <IncludeAssets>runtime; build; native; contentfiles; analyzers; buildtransitive</IncludeAssets>
+    </PackageReference>
+  </ItemGroup>
+
+  <ItemGroup>
+    <ProjectReference Include="..\HordeAgent\HordeAgent.csproj" />
+  </ItemGroup>
+
+  <ItemGroup>
+    <Content Include="remote-exec-test-bin\*.exe">
+      <CopyToOutputDirectory>PreserveNewest</CopyToOutputDirectory>
+    </Content>
+  </ItemGroup>
+
+</Project>
diff --git a/Engine/Source/Programs/Horde/HordeAgent/HordeAgent.csproj b/Engine/Source/Programs/Horde/HordeAgent/HordeAgent.csproj
new file mode 100644
--- /dev/null
+++ b/Engine/Source/Programs/Horde/HordeAgent/HordeAgent.csproj
@@ -0,0 +1,210 @@
+﻿<Project Sdk="Microsoft.NET.Sdk">
+
+  <PropertyGroup>
+    <OutputType>Exe</OutputType>
+    <TargetFramework>net6.0</TargetFramework>
+    <Nullable>enable</Nullable>
+    <ImplicitUsings>enable</ImplicitUsings>
+	<ApplicationIcon>HordeAgent.ico</ApplicationIcon>
+    <ApplicationManifest>app.manifest</ApplicationManifest>
+    <AssemblyName>HordeAgent</AssemblyName>
+    <PublishRepositoryUrl>true</PublishRepositoryUrl>
+    <EmbedUntrackedSources>false</EmbedUntrackedSources>
+    <DebugType>pdbonly</DebugType>
+    <Configurations>Debug;Release;Analyze</Configurations>
+    <NoWarn>1701;1702;NU3003</NoWarn>
+  </PropertyGroup>
+
+  <PropertyGroup Condition="'$(Configuration)|$(Platform)'=='Debug|AnyCPU'">
+    <DocumentationFile>HordeAgent.xml</DocumentationFile>
+    <TreatWarningsAsErrors>False</TreatWarningsAsErrors>
+  </PropertyGroup>
+
+  <PropertyGroup Condition=" '$(Configuration)|$(Platform)' == 'Analyze|AnyCPU' ">
+    <RunAnalyzersDuringBuild>True</RunAnalyzersDuringBuild>
+    <TreatWarningsAsErrors>false</TreatWarningsAsErrors>
+    <EnforceCodeStyleInBuild>True</EnforceCodeStyleInBuild>
+    <GenerateDocumentationFile>True</GenerateDocumentationFile>
+    <DocumentationFile></DocumentationFile>
+  </PropertyGroup>
+
+  <ItemGroup>
+    <None Remove="HordeAgent.xml" />
+  </ItemGroup>
+
+  <ItemGroup>
+    <PackageReference Include="AWSSDK.Core" Version="3.7.300.12" />
+    <PackageReference Include="AWSSDK.EC2" Version="3.7.306.6" />
+    <PackageReference Include="AWSSDK.SimpleSystemsManagement" Version="3.7.301.8" />
+    <PackageReference Include="Datadog.Trace.OpenTracing" Version="2.14.0" />
+    <PackageReference Include="Google.Protobuf" Version="3.25.1" />
+    <PackageReference Include="Grpc.Net.Client" Version="2.59.0" />
+    <PackageReference Include="Grpc.Tools" Version="2.59.0">
+      <PrivateAssets>all</PrivateAssets>
+      <IncludeAssets>runtime; build; native; contentfiles; analyzers; buildtransitive</IncludeAssets>
+    </PackageReference>
+    <PackageReference Include="Meziantou.Analyzer" Version="2.0.139">
+      <PrivateAssets>all</PrivateAssets>
+      <IncludeAssets>runtime; build; native; contentfiles; analyzers; buildtransitive</IncludeAssets>
+    </PackageReference>
+    <PackageReference Include="Microsoft.Extensions.Caching.Memory" Version="8.0.0" />
+    <PackageReference Include="Microsoft.Extensions.Hosting" Version="8.0.0" />
+    <PackageReference Include="Microsoft.Extensions.Hosting.Abstractions" Version="8.0.0" />
+    <PackageReference Include="Microsoft.Extensions.Hosting.WindowsServices" Version="8.0.0" />
+    <PackageReference Include="Microsoft.Extensions.Http" Version="8.0.0" />
+    <PackageReference Include="Microsoft.Extensions.Http.Polly" Version="8.0.0" />
+    <PackageReference Include="Microsoft.Extensions.Logging.Abstractions" Version="8.0.0" />
+    <PackageReference Include="Microsoft.Extensions.Options.DataAnnotations" Version="8.0.0" />
+    <PackageReference Include="Microsoft.Management.Infrastructure" Version="3.0.0" />
+    <PackageReference Include="Microsoft.VisualStudio.Threading.Analyzers" Version="17.8.14">
+      <PrivateAssets>all</PrivateAssets>
+      <IncludeAssets>runtime; build; native; contentfiles; analyzers; buildtransitive</IncludeAssets>
+    </PackageReference>
+    <PackageReference Include="OpenTracing" Version="0.12.1" />
+    <PackageReference Include="Serilog" Version="3.1.1" />
+    <PackageReference Include="Serilog.Extensions.Hosting" Version="8.0.0" />
+    <PackageReference Include="Serilog.Extensions.Logging" Version="8.0.0" />
+    <PackageReference Include="Serilog.Settings.Configuration" Version="8.0.0" />
+    <PackageReference Include="Serilog.Sinks.Console" Version="5.0.1" />
+    <PackageReference Include="Serilog.Sinks.File" Version="5.0.0" />
+    <PackageReference Include="System.Diagnostics.PerformanceCounter" Version="8.0.0" />
+    <PackageReference Include="System.ServiceProcess.ServiceController" Version="8.0.0" />
+  </ItemGroup>
+
+  <ItemGroup>
+    <Content Include="appsettings.json" Condition="Exists('appsettings.json')">
+      <Generator></Generator>
+      <CopyToOutputDirectory>Always</CopyToOutputDirectory>
+      <CopyToPublishDirectory>Always</CopyToPublishDirectory>
+    </Content>
+    <Content Include="appsettings.Development.json" Condition="Exists('appsettings.Development.json')">
+      <Generator></Generator>
+      <CopyToOutputDirectory>Always</CopyToOutputDirectory>
+      <CopyToPublishDirectory>Always</CopyToPublishDirectory>
+      <DependentUpon>appsettings.json</DependentUpon>
+    </Content>
+    <Content Include="appsettings.Local.json" Condition="Exists('appsettings.Local.json')">
+      <Generator></Generator>
+      <CopyToOutputDirectory>Always</CopyToOutputDirectory>
+      <CopyToPublishDirectory>Always</CopyToPublishDirectory>
+      <DependentUpon>appsettings.json</DependentUpon>
+    </Content>
+    <Content Include="appsettings.Production.json" Condition="Exists('appsettings.Production.json')">
+      <Generator></Generator>
+      <CopyToOutputDirectory>Always</CopyToOutputDirectory>
+      <CopyToPublishDirectory>Always</CopyToPublishDirectory>
+      <DependentUpon>appsettings.json</DependentUpon>
+    </Content>
+    <Content Include="appsettings.User.json" Condition="Exists('appsettings.User.json')">
+      <Generator></Generator>
+      <CopyToOutputDirectory>Always</CopyToOutputDirectory>
+      <CopyToPublishDirectory>Always</CopyToPublishDirectory>
+      <DependentUpon>appsettings.json</DependentUpon>
+    </Content>
+    <Content Include="Defaults\agent.json">
+      <Generator></Generator>
+      <CopyToOutputDirectory>Always</CopyToOutputDirectory>
+      <CopyToPublishDirectory>Always</CopyToPublishDirectory>
+    </Content>
+  </ItemGroup>
+
+  <ItemGroup>
+    <Content Include="..\HordeServer\Protos\horde\*.proto">
+      <Link>Protos\%(FileName)%(Extension)</Link>
+    </Content>
+    <Content Include="..\HordeServer\Protos\horde\common\**\*.proto">
+      <Link>Protos\horde\common\%(RecursiveDir)%(FileName)%(Extension)</Link>
+    </Content>
+    <Content Include="..\HordeServer\Protos\horde\messages\**\*.proto">
+      <Link>Protos\horde\messages\%(RecursiveDir)%(FileName)%(Extension)</Link>
+    </Content>
+    <Content Include="..\HordeServer\Protos\horde\tasks\**\*.proto">
+      <Link>Protos\horde\tasks\%(RecursiveDir)%(FileName)%(Extension)</Link>
+    </Content>
+<!--    
+    <Compile Include="..\HordeServer\Protos\horde\**\*.cs">
+      <Link>Protos\horde\%(RecursiveDir)%(FileName)%(Extension)</Link>
+    </Compile>
+    -->
+    <Compile Include="$(IntermediateOutputPath)VersionInfo.cs" />
+  </ItemGroup>
+
+  <ItemGroup>
+    <ProjectReference Include="..\..\Shared\EpicGames.Core\EpicGames.Core.csproj" />
+    <ProjectReference Include="..\..\Shared\EpicGames.Horde\EpicGames.Horde.csproj" />
+    <ProjectReference Include="..\..\Shared\EpicGames.Perforce.Managed\EpicGames.Perforce.Managed.csproj" />
+    <ProjectReference Include="..\..\Shared\EpicGames.Perforce\EpicGames.Perforce.csproj" />
+  </ItemGroup>
+
+  <ItemGroup>
+    <Protobuf Include="..\HordeServer\Protos\horde\horde_rpc.proto" ProtoRoot="..\HordeServer\Protos" GrpcServices="Client" Access="Public" />
+    <Protobuf Include="..\HordeServer\Protos\horde\horde_rpc_messages.proto" ProtoRoot="..\HordeServer\Protos" Access="Public" />
+    <Protobuf Include="..\HordeServer\Protos\horde\enrollment_rpc.proto" ProtoRoot="..\HordeServer\Protos" GrpcServices="Client" Access="Public" />
+    <Protobuf Include="..\HordeServer\Protos\horde\enrollment_rpc_messages.proto" ProtoRoot="..\HordeServer\Protos" Access="Public" />
+    <Protobuf Include="..\HordeServer\Protos\horde\common\agent_status.proto" ProtoRoot="..\HordeServer\Protos" Access="Public" />
+    <Protobuf Include="..\HordeServer\Protos\horde\common\io_hash.proto" ProtoRoot="..\HordeServer\Protos" Access="Public" />
+    <Protobuf Include="..\HordeServer\Protos\horde\common\lease_outcome.proto" ProtoRoot="..\HordeServer\Protos" Access="Public" />
+    <Protobuf Include="..\HordeServer\Protos\horde\common\lease_state.proto" ProtoRoot="..\HordeServer\Protos" Access="Public" />
+    <Protobuf Include="..\HordeServer\Protos\horde\common\namespace_id.proto" ProtoRoot="..\HordeServer\Protos" Access="Public" />
+    <Protobuf Include="..\HordeServer\Protos\horde\messages\agent_capabilities.proto" ProtoRoot="..\HordeServer\Protos" Access="Public" />
+    <Protobuf Include="..\HordeServer\Protos\horde\messages\agent_requirements.proto" ProtoRoot="..\HordeServer\Protos" Access="Public" />
+    <Protobuf Include="..\HordeServer\Protos\horde\messages\lease.proto" ProtoRoot="..\HordeServer\Protos" Access="Public" />
+    <Protobuf Include="..\HordeServer\Protos\horde\messages\aws_lambda.proto" ProtoRoot="..\HordeServer\Protos" Access="Public" />
+    <Protobuf Include="..\HordeServer\Protos\horde\messages\telemetry.proto" ProtoRoot="..\HordeServer\Protos" Access="Public" />
+    <Protobuf Include="..\HordeServer\Protos\horde\tasks\compute_task.proto" ProtoRoot="..\HordeServer\Protos" Access="Public" />
+    <Protobuf Include="..\HordeServer\Protos\horde\tasks\restart_task.proto" ProtoRoot="..\HordeServer\Protos" Access="Public" />
+    <Protobuf Include="..\HordeServer\Protos\horde\tasks\shutdown_task.proto" ProtoRoot="..\HordeServer\Protos" Access="Public" />
+    <Protobuf Include="..\HordeServer\Protos\horde\tasks\test_task.proto" ProtoRoot="..\HordeServer\Protos" Access="Public" />
+    <Protobuf Include="..\HordeServer\Protos\horde\tasks\upgrade_task.proto" ProtoRoot="..\HordeServer\Protos" Access="Public" />
+
+    <Protobuf Include="..\HordeServer\Protos\horde\job_rpc_messages.proto" ProtoRoot="..\HordeServer\Protos" Access="Public" />
+    <Protobuf Include="..\HordeServer\Protos\horde\common\label_change.proto" ProtoRoot="..\HordeServer\Protos" Access="Public" />
+    <Protobuf Include="..\HordeServer\Protos\horde\messages\agent_workspace.proto" ProtoRoot="..\HordeServer\Protos" Access="Public" />
+    <Protobuf Include="..\HordeServer\Protos\horde\tasks\conform_task.proto" ProtoRoot="..\HordeServer\Protos" Access="Public" />
+    <Protobuf Include="..\HordeServer\Protos\horde\tasks\job_task.proto" ProtoRoot="..\HordeServer\Protos" Access="Public" />
+  </ItemGroup>
+
+  <ItemGroup>
+    <Compile Update="Properties\Resources.Designer.cs">
+      <DesignTime>True</DesignTime>
+      <AutoGen>True</AutoGen>
+      <DependentUpon>Resources.resx</DependentUpon>
+    </Compile>
+  </ItemGroup>
+
+  <ItemGroup>
+    <EmbeddedResource Update="Properties\Resources.resx">
+      <Generator>ResXFileCodeGenerator</Generator>
+      <LastGenOutput>Resources.Designer.cs</LastGenOutput>
+    </EmbeddedResource>
+  </ItemGroup>
+
+  <ItemGroup>
+    <Content Update="appsettings.Local.json" CopyToPublishDirectory="Never" />
+  </ItemGroup>
+
+  <Target Name="PublishSymbols" AfterTargets="AfterBuild" Condition="'$(FullSymbolsDir)' != ''">
+    <ItemGroup>
+      <OutputFiles Include="$(OutDir)**" />
+      <BinaryFilesBase Include="$(OutDir)\*.dll" />
+      <BinaryFilesBase Include="$(OutDir)\*.pdb" />
+      <BinaryFiles Include="@(BinaryFilesBase)" />
+    </ItemGroup>
+
+    <MakeDir Directories="$(FullSymbolsDir)" />
+    <Copy SourceFiles="@(Files)" DestinationFolder="$(FullSymbolsDir)" />
+  </Target>
+
+  <!-- Encode version in a source file as inferring version from assemblies is unreliable on platform-specific binaries -->
+  <Target Name="PrepareVersionInfo" BeforeTargets="BeforeBuild">
+    <PropertyGroup>
+      <TemplateFile>Utility/VersionInfo.cs.template</TemplateFile>
+      <OutputFile>$(IntermediateOutputPath)/VersionInfo.cs</OutputFile>
+    </PropertyGroup>
+    <ReadLinesFromFile File="$(TemplateFile)">
+      <Output TaskParameter="Lines" ItemName="Lines" />
+    </ReadLinesFromFile>
+    <WriteLinesToFile File="$(OutputFile)" Lines="@(Lines-&gt;Replace('$VERSION$', '$(InformationalVersion)'))" Overwrite="true" Encoding="UTF-8" />
+  </Target>
+  
+</Project>
diff --git a/Engine/Source/Programs/Horde/HordeAgent/Logging.cs b/Engine/Source/Programs/Horde/HordeAgent/Logging.cs
new file mode 100644
--- /dev/null
+++ b/Engine/Source/Programs/Horde/HordeAgent/Logging.cs
@@ -0,0 +1,90 @@
+// Copyright Epic Games, Inc. All Rights Reserved.
+
+using System.Runtime.InteropServices;
+using EpicGames.Core;
+using Microsoft.Extensions.Configuration;
+using Microsoft.Extensions.Logging;
+using OpenTracing;
+using OpenTracing.Util;
+using Serilog;
+using Serilog.Core;
+using Serilog.Extensions.Logging;
+using Serilog.Formatting.Json;
+using Serilog.Sinks.SystemConsole.Themes;
+
+namespace Horde.Agent
+{
+	static class Logging
+	{
+		static string s_env = "default";
+
+		public static void SetEnv(string newEnv)
+		{
+			s_env = newEnv;
+		}
+
+		public static LoggingLevelSwitch LogLevelSwitch = new LoggingLevelSwitch();
+
+		private class DatadogLogEnricher : ILogEventEnricher
+		{
+			public void Enrich(Serilog.Events.LogEvent logEvent, ILogEventPropertyFactory propertyFactory)
+			{
+				logEvent.AddOrUpdateProperty(propertyFactory.CreateProperty("dd.env", s_env));
+				logEvent.AddOrUpdateProperty(propertyFactory.CreateProperty("dd.service", "hordeagent"));
+				logEvent.AddOrUpdateProperty(propertyFactory.CreateProperty("dd.version", AgentApp.Version));
+
+				ISpan? span = GlobalTracer.Instance?.ActiveSpan;
+				if (span != null)
+				{
+					logEvent.AddPropertyIfAbsent(propertyFactory.CreateProperty("dd.trace_id", span.Context.TraceId));
+					logEvent.AddPropertyIfAbsent(propertyFactory.CreateProperty("dd.span_id", span.Context.SpanId));
+				}
+			}
+		}
+
+		public static ILoggerFactory CreateLoggerFactory(IConfiguration configuration)
+		{
+			Serilog.ILogger logger = CreateSerilogLogger(configuration);
+			return new SerilogLoggerFactory(logger, true);
+		}
+
+		static Serilog.ILogger CreateSerilogLogger(IConfiguration configuration)
+		{
+			AgentSettings settings = AgentApp.BindSettings(configuration);
+			DirectoryReference.CreateDirectory(settings.LogsDir);
+
+			ConsoleTheme theme;
+			if (RuntimeInformation.IsOSPlatform(OSPlatform.Windows) && Environment.OSVersion.Version < new Version(10, 0))
+			{
+				theme = SystemConsoleTheme.Literate;
+			}
+			else
+			{
+				theme = AnsiConsoleTheme.Code;
+			}
+
+			return new LoggerConfiguration()
+				.WriteTo.Console(outputTemplate: "[{Timestamp:HH:mm:ss} {Level:w3}] {Indent}{Message:l}{NewLine}{Exception}", theme: theme)
+				.WriteTo.File(FileReference.Combine(settings.LogsDir, "Log-.txt").FullName, fileSizeLimitBytes: 50 * 1024 * 1024, rollingInterval: RollingInterval.Day, rollOnFileSizeLimit: true, retainedFileCountLimit: 10)
+				.WriteTo.File(new JsonFormatter(renderMessage: true), FileReference.Combine(settings.LogsDir, "Log-.json").FullName, fileSizeLimitBytes: 50 * 1024 * 1024, rollingInterval: RollingInterval.Day, rollOnFileSizeLimit: true, retainedFileCountLimit: 10)
+				.ReadFrom.Configuration(configuration)
+				.MinimumLevel.ControlledBy(LogLevelSwitch)
+				.Enrich.FromLogContext()
+				.Enrich.With<DatadogLogEnricher>()
+				.CreateLogger();
+		}
+
+		public static ILoggerProvider CreateFileLoggerProvider(DirectoryReference baseDir, string name)
+		{
+			DirectoryReference.CreateDirectory(baseDir);
+
+			Serilog.Core.Logger logger = new LoggerConfiguration()
+				.WriteTo.File(FileReference.Combine(baseDir, $"{name}.txt").FullName)
+				.WriteTo.File(new JsonFormatter(renderMessage: true), FileReference.Combine(baseDir, $"{name}.json").FullName)
+				.Enrich.FromLogContext()
+				.CreateLogger();
+
+			return new SerilogLoggerProvider(logger, true);
+		}
+	}
+}
diff --git a/Engine/Source/Programs/Horde/HordeAgent/Tracing.cs b/Engine/Source/Programs/Horde/HordeAgent/Tracing.cs
new file mode 100644
--- /dev/null
+++ b/Engine/Source/Programs/Horde/HordeAgent/Tracing.cs
@@ -0,0 +1,27 @@
+// Copyright Epic Games, Inc. All Rights Reserved.
+
+using OpenTracing;
+
+namespace Horde.Agent
+{
+	static class Tracing
+	{
+		public static ISpanBuilder WithResourceName(this ISpanBuilder builder, string? resourceName)
+		{
+			if (resourceName != null)
+			{
+				builder = builder.WithTag(Datadog.Trace.OpenTracing.DatadogTags.ResourceName, resourceName);
+			}
+			return builder;
+		}
+
+		public static ISpanBuilder WithServiceName(this ISpanBuilder builder, string? serviceName)
+		{
+			if (serviceName != null)
+			{
+				builder = builder.WithTag(Datadog.Trace.OpenTracing.DatadogTags.ServiceName, serviceName);
+			}
+			return builder;
+		}
+	}
+}
diff --git a/Engine/Source/Programs/Horde/HordeCmd/HordeCmd.csproj b/Engine/Source/Programs/Horde/HordeCmd/HordeCmd.csproj
new file mode 100644
--- /dev/null
+++ b/Engine/Source/Programs/Horde/HordeCmd/HordeCmd.csproj
@@ -0,0 +1,47 @@
+﻿<Project Sdk="Microsoft.NET.Sdk">
+
+  <PropertyGroup>
+    <OutputType>Exe</OutputType>
+    <TargetFramework>net8.0</TargetFramework>
+    <ImplicitUsings>enable</ImplicitUsings>
+    <Nullable>enable</Nullable>
+    <Configurations>Debug;Release;Analyze</Configurations>
+  </PropertyGroup>
+
+  <PropertyGroup Condition=" '$(Configuration)|$(Platform)' == 'Analyze|AnyCPU' ">
+    <RunAnalyzersDuringBuild>True</RunAnalyzersDuringBuild>
+    <TreatWarningsAsErrors>False</TreatWarningsAsErrors>
+    <EnforceCodeStyleInBuild>True</EnforceCodeStyleInBuild>
+    <GenerateDocumentationFile>True</GenerateDocumentationFile>
+    <DocumentationFile></DocumentationFile>
+  </PropertyGroup>
+
+  <ItemGroup>
+    <ProjectReference Include="..\..\Shared\EpicGames.Horde\EpicGames.Horde.csproj" />
+    <ProjectReference Include="..\..\Shared\EpicGames.OIDC\EpicGames.OIDC.csproj" />
+    <ProjectReference Include="..\..\Shared\EpicGames.Perforce.Managed\EpicGames.Perforce.Managed.csproj" />
+    <ProjectReference Include="..\..\Shared\EpicGames.Perforce\EpicGames.Perforce.csproj" />
+  </ItemGroup>
+
+  <ItemGroup>
+    <PackageReference Include="Microsoft.Extensions.Configuration.EnvironmentVariables" Version="8.0.0" />
+    <PackageReference Include="Serilog" Version="3.1.1" />
+    <PackageReference Include="Serilog.Extensions.Hosting" Version="8.0.0" />
+    <PackageReference Include="Serilog.Extensions.Logging" Version="8.0.0" />
+    <PackageReference Include="Serilog.Settings.Configuration" Version="8.0.0" />
+    <PackageReference Include="Serilog.Sinks.Console" Version="5.0.1" />
+    <PackageReference Include="Serilog.Sinks.File" Version="5.0.0" />
+  </ItemGroup>
+
+  <ItemGroup>
+    <Content Include="appsettings.json" Condition="Exists('appsettings.json')">
+      <Generator></Generator>
+      <CopyToOutputDirectory>Always</CopyToOutputDirectory>
+      <CopyToPublishDirectory>Always</CopyToPublishDirectory>
+    </Content>
+    <Content Include="UgsTool.json">
+      <CopyToOutputDirectory>Always</CopyToOutputDirectory>
+    </Content>
+  </ItemGroup>
+
+</Project>
diff --git a/Engine/Source/Programs/Horde/HordeServer.Tests/Agents/Relay/AgentRelayClientTest.cs b/Engine/Source/Programs/Horde/HordeServer.Tests/Agents/Relay/AgentRelayClientTest.cs
new file mode 100644
--- /dev/null
+++ b/Engine/Source/Programs/Horde/HordeServer.Tests/Agents/Relay/AgentRelayClientTest.cs
@@ -0,0 +1,82 @@
+// Copyright Epic Games, Inc. All Rights Reserved.
+
+using System;
+using System.Collections.Generic;
+using System.Threading;
+using System.Threading.Tasks;
+using Grpc.Core;
+using Horde.Common.Rpc;
+using Horde.Server.Agents.Relay;
+using Horde.Server.Tests.Server;
+using Horde.Server.Utilities;
+using Microsoft.Extensions.Logging.Abstractions;
+using Microsoft.VisualStudio.TestTools.UnitTesting;
+
+namespace Horde.Server.Tests.Agents.Relay;
+
+public class TestRelayRpcClient : RelayRpc.RelayRpcClient
+{
+	private static readonly ServerCallContext s_adminContext = new ServerCallContextStub(HordeClaims.AdminClaim.ToClaim());
+	public TestAsyncStreamReader<GetPortMappingsResponse> GetPortMappingsResponses { get; } = new(s_adminContext);
+
+	public override AsyncServerStreamingCall<GetPortMappingsResponse> GetPortMappings(
+		GetPortMappingsRequest request, Metadata headers = null!,
+		DateTime? deadline = null, CancellationToken cancellationToken = default)
+	{
+		return new AsyncServerStreamingCall<GetPortMappingsResponse>(
+			GetPortMappingsResponses, null!, null!, null!, () => { });
+	}
+}
+
+[TestClass]
+public class AgentRelayClientTest
+{
+	private readonly TestRelayRpcClient _grpcClient = new();
+	private readonly AgentRelayClient _relayClient;
+
+	public AgentRelayClientTest()
+	{
+		_relayClient = new AgentRelayClient("myCluster", "myAgent", new List<string> { "192.168.1.99" }, Nftables.CreateNull(), _grpcClient, NullLogger<AgentRelayClient>.Instance);
+		_relayClient.CooldownOnException = TimeSpan.Zero;
+	}
+
+	[TestMethod]
+	public async Task GetPortMappingsLongPoll_Results_Async()
+	{
+		using CancellationTokenSource cts = new(3000);
+		_grpcClient.GetPortMappingsResponses.AddMessage(new GetPortMappingsResponse { PortMappings = { NftablesTests.LeaseMap1 }, RevisionCount = 1 });
+		_grpcClient.GetPortMappingsResponses.AddMessage(new GetPortMappingsResponse { PortMappings = { NftablesTests.LeaseMap2, NftablesTests.LeaseMap1 }, RevisionCount = 2 });
+
+		List<PortMapping>? pm1 = await _relayClient.GetPortMappingsLongPollAsync(cts.Token);
+		Assert.AreEqual(1, pm1!.Count);
+		Assert.AreEqual("lease1", pm1[0].LeaseId);
+		Assert.AreEqual(1, _relayClient.RevisionNumber);
+
+		List<PortMapping>? pm2 = await _relayClient.GetPortMappingsLongPollAsync(cts.Token);
+		Assert.AreEqual(2, pm2!.Count);
+		Assert.AreEqual("lease2", pm2[0].LeaseId);
+		Assert.AreEqual("lease1", pm2[1].LeaseId);
+		Assert.AreEqual(2, _relayClient.RevisionNumber);
+	}
+
+	[TestMethod]
+	public async Task GetPortMappingsLongPoll_TimesOut_Async()
+	{
+		using CancellationTokenSource cts = new(3000);
+		await Assert.ThrowsExceptionAsync<OperationCanceledException>(() => _relayClient.GetPortMappingsLongPollAsync(cts.Token));
+	}
+
+	[TestMethod]
+	public async Task ListenForPortMappings_TimesOut_Async()
+	{
+		using CancellationTokenSource cts = new(3000);
+
+		Task t = _relayClient.ListenForPortMappingsAsync(cts.Token);
+		await Task.Delay(100, cts.Token);
+		_grpcClient.GetPortMappingsResponses.AddMessage(new GetPortMappingsResponse { PortMappings = { NftablesTests.LeaseMap1 } });
+		await Task.Delay(100, cts.Token);
+		_grpcClient.GetPortMappingsResponses.AddMessage(new GetPortMappingsResponse { PortMappings = { NftablesTests.LeaseMap2 } });
+
+		await t;
+	}
+}
diff --git a/Engine/Source/Programs/Horde/HordeServer.Tests/Fixture.cs b/Engine/Source/Programs/Horde/HordeServer.Tests/Fixture.cs
new file mode 100644
--- /dev/null
+++ b/Engine/Source/Programs/Horde/HordeServer.Tests/Fixture.cs
@@ -0,0 +1,137 @@
+// Copyright Epic Games, Inc. All Rights Reserved.
+
+using System.Collections.Generic;
+using System.IO;
+using System.Text;
+using System.Threading.Tasks;
+using EpicGames.Core;
+using EpicGames.Horde.Agents.Pools;
+using EpicGames.Horde.Jobs;
+using EpicGames.Horde.Jobs.Templates;
+using EpicGames.Horde.Projects;
+using EpicGames.Horde.Streams;
+using Horde.Server.Agents;
+using Horde.Server.Configuration;
+using Horde.Server.Jobs;
+using Horde.Server.Jobs.Artifacts;
+using Horde.Server.Jobs.Graphs;
+using Horde.Server.Jobs.Templates;
+using Horde.Server.Projects;
+using Horde.Server.Server;
+using Horde.Server.Streams;
+
+namespace Horde.Server.Tests
+{
+	public class Fixture
+	{
+		public IJob Job1 { get; private set; } = null!;
+		public IJob Job2 { get; private set; } = null!;
+		public ITemplate Template { get; private set; } = null!;
+		public IGraph Graph { get; private set; } = null!;
+		public StreamId StreamId { get; private set; }
+		public StreamConfig? StreamConfig { get; private set; }
+		public TemplateId TemplateRefId1 { get; private set; }
+		public TemplateId TemplateRefId2 { get; private set; }
+		public IArtifactV1 Job1Artifact { get; private set; } = null!;
+		public string Job1ArtifactData { get; private set; } = null!;
+		public IAgent Agent1 { get; private set; } = null!;
+		public string Agent1Name { get; private set; } = null!;
+		public const string PoolName = "TestingPool";
+
+		public static async Task<Fixture> CreateAsync(ConfigService configService, IGraphCollection graphCollection, ITemplateCollection templateCollection, JobService jobService, IArtifactCollectionV1 artifactCollection, AgentService agentService, ServerSettings serverSettings)
+		{
+			Fixture fixture = new Fixture();
+			await fixture.PopulateAsync(configService, graphCollection, templateCollection, jobService, artifactCollection, agentService, serverSettings);
+
+			//			(PerforceService as PerforceServiceStub)?.AddChange("//UE5/Main", 112233, "leet.coder", "Did stuff", new []{"file.cpp"});
+			//			(PerforceService as PerforceServiceStub)?.AddChange("//UE5/Main", 1111, "swarm", "A shelved CL here", new []{"renderer.cpp"});
+
+			return fixture;
+		}
+
+		private async Task PopulateAsync(ConfigService configService, IGraphCollection graphCollection, ITemplateCollection templateCollection, JobService jobService, IArtifactCollectionV1 artifactCollection, AgentService agentService, ServerSettings serverSettings)
+		{
+			FixtureGraph fg = new FixtureGraph();
+			fg.Id = ContentHash.Empty;
+			fg.Schema = 1122;
+			fg.Groups = new List<INodeGroup>();
+			fg.Aggregates = new List<IAggregate>();
+			fg.Labels = new List<ILabel>();
+
+			Template = await templateCollection.GetOrAddAsync(new TemplateConfig { Name = "Test template" });
+			Graph = await graphCollection.AddAsync(Template, null);
+
+			TemplateRefId1 = new TemplateId("template1");
+			TemplateRefId2 = new TemplateId("template2");
+
+			List<TemplateRefConfig> templates = new List<TemplateRefConfig>();
+			templates.Add(new TemplateRefConfig { Id = TemplateRefId1, Name = "Test Template" });
+			templates.Add(new TemplateRefConfig { Id = TemplateRefId2, Name = "Test Template" });
+
+			List<TabConfig> tabs = new List<TabConfig>();
+			tabs.Add(new TabConfig { Title = "foo", Templates = new List<TemplateId> { TemplateRefId1, TemplateRefId2 } });
+
+			Dictionary<string, AgentConfig> agentTypes = new()
+			{
+				{ "Win64", new() { Pool = new PoolId(PoolName) } }
+			};
+
+			StreamId streamId = new StreamId("ue5-main");
+			StreamConfig streamConfig = new StreamConfig { Id = streamId, Name = "//UE5/Main", Tabs = tabs, Templates = templates, AgentTypes = agentTypes };
+
+			ProjectId projectId = new ProjectId("ue5");
+			ProjectConfig projectConfig = new ProjectConfig { Id = projectId, Name = "UE5", Streams = new List<StreamConfig> { streamConfig } };
+
+			GlobalConfig globalConfig = new GlobalConfig { Projects = new List<ProjectConfig> { projectConfig } };
+			globalConfig.PostLoad(serverSettings);
+			configService.OverrideConfig(globalConfig);
+
+			StreamId = streamId;
+			StreamConfig = streamConfig;
+
+			Job1 = await jobService.CreateJobAsync(
+				jobId: JobId.Parse("5f283932841e7fdbcafb6ab5"),
+				streamConfig: streamConfig,
+				templateRefId: TemplateRefId1,
+				templateHash: Template.Hash,
+				graph: Graph,
+				name: "hello1",
+				change: 1000001,
+				codeChange: 1000002,
+				new CreateJobOptions { PreflightChange = 1001 }
+			);
+			Job1 = (await jobService.GetJobAsync(Job1.Id))!;
+
+			Job2 = await jobService.CreateJobAsync(
+				jobId: JobId.Parse("5f69ea1b68423e921b035106"),
+				streamConfig: streamConfig,
+				templateRefId: new TemplateId("template-id-1"),
+				templateHash: ContentHash.MD5("made-up-template-hash"),
+				graph: fg,
+				name: "hello2",
+				change: 2000001,
+				codeChange: 2000002,
+				new CreateJobOptions()
+			);
+			Job2 = (await jobService.GetJobAsync(Job2.Id))!;
+
+			Job1ArtifactData = "For The Horde!";
+			using MemoryStream job1ArtifactStream = new MemoryStream(Encoding.UTF8.GetBytes(Job1ArtifactData));
+			Job1Artifact = await artifactCollection.CreateArtifactAsync(Job1.Id, JobStepId.Parse("22"), "myfile.txt",
+				"text/plain", job1ArtifactStream);
+
+			Agent1Name = "testAgent1";
+			Agent1 = await agentService.CreateAgentAsync(Agent1Name, false, "");
+		}
+
+		private class FixtureGraph : IGraph
+		{
+			public ContentHash Id { get; set; } = ContentHash.Empty;
+			public int Schema { get; set; }
+			public IReadOnlyList<INodeGroup> Groups { get; set; } = null!;
+			public IReadOnlyList<IAggregate> Aggregates { get; set; } = null!;
+			public IReadOnlyList<ILabel> Labels { get; set; } = null!;
+			public IReadOnlyList<IGraphArtifact> Artifacts { get; set; } = null!;
+		}
+	}
+}
\ No newline at end of file
diff --git a/Engine/Source/Programs/Horde/HordeServer.Tests/Fleet/AwsReuseFleetManagerTest.cs b/Engine/Source/Programs/Horde/HordeServer.Tests/Fleet/AwsReuseFleetManagerTest.cs
new file mode 100644
--- /dev/null
+++ b/Engine/Source/Programs/Horde/HordeServer.Tests/Fleet/AwsReuseFleetManagerTest.cs
@@ -0,0 +1,55 @@
+// Copyright Epic Games, Inc. All Rights Reserved.
+
+using System.Collections.Generic;
+using System.Threading;
+using System.Threading.Tasks;
+using Amazon.EC2;
+using Amazon.EC2.Model;
+using Horde.Server.Agents;
+using Horde.Server.Agents.Fleet;
+using Horde.Server.Agents.Fleet.Providers;
+using Horde.Server.Agents.Pools;
+using Microsoft.Extensions.DependencyInjection;
+using Microsoft.Extensions.Logging;
+using Microsoft.VisualStudio.TestTools.UnitTesting;
+
+namespace Horde.Server.Tests.Fleet
+{
+	[TestClass]
+	public class AwsReuseFleetManagerTest : TestSetup
+	{
+		[TestMethod]
+		public async Task ExpandOneAgentAsync()
+		{
+			FakeAmazonEc2 ec2 = new();
+			Instance i = ec2.AddInstance(FakeAmazonEc2.StateStopped, InstanceType.M5Large);
+			ec2.SetCapacity(FakeAmazonEc2.AzUsEast1A, InstanceType.M5Large, 1);
+
+			await ExpandPoolAsync(ec2.Get(), 1, new());
+			Assert.AreEqual(InstanceType.M5Large, ec2.Instances[i.InstanceId].InstanceType);
+			Assert.AreEqual(FakeAmazonEc2.StatePending, ec2.Instances[i.InstanceId].State);
+		}
+
+		[TestMethod]
+		public async Task ExpandWithInstanceTypeChangeAsync()
+		{
+			FakeAmazonEc2 ec2 = new();
+			Instance i = ec2.AddInstance(FakeAmazonEc2.StateStopped, InstanceType.M5Large);
+			ec2.SetCapacity(FakeAmazonEc2.AzUsEast1A, InstanceType.M54xlarge, 1);
+
+			await ExpandPoolAsync(ec2.Get(), 1, new(new List<string> { InstanceType.M54xlarge }));
+			Assert.AreEqual(InstanceType.M54xlarge, ec2.Instances[i.InstanceId].InstanceType);
+			Assert.AreEqual(FakeAmazonEc2.StatePending, ec2.Instances[i.InstanceId].State);
+		}
+
+		private async Task ExpandPoolAsync(IAmazonEC2 ec2, int numRequestedInstances, AwsReuseFleetManagerSettings settings)
+		{
+			ILogger<AwsReuseFleetManager> logger = ServiceProvider.GetRequiredService<ILogger<AwsReuseFleetManager>>();
+#pragma warning disable CS0618 // Type or member is obsolete
+			IPool pool = await CreatePoolAsync(new PoolConfig { Name = "testPool", EnableAutoscaling = true, MinAgents = 0, NumReserveAgents = 0, SizeStrategy = PoolSizeStrategy.NoOp });
+#pragma warning restore CS0618 // Type or member is obsolete
+			AwsReuseFleetManager manager = new(ec2, AgentCollection, settings, Tracer, logger);
+			await manager.ExpandPoolAsync(pool, new List<IAgent>(), numRequestedInstances, CancellationToken.None);
+		}
+	}
+}
\ No newline at end of file
diff --git a/Engine/Source/Programs/Horde/HordeServer.Tests/HordeServer.Tests.csproj b/Engine/Source/Programs/Horde/HordeServer.Tests/HordeServer.Tests.csproj
new file mode 100644
--- /dev/null
+++ b/Engine/Source/Programs/Horde/HordeServer.Tests/HordeServer.Tests.csproj
@@ -0,0 +1,112 @@
+<Project Sdk="Microsoft.NET.Sdk">
+
+  <PropertyGroup>
+    <TargetFramework>net8.0</TargetFramework>
+    <Nullable>enable</Nullable>
+    <IsPackable>false</IsPackable>
+    <NeutralLanguage>en-US</NeutralLanguage>
+    <Configurations>Debug;Release;Development;Analyze</Configurations>
+    <NoWarn>1701;1702;1591</NoWarn>
+    <GenerateProgramFile>false</GenerateProgramFile>
+  </PropertyGroup>
+
+  <PropertyGroup Condition=" '$(Configuration)|$(Platform)' == 'Analyze|AnyCPU' ">
+    <RunAnalyzersDuringBuild>True</RunAnalyzersDuringBuild>
+    <TreatWarningsAsErrors>false</TreatWarningsAsErrors>
+    <EnforceCodeStyleInBuild>True</EnforceCodeStyleInBuild>
+    <GenerateDocumentationFile>True</GenerateDocumentationFile>
+    <DocumentationFile></DocumentationFile>
+  </PropertyGroup>
+
+  <ItemGroup>
+    <PackageReference Include="Microsoft.AspNetCore.Mvc.Testing" Version="6.0.25" />
+    <PackageReference Include="Microsoft.NET.Test.Sdk" Version="17.8.0" />
+    <PackageReference Include="Moq" Version="4.20.70" />
+    <PackageReference Include="MSTest.TestAdapter" Version="3.1.1" />
+    <PackageReference Include="MSTest.TestFramework" Version="3.1.1" />
+    <PackageReference Include="Microsoft.Extensions.Caching.Memory" Version="8.0.0" />
+    <PackageReference Include="coverlet.collector" Version="6.0.0">
+      <PrivateAssets>all</PrivateAssets>
+      <IncludeAssets>runtime; build; native; contentfiles; analyzers; buildtransitive</IncludeAssets>
+    </PackageReference>
+  </ItemGroup>
+
+  <ItemGroup>
+    <ProjectReference Include="..\..\Shared\EpicGames.Serialization\EpicGames.Serialization.csproj" />
+    <ProjectReference Include="..\HordeAgent\HordeAgent.csproj">
+      <Aliases>HordeAgent</Aliases>
+    </ProjectReference>
+    <ProjectReference Include="..\Drivers\JobDriver\JobDriver.csproj">
+      <Aliases>JobDriver</Aliases>
+    </ProjectReference>
+    <ProjectReference Include="..\HordeServer\HordeServer.csproj" />
+  </ItemGroup>
+
+  <ItemGroup>
+    <Compile Update="Properties\Resources.Designer.cs">
+      <DesignTime>True</DesignTime>
+      <AutoGen>True</AutoGen>
+      <DependentUpon>Resources.resx</DependentUpon>
+    </Compile>
+  </ItemGroup>
+
+  <ItemGroup>
+    <EmbeddedResource Update="Properties\Resources.resx">
+      <Generator>ResXFileCodeGenerator</Generator>
+      <LastGenOutput>Resources.Designer.cs</LastGenOutput>
+    </EmbeddedResource>
+  </ItemGroup>
+
+  <ItemGroup>
+    <Folder Include="Resources\" />
+  </ItemGroup>
+
+  <ItemGroup>
+    <None Update="Ddc\Functional\ContentId\Payloads\05d7c699a2668efdecbe48f10db0d621d736f449.uecomp">
+      <CopyToOutputDirectory>Always</CopyToOutputDirectory>
+    </None>
+    <None Update="Ddc\Functional\ContentId\Payloads\0b29a878f71cc1eb6a26663a1152537e8fe7dca9.uecomp">
+      <CopyToOutputDirectory>Always</CopyToOutputDirectory>
+    </None>
+    <None Update="Ddc\Functional\ContentId\Payloads\4958689fe783e02fb35b13c14b0c3d7beb91e50c">
+      <CopyToOutputDirectory>Always</CopyToOutputDirectory>
+    </None>
+    <None Update="Ddc\Functional\ContentId\Payloads\7835c353d7dc67e8a0531c88fbc75ddfda10dee4">
+      <CopyToOutputDirectory>Always</CopyToOutputDirectory>
+    </None>
+    <None Update="Ddc\Functional\ContentId\Payloads\dce31eb416f3dcb4c8250ac545eda3930919d3ff">
+      <CopyToOutputDirectory>Always</CopyToOutputDirectory>
+    </None>
+    <None Update="Ddc\Functional\ContentId\Payloads\largerfile">
+      <CopyToOutputDirectory>Always</CopyToOutputDirectory>
+    </None>
+    <None Update="Ddc\Functional\ContentId\Payloads\Oodle-f0b9c675fe21951ca27699f9baab9f9f5040b202">
+      <CopyToOutputDirectory>Always</CopyToOutputDirectory>
+    </None>
+    <None Update="Ddc\Functional\ContentId\Payloads\Oodle-f895ea954b37217270e88d8b728bd3c09152689c">
+      <CopyToOutputDirectory>Always</CopyToOutputDirectory>
+    </None>
+    <None Update="Ddc\Functional\ContentId\Payloads\OodleTexture_CAS_dbda9040e75c4674fcec173f982fddf12b021e24.udd">
+      <CopyToOutputDirectory>Always</CopyToOutputDirectory>
+    </None>
+    <None Update="Ddc\Functional\ContentId\Payloads\smallerfile">
+      <CopyToOutputDirectory>Always</CopyToOutputDirectory>
+    </None>
+    <None Update="Ddc\Functional\ContentId\Payloads\UncompressedTexture_CAS_dea81b6c3b565bb5089695377c98ce0f1c13b0c3.udd">
+      <CopyToOutputDirectory>Always</CopyToOutputDirectory>
+    </None>
+    <None Update="Ddc\Functional\Objects\Payloads\lyra.cb">
+      <CopyToOutputDirectory>Always</CopyToOutputDirectory>
+    </None>
+    <None Update="Ddc\UnitTests\CompactBinaryObjects\build">
+      <CopyToOutputDirectory>Always</CopyToOutputDirectory>
+    </None>
+    <None Update="Ddc\UnitTests\CompactBinaryObjects\compact_binary">
+      <CopyToOutputDirectory>Always</CopyToOutputDirectory>
+    </None>
+    <None Update="Ddc\UnitTests\CompactBinaryObjects\ReferenceOutput">
+      <CopyToOutputDirectory>Always</CopyToOutputDirectory>
+    </None>
+  </ItemGroup>
+
+</Project>
diff --git a/Engine/Source/Programs/Horde/HordeServer.Tests/Jobs/JobsControllerTest.cs b/Engine/Source/Programs/Horde/HordeServer.Tests/Jobs/JobsControllerTest.cs
new file mode 100644
--- /dev/null
+++ b/Engine/Source/Programs/Horde/HordeServer.Tests/Jobs/JobsControllerTest.cs
@@ -0,0 +1,36 @@
+﻿// Copyright Epic Games, Inc. All Rights Reserved.
+
+using System.Net;
+using System.Threading.Tasks;
+using Horde.Server.Jobs.Artifacts;
+using Microsoft.VisualStudio.TestTools.UnitTesting;
+
+#pragma warning disable CA2234 // Pass system uri objects instead of strings
+
+namespace Horde.Server.Tests.Jobs
+{
+	[TestClass]
+	public class JobsControllerTest : ControllerIntegrationTest
+	{
+		private static string GetUri(string jobId, string stepId, string fileName)
+		{
+			return $"/api/v1/jobs/{jobId}/steps/{stepId}/artifacts/{fileName}/data";
+		}
+
+		[TestMethod]
+		public async Task GetArtifactDataByFilenameTestAsync()
+		{
+			Fixture? fixture = await GetFixtureAsync();
+			IArtifactV1? art = fixture.Job1Artifact;
+
+			// Test existing filename
+			System.Net.Http.HttpResponseMessage? res = await Client.GetAsync(GetUri(art.JobId.ToString(), art.StepId.ToString()!, art.Name));
+			res.EnsureSuccessStatusCode();
+			Assert.AreEqual(fixture.Job1ArtifactData, await res.Content.ReadAsStringAsync());
+
+			// Test non-existing filename
+			res = await Client.GetAsync(GetUri(art.JobId.ToString(), art.StepId.ToString()!, "bogus.txt"));
+			Assert.AreEqual(res.StatusCode, HttpStatusCode.NotFound);
+		}
+	}
+}
\ No newline at end of file
diff --git a/Engine/Source/Programs/Horde/HordeServer.Tests/Logs/LogChunkTests.cs b/Engine/Source/Programs/Horde/HordeServer.Tests/Logs/LogChunkTests.cs
new file mode 100644
--- /dev/null
+++ b/Engine/Source/Programs/Horde/HordeServer.Tests/Logs/LogChunkTests.cs
@@ -0,0 +1,48 @@
+// Copyright Epic Games, Inc. All Rights Reserved.
+
+using System;
+using System.Collections.Generic;
+using System.Text;
+using EpicGames.Horde.Logs;
+using Horde.Server.Logs.Data;
+using Microsoft.Extensions.Logging.Abstractions;
+using Microsoft.VisualStudio.TestTools.UnitTesting;
+
+namespace Horde.Server.Tests.Logs
+{
+	[TestClass]
+	public class LogChunkTests
+	{
+		readonly string[] _text =
+		{
+			@"{ ""id"": 1, ""message"": ""foo"" }",
+			@"{ ""id"": 2, ""message"": ""bar"" }",
+			@"{ ""id"": 3, ""message"": ""baz"" }",
+		};
+
+		[TestMethod]
+		public void TestSerialization()
+		{
+			byte[] textData = Encoding.UTF8.GetBytes(String.Join("\n", _text) + "\n");
+
+			List<LogSubChunkData> subChunks = new List<LogSubChunkData>();
+			subChunks.Add(new LogSubChunkData(LogType.Text, 0, 0, new LogText(textData, textData.Length)));
+			subChunks.Add(new LogSubChunkData(LogType.Json, textData.Length, 3, new LogText(textData, textData.Length)));
+
+			LogChunkData oldChunkData = new LogChunkData(0, 0, subChunks);
+			byte[] data = oldChunkData.ToByteArray(NullLogger.Instance);
+			LogChunkData newChunkData = LogChunkData.FromMemory(data, 0, 0);
+
+			Assert.AreEqual(oldChunkData.Length, newChunkData.Length);
+			Assert.AreEqual(oldChunkData.SubChunks.Count, newChunkData.SubChunks.Count);
+			Assert.IsTrue(oldChunkData.SubChunkOffset.AsSpan(0, oldChunkData.SubChunkOffset.Length).SequenceEqual(newChunkData.SubChunkOffset.AsSpan(0, newChunkData.SubChunkOffset.Length)));
+			Assert.IsTrue(oldChunkData.SubChunkLineIndex.AsSpan(0, oldChunkData.SubChunkLineIndex.Length).SequenceEqual(newChunkData.SubChunkLineIndex.AsSpan(0, newChunkData.SubChunkLineIndex.Length)));
+			for (int idx = 0; idx < oldChunkData.SubChunks.Count; idx++)
+			{
+				LogSubChunkData oldSubChunkData = oldChunkData.SubChunks[idx];
+				LogSubChunkData newSubChunkData = newChunkData.SubChunks[idx];
+				Assert.IsTrue(oldSubChunkData.InflateText().Data.Span.SequenceEqual(newSubChunkData.InflateText().Data.Span));
+			}
+		}
+	}
+}
diff --git a/Engine/Source/Programs/Horde/HordeServer.Tests/Server/RpcServiceTest.cs b/Engine/Source/Programs/Horde/HordeServer.Tests/Server/RpcServiceTest.cs
new file mode 100644
--- /dev/null
+++ b/Engine/Source/Programs/Horde/HordeServer.Tests/Server/RpcServiceTest.cs
@@ -0,0 +1,453 @@
+// Copyright Epic Games, Inc. All Rights Reserved.
+
+using System;
+using System.Collections.Generic;
+using System.IO;
+using System.Linq;
+using System.Reflection;
+using System.Security.Claims;
+using System.Threading;
+using System.Threading.Channels;
+using System.Threading.Tasks;
+using EpicGames.Horde.Agents;
+using EpicGames.Horde.Agents.Pools;
+using Grpc.Core;
+using Horde.Server.Agents;
+using Horde.Server.Server;
+using Horde.Server.Utilities;
+using HordeCommon.Rpc;
+using HordeCommon.Rpc.Messages;
+using Microsoft.AspNetCore.Http;
+using Microsoft.AspNetCore.Http.Features;
+using Microsoft.Extensions.FileProviders;
+using Microsoft.Extensions.Hosting;
+using Microsoft.VisualStudio.TestTools.UnitTesting;
+
+namespace Horde.Server.Tests.Server
+{
+	using ISession = Microsoft.AspNetCore.Http.ISession;
+
+	public class AppLifetimeStub : IHostApplicationLifetime
+	{
+		public CancellationToken ApplicationStarted { get; }
+		public CancellationToken ApplicationStopping { get; }
+		public CancellationToken ApplicationStopped { get; }
+
+		public AppLifetimeStub()
+		{
+			ApplicationStarted = new CancellationToken();
+			ApplicationStopping = new CancellationToken();
+			ApplicationStopped = new CancellationToken();
+		}
+
+		public void StopApplication()
+		{
+			throw new NotImplementedException();
+		}
+	}
+
+	public class WebHostEnvironmentStub : IHostEnvironment
+	{
+		public string ApplicationName { get; set; } = "HordeTest";
+		public IFileProvider ContentRootFileProvider { get; set; }
+		public string ContentRootPath { get; set; }
+		public string EnvironmentName { get; set; } = "Testing";
+
+		public WebHostEnvironmentStub()
+		{
+			ContentRootPath = Directory.CreateTempSubdirectory("HordeTest").FullName;
+			ContentRootFileProvider = new PhysicalFileProvider(ContentRootPath);
+		}
+	}
+
+	sealed class HttpContextStub : HttpContext
+	{
+		public override ConnectionInfo Connection { get; } = null!;
+		public override IFeatureCollection Features { get; } = null!;
+		public override IDictionary<object, object?> Items { get; set; } = null!;
+		public override HttpRequest Request { get; } = null!;
+		public override CancellationToken RequestAborted { get; set; }
+		public override IServiceProvider RequestServices { get; set; } = null!;
+		public override HttpResponse Response { get; } = null!;
+		public override ISession Session { get; set; } = null!;
+		public override string TraceIdentifier { get; set; } = null!;
+		public override ClaimsPrincipal User { get; set; }
+		public override WebSocketManager WebSockets { get; } = null!;
+
+		public HttpContextStub(Claim roleClaimType)
+		{
+			User = new ClaimsPrincipal(new ClaimsIdentity(new List<Claim>
+			{
+				roleClaimType
+			}, "TestAuthType"));
+		}
+
+		public HttpContextStub(ClaimsPrincipal user)
+		{
+			User = user;
+		}
+
+		public override void Abort()
+		{
+			throw new NotImplementedException();
+		}
+	}
+
+	public class ServerCallContextStub : ServerCallContext
+	{
+		// Copied from ServerCallContextExtensions.cs in Grpc.Core
+		const string HttpContextKey = "__HttpContext";
+
+		protected override string MethodCore { get; } = null!;
+		protected override string HostCore { get; } = null!;
+		protected override string PeerCore { get; } = null!;
+		protected override DateTime DeadlineCore { get; } = DateTime.Now.AddHours(24);
+		protected override Metadata RequestHeadersCore { get; } = null!;
+		protected override CancellationToken CancellationTokenCore => _cancellationToken;
+		protected override Metadata ResponseTrailersCore { get; } = null!;
+		protected override Status StatusCore { get; set; }
+		protected override WriteOptions? WriteOptionsCore { get; set; } = null!;
+		protected override AuthContext AuthContextCore { get; } = null!;
+
+		private CancellationToken _cancellationToken;
+
+		public static ServerCallContext ForAdminWithAgentSessionId(string agentSessionId)
+		{
+			return new ServerCallContextStub(new ClaimsPrincipal(new ClaimsIdentity(new List<Claim>
+			{
+				HordeClaims.AdminClaim.ToClaim(),
+				new Claim(HordeClaimTypes.AgentSessionId, agentSessionId),
+			}, "TestAuthType")));
+		}
+
+		public static ServerCallContext ForAdmin()
+		{
+			return new ServerCallContextStub(new ClaimsPrincipal(new ClaimsIdentity(new List<Claim>
+			{
+				HordeClaims.AdminClaim.ToClaim()
+			}, "TestAuthType")));
+		}
+
+		public ServerCallContextStub(Claim roleClaimType)
+		{
+			// The GetHttpContext extension falls back to getting the HttpContext from UserState
+			// We can piggyback on that behavior during tests
+			UserState[HttpContextKey] = new HttpContextStub(roleClaimType);
+		}
+
+		public ServerCallContextStub(ClaimsPrincipal user)
+		{
+			// The GetHttpContext extension falls back to getting the HttpContext from UserState
+			// We can piggyback on that behavior during tests
+			UserState[HttpContextKey] = new HttpContextStub(user);
+		}
+
+		public void SetCancellationToken(CancellationToken cancellationToken)
+		{
+			_cancellationToken = cancellationToken;
+		}
+
+		protected override Task WriteResponseHeadersAsyncCore(Metadata responseHeaders)
+		{
+			throw new NotImplementedException();
+		}
+
+		protected override ContextPropagationToken CreatePropagationTokenCore(ContextPropagationOptions? options)
+		{
+			throw new NotImplementedException();
+		}
+	}
+
+	[TestClass]
+	public class RpcServiceTest : TestSetup
+	{
+		private readonly ServerCallContext _adminContext = new ServerCallContextStub(HordeClaims.AdminClaim.ToClaim());
+
+		class RpcServiceInvoker : CallInvoker
+		{
+			private readonly RpcService _rpcService;
+			private readonly ServerCallContext _serverCallContext;
+
+			public RpcServiceInvoker(RpcService rpcService, ServerCallContext serverCallContext)
+			{
+				_rpcService = rpcService;
+				_serverCallContext = serverCallContext;
+			}
+
+			public override TResponse BlockingUnaryCall<TRequest, TResponse>(Method<TRequest, TResponse> method, string? host, CallOptions options, TRequest request)
+			{
+				throw new NotImplementedException("Blocking calls are not supported! Method " + method.FullName);
+			}
+
+			public override AsyncUnaryCall<TResponse> AsyncUnaryCall<TRequest, TResponse>(Method<TRequest, TResponse> method, string? host, CallOptions options, TRequest request)
+			{
+				MethodInfo methodInfo = GetMethod(method.Name);
+				Task<TResponse> res = (methodInfo.Invoke(_rpcService, new object[] { request, _serverCallContext }) as Task<TResponse>)!;
+				return new AsyncUnaryCall<TResponse>(res, null!, null!, null!, null!, null!);
+			}
+
+			public override AsyncServerStreamingCall<TResponse> AsyncServerStreamingCall<TRequest, TResponse>(Method<TRequest, TResponse> method, string? host, CallOptions options,
+				TRequest request)
+			{
+				Console.WriteLine($"RpcServiceInvoker.AsyncServerStreamingCall(method={method.FullName} request={request})");
+				throw new NotImplementedException();
+			}
+
+			public override AsyncClientStreamingCall<TRequest, TResponse> AsyncClientStreamingCall<TRequest, TResponse>(Method<TRequest, TResponse> method, string? host, CallOptions options)
+			{
+				Console.WriteLine($"RpcServiceInvoker.AsyncClientStreamingCall(method={method.FullName})");
+				throw new NotImplementedException();
+			}
+
+			public override AsyncDuplexStreamingCall<TRequest, TResponse> AsyncDuplexStreamingCall<TRequest, TResponse>(Method<TRequest, TResponse> method, string? host, CallOptions options)
+			{
+				Console.WriteLine($"RpcServiceInvoker.AsyncDuplexStreamingCall(method={method.FullName})");
+
+				GrpcDuplexStreamHandler<TRequest> requestStream = new GrpcDuplexStreamHandler<TRequest>(_serverCallContext);
+				GrpcDuplexStreamHandler<TResponse> responseStream = new GrpcDuplexStreamHandler<TResponse>(_serverCallContext);
+
+				MethodInfo methodInfo = GetMethod(method.Name);
+				Task methodTask = (methodInfo.Invoke(_rpcService, new object[] { requestStream, responseStream, _serverCallContext }) as Task)!;
+				methodTask.ContinueWith(t => { Console.Error.WriteLine($"Uncaught exception in {method.Name}: {t}"); }, CancellationToken.None, TaskContinuationOptions.OnlyOnFaulted, TaskScheduler.Default);
+
+				return new AsyncDuplexStreamingCall<TRequest, TResponse>(requestStream, responseStream, null!, null!, null!, null!);
+			}
+
+			private MethodInfo GetMethod(string methodName)
+			{
+				MethodInfo? method = _rpcService.GetType().GetMethod(methodName);
+				if (method == null)
+				{
+					throw new ArgumentException($"Method {methodName} not found in RpcService");
+				}
+
+				return method;
+			}
+		}
+
+		/// <summary>
+		/// Combines and cross-writes streams for a duplex streaming call in gRPC
+		/// </summary>
+		/// <typeparam name="T"></typeparam>
+		class GrpcDuplexStreamHandler<T> : IServerStreamWriter<T>, IClientStreamWriter<T>, IAsyncStreamReader<T> where T : class
+		{
+			private readonly ServerCallContext _serverCallContext;
+			private readonly Channel<T> _channel;
+
+			public WriteOptions? WriteOptions { get; set; }
+
+			public GrpcDuplexStreamHandler(ServerCallContext serverCallContext)
+			{
+				_channel = System.Threading.Channels.Channel.CreateUnbounded<T>();
+
+				_serverCallContext = serverCallContext;
+			}
+
+			public void Complete()
+			{
+				_channel.Writer.Complete();
+			}
+
+			public IAsyncEnumerable<T> ReadAllAsync()
+			{
+				return _channel.Reader.ReadAllAsync();
+			}
+
+			public async Task<T?> ReadNextAsync()
+			{
+				if (await _channel.Reader.WaitToReadAsync())
+				{
+					_channel.Reader.TryRead(out T? message);
+					return message;
+				}
+				else
+				{
+					return null;
+				}
+			}
+
+			public Task WriteAsync(T message)
+			{
+				if (_serverCallContext.CancellationToken.IsCancellationRequested)
+				{
+					return Task.FromCanceled(_serverCallContext.CancellationToken);
+				}
+
+				if (!_channel.Writer.TryWrite(message))
+				{
+					throw new InvalidOperationException("Unable to write message.");
+				}
+
+				return Task.CompletedTask;
+			}
+
+			public Task CompleteAsync()
+			{
+				Complete();
+				return Task.CompletedTask;
+			}
+
+			public async Task<bool> MoveNext(CancellationToken cancellationToken)
+			{
+				_serverCallContext.CancellationToken.ThrowIfCancellationRequested();
+
+				if (await _channel.Reader.WaitToReadAsync(cancellationToken))
+				{
+					if (_channel.Reader.TryRead(out T? message))
+					{
+						Current = message;
+						return true;
+					}
+				}
+
+				Current = null!;
+				return false;
+			}
+
+			public T Current { get; private set; } = null!;
+		}
+
+		public RpcServiceTest()
+		{
+			UpdateConfig(x => x.Pools.Clear());
+		}
+
+		[TestMethod]
+		public async Task CreateSessionTestAsync()
+		{
+			RpcCreateSessionRequest req = new RpcCreateSessionRequest();
+			await Assert.ThrowsExceptionAsync<StructuredRpcException>(() => RpcService.CreateSession(req, _adminContext));
+
+			req.Id = new AgentId("MyName").ToString();
+			await Assert.ThrowsExceptionAsync<StructuredRpcException>(() => RpcService.CreateSession(req, _adminContext));
+
+			req.Capabilities = new RpcAgentCapabilities();
+			RpcCreateSessionResponse res = await RpcService.CreateSession(req, _adminContext);
+
+			Assert.AreEqual("MYNAME", res.AgentId);
+			// TODO: Check Token, ExpiryTime, SessionId 
+		}
+
+		[TestMethod]
+		public async Task AgentJoinsPoolThroughPropertiesAsync()
+		{
+			RpcCreateSessionRequest req = new() { Id = new AgentId("bogusAgentName").ToString(), Capabilities = new RpcAgentCapabilities() };
+			req.Capabilities.Properties.Add($"{KnownPropertyNames.RequestedPools}=fooPool,barPool");
+			RpcCreateSessionResponse res = await RpcService.CreateSession(req, _adminContext);
+
+			IAgent agent = (await AgentService.GetAgentAsync(new AgentId(res.AgentId)))!;
+			CollectionAssert.AreEquivalent(new List<PoolId> { new("fooPool"), new("barPool") }, agent.Pools.ToList());
+
+			// Connect a second time, when the agent has already been created
+			req = new() { Id = new AgentId("bogusAgentName").ToString(), Capabilities = new RpcAgentCapabilities() };
+			req.Capabilities.Properties.Add($"{KnownPropertyNames.RequestedPools}=bazPool");
+			res = await RpcService.CreateSession(req, _adminContext);
+
+			agent = (await AgentService.GetAgentAsync(new AgentId(res.AgentId)))!;
+			CollectionAssert.AreEquivalent(new List<PoolId> { new("fooPool"), new("barPool"), new("bazPool") }, agent.Pools.ToList());
+		}
+
+		[TestMethod]
+		public async Task PropertiesFromAgentCapabilitiesAsync()
+		{
+			RpcCreateSessionRequest req = new() { Id = new AgentId("bogusAgentName").ToString(), Capabilities = new RpcAgentCapabilities() };
+			req.Capabilities.Properties.Add("fooKey=barValue");
+			RpcCreateSessionResponse res = await RpcService.CreateSession(req, _adminContext);
+			IAgent agent = (await AgentService.GetAgentAsync(new AgentId(res.AgentId)))!;
+			Assert.IsTrue(agent.Properties.Contains("fooKey=barValue"));
+		}
+
+		[TestMethod]
+		public async Task PropertiesFromDeviceCapabilitiesAsync()
+		{
+			RpcCreateSessionRequest req = new() { Id = new AgentId("bogusAgentName").ToString(), Capabilities = new RpcAgentCapabilities() };
+			req.Capabilities.Devices.Add(new RpcDeviceCapabilities { Handle = "someHandle", Properties = { "foo=bar" } });
+			RpcCreateSessionResponse res = await RpcService.CreateSession(req, _adminContext);
+			IAgent agent = (await AgentService.GetAgentAsync(new AgentId(res.AgentId)))!;
+			Assert.IsTrue(agent.Properties.Contains("foo=bar"));
+		}
+
+		[TestMethod]
+		public async Task KnownPropertiesAreSetAsResourcesAsync()
+		{
+			RpcCreateSessionRequest req = new() { Id = new AgentId("bogusAgentName").ToString(), Capabilities = new RpcAgentCapabilities() };
+			req.Capabilities.Devices.Add(new RpcDeviceCapabilities { Handle = "someHandle", Properties = { $"{KnownPropertyNames.LogicalCores}=10" } });
+			RpcCreateSessionResponse res = await RpcService.CreateSession(req, _adminContext);
+			IAgent agent = (await AgentService.GetAgentAsync(new AgentId(res.AgentId)))!;
+			Assert.AreEqual(10, agent.Resources[KnownPropertyNames.LogicalCores]);
+		}
+
+		[TestMethod]
+		public async Task UpdateSessionTestAsync()
+		{
+			RpcCreateSessionRequest createReq = new RpcCreateSessionRequest
+			{
+				Id = new AgentId("UpdateSessionTest1").ToString(),
+				Capabilities = new RpcAgentCapabilities()
+			};
+			RpcCreateSessionResponse createRes = await RpcService.CreateSession(createReq, _adminContext);
+			string agentId = createRes.AgentId;
+			string sessionId = createRes.SessionId;
+
+			TestAsyncStreamReader<RpcUpdateSessionRequest> requestStream =
+				new TestAsyncStreamReader<RpcUpdateSessionRequest>(_adminContext);
+			TestServerStreamWriter<RpcUpdateSessionResponse> responseStream =
+				new TestServerStreamWriter<RpcUpdateSessionResponse>(_adminContext);
+			Task call = RpcService.UpdateSession(requestStream, responseStream, _adminContext);
+
+			requestStream.AddMessage(new RpcUpdateSessionRequest { AgentId = "does-not-exist", SessionId = sessionId });
+			StructuredRpcException re = await Assert.ThrowsExceptionAsync<StructuredRpcException>(() => call);
+			Assert.AreEqual(StatusCode.NotFound, re.StatusCode);
+			Assert.IsTrue(re.Message.Contains("Invalid agent name", StringComparison.OrdinalIgnoreCase));
+		}
+
+		[TestMethod]
+		public async Task FinishBatchTestAsync()
+		{
+			RpcCreateSessionRequest createReq = new RpcCreateSessionRequest
+			{
+				Id = new AgentId("UpdateSessionTest1").ToString(),
+				Capabilities = new RpcAgentCapabilities()
+			};
+			RpcCreateSessionResponse createRes = await RpcService.CreateSession(createReq, _adminContext);
+			string agentId = createRes.AgentId;
+			string sessionId = createRes.SessionId;
+
+			TestAsyncStreamReader<RpcUpdateSessionRequest> requestStream =
+				new TestAsyncStreamReader<RpcUpdateSessionRequest>(_adminContext);
+			TestServerStreamWriter<RpcUpdateSessionResponse> responseStream =
+				new TestServerStreamWriter<RpcUpdateSessionResponse>(_adminContext);
+			Task call = RpcService.UpdateSession(requestStream, responseStream, _adminContext);
+
+			requestStream.AddMessage(new RpcUpdateSessionRequest { AgentId = "does-not-exist", SessionId = sessionId });
+			StructuredRpcException re = await Assert.ThrowsExceptionAsync<StructuredRpcException>(() => call);
+			Assert.AreEqual(StatusCode.NotFound, re.StatusCode);
+			Assert.IsTrue(re.Message.Contains("Invalid agent name", StringComparison.OrdinalIgnoreCase));
+		}
+
+		/*
+		[TestMethod]
+		public async Task UploadSoftwareAsync()
+		{
+			TestSetup TestSetup = await GetTestSetup();
+			
+			MemoryStream OutputStream = new MemoryStream();
+			using (ZipArchive ZipFile = new ZipArchive(OutputStream, ZipArchiveMode.Create, false))
+			{
+				string TempFilename = Path.GetTempFileName();
+				File.WriteAllText(TempFilename, "{\"Horde\": {\"Version\": \"myVersion\"}}");
+				ZipFile.CreateEntryFromFile(TempFilename, "appsettings.json");
+			}
+
+			ByteString Data = ByteString.CopyFrom(OutputStream.ToArray());
+			UploadSoftwareRequest Req = new UploadSoftwareRequest { Channel = "boguschannel", Data = Data };
+
+			UploadSoftwareResponse Res1 = await TestSetup.RpcService.UploadSoftware(Req, AdminContext);
+			Assert.AreEqual("r1", Res1.Version);
+			
+			UploadSoftwareResponse Res2 = await TestSetup.RpcService.UploadSoftware(Req, AdminContext);
+			Assert.AreEqual("r2", Res2.Version);
+		}
+		*/
+	}
+}
\ No newline at end of file
diff --git a/Engine/Source/Programs/Horde/HordeServer.Tests/Streams/StreamCollectionTests.cs b/Engine/Source/Programs/Horde/HordeServer.Tests/Streams/StreamCollectionTests.cs
new file mode 100644
--- /dev/null
+++ b/Engine/Source/Programs/Horde/HordeServer.Tests/Streams/StreamCollectionTests.cs
@@ -0,0 +1,35 @@
+// Copyright Epic Games, Inc. All Rights Reserved.
+
+using System.Collections.Generic;
+using EpicGames.Horde.Jobs.Templates;
+using EpicGames.Horde.Streams;
+using Horde.Server.Streams;
+using Microsoft.VisualStudio.TestTools.UnitTesting;
+
+namespace Horde.Server.Tests.Streams
+{
+	[TestClass]
+	public class StreamCollectionTests : TestSetup
+	{
+		private readonly StreamId _streamId = new("bogusStreamId");
+
+		[TestMethod]
+		public void ValidateUndefinedTemplateIdInTabs()
+		{
+			StreamConfig config = new()
+			{
+				Tabs = new()
+				{
+					new TabConfig { Templates = new List<TemplateId> { new ("foo") }},
+					new TabConfig { Templates = new List<TemplateId> { new ("bar") }}
+				},
+				Templates = new() { new TemplateRefConfig { Id = new TemplateId("foo") } }
+			};
+
+			Assert.ThrowsException<InvalidStreamException>(() => Horde.Server.Streams.StreamCollection.Validate(_streamId, config));
+
+			config.Templates.Add(new TemplateRefConfig { Id = new TemplateId("bar") });
+			Horde.Server.Streams.StreamCollection.Validate(_streamId, config);
+		}
+	}
+}
diff --git a/Engine/Source/Programs/Horde/HordeServer.Tests/Streams/StreamServiceTest.cs b/Engine/Source/Programs/Horde/HordeServer.Tests/Streams/StreamServiceTest.cs
new file mode 100644
--- /dev/null
+++ b/Engine/Source/Programs/Horde/HordeServer.Tests/Streams/StreamServiceTest.cs
@@ -0,0 +1,39 @@
+// Copyright Epic Games, Inc. All Rights Reserved.
+
+using System;
+using System.Globalization;
+using System.Threading.Tasks;
+using Horde.Server.Streams;
+using Microsoft.VisualStudio.TestTools.UnitTesting;
+
+namespace Horde.Server.Tests.Streams
+{
+	[TestClass]
+	public class StreamServiceTests : TestSetup
+	{
+		[TestMethod]
+		public async Task PausingAsync()
+		{
+			Fixture fixture = await CreateFixtureAsync();
+
+			IStream stream = (await StreamCollection.GetAsync(fixture!.StreamConfig!))!;
+			Assert.IsFalse(stream.IsPaused(DateTime.UtcNow));
+			Assert.IsNull(stream.PausedUntil);
+			Assert.IsNull(stream.PauseComment);
+
+			DateTime pausedUntil = DateTime.UtcNow.AddHours(1);
+			await StreamCollection.UpdatePauseStateAsync(stream, newPausedUntil: pausedUntil, newPauseComment: "mycomment");
+			stream = (await StreamCollection.GetAsync(fixture!.StreamConfig!))!;
+			// Comparing by string to avoid comparing exact milliseconds as those are not persisted in MongoDB fields
+			Assert.IsTrue(stream.IsPaused(DateTime.UtcNow));
+			Assert.AreEqual(pausedUntil.ToString(CultureInfo.InvariantCulture), stream.PausedUntil!.Value.ToString(CultureInfo.InvariantCulture));
+			Assert.AreEqual("mycomment", stream.PauseComment);
+
+			await StreamCollection.UpdatePauseStateAsync(stream, newPausedUntil: null, newPauseComment: null);
+			stream = (await StreamCollection.GetAsync(fixture!.StreamConfig!))!;
+			Assert.IsFalse(stream.IsPaused(DateTime.UtcNow));
+			Assert.IsNull(stream.PausedUntil);
+			Assert.IsNull(stream.PauseComment);
+		}
+	}
+}
