Horde: Rename projects before splitting into plugins (Horde.Agent -> HordeAgent, Horde.Server -> HordeServer, Horde -> HordeCmd).

diff --git a/Engine/Source/Programs/Horde/HordeServer.Tests/TestSetup.cs b/Engine/Source/Programs/Horde/HordeServer.Tests/TestSetup.cs
new file mode 100644
--- /dev/null
+++ b/Engine/Source/Programs/Horde/HordeServer.Tests/TestSetup.cs
@@ -0,0 +1,539 @@
+// Copyright Epic Games, Inc. All Rights Reserved.
+
+using System;
+using System.Collections.Generic;
+using System.Diagnostics.Metrics;
+using System.Net;
+using System.Net.Sockets;
+using System.Reflection;
+using System.Security.Claims;
+using System.Threading.Tasks;
+using Amazon.AutoScaling;
+using Amazon.CloudWatch;
+using Amazon.EC2;
+using EpicGames.Core;
+using EpicGames.Horde.Agents;
+using EpicGames.Horde.Agents.Pools;
+using EpicGames.Horde.Storage;
+using EpicGames.Horde.Storage.Bundles;
+using EpicGames.Horde.Storage.ObjectStores;
+using Horde.Server.Accounts;
+using Horde.Server.Acls;
+using Horde.Server.Agents;
+using Horde.Server.Agents.Enrollment;
+using Horde.Server.Agents.Fleet;
+using Horde.Server.Agents.Leases;
+using Horde.Server.Agents.Pools;
+using Horde.Server.Agents.Relay;
+using Horde.Server.Agents.Sessions;
+using Horde.Server.Agents.Telemetry;
+using Horde.Server.Agents.Utilization;
+using Horde.Server.Artifacts;
+using Horde.Server.Auditing;
+using Horde.Server.Commits;
+using Horde.Server.Compute;
+using Horde.Server.Configuration;
+using Horde.Server.Dashboard;
+using Horde.Server.Devices;
+using Horde.Server.Issues;
+using Horde.Server.Jobs;
+using Horde.Server.Jobs.Artifacts;
+using Horde.Server.Jobs.Bisect;
+using Horde.Server.Jobs.Graphs;
+using Horde.Server.Jobs.Schedules;
+using Horde.Server.Jobs.Templates;
+using Horde.Server.Jobs.TestData;
+using Horde.Server.Jobs.Timing;
+using Horde.Server.Logs;
+using Horde.Server.Logs.Storage;
+using Horde.Server.Notifications;
+using Horde.Server.Perforce;
+using Horde.Server.Server;
+using Horde.Server.Storage;
+using Horde.Server.Streams;
+using Horde.Server.Tasks;
+using Horde.Server.Telemetry;
+using Horde.Server.Telemetry.Sinks;
+using Horde.Server.Tests.Server;
+using Horde.Server.Tests.Stubs.Services;
+using Horde.Server.Tools;
+using Horde.Server.Ugs;
+using Horde.Server.Users;
+using Horde.Server.Utilities;
+using HordeCommon;
+using Microsoft.AspNetCore.Http;
+using Microsoft.AspNetCore.Mvc;
+using Microsoft.Extensions.Caching.Memory;
+using Microsoft.Extensions.Configuration;
+using Microsoft.Extensions.DependencyInjection;
+using Microsoft.Extensions.Hosting;
+using Microsoft.Extensions.Logging;
+using Microsoft.Extensions.Options;
+using Microsoft.VisualStudio.TestTools.UnitTesting;
+using Moq;
+using OpenTelemetry.Trace;
+
+namespace Horde.Server.Tests
+{
+	/// <summary>
+	/// Handles set up of collections, services, fixtures etc during testing
+	///
+	/// Easier to pass all these things around in a single object.
+	/// </summary>
+	public class TestSetup : DatabaseIntegrationTest
+	{
+		public FakeClock Clock => ServiceProvider.GetRequiredService<FakeClock>();
+		public IMemoryCache Cache => ServiceProvider.GetRequiredService<IMemoryCache>();
+
+		public IGraphCollection GraphCollection => ServiceProvider.GetRequiredService<IGraphCollection>();
+		public INotificationTriggerCollection NotificationTriggerCollection => ServiceProvider.GetRequiredService<INotificationTriggerCollection>();
+		public IStreamCollection StreamCollection => ServiceProvider.GetRequiredService<IStreamCollection>();
+		public IJobCollection JobCollection => ServiceProvider.GetRequiredService<IJobCollection>();
+		public IAgentCollection AgentCollection => ServiceProvider.GetRequiredService<IAgentCollection>();
+		public IJobStepRefCollection JobStepRefCollection => ServiceProvider.GetRequiredService<IJobStepRefCollection>();
+		public IJobTimingCollection JobTimingCollection => ServiceProvider.GetRequiredService<IJobTimingCollection>();
+		public IUgsMetadataCollection UgsMetadataCollection => ServiceProvider.GetRequiredService<IUgsMetadataCollection>();
+		public IIssueCollection IssueCollection => ServiceProvider.GetRequiredService<IIssueCollection>();
+		public IPoolCollection PoolCollection => ServiceProvider.GetRequiredService<IPoolCollection>();
+		public ILeaseCollection LeaseCollection => ServiceProvider.GetRequiredService<ILeaseCollection>();
+		public ILogCollection LogCollection => ServiceProvider.GetRequiredService<ILogCollection>();
+		public ISessionCollection SessionCollection => ServiceProvider.GetRequiredService<ISessionCollection>();
+		public ITestDataCollection TestDataCollection => ServiceProvider.GetRequiredService<ITestDataCollection>();
+		public IUserCollection UserCollection => ServiceProvider.GetRequiredService<IUserCollection>();
+		public IDeviceCollection DeviceCollection => ServiceProvider.GetRequiredService<IDeviceCollection>();
+		public IDashboardPreviewCollection DashboardPreviewCollection => ServiceProvider.GetRequiredService<IDashboardPreviewCollection>();
+		public IBisectTaskCollection BisectTaskCollection => ServiceProvider.GetRequiredService<IBisectTaskCollection>();
+
+		public AclService AclService => ServiceProvider.GetRequiredService<AclService>();
+		public FleetService FleetService => ServiceProvider.GetRequiredService<FleetService>();
+		public AgentService AgentService => ServiceProvider.GetRequiredService<AgentService>();
+		public AgentRelayService AgentRelayService => ServiceProvider.GetRequiredService<AgentRelayService>();
+		public ICommitService CommitService => ServiceProvider.GetRequiredService<ICommitService>();
+		public GlobalsService GlobalsService => ServiceProvider.GetRequiredService<GlobalsService>();
+		public MongoService MongoService => ServiceProvider.GetRequiredService<MongoService>();
+		public ITemplateCollection TemplateCollection => ServiceProvider.GetRequiredService<ITemplateCollection>();
+		internal PerforceServiceStub PerforceService => (PerforceServiceStub)ServiceProvider.GetRequiredService<IPerforceService>();
+		public ISubscriptionCollection SubscriptionCollection => ServiceProvider.GetRequiredService<ISubscriptionCollection>();
+		public INotificationService NotificationService => ServiceProvider.GetRequiredService<INotificationService>();
+		public IssueService IssueService => ServiceProvider.GetRequiredService<IssueService>();
+		public JobTaskSource JobTaskSource => ServiceProvider.GetRequiredService<JobTaskSource>();
+		public JobService JobService => ServiceProvider.GetRequiredService<JobService>();
+		public IArtifactCollectionV1 ArtifactCollection => ServiceProvider.GetRequiredService<IArtifactCollectionV1>();
+		public IDowntimeService DowntimeService => ServiceProvider.GetRequiredService<IDowntimeService>();
+		public RpcService RpcService => ServiceProvider.GetRequiredService<RpcService>();
+		public PoolService PoolService => ServiceProvider.GetRequiredService<PoolService>();
+		public LifetimeService LifetimeService => ServiceProvider.GetRequiredService<LifetimeService>();
+		public ScheduleService ScheduleService => ServiceProvider.GetRequiredService<ScheduleService>();
+		public DeviceService DeviceService => ServiceProvider.GetRequiredService<DeviceService>();
+		public StorageService StorageService => ServiceProvider.GetRequiredService<StorageService>();
+		public TestDataService TestDataService => ServiceProvider.GetRequiredService<TestDataService>();
+		public ComputeService ComputeService => ServiceProvider.GetRequiredService<ComputeService>();
+		public ServerStatusService ServerStatusService => ServiceProvider.GetRequiredService<ServerStatusService>();
+
+		public ServerSettings ServerSettings => ServiceProvider.GetRequiredService<IOptions<ServerSettings>>().Value;
+		public IOptionsMonitor<ServerSettings> ServerSettingsMon => ServiceProvider.GetRequiredService<IOptionsMonitor<ServerSettings>>();
+
+		public JobsController JobsController => GetJobsController();
+		public AgentsController AgentsController => GetAgentsController();
+		public PoolsController PoolsController => GetPoolsController();
+		public LeasesController LeasesController => GetLeasesController();
+		public DevicesController DevicesController => GetDevicesController();
+		public DashboardController DashboardController => GetDashboardController();
+		public TestDataController TestDataController => GetTestDataController();
+		public BisectTasksController BisectTasksController => GetBisectTasksController();
+
+		public OpenTelemetry.Trace.Tracer Tracer => ServiceProvider.GetRequiredService<OpenTelemetry.Trace.Tracer>();
+		public Meter Meter => ServiceProvider.GetRequiredService<Meter>();
+		public ConfigService ConfigService => ServiceProvider.GetRequiredService<ConfigService>();
+		public IOptionsMonitor<GlobalConfig> GlobalConfig => ServiceProvider.GetRequiredService<IOptionsMonitor<GlobalConfig>>();
+		public IOptionsSnapshot<GlobalConfig> GlobalConfigSnapshot => ServiceProvider.GetRequiredService<IOptionsSnapshot<GlobalConfig>>();
+
+		private static bool s_datadogWriterPatched;
+
+		public TestSetup()
+		{
+			PatchDatadogWriter();
+		}
+
+		protected void SetConfig(GlobalConfig globalConfig)
+		{
+			globalConfig.PostLoad(ServerSettings);
+			ConfigService.OverrideConfig(globalConfig);
+		}
+
+		protected void UpdateConfig(Action<GlobalConfig> action)
+		{
+			GlobalConfig globalConfig = GlobalConfig.CurrentValue;
+			action(globalConfig);
+			globalConfig.PostLoad(ServerSettings);
+			ConfigService.OverrideConfig(globalConfig);
+		}
+
+		protected override void ConfigureSettings(ServerSettings settings)
+		{
+			DirectoryReference baseDir = DirectoryReference.Combine(ServerApp.DataDir, "Tests");
+			try
+			{
+				FileUtils.ForceDeleteDirectoryContents(baseDir);
+			}
+			catch
+			{
+			}
+
+			settings.WithAws = true;
+
+			settings.ForceConfigUpdateOnStartup = true;
+		}
+
+		protected override void ConfigureServices(IServiceCollection services)
+		{
+			base.ConfigureServices(services);
+
+			IConfiguration config = new ConfigurationBuilder().Build();
+			services.Configure<ServerSettings>(ConfigureSettings);
+			services.AddSingleton<IConfiguration>(config);
+
+			services.AddHttpClient<RpcService>();
+
+			services.AddLogging(builder => { builder.AddConsole().SetMinimumLevel(LogLevel.Debug); });
+			services.AddSingleton<IMemoryCache>(sp => new MemoryCache(new MemoryCacheOptions { }));
+
+			services.AddSingleton(typeof(IAuditLogFactory<>), typeof(AuditLogFactory<>));
+			services.AddSingleton<IAuditLog<AgentId>>(sp => sp.GetRequiredService<IAuditLogFactory<AgentId>>().Create("Agents.Log", "AgentId"));
+			services.AddSingleton<ITelemetrySink, NullTelemetrySink>();
+			services.AddSingleton<ITelemetrySinkInternal, MetricTelemetrySink>();
+			services.AddSingleton<OpenTelemetry.Trace.Tracer>(sp => TracerProvider.Default.GetTracer("TestTracer"));
+			services.AddSingleton(sp => new Meter("TestMeter"));
+
+			services.AddSingleton<ConfigService>();
+			services.AddSingleton<IOptionsFactory<GlobalConfig>>(sp => sp.GetRequiredService<ConfigService>());
+			services.AddSingleton<IOptionsChangeTokenSource<GlobalConfig>>(sp => sp.GetRequiredService<ConfigService>());
+
+			services.AddSingleton<IConfigSource, FileConfigSource>();
+
+			services.AddSingleton<IAccountCollection, AccountCollection>();
+			services.AddSingleton<IAgentCollection, AgentCollection>();
+			services.AddSingleton<IAgentTelemetryCollection, AgentTelemetryCollection>();
+			services.AddSingleton<IArtifactCollection, ArtifactCollection>();
+			services.AddSingleton<IArtifactCollectionV1, ArtifactCollectionV1>();
+			services.AddSingleton<ICommitService, CommitService>();
+			services.AddSingleton<IGraphCollection, GraphCollection>();
+			services.AddSingleton<IIssueCollection, IssueCollection>();
+			services.AddSingleton<IJobCollection, JobCollection>();
+			services.AddSingleton<IJobStepRefCollection, JobStepRefCollection>();
+			services.AddSingleton<IJobTimingCollection, JobTimingCollection>();
+			services.AddSingleton<ILeaseCollection, LeaseCollection>();
+			services.AddSingleton<ILogCollection, LogCollection>();
+			services.AddSingleton<INotificationTriggerCollection, NotificationTriggerCollection>();
+			services.AddSingleton<IPoolCollection, PoolCollection>();
+			services.AddSingleton<IBisectTaskCollection, BisectTaskCollection>();
+			services.AddSingleton<ISessionCollection, SessionCollection>();
+			services.AddSingleton<ISubscriptionCollection, SubscriptionCollection>();
+			services.AddSingleton<IStreamCollection, StreamCollection>();
+			services.AddSingleton<ITemplateCollection, TemplateCollection>();
+			services.AddSingleton<ITestDataCollection, TestDataCollection>();
+			services.AddSingleton<IUtilizationDataCollection, UtilizationDataCollection>();
+			services.AddSingleton<ITemplateCollection, TemplateCollection>();
+			services.AddSingleton<IToolCollection, ToolCollection>();
+			services.AddSingleton<IUgsMetadataCollection, UgsMetadataCollection>();
+			services.AddSingleton<IUserCollection, UserCollectionV2>();
+			services.AddSingleton<IDeviceCollection, DeviceCollection>();
+			services.AddSingleton<IDashboardPreviewCollection, DashboardPreviewCollection>();
+
+			services.AddSingleton<FakeClock>();
+			services.AddSingleton<IClock>(sp => sp.GetRequiredService<FakeClock>());
+			services.AddSingleton<IHostApplicationLifetime, AppLifetimeStub>();
+			services.AddSingleton<IHostEnvironment, WebHostEnvironmentStub>();
+
+			// Empty mocked object to satisfy basic test runs
+			services.AddSingleton<IAmazonEC2>(sp => new Mock<IAmazonEC2>().Object);
+			services.AddSingleton<IAmazonAutoScaling>(sp => new Mock<IAmazonAutoScaling>().Object);
+			services.AddSingleton<IAmazonCloudWatch>(sp => new Mock<IAmazonCloudWatch>().Object);
+			services.AddSingleton<IFleetManagerFactory, FleetManagerFactory>();
+
+			services.AddSingleton<IPoolSizeStrategyFactory, NoOpPoolSizeStrategyFactory>();
+			services.AddSingleton<IPoolSizeStrategyFactory, JobQueueStrategyFactory>();
+			services.AddSingleton<IPoolSizeStrategyFactory, LeaseUtilizationStrategyFactory>();
+			services.AddSingleton<IPoolSizeStrategyFactory, LeaseUtilizationAwsMetricStrategyFactory>();
+
+			services.AddSingleton<AclService>();
+			services.AddSingleton<AgentService>();
+			services.AddSingleton(provider => new Lazy<AgentService>(provider.GetRequiredService<AgentService>));
+			services.AddSingleton<AgentRelayService>();
+			services.AddSingleton<AwsAutoScalingLifecycleService>();
+			services.AddSingleton<ArtifactExpirationService>();
+			services.AddSingleton<FleetService>();
+			services.AddSingleton<ConsistencyService>();
+			services.AddSingleton<RequestTrackerService>();
+			services.AddSingleton<GlobalsService>();
+			services.AddSingleton<JobTaskSource>();
+			services.AddSingleton<IDowntimeService, DowntimeServiceStub>();
+			services.AddSingleton<IssueService>();
+			services.AddSingleton<JobService>();
+			services.AddSingleton<JobExpirationService>();
+			services.AddSingleton<LifetimeService>();
+			services.AddSingleton<ILogStorage, NullLogStorage>();
+			services.AddSingleton<LogTailService>();
+			services.AddSingleton<INotificationService, NotificationService>();
+			services.AddSingleton<IPerforceService, PerforceServiceStub>();
+			services.AddSingleton<PerforceLoadBalancer>();
+			services.AddSingleton<PoolService>();
+			services.AddSingleton<BisectService>();
+			services.AddSingleton<RpcService>();
+			services.AddSingleton<ScheduleService>();
+			services.AddSingleton<DeviceService>();
+			services.AddSingleton<TestDataService>();
+			services.AddSingleton<ComputeService>();
+			services.AddSingleton<EnrollmentService>();
+
+			services.AddSingleton(typeof(IHealthMonitor<>), typeof(HealthMonitor<>));
+			services.AddSingleton<ServerStatusService>();
+
+			services.AddSingleton<ConformTaskSource>();
+			services.AddSingleton<ICommitService, CommitService>();
+
+			services.AddSingleton<FileObjectStoreFactory>();
+			services.AddSingleton<IObjectStoreFactory, ObjectStoreFactory>();
+			services.AddSingleton<IObjectStore<PersistentLogStorage>>(sp => new MemoryObjectStore().ForType<PersistentLogStorage>());
+			services.AddSingleton<IObjectStore<ArtifactCollectionV1>>(sp => new MemoryObjectStore().ForType<ArtifactCollectionV1>());
+
+			services.AddSingleton<StorageService>();
+			services.AddSingleton<StorageBackendCache>();
+			services.AddSingleton<BundleCache>();
+		}
+
+		public Task<Fixture> CreateFixtureAsync()
+		{
+			return Fixture.CreateAsync(ConfigService, GraphCollection, TemplateCollection, JobService, ArtifactCollection, AgentService, ServerSettings);
+		}
+
+		private JobsController GetJobsController()
+		{
+			JobsController jobsCtrl = ActivatorUtilities.CreateInstance<JobsController>(ServiceProvider);
+			jobsCtrl.ControllerContext = GetControllerContext();
+			return jobsCtrl;
+		}
+
+		private DevicesController GetDevicesController()
+		{
+			DevicesController devicesCtrl = ActivatorUtilities.CreateInstance<DevicesController>(ServiceProvider);
+			devicesCtrl.ControllerContext = GetControllerContext();
+			return devicesCtrl;
+		}
+
+		private DashboardController GetDashboardController()
+		{
+			DashboardController dashboardCtrl = ActivatorUtilities.CreateInstance<DashboardController>(ServiceProvider);
+			dashboardCtrl.ControllerContext = GetControllerContext();
+			return dashboardCtrl;
+		}
+
+		private TestDataController GetTestDataController()
+		{
+			TestDataController dataCtrl = ActivatorUtilities.CreateInstance<TestDataController>(ServiceProvider);
+			dataCtrl.ControllerContext = GetControllerContext();
+			return dataCtrl;
+		}
+
+		private BisectTasksController GetBisectTasksController()
+		{
+			BisectTasksController bisectCtrl = ActivatorUtilities.CreateInstance<BisectTasksController>(ServiceProvider);
+			bisectCtrl.ControllerContext = GetControllerContext();
+			return bisectCtrl;
+		}
+
+		private AgentsController GetAgentsController()
+		{
+			AgentsController agentCtrl = ActivatorUtilities.CreateInstance<AgentsController>(ServiceProvider);
+			agentCtrl.ControllerContext = GetControllerContext();
+			return agentCtrl;
+		}
+
+		private PoolsController GetPoolsController()
+		{
+			PoolsController controller = ActivatorUtilities.CreateInstance<PoolsController>(ServiceProvider);
+			controller.ControllerContext = GetControllerContext();
+			return controller;
+		}
+
+		private LeasesController GetLeasesController()
+		{
+			LeasesController controller = ActivatorUtilities.CreateInstance<LeasesController>(ServiceProvider);
+			controller.ControllerContext = GetControllerContext();
+			return controller;
+		}
+
+		private static ControllerContext GetControllerContext()
+		{
+			ControllerContext controllerContext = new ControllerContext();
+			controllerContext.HttpContext = new DefaultHttpContext();
+			controllerContext.HttpContext.User = new ClaimsPrincipal(new ClaimsIdentity(
+				new List<Claim> { HordeClaims.AdminClaim.ToClaim() }, "TestAuthType"));
+			return controllerContext;
+		}
+
+		private static int s_agentIdCounter = 1;
+		public Task<IAgent> CreateAgentAsync(IPool pool, bool enabled = true, bool requestShutdown = false, List<string>? properties = null, List<AgentWorkspaceInfo>? workspaces = null, TimeSpan? adjustClockBy = null)
+		{
+			return CreateAgentAsync(pool.Id, enabled, requestShutdown, properties, workspaces, adjustClockBy);
+		}
+
+		/// <summary>
+		/// Helper function for setting up agents to be used in tests
+		/// </summary>
+		/// <param name="poolId">Pool ID which the agent should belong to</param>
+		/// <param name="enabled">Whether set the agent as enabled</param>
+		/// <param name="requestShutdown">Mark it with a request for shutdown</param>
+		/// <param name="properties">Any properties to assign</param>
+		/// <param name="workspaces">Any workspaces to assign</param>
+		/// <param name="adjustClockBy">Time span to temporarily skew the clock when creating the agent</param>
+		/// <param name="awsInstanceId">AWS instance ID for the agent (will be set in properties)</param>
+		/// <param name="lease">A lease to assign the agent</param>
+		/// <param name="ephemeral">Whether the agent is ephemeral</param>
+		/// <returns>A new agent</returns>
+		public async Task<IAgent> CreateAgentAsync(
+			PoolId? poolId,
+			bool enabled = true,
+			bool requestShutdown = false,
+			List<string>? properties = null,
+			List<AgentWorkspaceInfo>? workspaces = null,
+			TimeSpan? adjustClockBy = null,
+			string? awsInstanceId = null,
+			AgentLease? lease = null,
+			bool ephemeral = false)
+		{
+			DateTime now = Clock.UtcNow;
+			if (adjustClockBy != null)
+			{
+				Clock.UtcNow = now + adjustClockBy.Value;
+			}
+
+			Dictionary<string, int> resources = new();
+			List<string> tempProps = new(properties ?? new List<string>());
+			if (awsInstanceId != null)
+			{
+				tempProps.Add(KnownPropertyNames.AwsInstanceId + "=" + awsInstanceId);
+			}
+
+			IAgent? agent = await AgentService.CreateAgentAsync("TestAgent" + s_agentIdCounter++, ephemeral, "");
+			Assert.IsNotNull(agent);
+
+			agent = await agent.TryUpdateAsync(new UpdateAgentOptions { Enabled = enabled, ExplicitPools = poolId != null ? [poolId.Value] : [] });
+			Assert.IsNotNull(agent);
+
+			agent = await AgentService.CreateSessionAsync(agent, AgentStatus.Ok, tempProps, resources, null);
+			Assert.IsNotNull(agent);
+
+			if (workspaces is { Count: > 0 })
+			{
+				await agent.TryUpdateWorkspacesAsync(workspaces, false);
+			}
+			
+			if (requestShutdown)
+			{
+				await agent.TryUpdateAsync(new UpdateAgentOptions { RequestShutdown = true });
+			}
+
+			if (lease != null)
+			{
+				await agent.TryAddLeaseAsync(lease);
+			}
+
+			Clock.UtcNow = now;
+			return agent;
+		}
+
+		/// <summary>
+		/// Hack the Datadog tracing library to not block during shutdown of tests.
+		/// Without this fix, the lib will try to send traces to a host that isn't running and block for +20 secs
+		///
+		/// Since so many of the interfaces and classes in the lib are internal it was difficult to replace Tracer.Instance
+		/// </summary>
+		private static void PatchDatadogWriter()
+		{
+			if (s_datadogWriterPatched)
+			{
+				return;
+			}
+
+			s_datadogWriterPatched = true;
+
+			string msg = "Unable to patch Datadog agent writer! Tests will still work, but shutdown will block for +20 seconds.";
+
+			FieldInfo? agentWriterField = Datadog.Trace.Tracer.Instance.GetType().GetField("_agentWriter", BindingFlags.NonPublic | BindingFlags.Instance);
+			if (agentWriterField == null)
+			{
+				Console.Error.WriteLine(msg);
+				return;
+			}
+
+			object? agentWriterInstance = agentWriterField.GetValue(Datadog.Trace.Tracer.Instance);
+			if (agentWriterInstance == null)
+			{
+				Console.Error.WriteLine(msg);
+				return;
+			}
+
+			FieldInfo? processExitField = agentWriterInstance.GetType().GetField("_processExit", BindingFlags.NonPublic | BindingFlags.Instance);
+			if (processExitField == null)
+			{
+				Console.Error.WriteLine(msg);
+				return;
+			}
+
+			TaskCompletionSource<bool>? processExitInstance = (TaskCompletionSource<bool>?)processExitField.GetValue(agentWriterInstance);
+			if (processExitInstance == null)
+			{
+				Console.Error.WriteLine(msg);
+				return;
+			}
+
+			processExitInstance.TrySetResult(true);
+		}
+
+		/// <summary>
+		/// Find an available TCP/IP port
+		/// </summary>
+		/// <returns>Port number available</returns>
+		public static int GetAvailablePort()
+		{
+			using TcpListener listener = new(IPAddress.Loopback, 0);
+			listener.Start();
+			int port = ((IPEndPoint)listener.LocalEndpoint).Port;
+			listener.Stop();
+			return port;
+		}
+
+		/// <summary>
+		/// Create a console logger for tests
+		/// </summary>
+		/// <typeparam name="T">Type to instantiate</typeparam>
+		/// <returns>A logger</returns>
+		public static ILogger<T> CreateConsoleLogger<T>()
+		{
+			using ILoggerFactory loggerFactory = LoggerFactory.Create(builder =>
+			{
+				builder.SetMinimumLevel(LogLevel.Debug);
+				builder.AddSimpleConsole(options => { options.SingleLine = true; });
+			});
+
+			return loggerFactory.CreateLogger<T>();
+		}
+
+		protected async Task<IPool> CreatePoolAsync(PoolConfig poolConfig)
+		{
+			UpdateConfig(config => config.Pools.Add(poolConfig));
+			return await PoolCollection.GetAsync(poolConfig.Id) ?? throw new NotImplementedException();
+		}
+	}
+
+	public class DowntimeServiceStub : IDowntimeService
+	{
+		public DowntimeServiceStub(bool isDowntimeActive = false)
+		{
+			IsDowntimeActive = isDowntimeActive;
+		}
+
+		public bool IsDowntimeActive { get; set; }
+	}
+}
\ No newline at end of file
diff --git a/Engine/Source/Programs/Horde/HordeServer/Acls/AclService.cs b/Engine/Source/Programs/Horde/HordeServer/Acls/AclService.cs
new file mode 100644
--- /dev/null
+++ b/Engine/Source/Programs/Horde/HordeServer/Acls/AclService.cs
@@ -0,0 +1,137 @@
+// Copyright Epic Games, Inc. All Rights Reserved.
+
+using System;
+using System.Collections.Generic;
+using System.IdentityModel.Tokens.Jwt;
+using System.Linq;
+using System.Security.Claims;
+using System.Threading;
+using System.Threading.Tasks;
+using EpicGames.Horde.Agents;
+using EpicGames.Horde.Agents.Leases;
+using EpicGames.Horde.Agents.Sessions;
+using Horde.Server.Server;
+using Horde.Server.Utilities;
+using Microsoft.IdentityModel.Tokens;
+using MongoDB.Driver;
+
+namespace Horde.Server.Acls
+{
+	/// <summary>
+	/// Wraps functionality for manipulating permissions
+	/// </summary>
+	public class AclService
+	{
+		private readonly GlobalsService _globalsService;
+
+		/// <summary>
+		/// Constructor
+		/// </summary>
+		public AclService(GlobalsService globalsService)
+		{
+			_globalsService = globalsService;
+		}
+
+		/// <summary>
+		/// Issues a bearer token with the given roles
+		/// </summary>
+		/// <param name="claims">List of claims to include</param>
+		/// <param name="expiry">Time that the token expires</param>
+		/// <param name="cancellationToken">Cancellation token for the operation</param>
+		/// <returns>JWT security token with a claim for creating new agents</returns>
+		public async ValueTask<string> IssueBearerTokenAsync(IEnumerable<AclClaimConfig> claims, TimeSpan? expiry, CancellationToken cancellationToken = default)
+		{
+			return await IssueBearerTokenAsync(claims.Select(x => new Claim(x.Type, x.Value)), expiry, cancellationToken);
+		}
+
+		/// <summary>
+		/// Issues a bearer token with the given claims
+		/// </summary>
+		/// <param name="claims">List of claims to include</param>
+		/// <param name="expiry">Time that the token expires</param>
+		/// <param name="cancellationToken">Cancellation token for the operation</param>
+		/// <returns>JWT security token with a claim for creating new agents</returns>
+		public async ValueTask<string> IssueBearerTokenAsync(IEnumerable<Claim> claims, TimeSpan? expiry, CancellationToken cancellationToken = default)
+		{
+			IGlobals globals = await _globalsService.GetAsync(cancellationToken);
+			SigningCredentials signingCredentials = new(globals.JwtSigningKey, SecurityAlgorithms.HmacSha256);
+
+			JwtSecurityToken token = new(globals.JwtIssuer, null, claims.DistinctBy(x => (x.Type, x.Value)), null, DateTime.UtcNow + expiry, signingCredentials);
+			return new JwtSecurityTokenHandler().WriteToken(token);
+		}
+
+		/// <summary>
+		/// Gets the agent id associated with a particular user
+		/// </summary>
+		/// <param name="user"></param>
+		/// <returns></returns>
+		public static AgentId? GetAgentId(ClaimsPrincipal user)
+		{
+			Claim? claim = user.Claims.FirstOrDefault(x => x.Type == HordeClaimTypes.Agent);
+			if (claim == null)
+			{
+				return null;
+			}
+			else
+			{
+				return new AgentId(claim.Value);
+			}
+		}
+	}
+
+	internal static class ClaimExtensions
+	{
+		public static bool HasAdminClaim(this ClaimsPrincipal user)
+		{
+			return user.HasClaim(HordeClaims.AdminClaim.Type, HordeClaims.AdminClaim.Value);
+		}
+
+		public static bool HasAgentClaim(this ClaimsPrincipal user, AgentId agentId)
+		{
+			return user.HasClaim(HordeClaimTypes.Agent, agentId.ToString());
+		}
+
+		public static bool HasLeaseClaim(this ClaimsPrincipal user, LeaseId leaseId)
+		{
+			return user.HasClaim(HordeClaimTypes.Lease, leaseId.ToString());
+		}
+
+		public static bool HasSessionClaim(this ClaimsPrincipal user, SessionId sessionId)
+		{
+			return user.HasClaim(HordeClaimTypes.AgentSessionId, sessionId.ToString());
+		}
+
+		public static LeaseId? GetLeaseClaim(this ClaimsPrincipal user)
+		{
+			Claim? claim = user.FindFirst(HordeClaimTypes.Lease);
+			if (claim == null || !LeaseId.TryParse(claim.Value, out LeaseId leaseIdValue))
+			{
+				return null;
+			}
+			else
+			{
+				return leaseIdValue;
+			}
+		}
+
+		public static SessionId? GetSessionClaim(this ClaimsPrincipal user)
+		{
+			Claim? claim = user.FindFirst(HordeClaimTypes.AgentSessionId);
+			if (claim == null || !SessionId.TryParse(claim.Value, out SessionId sessionIdValue))
+			{
+				return null;
+			}
+			else
+			{
+				return sessionIdValue;
+			}
+		}
+
+		public static string GetSessionClaimsAsString(this ClaimsPrincipal user)
+		{
+			return String.Join(",", user.Claims
+				.Where(c => c.Type == HordeClaimTypes.AgentSessionId)
+				.Select(c => c.Value));
+		}
+	}
+}
diff --git a/Engine/Source/Programs/Horde/HordeServer/Agents/AgentService.cs b/Engine/Source/Programs/Horde/HordeServer/Agents/AgentService.cs
new file mode 100644
--- /dev/null
+++ b/Engine/Source/Programs/Horde/HordeServer/Agents/AgentService.cs
@@ -0,0 +1,1240 @@
+// Copyright Epic Games, Inc. All Rights Reserved.
+
+using System;
+using System.Collections.Generic;
+using System.Diagnostics;
+using System.Diagnostics.Metrics;
+using System.Linq;
+using System.Security.Claims;
+using System.Threading;
+using System.Threading.Tasks;
+using EpicGames.Core;
+using EpicGames.Horde.Agents;
+using EpicGames.Horde.Agents.Leases;
+using EpicGames.Horde.Agents.Pools;
+using EpicGames.Horde.Agents.Sessions;
+using EpicGames.Redis;
+using EpicGames.Serialization;
+using Google.Protobuf.WellKnownTypes;
+using Horde.Server.Acls;
+using Horde.Server.Agents.Leases;
+using Horde.Server.Agents.Pools;
+using Horde.Server.Agents.Sessions;
+using Horde.Server.Auditing;
+using Horde.Server.Server;
+using Horde.Server.Tasks;
+using Horde.Server.Utilities;
+using HordeCommon.Rpc;
+using HordeCommon.Rpc.Messages;
+using Microsoft.Extensions.Hosting;
+using Microsoft.Extensions.Logging;
+using OpenTelemetry.Trace;
+using StackExchange.Redis;
+
+namespace Horde.Server.Agents
+{
+	/// <summary>
+	/// Singleton used to store agent costs
+	/// </summary>
+	[RedisConverter(typeof(RedisCbConverter<>))]
+	public class AgentRateTable
+	{
+		/// <summary>
+		/// List of costs for different agent types
+		/// </summary>
+		[CbField]
+		public List<AgentRateConfig> Entries { get; set; } = new List<AgentRateConfig>();
+	}
+
+	/// <summary>
+	/// Wraps functionality for manipulating agents
+	/// </summary>
+	public sealed class AgentService : IHostedService, IAsyncDisposable
+	{
+		/// <summary>
+		/// Maximum time between updates for an agent to be considered online
+		/// </summary>
+		public static readonly TimeSpan SessionExpiryTime = TimeSpan.FromMinutes(5);
+
+		/// <summary>
+		/// Time before a session expires that we will poll until
+		/// </summary>
+		public static readonly TimeSpan SessionLongPollTime = TimeSpan.FromSeconds(55);
+
+		/// <summary>
+		/// Time after which a session will be renewed
+		/// </summary>
+		public static readonly TimeSpan SessionRenewTime = TimeSpan.FromSeconds(50);
+
+		/// <summary>
+		/// Collection of agent documents
+		/// </summary>
+		public IAgentCollection Agents { get; }
+
+		readonly ILeaseCollection _leases;
+		readonly ISessionCollection _sessions;
+		readonly AclService _aclService;
+		readonly IDowntimeService _downtimeService;
+		readonly ITaskSource[] _taskSources;
+		readonly RedisService _redisService;
+		readonly IHostApplicationLifetime _applicationLifetime;
+		readonly IClock _clock;
+		readonly Tracer _tracer;
+		readonly ILogger _logger;
+		readonly ITicker _ticker;
+		readonly ITicker _sharedTicker;
+
+		readonly RedisStringKey<AgentRateTable> _agentRateTableData = new ("agent-rates");
+
+		/// <summary>Lazily updated costs for different agent types</summary>
+		readonly AsyncCachedValue<AgentRateTable?> _cachedRates;
+
+		/// <summary>Lazily updated list of current pools</summary>
+		readonly AsyncCachedValue<IReadOnlyList<IPoolConfig>> _cachedPools;
+		
+		/// <summary>Manually updated cached list of current agents</summary>
+		IReadOnlyDictionary<AgentId, IAgent>? _cachedAgents;
+
+		/// <summary>All the agents currently performing a long poll for work on this server</summary>
+		readonly Dictionary<AgentId, CancellationTokenSource> _waitingAgents = new ();
+
+		/// <summary>OpenTelemetry measurements (gauges)</summary>
+		IEnumerable<Measurement<int>> _measurements = new List<Measurement<int>>();
+
+		/// <summary>Subscription for update events</summary>
+		IAsyncDisposable? _subscription;
+
+		/// <summary>
+		/// Constructor
+		/// </summary>
+		public AgentService(
+			IAgentCollection agents,
+			ILeaseCollection leases,
+			ISessionCollection sessions,
+			AclService aclService,
+			IDowntimeService downtimeService,
+			IPoolCollection poolCollection,
+			IEnumerable<ITaskSource> taskSources,
+			RedisService redisService,
+			IHostApplicationLifetime applicationLifetime,
+			IClock clock,
+			Tracer tracer,
+			Meter meter,
+			ILogger<AgentService> logger)
+		{
+			Agents = agents;
+			_leases = leases;
+			_sessions = sessions;
+			_aclService = aclService;
+			_downtimeService = downtimeService;
+			_cachedRates = new AsyncCachedValue<AgentRateTable?>(_ => redisService.GetDatabase().StringGetAsync(_agentRateTableData), TimeSpan.FromMinutes(2));
+			_cachedPools = new AsyncCachedValue<IReadOnlyList<IPoolConfig>>(poolCollection.GetConfigsAsync, TimeSpan.FromSeconds(30.0));
+			_taskSources = taskSources.ToArray();
+			_applicationLifetime = applicationLifetime;
+			_redisService = redisService;
+			_clock = clock;
+			_ticker = clock.AddTicker($"{nameof(AgentService)}.{nameof(TickAsync)}", TimeSpan.FromSeconds(30.0), TickAsync, logger);
+			_sharedTicker = clock.AddSharedTicker($"{nameof(AgentService)}.{nameof(TickSharedAsync)}", TimeSpan.FromSeconds(30.0), TickSharedAsync, logger);
+			_tracer = tracer;
+			_logger = logger;
+
+			meter.CreateObservableGauge("horde.agent.count", () => _measurements);
+		}
+
+		/// <inheritdoc/>
+		public async Task StartAsync(CancellationToken cancellationToken)
+		{
+			await _ticker.StartAsync();
+			await _sharedTicker.StartAsync();
+			_subscription = await Agents.SubscribeToUpdateEventsAsync(OnAgentUpdate);
+		}
+
+		/// <inheritdoc/>
+		public async Task StopAsync(CancellationToken cancellationToken)
+		{
+			if (_subscription != null)
+			{
+				await _subscription.DisposeAsync();
+				_subscription = null;
+			}
+			await _ticker.StopAsync();
+			await _sharedTicker.StopAsync();
+		}
+
+		/// <inheritdoc/>
+		public async ValueTask DisposeAsync()
+		{
+			await _cachedRates.DisposeAsync();
+			await _cachedPools.DisposeAsync();
+			await _ticker.DisposeAsync();
+			await _sharedTicker.DisposeAsync();
+		}
+
+		/// <summary>
+		/// Gets user-readable payload information
+		/// </summary>
+		/// <param name="payload">The payload data</param>
+		/// <param name="cancellationToken">Cancellation token for the operation</param>
+		/// <returns>Dictionary of key/value pairs for the payload</returns>
+		public async ValueTask<Dictionary<string, string>?> GetPayloadDetailsAsync(ReadOnlyMemory<byte>? payload, CancellationToken cancellationToken = default)
+		{
+			Dictionary<string, string>? details = null;
+			if (payload != null)
+			{
+				Any basePayload = Any.Parser.ParseFrom(payload.Value.ToArray());
+				foreach (ITaskSource taskSource in _taskSources)
+				{
+					if (basePayload.Is(taskSource.Descriptor))
+					{
+						details = new Dictionary<string, string>();
+						await taskSource.GetLeaseDetailsAsync(basePayload, details, cancellationToken);
+						break;
+					}
+				}
+			}
+			return details;
+		}
+
+		/// <summary>
+		/// Issues a bearer token for the given session id
+		/// </summary>
+		/// <param name="agentId">The agent id</param>
+		/// <param name="sessionId">The session id</param>
+		/// <param name="cancellationToken">Cancellation token for the operation</param>
+		/// <returns>Bearer token for the agent</returns>
+		public async ValueTask<string> IssueSessionTokenAsync(AgentId agentId, SessionId sessionId, CancellationToken cancellationToken = default)
+		{
+			List<AclClaimConfig> claims = new List<AclClaimConfig>();
+			claims.Add(HordeClaims.AgentRoleClaim);
+			claims.Add(HordeClaims.GetAgentClaim(agentId));
+			claims.Add(HordeClaims.GetSessionClaim(sessionId));
+			return await _aclService.IssueBearerTokenAsync(claims, null, cancellationToken);
+		}
+
+		/// <summary>
+		/// Register a new agent
+		/// </summary>
+		/// <param name="name">Name of the agent</param>
+		/// <param name="ephemeral">Whether the agent is ephemeral or not</param>
+		/// <param name="enrollmentKey">Key for enrolling the agent</param>
+		/// <param name="cancellationToken">Cancellation token for the operation</param>
+		/// <returns>Unique id for the agent</returns>
+		public Task<IAgent> CreateAgentAsync(string name, bool ephemeral, string enrollmentKey, CancellationToken cancellationToken = default)
+		{
+			return CreateAgentAsync(new AgentId(name), ephemeral, enrollmentKey, cancellationToken);
+		}
+
+		/// <summary>
+		/// Register a new agent
+		/// </summary>
+		/// <param name="agentId">Agent id</param>
+		/// <param name="ephemeral">Whether the agent is ephemeral or not</param>
+		/// <param name="enrollmentKey">Key for enrolling the agent</param>
+		/// <param name="cancellationToken">Cancellation token for the operation</param>
+		/// <returns>Unique id for the agent</returns>
+		public async Task<IAgent> CreateAgentAsync(AgentId agentId, bool ephemeral, string enrollmentKey, CancellationToken cancellationToken = default)
+		{
+			for (; ; )
+			{
+				IAgent? agent = await Agents.GetAsync(agentId, cancellationToken);
+				if (agent == null)
+				{
+					return await Agents.AddAsync(agentId, ephemeral, enrollmentKey, cancellationToken);
+				}
+
+				agent = await agent.TryResetAsync(ephemeral, enrollmentKey, cancellationToken);
+				if (agent != null)
+				{
+					return agent;
+				}
+			}
+		}
+
+		/// <summary>
+		/// Gets an agent by ID
+		/// </summary>
+		/// <param name="agentId">Unique id of the agent</param>
+		/// <param name="cancellationToken">Cancellation token for the operation</param>
+		/// <returns>The agent document</returns>
+		public Task<IAgent?> GetAgentAsync(AgentId agentId, CancellationToken cancellationToken = default)
+		{
+			return Agents.GetAsync(agentId, cancellationToken);
+		}
+
+		/// <summary>
+		/// Finds all agents matching certain criteria
+		/// </summary>
+		/// <param name="poolId">The pool containing the agent</param>
+		/// <param name="modifiedAfter">If set, only returns agents modified after this time</param>
+		/// <param name="property">If set, only return agents matching this property</param>
+		/// <param name="includeDeleted">If set, include agents marked as deleted</param>
+		/// <param name="index">Index within the list of results</param>
+		/// <param name="count">Number of results to return</param>
+		/// <param name="cancellationToken">Cancellation token for the operation</param>
+		/// <returns>List of agents matching the given criteria</returns>
+		public Task<IReadOnlyList<IAgent>> FindAgentsAsync(PoolId? poolId, DateTime? modifiedAfter, string? property, bool includeDeleted, int? index, int? count, CancellationToken cancellationToken)
+		{
+			return Agents.FindAsync(poolId, modifiedAfter, property, null, null, includeDeleted, index, count, true, cancellationToken);
+		}
+		
+		/// <summary>
+		/// Get all agents from local in-memory cache
+		/// </summary>
+		/// <param name="cancellationToken">Cancellation token for the operation</param>
+		/// <returns>List of all agents</returns>
+		public async Task<IReadOnlyList<IAgent>> GetCachedAgentsAsync(CancellationToken cancellationToken)
+		{
+			if (_cachedAgents == null)
+			{
+				// First time this method is called, cached agents may not be initialized
+				// This could lead to multiple concurrent agent refreshes (dog-piling), but that's fine as it should only happen once
+				await RefreshCachedAgentsAsync(cancellationToken);
+			}
+
+			if (_cachedAgents == null)
+			{
+				throw new Exception("Cached agents not initialized");
+			}
+
+			return _cachedAgents.Values.ToList();
+		}
+
+		/// <summary>
+		/// Update the current workspaces for an agent.
+		/// </summary>
+		/// <param name="agent">The agent to update</param>
+		/// <param name="workspaces">Current list of workspaces</param>
+		/// <param name="pendingConform">Whether the agent still needs to run another conform</param>
+		/// <param name="cancellationToken">Cancellation token for the operation</param>
+		/// <returns>New agent state</returns>
+		public async Task<bool> TryUpdateWorkspacesAsync(IAgent agent, List<AgentWorkspaceInfo> workspaces, bool pendingConform, CancellationToken cancellationToken)
+		{
+			IAgent? newAgent = await agent.TryUpdateWorkspacesAsync(workspaces, pendingConform, cancellationToken);
+			return newAgent != null;
+		}
+
+		/// <summary>
+		/// Marks the agent as deleted
+		/// </summary>
+		/// <param name="agent">The agent to delete</param>
+		/// <param name="forceDelete">Whether to fully delete the agent as opposed to just marking it as deleted</param>
+		/// <param name="cancellationToken">Cancellation token for the operation</param>
+		/// <returns>Async task</returns>
+		public async Task DeleteAgentAsync(IAgent? agent, bool forceDelete = false, CancellationToken cancellationToken = default)
+		{
+			if (agent == null)
+			{
+				return;
+			}
+
+			if (forceDelete)
+			{
+				await Agents.ForceDeleteAsync(agent.Id, cancellationToken);
+				return;
+			}
+
+			while (agent is { Deleted: false })
+			{
+				IAgent? newAgent = await agent.TryDeleteAsync(cancellationToken);
+				if (newAgent != null)
+				{
+					break;
+				}
+				agent = await GetAgentAsync(agent.Id, cancellationToken);
+			}
+		}
+
+		async ValueTask<List<PoolId>> GetDynamicPoolsAsync(IAgent agent, CancellationToken cancellationToken)
+		{
+			List<PoolId> newDynamicPools = new List<PoolId>();
+
+			IReadOnlyList<IPoolConfig> pools = await _cachedPools.GetAsync(cancellationToken);
+			foreach (IPoolConfig pool in pools)
+			{
+				if (pool.Condition != null && agent.SatisfiesCondition(pool.Condition))
+				{
+					newDynamicPools.Add(pool.Id);
+				}
+			}
+
+			return newDynamicPools;
+		}
+
+		/// <summary>
+		/// Callback for an agents 
+		/// </summary>
+		/// <param name="agentId"></param>
+		void OnAgentUpdate(AgentId agentId)
+		{
+			lock (_waitingAgents)
+			{
+				if (_waitingAgents.TryGetValue(agentId, out CancellationTokenSource? cancellationSource))
+				{
+					cancellationSource.Cancel();
+				}
+			}
+		}
+
+		/// <summary>
+		/// Creates a new agent session
+		/// </summary>
+		/// <param name="agent">The agent to create a session for</param>
+		/// <param name="status">Current status of the agent</param>
+		/// <param name="properties">Properties for the agent</param>
+		/// <param name="resources">Resources which the agent has</param>
+		/// <param name="version">Version of the software that's running</param>
+		/// <param name="cancellationToken">Cancellation token for the operation</param>
+		/// <returns>New agent state</returns>
+		public async Task<IAgent> CreateSessionAsync(IAgent agent, AgentStatus status, IReadOnlyList<string> properties, IReadOnlyDictionary<string, int> resources, string? version, CancellationToken cancellationToken = default)
+		{
+			DateTime? lastStatusChange = null;
+			for (; ; )
+			{
+				IAuditLogChannel<AgentId> agentLogger = Agents.GetLogger(agent.Id);
+
+				// Check if there's already a session running for this agent.
+				IAgent? newAgent;
+				if (agent.SessionId != null)
+				{
+					// Save last status change timestamp to avoid registering the change to "stopped" when it's re-created immediately after to "ok".
+					lastStatusChange = agent.LastStatusChange;
+
+					// Try to terminate the current session
+					await TryTerminateSessionAsync(agent, cancellationToken);
+				}
+				else
+				{
+					DateTime utcNow = _clock.UtcNow;
+
+					// Remove any outstanding leases
+					foreach (AgentLease lease in agent.Leases)
+					{
+						agentLogger.LogInformation("Removing outstanding lease {LeaseId}", lease.Id);
+						await RemoveLeaseAsync(agent, lease, utcNow, LeaseOutcome.Failed, null, cancellationToken);
+					}
+
+					// Create a new session document
+					ISession newSession = await _sessions.AddAsync(SessionIdUtils.GenerateNewId(), agent.Id, _clock.UtcNow, properties, resources, version, cancellationToken);
+					DateTime sessionExpiresAt = utcNow + SessionExpiryTime;
+
+					// Get the new pools for the agent
+					List<PoolId> dynamicPools = await GetDynamicPoolsAsync(agent, cancellationToken);
+
+					// Reset the agent to use the new session
+					newAgent = await agent.TryCreateSessionAsync(new CreateSessionOptions(newSession.Id, sessionExpiresAt, status, properties, resources, dynamicPools, lastStatusChange ?? utcNow, version), cancellationToken);
+					if (newAgent != null)
+					{
+						LogPropertyChanges(agentLogger, agent.Properties, newAgent.Properties);
+						agent = newAgent;
+						agentLogger.LogInformation("Session {SessionId} started", newSession.Id);
+						break;
+					}
+
+					// Remove the session we didn't use
+					await _sessions.DeleteAsync(newSession.Id, cancellationToken);
+				}
+
+				// Get the current agent state
+				newAgent = await GetAgentAsync(agent.Id, cancellationToken);
+				if (newAgent == null)
+				{
+					throw new InvalidOperationException($"Invalid agent id '{agent.Id}'");
+				}
+				agent = newAgent;
+			}
+			return agent;
+		}
+
+		/// <summary>
+		/// Determines whether a task source can currently issue tasks
+		/// </summary>
+		bool CanUseTaskSource(IAgent agent, ITaskSource taskSource)
+		{
+			TaskSourceFlags flags = taskSource.Flags;
+			if ((flags & TaskSourceFlags.AllowWhenBusy) == 0 && agent.Status == AgentStatus.Busy)
+			{
+				return false;
+			}
+			if ((flags & TaskSourceFlags.AllowWhenDisabled) == 0 && !agent.Enabled)
+			{
+				return false;
+			}
+			if ((flags & TaskSourceFlags.AllowDuringDowntime) == 0 && _downtimeService.IsDowntimeActive)
+			{
+				return false;
+			}
+			return true;
+		}
+
+		/// <summary>
+		/// Compare properties and write changes to audit log
+		/// </summary>
+		private static void LogPropertyChanges(IAuditLogChannel<AgentId> agentLogger, IReadOnlyList<string> before, IReadOnlyList<string> after)
+		{
+			const string AwsInstanceTypeKey = KnownPropertyNames.AwsInstanceType + "=";
+			string beforeProp = before.FirstOrDefault(x => x.StartsWith(AwsInstanceTypeKey, StringComparison.Ordinal), String.Empty);
+			string afterProp = after.FirstOrDefault(x => x.StartsWith(AwsInstanceTypeKey, StringComparison.Ordinal), String.Empty);
+
+			if (!String.IsNullOrEmpty(beforeProp) && !String.IsNullOrEmpty(afterProp) && beforeProp != afterProp)
+			{
+				string oldInstanceType = beforeProp.Replace(AwsInstanceTypeKey, "", StringComparison.Ordinal);
+				string newInstanceType = afterProp.Replace(AwsInstanceTypeKey, "", StringComparison.Ordinal);
+				agentLogger.LogInformation("AWS EC2 instance type changed from {OldInstanceType} to {NewInstanceType}", oldInstanceType, newInstanceType);
+			}
+		}
+
+		/// <summary>
+		/// Waits for a lease to be assigned to an agent
+		/// </summary>
+		/// <param name="agent">The agent to assign a lease to</param>
+		/// <param name="newLeases">Leases that the agent knows about</param>
+		/// <param name="cancellationToken"></param>
+		/// <returns>True if a lease was assigned, false otherwise</returns>
+		public async Task<IAgent?> WaitForLeaseAsync(IAgent? agent, IList<HordeCommon.Rpc.Messages.RpcLease> newLeases, CancellationToken cancellationToken = default)
+		{
+			HashSet<string> knownLeases = new HashSet<string>(newLeases.Select(x => x.Id), StringComparer.OrdinalIgnoreCase);
+			while (agent != null && agent.Leases.All(x => knownLeases.Contains(x.Id.ToString())))
+			{
+				if (!agent.SessionExpiresAt.HasValue)
+				{
+					break;
+				}
+
+				// Check we have some time to wait
+				DateTime utcNow = _clock.UtcNow;
+				TimeSpan maxWaitTime = (agent.SessionExpiresAt.Value - SessionExpiryTime + SessionLongPollTime) - utcNow;
+				if (maxWaitTime <= TimeSpan.Zero)
+				{
+					break;
+				}
+
+				// Create a cancellation token that will expire with the session
+				using CancellationTokenSource cancellationSource = CancellationTokenSource.CreateLinkedTokenSource(_applicationLifetime.ApplicationStopping, cancellationToken);
+				cancellationSource.CancelAfter(maxWaitTime);
+
+				// Assign a new lease
+				(ITaskSource, AgentLease)? result = null;
+				try
+				{
+					// Add the cancellation source to the set of waiting agents
+					lock (_waitingAgents)
+					{
+						_waitingAgents[agent.Id] = cancellationSource;
+					}
+
+					// Create all the tasks to wait for
+					List<Task<(ITaskSource, AgentLease)?>> tasks = new List<Task<(ITaskSource, AgentLease)?>>();
+					foreach (ITaskSource taskSource in _taskSources)
+					{
+						if (CanUseTaskSource(agent, taskSource) && !cancellationSource.IsCancellationRequested)
+						{
+							Task<(ITaskSource, AgentLease)?> task = await GuardedAssignLeaseAsync(taskSource, agent, cancellationSource);
+							tasks.Add(task);
+						}
+					}
+
+					// If no task source is valid, just add a delay
+					if (tasks.Count == 0)
+					{
+						_logger.LogInformation("No task source valid for agent {AgentId}. Waiting {WaitTimeMs} ms", agent.Id, maxWaitTime.TotalMilliseconds);
+						await AsyncUtils.DelayNoThrow(maxWaitTime, cancellationToken);
+						break;
+					}
+
+					// Wait for all the tasks to complete. Once the first task completes it will set the cancellation source, triggering the 
+					// others to terminate.
+					await Task.WhenAll(tasks);
+
+					// Find the first result
+					foreach (Task<(ITaskSource, AgentLease)?> task in tasks)
+					{
+						(ITaskSource, AgentLease)? taskResult;
+						if (task.TryGetResult(out taskResult) && taskResult != null)
+						{
+							(ITaskSource taskSource, AgentLease taskLease) = taskResult.Value;
+							if (result == null)
+							{
+								result = (taskSource, taskLease);
+							}
+							else
+							{
+								await taskSource.CancelLeaseAsync(agent, taskLease.Id, Any.Parser.ParseFrom(taskLease.Payload), CancellationToken.None);
+							}
+						}
+					}
+				}
+				finally
+				{
+					lock (_waitingAgents)
+					{
+						_waitingAgents.Remove(agent.Id);
+					}
+				}
+
+				// Exit if we didn't find any work to do. It may be that all the task sources returned null, in which case wait for the time period to expire.
+				if (result == null)
+				{
+					if (!cancellationSource.IsCancellationRequested)
+					{
+						await cancellationSource.Token.AsTask();
+					}
+					break;
+				}
+
+				// Get the resulting lease
+				(ITaskSource source, AgentLease lease) = result.Value;
+
+				// Add the new lease to the agent
+				IAgent? newAgent = await agent.TryAddLeaseAsync(lease, cancellationToken);
+				if (newAgent != null)
+				{
+					await source.OnLeaseStartedAsync(newAgent, lease.Id, Any.Parser.ParseFrom(lease.Payload), Agents.GetLogger(agent.Id), CancellationToken.None);
+					await CreateLeaseAsync(agent, lease, CancellationToken.None);
+					return newAgent;
+				}
+				else
+				{
+					_logger.LogInformation("Failed adding lease {LeaseId} for agent {AgentId}", lease.Id.ToString(), agent.Id.ToString());
+					await source.CancelLeaseAsync(agent, lease.Id, Any.Parser.ParseFrom(lease.Payload), CancellationToken.None);
+				}
+
+				// Update the agent
+				agent = await GetAgentAsync(agent.Id, cancellationToken);
+			}
+			return agent;
+		}
+
+		async Task<Task<(ITaskSource, AgentLease)?>> GuardedAssignLeaseAsync(ITaskSource source, IAgent agent, CancellationTokenSource cancellationSource)
+		{
+			CancellationToken cancellationToken = cancellationSource.Token;
+			try
+			{
+				Task<AgentLease?> task = await source.AssignLeaseAsync(agent, cancellationToken);
+				return task.ContinueWith(x => WrapAssignedLease(source, x, cancellationSource), TaskScheduler.Default);
+			}
+			catch (OperationCanceledException) when (cancellationToken.IsCancellationRequested)
+			{
+				return Task.FromResult<(ITaskSource, AgentLease)?>(null);
+			}
+			catch (Exception ex)
+			{
+				_logger.LogError(ex, "Exception while trying to assign lease");
+				return Task.FromResult<(ITaskSource, AgentLease)?>(null);
+			}
+		}
+
+		(ITaskSource, AgentLease)? WrapAssignedLease(ITaskSource source, Task<AgentLease?> task, CancellationTokenSource cancellationSource)
+		{
+			if (task.IsCanceled)
+			{
+				return null;
+			}
+			else if (task.TryGetResult(out AgentLease? lease))
+			{
+				if (lease == null)
+				{
+					return null;
+				}
+				else
+				{
+					cancellationSource.Cancel();
+					return (source, lease);
+				}
+			}
+			else if (task.IsFaulted)
+			{
+				_logger.LogError(task.Exception, "Exception while trying to assign lease");
+				return null;
+			}
+			else
+			{
+				_logger.LogWarning("Unhandled task state: {Status}", task.Status);
+				return null;
+			}
+		}
+
+		/// <summary>
+		/// Cancels the specified agent lease
+		/// </summary>
+		/// <param name="agent">Agent to cancel the lease on</param>
+		/// <param name="leaseId">The lease id to cancel</param>
+		/// <param name="cancellationToken">Cancellation token for the operation</param>
+		/// <returns></returns>
+		public async Task<bool> CancelLeaseAsync(IAgent agent, LeaseId leaseId, CancellationToken cancellationToken)
+		{
+			int index = 0;
+			while (index < agent.Leases.Count && agent.Leases[index].Id != leaseId)
+			{
+				index++;
+			}
+
+			if (index == agent.Leases.Count)
+			{
+				return false;
+			}
+
+			if (agent.Leases[index].State == LeaseState.Cancelled)
+			{
+				return false;
+			}
+
+			await agent.TryCancelLeaseAsync(index, cancellationToken);
+			return true;
+		}
+
+		/// <summary>
+		/// 
+		/// </summary>
+		public async Task<IAgent?> UpdateSessionWithWaitAsync(IAgent inAgent, SessionId sessionId, AgentStatus status, IReadOnlyList<string>? properties, IReadOnlyDictionary<string, int>? resources, IList<HordeCommon.Rpc.Messages.RpcLease> newLeases, CancellationToken cancellationToken = default)
+		{
+			IAgent? agent = inAgent;
+
+			// Capture the current agent update index. This allows us to detect if anything has changed.
+			uint updateIndex = agent.UpdateIndex;
+
+			// Update the agent session and return to the caller if anything changes
+			agent = await UpdateSessionAsync(agent, sessionId, status, properties, resources, newLeases, cancellationToken);
+			if (agent != null && agent.UpdateIndex == updateIndex && (agent.Leases.Count > 0 || agent.Status != AgentStatus.Stopping))
+			{
+				agent = await WaitForLeaseAsync(agent, newLeases, cancellationToken);
+			}
+			return agent;
+		}
+
+		/// <summary>
+		/// Updates the state of the current agent session
+		/// </summary>
+		/// <param name="inAgent">The current agent state</param>
+		/// <param name="sessionId">Id of the session</param>
+		/// <param name="status">New status for the agent</param>
+		/// <param name="properties">New agent properties</param>
+		/// <param name="resources">New agent resources</param>
+		/// <param name="newLeases">New list of leases for this session</param>
+		/// <param name="cancellationToken">Cancellation token for the operation</param>
+		/// <returns>Updated agent state</returns>
+		public async Task<IAgent?> UpdateSessionAsync(IAgent inAgent, SessionId sessionId, AgentStatus status, IReadOnlyList<string>? properties, IReadOnlyDictionary<string, int>? resources, IList<HordeCommon.Rpc.Messages.RpcLease> newLeases, CancellationToken cancellationToken = default)
+		{
+			DateTime utcNow = _clock.UtcNow;
+
+			Stopwatch timer = Stopwatch.StartNew();
+
+			IAgent? agent = inAgent;
+			while (agent != null)
+			{
+				// Check the session id is correct.
+				if (agent.SessionId != sessionId)
+				{
+					if (status == AgentStatus.Stopping)
+					{
+						break; // Harmless; agent is not doing any work.
+					}
+					else
+					{
+						throw new InvalidOperationException($"Invalid agent session {sessionId}");
+					}
+				}
+
+				// Check the session hasn't expired
+				if (!agent.IsSessionValid(utcNow))
+				{
+					throw new InvalidOperationException("Session has already expired");
+				}
+
+				// Extend the current session time if we're within a time period of the current time expiring. This reduces
+				// unnecessary DB writes a little, but it also allows us to skip the update and jump into a long poll state
+				// if there's still time on the session left.
+				DateTime? sessionExpiresAt = null;
+				if (!agent.SessionExpiresAt.HasValue || utcNow > (agent.SessionExpiresAt - SessionExpiryTime) + SessionRenewTime)
+				{
+					sessionExpiresAt = utcNow + SessionExpiryTime;
+				}
+
+				// Flag for whether the leases array should be updated
+				bool updateLeases = false;
+				List<AgentLease> leases = new List<AgentLease>(agent.Leases);
+
+				// Remove any completed leases from the agent
+				Dictionary<LeaseId, RpcLease> leaseIdToNewState = newLeases.ToDictionary(x => LeaseId.Parse(x.Id), x => x);
+				for (int idx = 0; idx < leases.Count; idx++)
+				{
+					AgentLease lease = leases[idx];
+					if (lease.State == LeaseState.Cancelled)
+					{
+						RpcLease? newLease;
+						if (!leaseIdToNewState.TryGetValue(lease.Id, out newLease) || newLease.State == RpcLeaseState.Cancelled || newLease.State == RpcLeaseState.Completed)
+						{
+							await RemoveLeaseAsync(agent, lease, utcNow, LeaseOutcome.Cancelled, null, cancellationToken);
+							leases.RemoveAt(idx--);
+							updateLeases = true;
+						}
+					}
+					else
+					{
+						RpcLease? newLease;
+						if (leaseIdToNewState.TryGetValue(lease.Id, out newLease) && (LeaseState)newLease.State != lease.State)
+						{
+							if (newLease.State == RpcLeaseState.Cancelled || newLease.State == RpcLeaseState.Completed)
+							{
+								await RemoveLeaseAsync(agent, lease, utcNow, (LeaseOutcome)newLease.Outcome, newLease.Output.ToByteArray(), cancellationToken);
+								leases.RemoveAt(idx--);
+							}
+							else if (newLease.State == RpcLeaseState.Active && lease.State == LeaseState.Pending)
+							{
+								lease.State = LeaseState.Active;
+							}
+							updateLeases = true;
+						}
+					}
+				}
+
+				// If the agent is stopping, cancel all the leases. Clear out the current session once it's complete.
+				if (status == AgentStatus.Stopping || status == AgentStatus.Busy)
+				{
+					foreach (AgentLease lease in leases)
+					{
+						if (lease.State != LeaseState.Cancelled)
+						{
+							lease.State = LeaseState.Cancelled;
+							updateLeases = true;
+						}
+					}
+				}
+
+				// Get the new dynamic pools for the agent
+				List<PoolId> dynamicPools = await GetDynamicPoolsAsync(agent, cancellationToken);
+
+				// Update the agent, and try to create new lease documents if we succeed
+				IAgent? newAgent = await agent.TryUpdateSessionAsync(new UpdateSessionOptions(status, sessionExpiresAt, properties, resources, dynamicPools, updateLeases ? leases : null), cancellationToken);
+				if (newAgent != null)
+				{
+					agent = newAgent;
+					break;
+				}
+
+				// Fetch the agent again
+				agent = await GetAgentAsync(agent.Id, cancellationToken);
+			}
+
+			// If the agent is stopping, terminate the session
+			while (agent != null && agent.Status == AgentStatus.Stopping && agent.Leases.Count == 0)
+			{
+				IAgent? terminatedAgent = await TryTerminateSessionAsync(agent, cancellationToken);
+				if (terminatedAgent != null)
+				{
+					_logger.LogInformation("Terminated session {SessionId} for {AgentId}; agent is stopping", agent.SessionId, agent.Id);
+					agent = terminatedAgent;
+					break;
+				}
+				agent = await GetAgentAsync(agent.Id, cancellationToken);
+			}
+
+			return agent;
+		}
+
+		/// <summary>
+		/// Terminates an existing session. Does not update the agent itself, if it's currently 
+		/// </summary>
+		/// <param name="agent">The agent whose current session should be terminated</param>
+		/// <param name="cancellationToken">Cancellation token for the operation</param>
+		/// <returns>An up-to-date IAgent if the session was terminated</returns>
+		private async Task<IAgent?> TryTerminateSessionAsync(IAgent agent, CancellationToken cancellationToken = default)
+		{
+			// Make sure the agent has a valid session id
+			if (agent.SessionId == null)
+			{
+				return agent;
+			}
+
+			// Get the time that the session finishes at
+			DateTime finishTime = _clock.UtcNow;
+			if (agent.SessionExpiresAt.HasValue && agent.SessionExpiresAt.Value < finishTime)
+			{
+				finishTime = agent.SessionExpiresAt.Value;
+			}
+
+			// Save off the session id and current leases
+			SessionId sessionId = agent.SessionId.Value;
+			List<AgentLease> leases = new List<AgentLease>(agent.Leases);
+
+			// Clear the current session
+			IAgent? newAgent = await agent.TryTerminateSessionAsync(cancellationToken);
+			if (newAgent != null)
+			{
+				agent = newAgent;
+
+				// Remove any outstanding leases
+				foreach (AgentLease lease in leases)
+				{
+					Agents.GetLogger(agent.Id).LogInformation("Removing lease {LeaseId} during session terminate...", lease.Id);
+					await RemoveLeaseAsync(agent, lease, finishTime, LeaseOutcome.Failed, null, cancellationToken);
+				}
+
+				// Update the session document
+				Agents.GetLogger(agent.Id).LogInformation("Terminated session {SessionId}", sessionId);
+				await _sessions.UpdateAsync(sessionId, finishTime, agent.Properties, agent.Resources, cancellationToken);
+				return agent;
+			}
+			return null;
+		}
+
+		/// <summary>
+		/// Creates a new lease document
+		/// </summary>
+		/// <param name="agent">Agent that will be executing the lease</param>
+		/// <param name="agentLease">The new agent lease</param>
+		/// <param name="cancellationToken">Cancellation token for the operation</param>
+		/// <returns>New lease document</returns>
+		internal Task<ILease> CreateLeaseAsync(IAgent agent, AgentLease agentLease, CancellationToken cancellationToken = default)
+		{
+			try
+			{
+				return _leases.AddAsync(agentLease.Id, agentLease.ParentId, agentLease.Name, agent.Id, agent.SessionId!.Value, agentLease.StreamId, agentLease.PoolId, agentLease.LogId, agentLease.StartTime, agentLease.Payload!, cancellationToken);
+			}
+			catch (Exception ex)
+			{
+				_logger.LogError(ex, "Unable to create lease {LeaseId} for agent {AgentId}; lease already exists?", agentLease.Id, agent.Id);
+				throw;
+			}
+		}
+
+		/// <summary>
+		/// Finds all leases matching a set of criteria
+		/// </summary>
+		/// <param name="agentId">Unqiue id of the agent executing this lease</param>
+		/// <param name="sessionId">Unique id of the agent session</param>
+		/// <param name="startTime">Start of the search window to return results for</param>
+		/// <param name="finishTime">End of the search window to return results for</param>
+		/// <param name="index">Index of the first result to return</param>
+		/// <param name="count">Number of results to return</param>
+		/// <param name="cancellationToken">Cancellation token for the operation</param>
+		/// <returns>List of leases matching the given criteria</returns>
+		public Task<IReadOnlyList<ILease>> FindLeasesAsync(AgentId? agentId, SessionId? sessionId, DateTime? startTime, DateTime? finishTime, int index, int count, CancellationToken cancellationToken = default)
+		{
+			return _leases.FindLeasesAsync(null, agentId, sessionId, startTime, finishTime, index, count, cancellationToken: cancellationToken);
+		}
+
+		/// <summary>
+		/// Finds all leases by finish time
+		/// </summary>
+		/// <param name="minFinishTime">Start of the search window to return results for</param>
+		/// <param name="maxFinishTime">End of the search window to return results for</param>
+		/// <param name="index">Index of the first result to return</param>
+		/// <param name="count">Number of results to return</param>
+		/// <param name="cancellationToken">Cancellation token for the operation</param>
+		/// <returns>List of leases matching the given criteria</returns>
+		public Task<IReadOnlyList<ILease>> FindLeasesByFinishTimeAsync(DateTime? minFinishTime, DateTime? maxFinishTime, int? index, int? count, CancellationToken cancellationToken = default)
+		{
+			return _leases.FindLeasesByFinishTimeAsync(minFinishTime, maxFinishTime, index, count, null, false, cancellationToken);
+		}
+
+		/// <summary>
+		/// Gets a specific lease
+		/// </summary>
+		/// <param name="leaseId">Unique id of the lease</param>
+		/// <param name="cancellationToken"></param>
+		/// <returns>The lease that was found, or null if it does not exist</returns>
+		public Task<ILease?> GetLeaseAsync(LeaseId leaseId, CancellationToken cancellationToken = default)
+		{
+			return _leases.GetAsync(leaseId, cancellationToken);
+		}
+
+		/// <summary>
+		/// Removes a lease with the given id. Updates the lease state in the database, and removes the item from the agent's leases array.
+		/// </summary>
+		/// <param name="agent">The agent to remove a lease from</param>
+		/// <param name="lease">The lease to cancel</param>
+		/// <param name="utcNow">The current time</param>
+		/// <param name="outcome">Final status of the lease</param>
+		/// <param name="output">Output from executing the task</param>
+		/// <param name="cancellationToken">Cancellation token for the operation</param>
+		/// <returns>Async task</returns>
+		private async Task RemoveLeaseAsync(IAgent agent, AgentLease lease, DateTime utcNow, LeaseOutcome outcome, byte[]? output, CancellationToken cancellationToken = default)
+		{
+			// Make sure the lease is terminated correctly
+			if (lease.Payload == null)
+			{
+				_logger.LogWarning("Removing lease {LeaseId} (no payload)", lease.Id);
+			}
+			else
+			{
+				Any any = Any.Parser.ParseFrom(lease.Payload);
+				_logger.LogInformation("Removing lease {LeaseId} ({LeaseType})", lease.Id, any.TypeUrl);
+
+				foreach (ITaskSource taskSource in _taskSources)
+				{
+					if (any.Is(taskSource.Descriptor))
+					{
+						await taskSource.OnLeaseFinishedAsync(agent, lease.Id, any, outcome, output, Agents.GetLogger(agent.Id), cancellationToken);
+						break;
+					}
+				}
+			}
+
+			// Figure out what time the lease finished
+			DateTime finishTime = utcNow;
+			if (agent.SessionExpiresAt.HasValue && agent.SessionExpiresAt.Value < finishTime)
+			{
+				finishTime = agent.SessionExpiresAt.Value;
+			}
+			if (lease.ExpiryTime.HasValue && lease.ExpiryTime.Value < finishTime)
+			{
+				finishTime = lease.ExpiryTime.Value;
+			}
+
+			// Update the lease
+			await _leases.TrySetOutcomeAsync(lease.Id, finishTime, outcome, output, cancellationToken);
+
+			// Temporarily disabling due to gRPC timeouts.
+#if false
+			// Terminate any child leases
+			List<ILease> childLeases = await _leases.FindLeasesAsync(parentId: lease.Id);
+			foreach (ILease childLease in childLeases)
+			{
+				if (childLease.Outcome == LeaseOutcome.Unspecified)
+				{
+					IAgent? otherAgent = await GetAgentAsync(childLease.AgentId);
+					if (otherAgent != null)
+					{
+						AgentLease? otherLease = otherAgent.Leases.FirstOrDefault(x => x.Id == childLease.Id);
+						if (otherLease != null)
+						{
+							_logger.LogInformation("Terminating child lease {LeaseId} with parent {ParentLeaseId}", otherLease.Id, lease.Id);
+							await RemoveLeaseAsync(otherAgent, otherLease, utcNow, LeaseOutcome.Failed, null);
+						}
+					}
+				}
+			}
+#endif
+		}
+
+		/// <summary>
+		/// Gets information about a particular session
+		/// </summary>
+		/// <param name="sessionId">The unique session id</param>
+		/// <returns>The session information</returns>
+		public Task<ISession?> GetSessionAsync(SessionId sessionId)
+		{
+			return _sessions.GetAsync(sessionId);
+		}
+
+		/// <summary>
+		/// Find sessions for the given agent
+		/// </summary>
+		/// <param name="agentId">The unique agent id</param>
+		/// <param name="startTime">Start time to include in the search</param>
+		/// <param name="finishTime">Finish time to include in the search</param>
+		/// <param name="index">Index of the first result to return</param>
+		/// <param name="count">Number of results to return</param>
+		/// <returns>List of sessions matching the given criteria</returns>
+		public Task<List<ISession>> FindSessionsAsync(AgentId agentId, DateTime? startTime, DateTime? finishTime, int index, int count)
+		{
+			return _sessions.FindAsync(agentId, startTime, finishTime, index, count);
+		}
+
+		/// <summary>
+		/// Update the rate table for different agent types
+		/// </summary>
+		/// <param name="entries">New entries for the rate table</param>
+		/// <returns></returns>
+		public async Task UpdateRateTableAsync(List<AgentRateConfig> entries)
+		{
+			await _redisService.GetDatabase().StringSetAsync(_agentRateTableData, new AgentRateTable { Entries = entries });
+		}
+
+		/// <summary>
+		/// Gets the rate for the given agent
+		/// </summary>
+		/// <param name="agentId">Agent id to query</param>
+		/// <param name="cancellationToken">Cancellation token for the operation</param>
+		/// <returns>Hourly rate of running the given agent</returns>
+		public async ValueTask<double?> GetRateAsync(AgentId agentId, CancellationToken cancellationToken = default)
+		{
+			RedisKey key = $"agent-rate/{agentId}";
+
+			// Try to get the current value
+			RedisValue value = await _redisService.GetDatabase().StringGetAsync(key).WaitAsync(cancellationToken);
+			if (!value.IsNull)
+			{
+				double rate = (double)value;
+				if (rate == 0.0)
+				{
+					return null;
+				}
+				else
+				{
+					return rate;
+				}
+			}
+			else
+			{
+				double rate = 0.0;
+
+				// Get the rate table
+				AgentRateTable? rateTable = await _cachedRates.GetAsync(cancellationToken);
+				if (rateTable != null && rateTable.Entries.Count > 0)
+				{
+					IAgent? agent = await GetAgentAsync(agentId, cancellationToken);
+					if (agent != null)
+					{
+						foreach (AgentRateConfig config in rateTable.Entries)
+						{
+							if (config.Condition != null && config.Condition.Evaluate(x => agent.GetPropertyValues(x)))
+							{
+								rate = config.Rate;
+								break;
+							}
+						}
+					}
+				}
+
+				// Cache it for future reference
+				await _redisService.GetDatabase().StringSetAsync(key, rate, TimeSpan.FromMinutes(5.0), flags: CommandFlags.FireAndForget).WaitAsync(cancellationToken);
+				return rate;
+			}
+		}
+
+		internal async ValueTask TickAsync(CancellationToken stoppingToken)
+		{
+			using TelemetrySpan span = _tracer.StartActiveSpan($"{nameof(AgentService)}.{nameof(TickAsync)}");
+			await RefreshCachedAgentsAsync(stoppingToken);
+			await CollectMetricsAsync(stoppingToken);
+		}
+
+		internal async ValueTask TickSharedAsync(CancellationToken stoppingToken)
+		{
+			using TelemetrySpan span = _tracer.StartActiveSpan($"{nameof(AgentService)}.{nameof(TickSharedAsync)}");
+
+			await TerminateExpiredSessionsAsync(stoppingToken);
+			await DeleteExpiredEphemeralAgentsAsync(stoppingToken);
+		}
+
+		private async Task TerminateExpiredSessionsAsync(CancellationToken cancellationToken)
+		{
+			using TelemetrySpan span = _tracer.StartActiveSpan($"{nameof(AgentService)}.{nameof(TerminateExpiredSessionsAsync)}");
+
+			int c = 0;
+			while (!cancellationToken.IsCancellationRequested)
+			{
+				// Find all the agents which are ready to be expired
+				const int MaxAgents = 100;
+				DateTime utcNow = _clock.UtcNow;
+				IReadOnlyList<IAgent> expiredAgents = await Agents.FindExpiredAsync(utcNow, MaxAgents, cancellationToken);
+
+				// Transition each agent to being offline
+				foreach (IAgent expiredAgent in expiredAgents)
+				{
+					cancellationToken.ThrowIfCancellationRequested();
+					_logger.LogDebug("Terminating session {SessionId} for agent {Agent}", expiredAgent.SessionId, expiredAgent.Id);
+					await TryTerminateSessionAsync(expiredAgent, cancellationToken);
+				}
+				c += expiredAgents.Count;
+
+				// Try again if we didn't fetch everything
+				if (expiredAgents.Count < MaxAgents)
+				{
+					break;
+				}
+			}
+			span.SetAttribute("NumAgentsTerminated", c);
+		}
+
+		private async Task DeleteExpiredEphemeralAgentsAsync(CancellationToken cancellationToken)
+		{
+			using TelemetrySpan span = _tracer.StartActiveSpan($"{nameof(AgentService)}.{nameof(DeleteExpiredEphemeralAgentsAsync)}");
+			int c = 0;
+			foreach (IAgent agent in await Agents.FindDeletedAsync(cancellationToken))
+			{
+				cancellationToken.ThrowIfCancellationRequested();
+				bool noStatusChangeDuringPeriod = _clock.UtcNow > agent.LastStatusChange + TimeSpan.FromHours(1);
+				if (agent is { Status: AgentStatus.Stopped, Ephemeral: true } && noStatusChangeDuringPeriod)
+				{
+					_logger.LogDebug("Deleting ephemeral agent {Agent}", agent.Id);
+					await DeleteAgentAsync(agent, true, cancellationToken);
+					c++;
+				}
+			}
+
+			span.SetAttribute("NumAgentsDeleted", c);
+		}
+
+		/// <summary>
+		/// Updates the in-memory cache of current agents
+		/// Called by ticker to avoid blocking reads. Updates in the background, in favor of slightly more stale agents
+		/// </summary>
+		/// <param name="cancellationToken">Cancellation token</param>
+		private async Task RefreshCachedAgentsAsync(CancellationToken cancellationToken = default)
+		{
+			using TelemetrySpan span = _tracer.StartActiveSpan($"{nameof(AgentService)}.{nameof(RefreshCachedAgentsAsync)}");
+			Dictionary<AgentId, IAgent> agents = new();
+			IReadOnlyList<IAgent> agentList = await Agents.FindAsync(consistentRead: false, cancellationToken: cancellationToken);
+			foreach (IAgent agent in agentList)
+			{
+				agents[agent.Id] = agent;
+			}
+			
+			Interlocked.Exchange(ref _cachedAgents, agents);
+		}
+
+		private async Task CollectMetricsAsync(CancellationToken cancellationToken = default)
+		{
+			using TelemetrySpan span = _tracer.StartActiveSpan($"{nameof(AgentService)}.{nameof(CollectMetricsAsync)}");
+
+			IReadOnlyList<IAgent> agentList = await GetCachedAgentsAsync(cancellationToken);
+			int numAgentsTotal = agentList.Count;
+			int numAgentsTotalDeleted = agentList.Count(a => a.Deleted);
+			int numAgentsTotalEnabled = agentList.Count(a => a.Enabled);
+			int numAgentsTotalDisabled = agentList.Count(a => !a.Enabled);
+			int numAgentsTotalOk = agentList.Count(a => a.Enabled && a.Status == AgentStatus.Ok);
+			int numAgentsTotalStopping = agentList.Count(a => a.Enabled && a.Status == AgentStatus.Stopping);
+			int numAgentsTotalUnhealthy = agentList.Count(a => a.Enabled && a.Status == AgentStatus.Unhealthy);
+			int numAgentsTotalBusy = agentList.Count(a => a.Enabled && a.Status == AgentStatus.Busy);
+			int numAgentsTotalUnspecified = agentList.Count(a => a.Enabled && a.Status == AgentStatus.Unspecified);
+
+			List<Measurement<int>> newMeasurements =
+			[
+				new Measurement<int>(numAgentsTotal),
+				new Measurement<int>(numAgentsTotalDeleted, new KeyValuePair<string, object?>("status", "deleted")),
+				new Measurement<int>(numAgentsTotalEnabled, new KeyValuePair<string, object?>("status", "enabled")),
+				new Measurement<int>(numAgentsTotalDisabled, new KeyValuePair<string, object?>("status", "disabled")),
+				new Measurement<int>(numAgentsTotalOk, new KeyValuePair<string, object?>("status", "ok")),
+				new Measurement<int>(numAgentsTotalStopping, new KeyValuePair<string, object?>("status", "stopping")),
+				new Measurement<int>(numAgentsTotalUnhealthy, new KeyValuePair<string, object?>("status", "unhealthy")),
+				new Measurement<int>(numAgentsTotalBusy, new KeyValuePair<string, object?>("status", "paused")),
+				new Measurement<int>(numAgentsTotalUnspecified, new KeyValuePair<string, object?>("status", "unspecified"))
+			];
+
+			_measurements = newMeasurements;
+		}
+
+		/// <summary>
+		/// Determines if the user is authorized to perform an action on a particular agent
+		/// </summary>
+		/// <param name="agent">The agent to check</param>
+		/// <param name="user">The principal to authorize</param>
+		/// <param name="reason">Reason for being authorized or not</param>
+		/// <returns>True if the action is authorized</returns>
+		public bool AuthorizeSession(IAgent agent, ClaimsPrincipal user, out string reason)
+		{
+			if (agent.SessionId == null)
+			{
+				reason = $"{nameof(agent.SessionId)} is null";
+				return false;
+			}
+
+			if (!user.HasSessionClaim(agent.SessionId.Value))
+			{
+				reason = $"Missing session claim for {agent.SessionId.Value}";
+				return false;
+			}
+
+			if (!agent.IsSessionValid(_clock.UtcNow))
+			{
+				reason = $"Session has expired";
+				return false;
+			}
+
+			reason = "Session is valid";
+			return true;
+		}
+	}
+}
diff --git a/Engine/Source/Programs/Horde/HordeServer/Agents/Fleet/JobQueueStrategy.cs b/Engine/Source/Programs/Horde/HordeServer/Agents/Fleet/JobQueueStrategy.cs
new file mode 100644
--- /dev/null
+++ b/Engine/Source/Programs/Horde/HordeServer/Agents/Fleet/JobQueueStrategy.cs
@@ -0,0 +1,278 @@
+// Copyright Epic Games, Inc. All Rights Reserved.
+
+using System;
+using System.Collections.Generic;
+using System.Linq;
+using System.Text;
+using System.Text.Json.Serialization;
+using System.Threading;
+using System.Threading.Tasks;
+using EpicGames.Core;
+using EpicGames.Horde.Agents.Pools;
+using EpicGames.Horde.Jobs;
+using EpicGames.Horde.Streams;
+using Horde.Server.Agents.Pools;
+using Horde.Server.Jobs;
+using Horde.Server.Jobs.Graphs;
+using Horde.Server.Server;
+using Horde.Server.Streams;
+using Horde.Server.Utilities;
+using Microsoft.Extensions.Caching.Memory;
+using Microsoft.Extensions.Options;
+using OpenTelemetry.Trace;
+
+namespace Horde.Server.Agents.Fleet
+{
+	/// <summary>
+	/// Job queue sizing settings for a pool
+	/// </summary>
+	public class JobQueueSettings
+	{
+		/// <summary>
+		/// Factor translating queue size to additional agents to grow the pool with
+		/// The result is always rounded up to nearest integer. 
+		/// Example: if there are 20 jobs in queue, a factor 0.25 will result in 5 new agents being added (20 * 0.25)
+		/// </summary>
+		public double ScaleOutFactor { get; set; } = 0.25;
+
+		/// <summary>
+		/// Factor by which to shrink the pool size with when queue is empty
+		/// The result is always rounded up to nearest integer.
+		/// Example: when the queue size is zero, a default value of 0.9 will shrink the pool by 10% (current agent count * 0.9)
+		/// </summary>
+		public double ScaleInFactor { get; set; } = 0.9;
+
+		/// <summary>
+		/// How far back in time to look for job batches (that potentially are in the queue)
+		/// </summary>
+		public int SamplePeriodMin { get; set; } = 60 * 24 * 1; // 1 day
+
+		/// <summary>
+		/// Time spent in ready state before considered truly waiting for an agent
+		///
+		/// A job batch can be in ready state before getting picked up and executed.
+		/// This threshold will help ensure only batches that have been waiting longer than this value will be considered.
+		/// </summary>
+		public int ReadyTimeThresholdSec { get; set; } = 45;
+
+		/// <summary>
+		/// Constructor used for JSON serialization
+		/// </summary>
+		[JsonConstructor]
+		public JobQueueSettings()
+		{
+		}
+
+		/// <summary>
+		/// Constructor
+		/// </summary>
+		/// <param name="scaleOutFactor"></param>
+		/// <param name="scaleInFactor"></param>
+		public JobQueueSettings(double? scaleOutFactor = null, double? scaleInFactor = null)
+		{
+			ScaleOutFactor = scaleOutFactor.GetValueOrDefault(ScaleOutFactor);
+			ScaleInFactor = scaleInFactor.GetValueOrDefault(ScaleInFactor);
+		}
+
+		/// <inheritdoc />
+		public override string ToString()
+		{
+			StringBuilder sb = new(50);
+			sb.AppendFormat("{0}={1} ", nameof(ScaleOutFactor), ScaleOutFactor);
+			sb.AppendFormat("{0}={1} ", nameof(ScaleInFactor), ScaleInFactor);
+			sb.AppendFormat("{0}={1} ", nameof(SamplePeriodMin), SamplePeriodMin);
+			sb.AppendFormat("{0}={1} ", nameof(ReadyTimeThresholdSec), ReadyTimeThresholdSec);
+			return sb.ToString();
+		}
+	}
+
+	/// <summary>
+	/// Calculate pool size by observing the number of jobs in waiting state
+	///
+	/// Allows for more proactive scaling compared to LeaseUtilizationStrategy.
+	/// <see cref="LeaseUtilizationStrategy"/> 
+	/// </summary>
+	public class JobQueueStrategy : IPoolSizeStrategy
+	{
+		private const string CacheKey = nameof(JobQueueStrategy);
+		internal JobQueueSettings Settings { get; }
+
+		private readonly IJobCollection _jobs;
+		private readonly IGraphCollection _graphs;
+		private readonly IStreamCollection _streamCollection;
+		private readonly IClock _clock;
+		private readonly IMemoryCache _cache;
+		private readonly bool _isDowntimeActive;
+		private readonly IOptionsMonitor<GlobalConfig> _globalConfig;
+
+		/// <summary>
+		/// Constructor
+		/// </summary>
+		/// <param name="jobs"></param>
+		/// <param name="graphs"></param>
+		/// <param name="streamCollection"></param>
+		/// <param name="clock"></param>
+		/// <param name="cache"></param>
+		/// <param name="isDowntimeActive"></param>
+		/// <param name="globalConfig"></param>
+		/// <param name="settings"></param>
+		public JobQueueStrategy(IJobCollection jobs, IGraphCollection graphs, IStreamCollection streamCollection, IClock clock, IMemoryCache cache, bool isDowntimeActive, IOptionsMonitor<GlobalConfig> globalConfig, JobQueueSettings? settings = null)
+		{
+			_jobs = jobs;
+			_graphs = graphs;
+			_streamCollection = streamCollection;
+			_clock = clock;
+			_cache = cache;
+			_isDowntimeActive = isDowntimeActive;
+			_globalConfig = globalConfig;
+			Settings = settings ?? new JobQueueSettings();
+		}
+
+		/// <inheritdoc/>
+		public string Name { get; } = "JobQueue";
+
+		/// <summary>
+		/// Extract all job step batches from a job, with their associated pool 
+		/// </summary>
+		/// <param name="job">Job to extract from</param>
+		/// <param name="streams">Cached lookup table of streams</param>
+		/// <param name="cancellationToken">Cancellation token for the operation</param>
+		/// <returns></returns>
+		private async Task<List<(IJob Job, IJobStepBatch Batch, PoolId PoolId)>> GetJobBatchesWithPoolsAsync(IJob job, Dictionary<StreamId, StreamConfig> streams, CancellationToken cancellationToken)
+		{
+			IGraph graph = await _graphs.GetAsync(job.GraphHash, cancellationToken);
+
+			List<(IJob Job, IJobStepBatch Batch, PoolId PoolId)> jobBatches = new();
+			foreach (IJobStepBatch batch in job.Batches)
+			{
+				if (batch.State != JobStepBatchState.Ready)
+				{
+					continue;
+				}
+
+				TimeSpan? waitTime = _clock.UtcNow - batch.ReadyTimeUtc;
+				if (waitTime == null)
+				{
+					continue;
+				}
+				if (waitTime.Value < TimeSpan.FromSeconds(Settings.ReadyTimeThresholdSec))
+				{
+					continue;
+				}
+
+				if (!streams.TryGetValue(job.StreamId, out StreamConfig? streamConfig))
+				{
+					continue;
+				}
+
+				string batchAgentType = graph.Groups[batch.GroupIdx].AgentType;
+				if (!streamConfig.AgentTypes.TryGetValue(batchAgentType, out AgentConfig? agentType))
+				{
+					continue;
+				}
+
+				jobBatches.Add((job, batch, agentType.Pool));
+			}
+
+			return jobBatches;
+		}
+
+		internal async Task<Dictionary<PoolId, int>> GetPoolQueueSizesAsync(DateTimeOffset jobsCreatedAfter, CancellationToken cancellationToken = default)
+		{
+			using TelemetrySpan span = OpenTelemetryTracers.Horde.StartActiveSpan($"{nameof(JobQueueStrategy)}.{nameof(GetPoolQueueSizesAsync)}");
+			span.SetAttribute("after", jobsCreatedAfter);
+
+			Dictionary<StreamId, StreamConfig> streams = _globalConfig.CurrentValue.Streams.ToDictionary(x => x.Id, x => (StreamConfig)x);
+			IReadOnlyList<IJob> recentJobs = await _jobs.FindAsync(minCreateTime: jobsCreatedAfter, batchState: JobStepBatchState.Ready, cancellationToken: cancellationToken);
+			span.SetAttribute("numJobs", recentJobs.Count);
+			span.SetAttribute("numUniqueGraphs", recentJobs.Select(x => x.GraphHash).Distinct().Count());
+
+			List<(IJob Job, IJobStepBatch Batch, PoolId PoolId)> jobBatches = new();
+			foreach (IJob job in recentJobs)
+			{
+				jobBatches.AddRange(await GetJobBatchesWithPoolsAsync(job, streams, cancellationToken));
+			}
+
+			List<(PoolId PoolId, int QueueSize)> poolsWithQueueSize = jobBatches.GroupBy(t => t.PoolId).Select(t => (t.Key, t.Count())).ToList();
+
+			span.SetAttribute("numPools", poolsWithQueueSize.Count);
+
+			if (_isDowntimeActive)
+			{
+				// As an optimization, assume queue size is zero during maintenance windows.
+				return poolsWithQueueSize.ToDictionary(x => x.PoolId, x => 0);
+			}
+
+			return poolsWithQueueSize.ToDictionary(x => x.PoolId, x => x.QueueSize);
+		}
+
+		/// <inheritdoc/>
+		public async Task<PoolSizeResult> CalculatePoolSizeAsync(IPoolConfig pool, List<IAgent> agents, CancellationToken cancellationToken = default)
+		{
+			using TelemetrySpan span = OpenTelemetryTracers.Horde.StartActiveSpan($"{nameof(JobQueueStrategy)}.{nameof(CalculatePoolSizeAsync)}");
+			span.SetAttribute(OpenTelemetryTracers.DatadogResourceAttribute, pool.Id.ToString());
+			span.SetAttribute("currentAgentCount", agents.Count);
+			span.SetAttribute("samplePeriodMin", Settings.SamplePeriodMin);
+
+			DateTimeOffset minCreateTime = _clock.UtcNow - TimeSpan.FromMinutes(Settings.SamplePeriodMin);
+
+			// Cache pool queue sizes for a short while for faster runs when many pools are scaled
+			if (!_cache.TryGetValue(CacheKey, out Dictionary<PoolId, int>? poolQueueSizes) || poolQueueSizes == null)
+			{
+				// Pool sizes haven't been cached, update them (might happen from multiple tasks but that is fine)
+				poolQueueSizes = await GetPoolQueueSizesAsync(minCreateTime, cancellationToken);
+				_cache.Set(CacheKey, poolQueueSizes, TimeSpan.FromSeconds(60));
+			}
+
+			poolQueueSizes.TryGetValue(pool.Id, out int queueSize);
+
+			Dictionary<string, object> status = new()
+			{
+				["Name"] = GetType().Name,
+				["QueueSize"] = queueSize,
+				["ScaleOutFactor"] = Settings.ScaleOutFactor,
+				["ScaleInFactor"] = Settings.ScaleInFactor,
+				["SamplePeriodMin"] = Settings.SamplePeriodMin,
+				["ReadyTimeThresholdSec"] = Settings.ReadyTimeThresholdSec,
+			};
+
+			if (queueSize > 0)
+			{
+				int additionalAgentCount = (int)Math.Ceiling(queueSize * Settings.ScaleOutFactor);
+				int desiredAgentCount = agents.Count + additionalAgentCount;
+				return new PoolSizeResult(agents.Count, desiredAgentCount, status);
+			}
+			else
+			{
+				int desiredAgentCount = (int)(agents.Count * Settings.ScaleInFactor);
+				return new PoolSizeResult(agents.Count, desiredAgentCount, status);
+			}
+		}
+	}
+
+	class JobQueueStrategyFactory : PoolSizeStrategyFactory<JobQueueSettings>
+	{
+		readonly IGraphCollection _graphCollection;
+		readonly IJobCollection _jobCollection;
+		readonly IStreamCollection _streamCollection;
+		readonly IDowntimeService _downtimeService;
+		readonly IClock _clock;
+		readonly IMemoryCache _memoryCache;
+		readonly IOptionsMonitor<GlobalConfig> _globalConfig;
+
+		public JobQueueStrategyFactory(IGraphCollection graphCollection, IJobCollection jobCollection, IStreamCollection streamCollection, IDowntimeService downtimeService, IClock clock, IMemoryCache memoryCache, IOptionsMonitor<GlobalConfig> globalConfig)
+			: base(PoolSizeStrategy.JobQueue)
+		{
+			_graphCollection = graphCollection;
+			_jobCollection = jobCollection;
+			_streamCollection = streamCollection;
+			_downtimeService = downtimeService;
+			_clock = clock;
+			_memoryCache = memoryCache;
+			_globalConfig = globalConfig;
+		}
+
+		public override IPoolSizeStrategy Create(JobQueueSettings settings)
+			=> new JobQueueStrategy(_jobCollection, _graphCollection, _streamCollection, _clock, _memoryCache, _downtimeService.IsDowntimeActive, _globalConfig, settings);
+	}
+}
diff --git a/Engine/Source/Programs/Horde/HordeServer/Agents/IAgentCollection.cs b/Engine/Source/Programs/Horde/HordeServer/Agents/IAgentCollection.cs
new file mode 100644
--- /dev/null
+++ b/Engine/Source/Programs/Horde/HordeServer/Agents/IAgentCollection.cs
@@ -0,0 +1,119 @@
+// Copyright Epic Games, Inc. All Rights Reserved.
+
+using System;
+using System.Collections.Generic;
+using System.Threading;
+using System.Threading.Tasks;
+using EpicGames.Horde.Agents;
+using EpicGames.Horde.Agents.Leases;
+using EpicGames.Horde.Agents.Pools;
+using Horde.Server.Auditing;
+
+namespace Horde.Server.Agents
+{
+	/// <summary>
+	/// Interface for a collection of agent documents
+	/// </summary>
+	public interface IAgentCollection
+	{
+		/// <summary>
+		/// Adds a new agent with the given properties
+		/// </summary>
+		/// <param name="id">Id for the new agent</param>
+		/// <param name="ephemeral">Whether the agent is ephemeral or not</param>
+		/// <param name="enrollmentKey">Key used to identify a unique enrollment for the agent with this id</param>
+		/// <param name="cancellationToken">Cancellation token for the operation</param>
+		Task<IAgent> AddAsync(AgentId id, bool ephemeral, string enrollmentKey, CancellationToken cancellationToken = default);
+
+		/// <summary>
+		/// Deletes an agent
+		/// </summary>
+		/// <param name="agentId">Unique id of the agent</param>
+		/// <param name="cancellationToken">Cancellation token for the operation</param>
+		/// <returns>Async task</returns>
+		Task ForceDeleteAsync(AgentId agentId, CancellationToken cancellationToken = default);
+
+		/// <summary>
+		/// Gets an agent by ID
+		/// </summary>
+		/// <param name="agentId">Unique id of the agent</param>
+		/// <param name="cancellationToken">Cancellation token for the operation</param>
+		/// <returns>The agent document</returns>
+		Task<IAgent?> GetAsync(AgentId agentId, CancellationToken cancellationToken = default);
+
+		/// <summary>
+		/// Gets multiple agents by ID
+		/// </summary>
+		/// <param name="agentIds">List of unique IDs of the agents</param>
+		/// <param name="cancellationToken">Cancellation token for the operation</param>
+		/// <returns>The agent documents</returns>
+		Task<IReadOnlyList<IAgent>> GetManyAsync(List<AgentId> agentIds, CancellationToken cancellationToken = default);
+
+		/// <summary>
+		/// Finds all agents matching certain criteria
+		/// </summary>
+		/// <param name="poolId">The pool ID in string form containing the agent</param>
+		/// <param name="modifiedAfter">If set, only returns agents modified after this time</param>
+		/// <param name="property">Property to look for</param>
+		/// <param name="status">Status to look for</param>
+		/// <param name="enabled">Enabled/disabled status to look for</param>
+		/// <param name="includeDeleted">Whether agents marked as deleted should be included</param>
+		/// <param name="index">Index of the first result</param>
+		/// <param name="count">Number of results to return</param>
+		/// <param name="consistentRead">If the database read should be made to the replica server</param>
+		/// <param name="cancellationToken">Cancellation token for the operation</param>
+		/// <returns>List of agents matching the given criteria</returns>
+		Task<IReadOnlyList<IAgent>> FindAsync(PoolId? poolId = null, DateTime? modifiedAfter = null, string? property = null, AgentStatus? status = null, bool? enabled = null, bool includeDeleted = false, int? index = null, int? count = null, bool consistentRead = true, CancellationToken cancellationToken = default);
+
+		/// <summary>
+		/// Finds all agents with an expired session
+		/// </summary>
+		/// <param name="utcNow">The current time</param>
+		/// <param name="maxAgents">Maximum number of agents to return</param>
+		/// <param name="cancellationToken">Cancellation token for the operation</param>
+		/// <returns>List of agents</returns>
+		Task<IReadOnlyList<IAgent>> FindExpiredAsync(DateTime utcNow, int maxAgents, CancellationToken cancellationToken = default);
+
+		/// <summary>
+		/// Finds all agents marked as deleted
+		/// </summary>
+		/// <param name="cancellationToken">Cancellation token for the operation</param>
+		/// <returns>List of agents</returns>
+		Task<IReadOnlyList<IAgent>> FindDeletedAsync(CancellationToken cancellationToken = default);
+
+		/// <summary>
+		/// Finds all active agent lease IDs
+		/// </summary>
+		/// <param name="cancellationToken">Cancellation token for the operation</param>
+		/// <returns>List of agent lease IDs</returns>
+		Task<List<LeaseId>> FindActiveLeaseIdsAsync(CancellationToken cancellationToken = default);
+
+		/// <summary>
+		/// Get all child lease IDs
+		/// </summary>
+		/// <param name="id">Lease ID</param>
+		/// <param name="cancellationToken">Cancellation token for the operation</param>
+		/// <returns>List of agent lease IDs</returns>
+		Task<List<LeaseId>> GetChildLeaseIdsAsync(LeaseId id, CancellationToken cancellationToken = default);
+
+		/// <summary>
+		/// Gets the log channel for an agent
+		/// </summary>
+		/// <param name="agentId"></param>
+		/// <returns></returns>
+		IAuditLogChannel<AgentId> GetLogger(AgentId agentId);
+
+		/// <summary>
+		/// Sends a notification that an update event has ocurred
+		/// </summary>
+		/// <param name="agentId">Agent that has been updated</param>
+		Task PublishUpdateEventAsync(AgentId agentId);
+
+		/// <summary>
+		/// Subscribe to notifications on agent states being updated
+		/// </summary>
+		/// <param name="onUpdate">Callback for updates</param>
+		/// <returns>Disposable subscription object</returns>
+		Task<IAsyncDisposable> SubscribeToUpdateEventsAsync(Action<AgentId> onUpdate);
+	}
+}
diff --git a/Engine/Source/Programs/Horde/HordeServer/Agents/Sessions/ISessionCollection.cs b/Engine/Source/Programs/Horde/HordeServer/Agents/Sessions/ISessionCollection.cs
new file mode 100644
--- /dev/null
+++ b/Engine/Source/Programs/Horde/HordeServer/Agents/Sessions/ISessionCollection.cs
@@ -0,0 +1,77 @@
+// Copyright Epic Games, Inc. All Rights Reserved.
+
+using System;
+using System.Collections.Generic;
+using System.Threading;
+using System.Threading.Tasks;
+using EpicGames.Horde.Agents;
+using EpicGames.Horde.Agents.Sessions;
+
+namespace Horde.Server.Agents.Sessions
+{
+	/// <summary>
+	/// Interface for a collection of session documents
+	/// </summary>
+	public interface ISessionCollection
+	{
+		/// <summary>
+		/// Adds a new session to the collection
+		/// </summary>
+		/// <param name="id">The session id</param>
+		/// <param name="agentId">The agent this session is for</param>
+		/// <param name="startTime">Start time of this session</param>
+		/// <param name="properties">Properties of this agent at the time the session started</param>
+		/// <param name="resources">Resources which the agent has</param>
+		/// <param name="version">Version of the agent software</param>
+		/// <param name="cancellationToken">Cancellation token for the operation</param>
+		Task<ISession> AddAsync(SessionId id, AgentId agentId, DateTime startTime, IReadOnlyList<string>? properties, IReadOnlyDictionary<string, int>? resources, string? version, CancellationToken cancellationToken = default);
+
+		/// <summary>
+		/// Gets information about a particular session
+		/// </summary>
+		/// <param name="sessionId">The unique session id</param>
+		/// <param name="cancellationToken">Cancellation token for the operation</param>
+		/// <returns>The session information</returns>
+		Task<ISession?> GetAsync(SessionId sessionId, CancellationToken cancellationToken = default);
+
+		/// <summary>
+		/// Find sessions for the given agent
+		/// </summary>
+		/// <param name="agentId">The unique agent id</param>
+		/// <param name="startTime">Start time to include in the search</param>
+		/// <param name="finishTime">Finish time to include in the search</param>
+		/// <param name="index">Index of the first result to return</param>
+		/// <param name="count">Number of results to return</param>
+		/// <param name="cancellationToken">Cancellation token for the operation</param>
+		/// <returns>List of sessions matching the given criteria</returns>
+		Task<List<ISession>> FindAsync(AgentId agentId, DateTime? startTime, DateTime? finishTime, int index, int count, CancellationToken cancellationToken = default);
+
+		/// <summary>
+		/// Finds any active sessions
+		/// </summary>
+		/// <param name="index">Index of the first result to return</param>
+		/// <param name="count">Number of results to return</param>
+		/// <param name="cancellationToken">Cancellation token for the operation</param>
+		/// <returns>List of sessions</returns>
+		Task<List<ISession>> FindActiveSessionsAsync(int? index = null, int? count = null, CancellationToken cancellationToken = default);
+
+		/// <summary>
+		/// Update a session from the collection
+		/// </summary>
+		/// <param name="sessionId">The session to update</param>
+		/// <param name="finishTime">Time at which the session finished</param>
+		/// <param name="properties">The agent properties</param>
+		/// <param name="resources">Resources which the agent has</param>
+		/// <param name="cancellationToken">Cancellation token for the operation</param>
+		/// <returns>Async task</returns>
+		Task UpdateAsync(SessionId sessionId, DateTime? finishTime, IReadOnlyList<string>? properties, IReadOnlyDictionary<string, int>? resources, CancellationToken cancellationToken = default);
+
+		/// <summary>
+		/// Delete a session from the collection
+		/// </summary>
+		/// <param name="sessionId">The session id</param>
+		/// <param name="cancellationToken">Cancellation token for the operation</param>
+		/// <returns>Async task</returns>
+		Task DeleteAsync(SessionId sessionId, CancellationToken cancellationToken = default);
+	}
+}
