Horde: Rename projects before splitting into plugins (Horde.Agent -> HordeAgent, Horde.Server -> HordeServer, Horde -> HordeCmd).

diff --git a/Engine/Source/Programs/Horde/HordeServer/Agents/Sessions/SessionCollection.cs b/Engine/Source/Programs/Horde/HordeServer/Agents/Sessions/SessionCollection.cs
new file mode 100644
--- /dev/null
+++ b/Engine/Source/Programs/Horde/HordeServer/Agents/Sessions/SessionCollection.cs
@@ -0,0 +1,150 @@
+// Copyright Epic Games, Inc. All Rights Reserved.
+
+using System;
+using System.Collections.Generic;
+using System.Linq;
+using System.Threading;
+using System.Threading.Tasks;
+using EpicGames.Horde.Agents;
+using EpicGames.Horde.Agents.Sessions;
+using Horde.Server.Server;
+using Horde.Server.Utilities;
+using MongoDB.Bson.Serialization.Attributes;
+using MongoDB.Driver;
+
+namespace Horde.Server.Agents.Sessions
+{
+	/// <summary>
+	/// Collection of session documents
+	/// </summary>
+	public class SessionCollection : ISessionCollection
+	{
+		/// <summary>
+		/// Concrete implementation of ISession
+		/// </summary>
+		class SessionDocument : ISession
+		{
+			[BsonRequired, BsonId]
+			public SessionId Id { get; set; }
+
+			[BsonRequired]
+			public AgentId AgentId { get; set; }
+
+			public DateTime StartTime { get; set; }
+			public DateTime? FinishTime { get; set; }
+			public List<string> Properties { get; set; } = new List<string>();
+			public Dictionary<string, int> Resources { get; set; } = new Dictionary<string, int>();
+			public string Version { get; set; } = String.Empty;
+
+			IReadOnlyList<string> ISession.Properties => Properties;
+			IReadOnlyDictionary<string, int> ISession.Resources => Resources;
+
+			[BsonConstructor]
+			private SessionDocument()
+			{
+			}
+
+			public SessionDocument(SessionId id, AgentId agentId, DateTime startTime, IReadOnlyList<string>? properties, IReadOnlyDictionary<string, int>? resources, string? version)
+			{
+				Id = id;
+				AgentId = agentId;
+				StartTime = startTime;
+				if (properties != null)
+				{
+					Properties = new List<string>(properties);
+				}
+				if (resources != null)
+				{
+					Resources = new Dictionary<string, int>(resources);
+				}
+				if (version != null)
+				{
+					Version = version;
+				}
+			}
+		}
+
+		/// <summary>
+		/// Collection of session documents
+		/// </summary>
+		readonly IMongoCollection<SessionDocument> _sessions;
+
+		/// <summary>
+		/// Constructor
+		/// </summary>
+		/// <param name="mongoService">The database service</param>
+		public SessionCollection(MongoService mongoService)
+		{
+			List<MongoIndex<SessionDocument>> indexes = new List<MongoIndex<SessionDocument>>();
+			indexes.Add(keys => keys.Ascending(x => x.AgentId).Ascending(x => x.StartTime).Ascending(x => x.FinishTime));
+			indexes.Add(keys => keys.Ascending(x => x.FinishTime));
+
+			_sessions = mongoService.GetCollection<SessionDocument>("Sessions", indexes);
+		}
+
+		/// <inheritdoc/>
+		public async Task<ISession> AddAsync(SessionId id, AgentId agentId, DateTime startTime, IReadOnlyList<string>? properties, IReadOnlyDictionary<string, int>? resources, string? version, CancellationToken cancellationToken = default)
+		{
+			SessionDocument newSession = new SessionDocument(id, agentId, startTime, properties, resources, version);
+			await _sessions.InsertOneAsync(newSession, null, cancellationToken);
+			return newSession;
+		}
+
+		/// <inheritdoc/>
+		public async Task<ISession?> GetAsync(SessionId sessionId, CancellationToken cancellationToken = default)
+		{
+			return await _sessions.Find(x => x.Id == sessionId).FirstOrDefaultAsync(cancellationToken);
+		}
+
+		/// <inheritdoc/>
+		public async Task<List<ISession>> FindAsync(AgentId agentId, DateTime? startTime, DateTime? finishTime, int index, int count, CancellationToken cancellationToken = default)
+		{
+			FilterDefinitionBuilder<SessionDocument> filterBuilder = Builders<SessionDocument>.Filter;
+
+			FilterDefinition<SessionDocument> filter = filterBuilder.Eq(x => x.AgentId, agentId);
+			if (startTime != null)
+			{
+				filter &= filterBuilder.Gte(x => x.StartTime, startTime.Value);
+			}
+			if (finishTime != null)
+			{
+				filter &= filterBuilder.Or(filterBuilder.Eq(x => x.FinishTime, null), filterBuilder.Lte(x => x.FinishTime, finishTime.Value));
+			}
+
+			List<SessionDocument> results = await _sessions.Find(filter).SortByDescending(x => x.StartTime).Skip(index).Limit(count).ToListAsync(cancellationToken);
+			return results.ConvertAll<ISession>(x => x);
+		}
+
+		/// <inheritdoc/>
+		public async Task<List<ISession>> FindActiveSessionsAsync(int? index, int? count, CancellationToken cancellationToken = default)
+		{
+			List<SessionDocument> results = await _sessions.Find(x => x.FinishTime == null).Range(index, count).ToListAsync(cancellationToken);
+			return results.ConvertAll<ISession>(x => x);
+		}
+
+		/// <inheritdoc/>
+		public Task UpdateAsync(SessionId sessionId, DateTime? finishTime, IReadOnlyList<string>? properties, IReadOnlyDictionary<string, int>? resources, CancellationToken cancellationToken = default)
+		{
+			List<UpdateDefinition<SessionDocument>> updates = new List<UpdateDefinition<SessionDocument>>();
+			if (finishTime != null)
+			{
+				updates.Add(Builders<SessionDocument>.Update.Set(x => x.FinishTime, finishTime));
+			}
+			if (properties != null)
+			{
+				updates.Add(Builders<SessionDocument>.Update.Set(x => x.Properties, new List<string>(properties)));
+			}
+			if (resources != null)
+			{
+				updates.Add(Builders<SessionDocument>.Update.Set(x => x.Resources, new Dictionary<string, int>(resources)));
+			}
+			return _sessions.FindOneAndUpdateAsync(x => x.Id == sessionId, Builders<SessionDocument>.Update.Combine(updates), cancellationToken: cancellationToken);
+		}
+
+		/// <inheritdoc/>
+		public Task DeleteAsync(SessionId sessionId, CancellationToken cancellationToken = default)
+		{
+			return _sessions.DeleteOneAsync(x => x.Id == sessionId, null, cancellationToken);
+		}
+	}
+}
diff --git a/Engine/Source/Programs/Horde/HordeServer/HordeServer.csproj b/Engine/Source/Programs/Horde/HordeServer/HordeServer.csproj
new file mode 100644
--- /dev/null
+++ b/Engine/Source/Programs/Horde/HordeServer/HordeServer.csproj
@@ -0,0 +1,163 @@
+<Project Sdk="Microsoft.NET.Sdk.Web">
+
+  <PropertyGroup>
+    <TargetFramework>net8.0</TargetFramework>
+    <Nullable>enable</Nullable>
+    <Deterministic>true</Deterministic>
+    <AssemblyName>HordeServer</AssemblyName>
+    <DocumentationFile>HordeServer.xml</DocumentationFile>
+    <ApplicationIcon>HordeServer.ico</ApplicationIcon>
+    <RunAnalyzersDuringBuild>False</RunAnalyzersDuringBuild>
+    <Configurations>Debug;Release;Analyze</Configurations>
+  </PropertyGroup>
+
+  <PropertyGroup Condition=" '$(Configuration)|$(Platform)' == 'Analyze|AnyCPU' ">
+    <RunAnalyzersDuringBuild>True</RunAnalyzersDuringBuild>
+    <TreatWarningsAsErrors>false</TreatWarningsAsErrors>
+    <EnforceCodeStyleInBuild>True</EnforceCodeStyleInBuild>
+    <GenerateDocumentationFile>True</GenerateDocumentationFile>
+    <DocumentationFile></DocumentationFile>
+  </PropertyGroup>
+
+  <ItemGroup>
+    <PackageReference Include="AWSSDK.AutoScaling" Version="3.7.301.9" />
+    <PackageReference Include="AWSSDK.CloudWatch" Version="3.7.300.12" />
+    <PackageReference Include="AWSSDK.EC2" Version="3.7.306.6" />
+    <PackageReference Include="AWSSDK.Extensions.NETCore.Setup" Version="3.7.300" />
+    <PackageReference Include="AWSSDK.S3" Version="3.7.304.1" />
+    <PackageReference Include="AWSSDK.SimpleSystemsManagement" Version="3.7.304.24" />
+    <PackageReference Include="AWSSDK.SQS" Version="3.7.300.12" />
+    <PackageReference Include="AWSSDK.SecurityToken" Version="3.7.300.13" />
+    <PackageReference Include="Azure.Storage.Blobs" Version="12.19.1" />
+    <PackageReference Include="Crc32.NET" Version="1.2.0" />
+    <PackageReference Include="Grpc.AspNetCore" Version="2.59.0" />
+    <PackageReference Include="Grpc.AspNetCore.Server.Reflection" Version="2.59.0" />
+    <PackageReference Include="Grpc.Net.Client" Version="2.59.0" />
+    <PackageReference Include="Grpc.Tools" Version="2.59.0">
+      <PrivateAssets>all</PrivateAssets>
+      <IncludeAssets>runtime; build; native; contentfiles; analyzers; buildtransitive</IncludeAssets>
+    </PackageReference>
+    <PackageReference Include="JetBrains.Profiler.SelfApi" Version="2.5.0" />
+    <PackageReference Include="JsonPath.Net" Version="0.6.7" />
+    <PackageReference Include="Meziantou.Analyzer" Version="2.0.139">
+      <PrivateAssets>all</PrivateAssets>
+      <IncludeAssets>runtime; build; native; contentfiles; analyzers; buildtransitive</IncludeAssets>
+    </PackageReference>
+    <PackageReference Include="Microsoft.AspNetCore.Authentication.Cookies" Version="2.2.0" />
+    <PackageReference Include="Microsoft.AspNetCore.Authentication.JwtBearer" Version="8.0.0" />
+    <PackageReference Include="Microsoft.AspNetCore.Authentication.OpenIdConnect" Version="8.0.0" />
+    <PackageReference Include="Microsoft.AspNetCore.DataProtection.StackExchangeRedis" Version="8.0.0" />
+    <PackageReference Include="Microsoft.AspNetCore.Mvc.Razor.RuntimeCompilation" Version="6.0.25" />
+    <PackageReference Include="Microsoft.AspNetCore.SpaServices.Extensions" Version="8.0.0" />
+    <PackageReference Include="Microsoft.Extensions.Caching.Memory" Version="8.0.0" />
+    <PackageReference Include="Microsoft.Extensions.Hosting.WindowsServices" Version="8.0.0" />
+    <PackageReference Include="Microsoft.Extensions.Http" Version="8.0.0" />
+    <PackageReference Include="Microsoft.Extensions.Http.Polly" Version="8.0.0" />
+    <PackageReference Include="Microsoft.Extensions.Logging" Version="8.0.0" />
+    <PackageReference Include="Microsoft.VisualStudio.Threading.Analyzers" Version="17.8.14">
+      <PrivateAssets>all</PrivateAssets>
+      <IncludeAssets>runtime; build; native; contentfiles; analyzers; buildtransitive</IncludeAssets>
+    </PackageReference>
+    <PackageReference Include="MongoDB.Driver" Version="2.22.0" />
+    <PackageReference Include="OpenTelemetry" Version="1.8.0" />
+    <PackageReference Include="OpenTelemetry.Exporter.Console" Version="1.8.0" />
+    <PackageReference Include="OpenTelemetry.Exporter.OpenTelemetryProtocol" Version="1.8.0" />
+    <PackageReference Include="OpenTelemetry.Extensions.Hosting" Version="1.8.0" />
+    <PackageReference Include="OpenTelemetry.Instrumentation.AspNetCore" Version="1.8.1" />
+    <PackageReference Include="OpenTelemetry.Instrumentation.GrpcNetClient" Version="1.8.0-beta.1" />
+    <PackageReference Include="OpenTelemetry.Instrumentation.Http" Version="1.8.1" />
+    <PackageReference Include="OpenTelemetry.Instrumentation.StackExchangeRedis" Version="1.0.0-rc9.14" />
+    <PackageReference Include="Polly" Version="8.2.0" />
+    <PackageReference Include="Polly.Extensions.Http" Version="3.0.0" />
+    <PackageReference Include="Serilog.AspNetCore" Version="8.0.0" />
+    <PackageReference Include="Serilog.Exceptions.Grpc" Version="8.4.0" />
+    <PackageReference Include="Serilog.Extensions.Hosting" Version="8.0.0" />
+    <PackageReference Include="Serilog.Sinks.Console" Version="5.0.1" />
+    <PackageReference Include="Serilog.Sinks.File" Version="5.0.0" />
+    <PackageReference Include="SerilogAnalyzer" Version="0.15.0" />
+    <PackageReference Include="SharpZipLib" Version="1.4.2" />
+    <PackageReference Include="StackExchange.Redis" Version="2.7.4" />
+    <PackageReference Include="Swashbuckle.AspNetCore" Version="6.5.0" />
+    <PackageReference Include="System.Linq.Async" Version="6.0.1" />
+    <PackageReference Include="T-Digest.NET" Version="1.0.3" />
+    <PackageReference Include="TimeZoneConverter" Version="6.1.0" />
+    <PackageReference Include="ZstdSharp.Port" Version="0.8.1" />
+  </ItemGroup>
+
+  <ItemGroup>
+    <Content Include="..\..\Shared\EpicGames.Horde.Common\Protos\**\*.*">
+      <Link>Protos\%(RecursiveDir)%(FileName)%(Extension)</Link>
+    </Content>
+  </ItemGroup>
+
+  <ItemGroup>
+    <None Remove="Horde.Server.xml" />
+    <None Remove="HordeServer.xml" />
+  </ItemGroup>
+
+  <ItemGroup>
+    <Protobuf Include="Protos\horde\job_rpc.proto" ProtoRoot="Protos" GrpcServices="Server" AdditionalImportDirs="..\..\Shared\EpicGames.Horde\Protos" />
+    <Protobuf Include="Protos\horde\job_rpc_messages.proto" ProtoRoot="Protos" />
+    <Protobuf Include="Protos\horde\enrollment_rpc.proto" ProtoRoot="Protos" GrpcServices="Server" />
+    <Protobuf Include="Protos\horde\enrollment_rpc_messages.proto" ProtoRoot="Protos" />
+    <Protobuf Include="Protos\horde\horde_rpc.proto" ProtoRoot="Protos" GrpcServices="Server" />
+    <Protobuf Include="Protos\horde\horde_rpc_messages.proto" ProtoRoot="Protos" />
+    <Protobuf Include="Protos\horde\relay_rpc.proto" ProtoRoot="Protos" />
+    <Protobuf Include="Protos\horde\common\agent_status.proto" ProtoRoot="Protos" />
+    <Protobuf Include="Protos\horde\common\io_hash.proto" ProtoRoot="Protos" />
+    <Protobuf Include="Protos\horde\common\label_change.proto" ProtoRoot="Protos" />
+    <Protobuf Include="Protos\horde\common\lease_outcome.proto" ProtoRoot="Protos" />
+    <Protobuf Include="Protos\horde\common\lease_state.proto" ProtoRoot="Protos" />
+    <Protobuf Include="Protos\horde\common\namespace_id.proto" ProtoRoot="Protos" />
+    <Protobuf Include="Protos\horde\messages\agent_capabilities.proto" ProtoRoot="Protos" />
+    <Protobuf Include="Protos\horde\messages\agent_requirements.proto" ProtoRoot="Protos" />
+    <Protobuf Include="Protos\horde\messages\agent_workspace.proto" ProtoRoot="Protos" />
+    <Protobuf Include="Protos\horde\messages\lease.proto" ProtoRoot="Protos" />
+    <Protobuf Include="Protos\horde\messages\telemetry.proto" ProtoRoot="Protos" />
+    <Protobuf Include="Protos\horde\tasks\compute_task.proto" ProtoRoot="Protos" />
+    <Protobuf Include="Protos\horde\tasks\conform_task.proto" ProtoRoot="Protos" />
+    <Protobuf Include="Protos\horde\tasks\job_task.proto" ProtoRoot="Protos" />
+    <Protobuf Include="Protos\horde\tasks\restart_task.proto" ProtoRoot="Protos" />
+    <Protobuf Include="Protos\horde\tasks\shutdown_task.proto" ProtoRoot="Protos" />
+    <Protobuf Include="Protos\horde\tasks\upgrade_task.proto" ProtoRoot="Protos" />
+    <Protobuf Include="Protos\grpc\health\v1\health.proto" ProtoRoot="Protos" GrpcServices="Server" />
+  </ItemGroup>
+
+  <ItemGroup>
+    <ProjectReference Include="..\..\Shared\EpicGames.AspNet\EpicGames.AspNet.csproj" />
+    <ProjectReference Include="..\..\Shared\EpicGames.BuildGraph\EpicGames.BuildGraph.csproj" />
+    <ProjectReference Include="..\..\Shared\EpicGames.Core\EpicGames.Core.csproj" />
+    <ProjectReference Include="..\..\Shared\EpicGames.Horde\EpicGames.Horde.csproj" />
+    <ProjectReference Include="..\..\Shared\EpicGames.IoHash\EpicGames.IoHash.csproj" />
+    <ProjectReference Include="..\..\Shared\EpicGames.Oodle\EpicGames.Oodle.csproj" />
+    <ProjectReference Include="..\..\Shared\EpicGames.Perforce.Managed\EpicGames.Perforce.Managed.csproj" />
+    <ProjectReference Include="..\..\Shared\EpicGames.Perforce\EpicGames.Perforce.csproj" />
+    <ProjectReference Include="..\..\Shared\EpicGames.Redis\EpicGames.Redis.csproj" />
+    <ProjectReference Include="..\..\Shared\EpicGames.Serialization\EpicGames.Serialization.csproj" />
+    <ProjectReference Include="..\..\Shared\EpicGames.Slack\EpicGames.Slack.csproj" />
+  </ItemGroup>
+
+  <ItemGroup>
+    <Content Include="..\ThirdParty\**" Condition="$([MSBuild]::IsOSPlatform('Windows'))">
+      <CopyToOutputDirectory>Always</CopyToOutputDirectory>
+      <Link>ThirdParty/%(RecursiveDir)/%(Filename)%(Extension)</Link>
+    </Content>
+  </ItemGroup>
+
+  <ItemGroup>
+    <None Include="DashboardApp\**">
+      <CopyToOutputDirectory>Always</CopyToOutputDirectory>
+    </None>
+  </ItemGroup>
+
+  <ItemGroup>
+    <Content Include="Defaults\ue5.project.png">
+      <CopyToOutputDirectory>Always</CopyToOutputDirectory>
+    </Content>
+  </ItemGroup>
+
+  <ItemGroup>
+    <Content Update="appsettings.Local.json" CopyToPublishDirectory="Never" />
+  </ItemGroup>
+
+</Project>
diff --git a/Engine/Source/Programs/Horde/HordeServer/Issues/IssueCollection.cs b/Engine/Source/Programs/Horde/HordeServer/Issues/IssueCollection.cs
new file mode 100644
--- /dev/null
+++ b/Engine/Source/Programs/Horde/HordeServer/Issues/IssueCollection.cs
@@ -0,0 +1,1702 @@
+// Copyright Epic Games, Inc. All Rights Reserved.
+
+using System;
+using System.Collections.Generic;
+using System.Data;
+using System.Diagnostics;
+using System.Diagnostics.CodeAnalysis;
+using System.Linq;
+using System.Threading;
+using System.Threading.Tasks;
+using EpicGames.Core;
+using EpicGames.Horde.Issues;
+using EpicGames.Horde.Jobs;
+using EpicGames.Horde.Jobs.Templates;
+using EpicGames.Horde.Logs;
+using EpicGames.Horde.Streams;
+using EpicGames.Horde.Telemetry;
+using EpicGames.Horde.Users;
+using EpicGames.Redis.Utility;
+using Horde.Server.Auditing;
+using Horde.Server.Jobs;
+using Horde.Server.Jobs.Graphs;
+using Horde.Server.Server;
+using Horde.Server.Streams;
+using Horde.Server.Telemetry;
+using Horde.Server.Users;
+using Horde.Server.Utilities;
+using Microsoft.Extensions.Caching.Memory;
+using Microsoft.Extensions.Logging;
+using Microsoft.Extensions.Options;
+using MongoDB.Bson;
+using MongoDB.Bson.Serialization;
+using MongoDB.Bson.Serialization.Attributes;
+using MongoDB.Driver;
+using OpenTelemetry.Trace;
+
+namespace Horde.Server.Issues
+{
+	class IssueCollection : IIssueCollection
+	{
+		[SingletonDocument("issue-ledger", "5e4c226440ce25fa3207a9af")]
+		class IssueLedger : SingletonBase
+		{
+			public int NextId { get; set; }
+		}
+
+		[DebuggerDisplay("{Id}: {Summary}")]
+		class Issue : IIssue
+		{
+			[BsonId]
+			public int Id { get; set; }
+
+			public string Summary { get; set; }
+
+			[BsonIgnoreIfNull]
+			public string? UserSummary { get; set; }
+
+			[BsonIgnoreIfNull]
+			public string? Description { get; set; }
+
+			[BsonIgnoreIfNull]
+			public IssueFingerprint? Fingerprint { get; set; }
+			public List<IssueFingerprint>? Fingerprints { get; set; }
+
+			public IssueSeverity Severity { get; set; }
+
+			[BsonElement("PromotedV2")]
+			public bool Promoted { get; set; }
+
+			[BsonIgnoreIfNull]
+			public bool? ManuallyPromoted { get; set; }
+
+			[BsonIgnoreIfNull, BsonElement("Promoted")]
+			public bool? ManuallyPromotedDeprecated { get; set; }
+
+			[BsonIgnoreIfNull]
+			public UserId? OwnerId { get; set; }
+
+			[BsonIgnoreIfNull]
+			public UserId? DefaultOwnerId { get; set; }
+
+			[BsonIgnoreIfNull]
+			public UserId? NominatedById { get; set; }
+
+			public DateTime CreatedAt { get; set; }
+
+			[BsonIgnoreIfNull]
+			public DateTime? NominatedAt { get; set; }
+
+			[BsonIgnoreIfNull]
+			public DateTime? AcknowledgedAt { get; set; }
+
+			[BsonIgnoreIfNull]
+			public DateTime? ResolvedAt { get; set; }
+
+			[BsonIgnoreIfNull]
+			public UserId? ResolvedById { get; set; }
+
+			[BsonIgnoreIfNull]
+			public DateTime? VerifiedAt { get; set; }
+
+			public DateTime LastSeenAt { get; set; }
+
+			[BsonIgnoreIfNull]
+			public int? FixChange { get; set; }
+
+			public List<IssueStream> Streams { get; set; } = new List<IssueStream>();
+
+			public int MinSuspectChange { get; set; }
+			public int MaxSuspectChange { get; set; }
+
+			[BsonElement("Suspects"), BsonIgnoreIfNull]
+			public List<IssueSuspect>? SuspectsDeprecated { get; set; }
+
+			[BsonIgnoreIfNull]
+			public List<ObjectId>? ExcludeSpans { get; set; }
+
+			public int UpdateIndex { get; set; }
+
+			IReadOnlyList<IIssueFingerprint> IIssue.Fingerprints => Fingerprints ?? ((Fingerprint == null) ? new List<IssueFingerprint>() : new List<IssueFingerprint> { Fingerprint });
+			UserId? IIssue.OwnerId => OwnerId ?? DefaultOwnerId ?? GetDefaultOwnerId();
+			IReadOnlyList<IIssueStream> IIssue.Streams => Streams;
+			DateTime IIssue.LastSeenAt => (LastSeenAt == default) ? DateTime.UtcNow : LastSeenAt;
+
+			[BsonIgnoreIfNull]
+			public string? ExternalIssueKey { get; set; }
+
+			[BsonIgnoreIfNull]
+			public UserId? QuarantinedByUserId { get; set; }
+
+			[BsonIgnoreIfNull]
+			public DateTime? QuarantineTimeUtc { get; set; }
+
+			[BsonIgnoreIfNull]
+			public UserId? ForceClosedByUserId { get; set; }
+
+			[BsonIgnoreIfNull]
+			public Uri? WorkflowThreadUrl { get; set; }
+
+			[BsonConstructor]
+			private Issue()
+			{
+				Summary = String.Empty;
+				Fingerprint = null!;
+			}
+
+			public Issue(int id, string summary)
+			{
+				Id = id;
+				Summary = summary;
+				CreatedAt = DateTime.UtcNow;
+				LastSeenAt = DateTime.UtcNow;
+			}
+
+			UserId? GetDefaultOwnerId()
+			{
+				if (SuspectsDeprecated != null && SuspectsDeprecated.Count > 0)
+				{
+					UserId possibleOwner = SuspectsDeprecated[0].AuthorId;
+					if (SuspectsDeprecated.All(x => x.AuthorId == possibleOwner) && SuspectsDeprecated.Any(x => x.DeclinedAt == null))
+					{
+						return SuspectsDeprecated[0].AuthorId;
+					}
+				}
+				return null;
+			}
+
+			public string FingerprintsDesc
+			{
+				get
+				{
+					if (Fingerprints == null || Fingerprints.Count == 0)
+					{
+						return String.Empty;
+					}
+
+					return String.Join(", ", Fingerprints.Select(x =>
+					{
+						return $"(Type: {x.Type} / Keys: {String.Join(", ", x.Keys)} / RejectKeys: {String.Join(", ", x.RejectKeys ?? new HashSet<IssueKey>())})";
+					}));
+				}
+			}
+		}
+
+		class IssueStream : IIssueStream
+		{
+			public StreamId StreamId { get; set; }
+			public bool? MergeOrigin { get; set; }
+			public bool? ContainsFix { get; set; }
+			public bool? FixFailed { get; set; }
+
+			public IssueStream()
+			{
+			}
+
+			public IssueStream(IIssueStream other)
+			{
+				StreamId = other.StreamId;
+				MergeOrigin = other.MergeOrigin;
+				ContainsFix = other.ContainsFix;
+				FixFailed = other.FixFailed;
+			}
+		}
+
+		class IssueSuspect : IIssueSuspect
+		{
+			public ObjectId Id { get; set; }
+			public int IssueId { get; set; }
+			public UserId AuthorId { get; set; }
+			public int Change { get; set; }
+			public DateTime? DeclinedAt { get; set; }
+			public DateTime? ResolvedAt { get; set; } // Degenerate
+
+			private IssueSuspect()
+			{
+			}
+
+			public IssueSuspect(int issueId, NewIssueSuspectData newSuspect, DateTime? resolvedAt)
+			{
+				Id = ObjectId.GenerateNewId();
+				IssueId = issueId;
+				AuthorId = newSuspect.AuthorId;
+				Change = newSuspect.Change;
+				ResolvedAt = resolvedAt;
+			}
+
+			public IssueSuspect(int issueId, IIssueSpanSuspect suspect)
+				: this(issueId, suspect.AuthorId, suspect.OriginatingChange ?? suspect.Change, null, null)
+			{
+			}
+
+			public IssueSuspect(int issueId, UserId authorId, int change, DateTime? declinedAt, DateTime? resolvedAt)
+			{
+				Id = ObjectId.GenerateNewId();
+				IssueId = issueId;
+				AuthorId = authorId;
+				Change = change;
+				DeclinedAt = declinedAt;
+				ResolvedAt = resolvedAt;
+			}
+		}
+
+		class IssueFingerprint : IIssueFingerprint
+		{
+			public string Type { get; set; }
+
+			[BsonElement("summary"), BsonIgnoreIfNull]
+			string? SummaryTemplateValue { get; set; }
+
+			[BsonIgnore]
+			public string SummaryTemplate => SummaryTemplateValue ?? GetLegacyHandlerInfo(Type).SummaryTemplate;
+
+			[BsonElement("inc")]
+			public HashSet<IssueKey> Keys { get; set; } = new HashSet<IssueKey>();
+
+			IReadOnlySet<IssueKey> IIssueFingerprint.Keys => Keys;
+
+			[BsonElement("exc"), BsonIgnoreIfNull]
+			public HashSet<IssueKey>? RejectKeys { get; set; } = new HashSet<IssueKey>();
+
+			IReadOnlySet<IssueKey>? IIssueFingerprint.RejectKeys => RejectKeys;
+
+			[BsonElement("met"), BsonIgnoreIfNull]
+			public HashSet<IssueMetadata>? Metadata { get; set; } = new HashSet<IssueMetadata>();
+
+			IReadOnlySet<IssueMetadata>? IIssueFingerprint.Metadata => Metadata;
+
+			[BsonElement("flt"), BsonIgnoreIfNull]
+			public List<string>? ChangeFilter { get; set; }
+
+			IReadOnlyList<string> IIssueFingerprint.ChangeFilter => ChangeFilter ?? GetLegacyHandlerInfo(Type).ChangeFilter;
+
+			[BsonConstructor]
+			private IssueFingerprint()
+			{
+				Type = String.Empty;
+			}
+
+			public IssueFingerprint(IIssueFingerprint fingerprint)
+			{
+				Type = fingerprint.Type;
+				SummaryTemplateValue = fingerprint.SummaryTemplate;
+
+				Keys = new HashSet<IssueKey>(fingerprint.Keys);
+				if (fingerprint.RejectKeys != null && fingerprint.RejectKeys.Count > 0)
+				{
+					RejectKeys = new HashSet<IssueKey>(fingerprint.RejectKeys);
+				}
+				if (fingerprint.Metadata != null && fingerprint.Metadata.Count > 0)
+				{
+					Metadata = new HashSet<IssueMetadata>(fingerprint.Metadata);
+				}
+
+				ChangeFilter = new List<string>(fingerprint.ChangeFilter);
+			}
+
+			#region Legacy
+
+#pragma warning disable IDE0051
+			[BsonElement("Keys"), BsonIgnoreIfNull]
+			CaseInsensitiveStringSet? LegacyKeys
+			{
+				get => null;
+				set => Keys = ParseKeySet(value) ?? new HashSet<IssueKey>();
+			}
+
+			[BsonElement("RejectKeys"), BsonIgnoreIfNull]
+			CaseInsensitiveStringSet? LegacyRejectKeys
+			{
+				get => null;
+				set => RejectKeys = ParseKeySet(value);
+			}
+
+			[BsonElement("Metadata"), BsonIgnoreIfNull]
+			CaseInsensitiveStringSet? LegacyMetadata
+			{
+				get => null;
+				set => Metadata = ParseMetadata(value);
+			}
+#pragma warning restore IDE0051
+
+			[return: NotNullIfNotNull("set")]
+			static HashSet<IssueKey>? ParseKeySet(CaseInsensitiveStringSet? set) => (set == null) ? null : new HashSet<IssueKey>(set.Select(x => ParseKey(x)));
+
+			static IssueKey ParseKey(string key)
+			{
+				int colonIdx = key.IndexOf(':', StringComparison.Ordinal);
+				if (colonIdx != -1)
+				{
+					ReadOnlySpan<char> prefix = key.AsSpan(0, colonIdx);
+					if (prefix.Equals("hash", StringComparison.Ordinal))
+					{
+						return new IssueKey(key.Substring(colonIdx + 1), IssueKeyType.Hash);
+					}
+					if (prefix.Equals("note", StringComparison.Ordinal))
+					{
+						return new IssueKey(key.Substring(colonIdx + 1), IssueKeyType.Note);
+					}
+					if (prefix.Equals("step", StringComparison.Ordinal))
+					{
+						return new IssueKey(key.Substring(colonIdx + 1), IssueKeyType.None);
+					}
+				}
+				return new IssueKey(key, IssueKeyType.None);
+			}
+
+			[return: NotNullIfNotNull("set")]
+			static HashSet<IssueMetadata>? ParseMetadata(CaseInsensitiveStringSet? set)
+			{
+				HashSet<IssueMetadata>? entries = null;
+				if (set != null)
+				{
+					entries = new HashSet<IssueMetadata>();
+					foreach (string entry in set)
+					{
+						int idx = entry.IndexOf('=', StringComparison.Ordinal);
+						entries.Add(new IssueMetadata(entry.Substring(0, idx), entry.Substring(idx + 1)));
+					}
+				}
+				return entries;
+			}
+
+			record class LegacyHandlerInfo(string SummaryTemplate, IReadOnlyList<string> ChangeFilter);
+
+			static readonly Dictionary<StringView, LegacyHandlerInfo> s_legacyHandlers = new Dictionary<StringView, LegacyHandlerInfo>
+			{
+				["BuildGraph"] = new("BuildGraph {Severity} in {Files}", IssueChangeFilter.All),
+				["Compile"] = new("{Meta:CompileType} {Severity} in {Files}", IssueChangeFilter.Code),
+				["Content"] = new("{Severity} in {Files}", IssueChangeFilter.Content),
+				["Copyright"] = new("Missing copyright notice in {Files}", IssueChangeFilter.Code),
+				["Default"] = new("{Severity} in {Nodes}", IssueChangeFilter.All),
+				["Gauntlet"] = new("Gauntlet {Meta:Type} {Severity} {Meta:Context}", IssueChangeFilter.Code),
+				["Hashed"] = new("{Severity} in {Meta:Node}", IssueChangeFilter.All),
+				["Localization"] = new("Localization {Severity} in {Files}", IssueChangeFilter.Code),
+				["PerforceCase"] = new("Inconsistent case for {Files}", IssueChangeFilter.All),
+				["Scoped"] = new("{Severity} in {Meta:Node} - {Meta:Scope}", IssueChangeFilter.All),
+				["Shader"] = new("Shader compile {Severity} in {Files}", IssueChangeFilter.Code),
+				["Symbol"] = new("{LegacySymbolIssueHandler}", IssueChangeFilter.Code),
+				["Systemic"] = new("Systemic {Severity} in {Nodes}", IssueChangeFilter.None),
+				["UnacceptableWords"] = new("Unacceptable words in {Files}", IssueChangeFilter.Code)
+			};
+
+			static LegacyHandlerInfo GetLegacyHandlerInfo(string type)
+			{
+				StringView baseType = type;
+
+				int endIdx = type.IndexOf(':', StringComparison.Ordinal);
+				if (endIdx != -1)
+				{
+					baseType = new StringView(type, 0, endIdx);
+				}
+
+				s_legacyHandlers.TryGetValue(baseType, out LegacyHandlerInfo? info);
+				return info ?? new LegacyHandlerInfo($"{type} {{Severity}}", IssueChangeFilter.All);
+			}
+
+			#endregion
+		}
+
+		class IssueSpan : IIssueSpan
+		{
+			public ObjectId Id { get; set; }
+
+			[BsonRequired]
+			public StreamId StreamId { get; set; }
+
+			[BsonRequired]
+			public string StreamName { get; set; }
+
+			[BsonRequired]
+			public TemplateId TemplateRefId { get; set; }
+
+			[BsonRequired]
+			public string NodeName { get; set; }
+			public DateTime? ResolvedAt { get; set; } // Propagated from the owning issue
+
+			[BsonRequired]
+			public IssueFingerprint Fingerprint { get; set; }
+
+			public int MinChange { get; set; }
+			public int MaxChange { get; set; } = Int32.MaxValue;
+
+			public IssueStep? LastSuccess { get; set; }
+
+			[BsonRequired]
+			public IssueStep FirstFailure { get; set; }
+
+			[BsonRequired]
+			public IssueStep LastFailure { get; set; }
+
+			public IssueStep? NextSuccess { get; set; }
+
+			public bool? PromoteByDefault { get; set; }
+
+			[BsonElement("NotifySuspects"), BsonIgnoreIfDefault(false)]
+			public bool NotifySuspectsDeprecated { get; set; }
+
+			bool IIssueSpan.PromoteByDefault => PromoteByDefault ?? NotifySuspectsDeprecated;
+
+			public List<IssueSpanSuspect> Suspects { get; set; }
+			public int IssueId { get; set; }
+			public int UpdateIndex { get; set; }
+
+			IIssueStep? IIssueSpan.LastSuccess => LastSuccess;
+			IIssueStep IIssueSpan.FirstFailure => FirstFailure;
+			IIssueStep IIssueSpan.LastFailure => LastFailure;
+			IIssueStep? IIssueSpan.NextSuccess => NextSuccess;
+			IReadOnlyList<IIssueSpanSuspect> IIssueSpan.Suspects => Suspects;
+			IIssueFingerprint IIssueSpan.Fingerprint => Fingerprint;
+
+			private IssueSpan()
+			{
+				StreamName = null!;
+				NodeName = null!;
+				Fingerprint = null!;
+				FirstFailure = null!;
+				LastFailure = null!;
+				Suspects = new List<IssueSpanSuspect>();
+			}
+
+			public IssueSpan(int issueId, NewIssueSpanData newSpan)
+			{
+				Id = ObjectId.GenerateNewId();
+				StreamId = newSpan.StreamId;
+				StreamName = newSpan.StreamName;
+				TemplateRefId = newSpan.TemplateRefId;
+				NodeName = newSpan.NodeName;
+				Fingerprint = new IssueFingerprint(newSpan.Fingerprint);
+				if (newSpan.LastSuccess != null)
+				{
+					MinChange = newSpan.LastSuccess.Change;
+					LastSuccess = new IssueStep(Id, newSpan.LastSuccess);
+				}
+				FirstFailure = new IssueStep(Id, newSpan.FirstFailure);
+				LastFailure = new IssueStep(Id, newSpan.FirstFailure);
+				if (newSpan.NextSuccess != null)
+				{
+					MaxChange = newSpan.NextSuccess.Change;
+					NextSuccess = new IssueStep(Id, newSpan.NextSuccess);
+				}
+				PromoteByDefault = newSpan.FirstFailure.PromoteByDefault;
+				Suspects = newSpan.Suspects.ConvertAll(x => new IssueSpanSuspect(x));
+				IssueId = issueId;
+			}
+		}
+
+		class IssueSpanSuspect : IIssueSpanSuspect
+		{
+			public int Change { get; set; }
+			public UserId AuthorId { get; set; }
+			public int? OriginatingChange { get; set; }
+
+			[BsonConstructor]
+			private IssueSpanSuspect()
+			{
+			}
+
+			public IssueSpanSuspect(NewIssueSpanSuspectData newSuspectData)
+			{
+				Change = newSuspectData.Change;
+				AuthorId = newSuspectData.AuthorId;
+				OriginatingChange = newSuspectData.OriginatingChange;
+			}
+		}
+
+		class IssueStep : IIssueStep
+		{
+			public ObjectId Id { get; set; }
+			public ObjectId SpanId { get; set; }
+
+			public int Change { get; set; }
+			public IssueSeverity Severity { get; set; }
+
+			[BsonRequired]
+			public string JobName { get; set; }
+
+			[BsonRequired]
+			public JobId JobId { get; set; }
+
+			[BsonRequired]
+			public JobStepBatchId BatchId { get; set; }
+
+			[BsonRequired]
+			public JobStepId StepId { get; set; }
+
+			public DateTime StepTime { get; set; }
+
+			public LogId? LogId { get; set; }
+
+			[BsonIgnoreIfNull]
+			public NodeAnnotations? Annotations { get; set; }
+
+			public bool? PromoteByDefault { get; set; }
+
+			[BsonElement("NotifySuspects"), BsonIgnoreIfDefault(false)]
+			public bool NotifySuspectsDeprecated { get; set; }
+
+			IReadOnlyNodeAnnotations IIssueStep.Annotations => Annotations ?? NodeAnnotations.Empty;
+			bool IIssueStep.PromoteByDefault => PromoteByDefault ?? NotifySuspectsDeprecated;
+
+			[BsonConstructor]
+			private IssueStep()
+			{
+				JobName = null!;
+			}
+
+			public IssueStep(ObjectId spanId, NewIssueStepData stepData)
+			{
+				Id = ObjectId.GenerateNewId();
+				SpanId = spanId;
+				Change = stepData.Change;
+				Severity = stepData.Severity;
+				JobName = stepData.JobName;
+				JobId = stepData.JobId;
+				BatchId = stepData.BatchId;
+				StepId = stepData.StepId;
+				StepTime = stepData.StepTime;
+				LogId = stepData.LogId;
+				Annotations = stepData.Annotations;
+				PromoteByDefault = stepData.PromoteByDefault;
+			}
+		}
+
+		// Wraps a redis lock for the issue collection
+		sealed class IssueLock : IAsyncDisposable
+		{
+			readonly RedisLock _redisLock;
+			readonly Stopwatch _timer = Stopwatch.StartNew();
+			readonly TelemetrySpan _telemetrySpan;
+			readonly ILogger _logger;
+
+			public IssueLock(RedisLock redisLock, Tracer tracer, ILogger logger)
+			{
+				_redisLock = redisLock;
+				_telemetrySpan = tracer.StartActiveSpan("Holding Issue Lock").SetAttribute("Stack", Environment.StackTrace);
+				_logger = logger;
+			}
+
+			public async ValueTask DisposeAsync()
+			{
+				_telemetrySpan.Dispose();
+				if (_timer.Elapsed.TotalSeconds >= 2.5)
+				{
+					_logger.LogWarning("Issue lock held for {TimeSpan}s. Released from {Stack}.", (int)_timer.Elapsed.TotalSeconds, Environment.StackTrace);
+					_timer.Reset();
+				}
+				await _redisLock.DisposeAsync();
+			}
+		}
+
+		readonly RedisService _redisService;
+		readonly IUserCollection _userCollection;
+		readonly ISingletonDocument<IssueLedger> _ledgerSingleton;
+		readonly IMongoCollection<Issue> _issues;
+		readonly IMongoCollection<IssueSpan> _issueSpans;
+		readonly IMongoCollection<IssueStep> _issueSteps;
+		readonly IMongoCollection<IssueSuspect> _issueSuspects;
+		readonly IAuditLog<int> _auditLog;
+		readonly ITelemetrySink _telemetrySink;
+		readonly IOptionsMonitor<GlobalConfig> _globalConfig;
+		readonly Tracer _tracer;
+		readonly ILogger _logger;
+
+		static IssueCollection()
+		{
+			BsonClassMap.RegisterClassMap<IssueKey>(cm =>
+			{
+				cm.MapConstructor(() => new IssueKey("", IssueKeyType.None), nameof(IssueKey.Name), nameof(IssueKey.Type));
+				cm.MapProperty(x => x.Name).SetElementName("n");
+				cm.MapProperty(x => x.Type).SetElementName("t");
+				cm.MapProperty(x => x.Scope).SetElementName("s").SetIgnoreIfNull(true);
+			});
+
+			BsonClassMap.RegisterClassMap<IssueMetadata>(cm =>
+			{
+				cm.MapConstructor(() => new IssueMetadata("", ""), nameof(IssueMetadata.Key), nameof(IssueMetadata.Value));
+				cm.MapProperty(x => x.Key).SetElementName("k");
+				cm.MapProperty(x => x.Value).SetElementName("v");
+			});
+		}
+
+		public IssueCollection(MongoService mongoService, RedisService redisService, IUserCollection userCollection, IAuditLogFactory<int> auditLogFactory, ITelemetrySink telemetrySink, IOptionsMonitor<GlobalConfig> globalConfig, Tracer tracer, ILogger<IssueCollection> logger)
+		{
+			_redisService = redisService;
+			_userCollection = userCollection;
+			_telemetrySink = telemetrySink;
+			_globalConfig = globalConfig;
+			_tracer = tracer;
+			_logger = logger;
+
+			_ledgerSingleton = new SingletonDocument<IssueLedger>(mongoService);
+
+			List<MongoIndex<Issue>> issueIndexes = new List<MongoIndex<Issue>>();
+			issueIndexes.Add(keys => keys.Ascending(x => x.ResolvedAt));
+			issueIndexes.Add(keys => keys.Ascending(x => x.VerifiedAt));
+			_issues = mongoService.GetCollection<Issue>("IssuesV2");
+
+			List<MongoIndex<IssueSpan>> issueSpanIndexes = new List<MongoIndex<IssueSpan>>();
+			issueSpanIndexes.Add(keys => keys.Ascending(x => x.IssueId));
+			issueSpanIndexes.Add(keys => keys.Ascending(x => x.StreamId).Ascending(x => x.MinChange).Ascending(x => x.MaxChange));
+			issueSpanIndexes.Add("StreamChanges", keys => keys.Ascending(x => x.StreamId).Ascending(x => x.TemplateRefId).Ascending(x => x.NodeName).Ascending(x => x.MinChange).Ascending(x => x.MaxChange));
+			_issueSpans = mongoService.GetCollection<IssueSpan>("IssuesV2.Spans", issueSpanIndexes);
+
+			List<MongoIndex<IssueStep>> issueStepIndexes = new List<MongoIndex<IssueStep>>();
+			issueStepIndexes.Add(keys => keys.Ascending(x => x.SpanId));
+			issueStepIndexes.Add(keys => keys.Ascending(x => x.JobId).Ascending(x => x.BatchId).Ascending(x => x.StepId));
+			_issueSteps = mongoService.GetCollection<IssueStep>("IssuesV2.Steps", issueStepIndexes);
+
+			List<MongoIndex<IssueSuspect>> issueSuspectIndexes = new List<MongoIndex<IssueSuspect>>();
+			issueSuspectIndexes.Add(keys => keys.Ascending(x => x.Change));
+			issueSuspectIndexes.Add(keys => keys.Ascending(x => x.AuthorId).Ascending(x => x.ResolvedAt));
+			issueSuspectIndexes.Add(keys => keys.Ascending(x => x.IssueId).Ascending(x => x.Change), unique: true);
+			_issueSuspects = mongoService.GetCollection<IssueSuspect>("IssuesV2.Suspects", issueSuspectIndexes);
+
+			_auditLog = auditLogFactory.Create("IssuesV2.History", "IssueId");
+		}
+
+		/// <inheritdoc/>
+		public async Task<IAsyncDisposable> EnterCriticalSectionAsync()
+		{
+			Stopwatch timer = Stopwatch.StartNew();
+			TimeSpan nextNotifyTime = TimeSpan.FromSeconds(10.0);
+
+			RedisLock issueLock = new(_redisService.GetDatabase(), "issues/lock");
+			using (TelemetrySpan telemetrySpan = _tracer.StartActiveSpan("Wait for Issue Lock").SetAttribute("Trace", Environment.StackTrace))
+			{
+				while (!await issueLock.AcquireAsync(TimeSpan.FromMinutes(1)))
+				{
+					if (timer.Elapsed > nextNotifyTime)
+					{
+						_logger.LogWarning("Waiting on lock over issue collection for {TimeSpan}", timer.Elapsed);
+						nextNotifyTime *= 2;
+					}
+					await Task.Delay(TimeSpan.FromMilliseconds(100));
+				}
+			}
+			return new IssueLock(issueLock, _tracer, _logger);
+		}
+
+		async Task<Issue?> TryUpdateIssueAsync(IIssue issue, UpdateDefinition<Issue> update, CancellationToken cancellationToken)
+		{
+			Issue issueDocument = (Issue)issue;
+
+			int prevUpdateIndex = issueDocument.UpdateIndex;
+			update = update.Set(x => x.UpdateIndex, prevUpdateIndex + 1);
+
+			FindOneAndUpdateOptions<Issue, Issue> options = new FindOneAndUpdateOptions<Issue, Issue> { ReturnDocument = ReturnDocument.After };
+
+			Issue? newIssue = await _issues.FindOneAndUpdateAsync<Issue>(x => x.Id == issueDocument.Id && x.UpdateIndex == prevUpdateIndex, update, options, cancellationToken);
+			if (newIssue != null)
+			{
+				SendTelemetry(newIssue);
+			}
+			return newIssue;
+		}
+
+		async Task<IssueSpan?> TryUpdateSpanAsync(IIssueSpan issueSpan, UpdateDefinition<IssueSpan> update, CancellationToken cancellationToken)
+		{
+			IssueSpan issueSpanDocument = (IssueSpan)issueSpan;
+
+			int prevUpdateIndex = issueSpanDocument.UpdateIndex;
+			update = update.Set(x => x.UpdateIndex, prevUpdateIndex + 1);
+
+			FindOneAndUpdateOptions<IssueSpan, IssueSpan> options = new FindOneAndUpdateOptions<IssueSpan, IssueSpan> { ReturnDocument = ReturnDocument.After };
+
+			IssueSpan? newIssueSpan = await _issueSpans.FindOneAndUpdateAsync<IssueSpan>(x => x.Id == issueSpanDocument.Id && x.UpdateIndex == prevUpdateIndex, update, options, cancellationToken);
+			if (newIssueSpan != null)
+			{
+				SendTelemetry(newIssueSpan);
+			}
+			return newIssueSpan;
+		}
+
+		#region Issues
+
+		/// <inheritdoc/>
+		public async Task<IIssue> AddIssueAsync(string summary, CancellationToken cancellationToken)
+		{
+			IssueLedger ledger = await _ledgerSingleton.UpdateAsync(x => x.NextId++, cancellationToken);
+
+			Issue newIssue = new Issue(ledger.NextId, summary);
+			await _issues.InsertOneAsync(newIssue, null, cancellationToken);
+			SendTelemetry(newIssue);
+
+			ILogger issueLogger = GetLogger(newIssue.Id);
+			issueLogger.LogInformation("Created issue {IssueId}", newIssue.Id);
+
+			return newIssue;
+		}
+
+		void SendTelemetry(IIssue issue)
+		{
+			List<TelemetryStoreId> telemetryStoreIds = new List<TelemetryStoreId>();
+			foreach (IIssueStream stream in issue.Streams)
+			{
+				if (_globalConfig.CurrentValue.TryGetStream(stream.StreamId, out StreamConfig? streamConfig) && !streamConfig.TelemetryStoreId.IsEmpty)
+				{
+					telemetryStoreIds.Add(streamConfig.TelemetryStoreId);
+				}
+			}
+
+			foreach (TelemetryStoreId telemetryStoreId in telemetryStoreIds)
+			{
+				_telemetrySink.SendEvent(telemetryStoreId, TelemetryRecordMeta.CurrentHordeInstance, new
+				{
+					EventName = "State.Issue",
+					Id = issue.Id,
+					AcknowledgedAt = issue.AcknowledgedAt,
+					CreatedAt = issue.CreatedAt,
+					OwnerId = issue.OwnerId,
+					FixChange = issue.FixChange,
+					LastSeenAt = issue.LastSeenAt,
+					NominatedAt = issue.NominatedAt,
+					NominatedById = issue.NominatedById,
+					ResolvedAt = issue.ResolvedAt,
+					ResolvedById = issue.ResolvedById,
+					Severity = issue.Severity,
+					Summary = issue.Summary,
+					VerifiedAt = issue.VerifiedAt
+				});
+			}
+		}
+
+		async ValueTask<string> GetUserNameAsync(UserId? userId, CancellationToken cancellationToken)
+		{
+			if (userId == null)
+			{
+				return "null";
+			}
+			else if (userId == IIssue.ResolvedByUnknownId)
+			{
+				return "Horde (Unknown)";
+			}
+			else if (userId == IIssue.ResolvedByTimeoutId)
+			{
+				return "Horde (Timeout)";
+			}
+
+			IUser? user = await _userCollection.GetCachedUserAsync(userId, cancellationToken);
+			if (user == null)
+			{
+				return "Unknown user";
+			}
+
+			return user.Name;
+		}
+
+		async Task LogIssueChangesAsync(UserId? initiatedByUserId, Issue oldIssue, Issue newIssue, CancellationToken cancellationToken)
+		{
+			ILogger issueLogger = GetLogger(oldIssue.Id);
+			using IDisposable? scope = issueLogger.BeginScope("User {UserName} ({UserId})", await GetUserNameAsync(initiatedByUserId, cancellationToken), initiatedByUserId ?? UserId.Empty);
+			await LogIssueChangesImplAsync(issueLogger, oldIssue, newIssue, cancellationToken);
+		}
+
+		async Task LogIssueChangesImplAsync(ILogger issueLogger, Issue oldIssue, Issue newIssue, CancellationToken cancellationToken)
+		{
+			if (newIssue.Severity != oldIssue.Severity)
+			{
+				issueLogger.LogInformation("Changed severity to {Severity}", newIssue.Severity);
+			}
+			if (newIssue.Summary != oldIssue.Summary)
+			{
+				issueLogger.LogInformation("Changed summary to \"{Summary}\"", newIssue.Summary);
+			}
+			if (newIssue.Description != oldIssue.Description)
+			{
+				issueLogger.LogInformation("Description set to {Value}", newIssue.Description);
+			}
+			if (((IIssue)newIssue).Promoted != ((IIssue)oldIssue).Promoted)
+			{
+				issueLogger.LogInformation("Promoted set to {Value}", ((IIssue)newIssue).Promoted);
+			}
+			if (newIssue.OwnerId != oldIssue.OwnerId)
+			{
+				if (newIssue.NominatedById != null)
+				{
+					issueLogger.LogInformation("User {UserName} ({UserId}) was nominated by {NominatedByUserName} ({NominatedByUserId})", await GetUserNameAsync(newIssue.OwnerId, cancellationToken), newIssue.OwnerId, await GetUserNameAsync(newIssue.NominatedById, cancellationToken), newIssue.NominatedById);
+				}
+				else
+				{
+					issueLogger.LogInformation("User {UserName} ({UserId}) was nominated by default", await GetUserNameAsync(newIssue.OwnerId, cancellationToken), newIssue.OwnerId);
+				}
+			}
+			if (newIssue.AcknowledgedAt != oldIssue.AcknowledgedAt)
+			{
+				if (newIssue.AcknowledgedAt == null)
+				{
+					issueLogger.LogInformation("Issue was un-acknowledged by {UserName} ({UserId})", await GetUserNameAsync(oldIssue.OwnerId, cancellationToken), oldIssue.OwnerId);
+				}
+				else
+				{
+					issueLogger.LogInformation("Issue was acknowledged by {UserName} ({UserId})", await GetUserNameAsync(newIssue.OwnerId, cancellationToken), newIssue.OwnerId);
+				}
+			}
+			if (newIssue.FixChange != oldIssue.FixChange)
+			{
+				if (newIssue.FixChange == 0)
+				{
+					issueLogger.LogInformation("Issue was marked as not fixed");
+				}
+				else
+				{
+					issueLogger.LogInformation("Issue was marked as fixed in {Change}", newIssue.FixChange);
+				}
+			}
+			if (newIssue.ResolvedById != oldIssue.ResolvedById)
+			{
+				if (newIssue.ResolvedById == null)
+				{
+					issueLogger.LogInformation("Marking as unresolved");
+				}
+				else
+				{
+					issueLogger.LogInformation("Resolved by {UserName} ({UserId})", await GetUserNameAsync(newIssue.ResolvedById, cancellationToken), newIssue.ResolvedById);
+				}
+			}
+
+			if (newIssue.ResolvedAt != oldIssue.ResolvedAt)
+			{
+				if (newIssue.ResolvedAt == null)
+				{
+					issueLogger.LogInformation("Clearing resolved at time");
+				}
+				else
+				{
+					issueLogger.LogInformation("Setting resolved at time");
+				}
+			}
+
+			if (newIssue.VerifiedAt != oldIssue.VerifiedAt)
+			{
+				if (newIssue.VerifiedAt == null)
+				{
+					issueLogger.LogInformation("Clearing verified at time");
+				}
+				else
+				{
+					issueLogger.LogInformation("Setting verified at time");
+				}
+			}
+
+			if (newIssue.ExternalIssueKey != oldIssue.ExternalIssueKey)
+			{
+				if (newIssue.ExternalIssueKey != null)
+				{
+					issueLogger.LogInformation("Linked to external issue {ExternalIssueKey}", newIssue.ExternalIssueKey);
+				}
+				else
+				{
+					issueLogger.LogInformation("Unlinked from external issue {ExternalIssueKey}", oldIssue.ExternalIssueKey);
+				}
+			}
+
+			if (newIssue.QuarantinedByUserId != oldIssue.QuarantinedByUserId)
+			{
+				if (newIssue.QuarantinedByUserId != null)
+				{
+					issueLogger.LogInformation("Quarantined by {UserName} ({UserId})", await GetUserNameAsync(newIssue.QuarantinedByUserId, cancellationToken), newIssue.QuarantinedByUserId);
+				}
+				else
+				{
+					issueLogger.LogInformation("Quarantine cleared");
+				}
+			}
+
+			if (newIssue.ForceClosedByUserId != oldIssue.ForceClosedByUserId)
+			{
+				if (newIssue.ForceClosedByUserId != null)
+				{
+					issueLogger.LogInformation("Forced closed by {UserName} ({UserId})", await GetUserNameAsync(newIssue.ForceClosedByUserId, cancellationToken), newIssue.ForceClosedByUserId);
+				}
+				else
+				{
+					issueLogger.LogInformation("Force closed cleared");
+				}
+			}
+
+			string oldFingerprints = oldIssue.FingerprintsDesc;
+			string newFingerprints = newIssue.FingerprintsDesc;
+			if (oldFingerprints != newFingerprints)
+			{
+				issueLogger.LogInformation("Fingerprints changed {Fingerprints}", newFingerprints);
+			}
+
+			HashSet<StreamId> oldFixStreams = new HashSet<StreamId>(oldIssue.Streams.Where(x => x.ContainsFix ?? false).Select(x => x.StreamId));
+			HashSet<StreamId> newFixStreams = new HashSet<StreamId>(newIssue.Streams.Where(x => x.ContainsFix ?? false).Select(x => x.StreamId));
+			foreach (StreamId streamId in newFixStreams.Where(x => !oldFixStreams.Contains(x)))
+			{
+				issueLogger.LogInformation("Marking stream {StreamId} as fixed", streamId);
+			}
+			foreach (StreamId streamId in oldFixStreams.Where(x => !newFixStreams.Contains(x)))
+			{
+				issueLogger.LogInformation("Marking stream {StreamId} as not fixed", streamId);
+			}
+		}
+
+		async Task LogIssueSuspectChangesAsync(ILogger issueLogger, IReadOnlyList<IssueSuspect> oldIssueSuspects, List<IssueSuspect> newIssueSuspects, CancellationToken cancellationToken)
+		{
+			HashSet<(UserId, int)> oldSuspects = new HashSet<(UserId, int)>(oldIssueSuspects.Select(x => (x.AuthorId, x.Change)));
+			HashSet<(UserId, int)> newSuspects = new HashSet<(UserId, int)>(newIssueSuspects.Select(x => (x.AuthorId, x.Change)));
+			foreach ((UserId userId, int change) in newSuspects.Where(x => !oldSuspects.Contains(x)))
+			{
+				issueLogger.LogInformation("Added suspect {UserName} ({UserId}) for change {Change}", await GetUserNameAsync(userId, cancellationToken), userId, change);
+			}
+			foreach ((UserId userId, int change) in oldSuspects.Where(x => !newSuspects.Contains(x)))
+			{
+				issueLogger.LogInformation("Removed suspect {UserName} ({UserId}) for change {Change}", await GetUserNameAsync(userId, cancellationToken), userId, change);
+			}
+
+			HashSet<UserId> oldDeclinedBy = new HashSet<UserId>(oldIssueSuspects.Where(x => x.DeclinedAt != null).Select(x => x.AuthorId));
+			HashSet<UserId> newDeclinedBy = new HashSet<UserId>(newIssueSuspects.Where(x => x.DeclinedAt != null).Select(x => x.AuthorId));
+			foreach (UserId addDeclinedBy in newDeclinedBy.Where(x => !oldDeclinedBy.Contains(x)))
+			{
+				issueLogger.LogInformation("Declined by {UserName} ({UserId})", await GetUserNameAsync(addDeclinedBy, cancellationToken), addDeclinedBy);
+			}
+			foreach (UserId removeDeclinedBy in oldDeclinedBy.Where(x => !newDeclinedBy.Contains(x)))
+			{
+				issueLogger.LogInformation("Un-declined by {UserName} ({UserId})", await GetUserNameAsync(removeDeclinedBy, cancellationToken), removeDeclinedBy);
+			}
+		}
+
+		/// <inheritdoc/>
+		public async Task<IIssue?> GetIssueAsync(int issueId, CancellationToken cancellationToken)
+		{
+			Issue issue = await _issues.Find(x => x.Id == issueId).FirstOrDefaultAsync(cancellationToken);
+			return issue;
+		}
+
+		/// <inheritdoc/>
+		public async Task<IReadOnlyList<IIssueSuspect>> FindSuspectsAsync(int issueId, CancellationToken cancellationToken)
+		{
+			return await _issueSuspects.Find(x => x.IssueId == issueId).ToListAsync(cancellationToken);
+		}
+
+		class ProjectedIssueId
+		{
+			[System.Diagnostics.CodeAnalysis.SuppressMessage("Style", "IDE1006:Naming Styles")]
+			public int? _id { get; set; }
+		}
+
+		/// <inheritdoc/>
+		public async Task<IReadOnlyList<IIssue>> FindIssuesAsync(IEnumerable<int>? ids = null, UserId? ownerId = null, StreamId? streamId = null, int? minChange = null, int? maxChange = null, bool? resolved = null, bool? promoted = null, int? index = null, int? count = null, CancellationToken cancellationToken = default)
+		{
+			IReadOnlyList<IIssue> results;
+
+			if (ownerId == null)
+			{
+				results = await FilterIssuesByStreamIdAsync(ids, streamId, minChange, maxChange, resolved ?? false, promoted, index ?? 0, count, cancellationToken);
+			}
+			else
+			{
+				results = await _issues.Find(x => x.OwnerId == ownerId).ToListAsync(cancellationToken);
+			}
+
+			return results;
+		}
+
+		async Task<IReadOnlyList<IIssue>> FilterIssuesByStreamIdAsync(IEnumerable<int>? ids, StreamId? streamId, int? minChange, int? maxChange, bool? resolved, bool? promoted, int index, int? count, CancellationToken cancellationToken)
+		{
+			if (streamId == null)
+			{
+				return await FilterIssuesByOtherFieldsAsync(ids, minChange, maxChange, resolved, promoted, index, count, cancellationToken);
+			}
+			else
+			{
+				FilterDefinition<IssueSpan> filter = Builders<IssueSpan>.Filter.Eq(x => x.StreamId, streamId.Value);
+				if (ids != null)
+				{
+					filter &= Builders<IssueSpan>.Filter.In(x => x.IssueId, ids.Select<int, int?>(x => x));
+				}
+				else
+				{
+					filter &= Builders<IssueSpan>.Filter.Exists(x => x.IssueId);
+				}
+
+				if (minChange != null)
+				{
+					filter &= Builders<IssueSpan>.Filter.Not(Builders<IssueSpan>.Filter.Lt(x => x.MaxChange, minChange.Value));
+				}
+				if (maxChange != null)
+				{
+					filter &= Builders<IssueSpan>.Filter.Not(Builders<IssueSpan>.Filter.Gt(x => x.MinChange, maxChange.Value));
+				}
+
+				if (resolved != null)
+				{
+					if (resolved.Value)
+					{
+						filter &= Builders<IssueSpan>.Filter.Ne(x => x.ResolvedAt, null);
+					}
+					else
+					{
+						filter &= Builders<IssueSpan>.Filter.Eq(x => x.ResolvedAt, null);
+					}
+				}
+
+				using (IAsyncCursor<ProjectedIssueId> cursor = await _issueSpans.Aggregate().Match(filter).Group(x => x.IssueId, x => new ProjectedIssueId { _id = x.Key }).SortByDescending(x => x._id).ToCursorAsync(cancellationToken))
+				{
+					List<Issue> results = await PaginatedJoinAsync(cursor, (nextIds, nextIndex, nextCount) => FilterIssuesByOtherFieldsAsync(nextIds, null, null, null, promoted, nextIndex, nextCount, cancellationToken), index, count, cancellationToken);
+					if (resolved != null)
+					{
+						for (int idx = results.Count - 1; idx >= 0; idx--)
+						{
+							Issue issue = results[idx];
+							if ((issue.ResolvedAt != null) != resolved.Value && issue.ResolvedById != IIssue.ResolvedByTimeoutId)
+							{
+								_logger.LogWarning("Issue {IssueId} has resolved state out of sync with spans", issue.Id);
+								results.RemoveAt(idx);
+							}
+						}
+					}
+					return results;
+				}
+			}
+		}
+
+		async Task<List<Issue>> FilterIssuesByOtherFieldsAsync(IEnumerable<int>? ids, int? minChange, int? maxChange, bool? resolved, bool? promoted, int index, int? count, CancellationToken cancellationToken)
+		{
+			FilterDefinition<Issue> filter = FilterDefinition<Issue>.Empty;
+			if (ids != null)
+			{
+				filter &= Builders<Issue>.Filter.In(x => x.Id, ids);
+			}
+			if (resolved != null)
+			{
+				if (resolved.Value)
+				{
+					filter &= Builders<Issue>.Filter.Ne(x => x.ResolvedAt, null);
+				}
+				else
+				{
+					filter &= Builders<Issue>.Filter.Eq(x => x.ResolvedAt, null);
+				}
+			}
+			if (minChange != null)
+			{
+				filter &= Builders<Issue>.Filter.Not(Builders<Issue>.Filter.Lt(x => x.MaxSuspectChange, minChange.Value));
+			}
+			if (maxChange != null)
+			{
+				filter &= Builders<Issue>.Filter.Not(Builders<Issue>.Filter.Gt(x => x.MinSuspectChange, maxChange.Value));
+			}
+			if (promoted != null)
+			{
+				if (promoted.Value)
+				{
+					filter &= Builders<Issue>.Filter.Eq(x => x.Promoted, true);
+				}
+				else
+				{
+					filter &= Builders<Issue>.Filter.Ne(x => x.Promoted, true); // Handle the field not existing as well as being set to false.
+				}
+			}
+			return await _issues.Find(filter).SortByDescending(x => x.Id).Range(index, count).ToListAsync(cancellationToken);
+		}
+
+		/// <summary>
+		/// Performs a client-side join of a filtered set of issues against another query
+		/// </summary>
+		/// <param name="cursor"></param>
+		/// <param name="nextStageFunc"></param>
+		/// <param name="index"></param>
+		/// <param name="count"></param>
+		/// <param name="cancellationToken"></param>
+		/// <returns></returns>
+		static async Task<List<Issue>> PaginatedJoinAsync(IAsyncCursor<ProjectedIssueId> cursor, Func<IEnumerable<int>, int, int?, Task<List<Issue>>> nextStageFunc, int index, int? count, CancellationToken cancellationToken)
+		{
+			if (count == null)
+			{
+				List<ProjectedIssueId> issueIds = await cursor.ToListAsync(cancellationToken);
+				return await nextStageFunc(issueIds.Where(x => x._id != null).Select(x => x._id!.Value), index, null);
+			}
+			else
+			{
+				List<Issue> results = new List<Issue>();
+				while (await cursor.MoveNextAsync(cancellationToken) && results.Count < count.Value)
+				{
+					List<Issue> nextResults = await nextStageFunc(cursor.Current.Where(x => x._id != null).Select(x => x._id!.Value), 0, count.Value - results.Count);
+					int removeCount = Math.Min(index, nextResults.Count);
+					nextResults.RemoveRange(0, removeCount);
+					index -= removeCount;
+					results.AddRange(nextResults);
+				}
+				return results;
+			}
+		}
+
+		/// <inheritdoc/>
+		public async Task<IReadOnlyList<IIssue>> FindIssuesForChangesAsync(List<int> changes, CancellationToken cancellationToken)
+		{
+			List<int> issueIds = await (await _issueSuspects.DistinctAsync(x => x.IssueId, Builders<IssueSuspect>.Filter.In(x => x.Change, changes), cancellationToken: cancellationToken)).ToListAsync(cancellationToken);
+			return await _issues.Find(Builders<Issue>.Filter.In(x => x.Id, issueIds)).ToListAsync(cancellationToken);
+		}
+
+		/// <inheritdoc/>
+		public async Task<IIssue?> TryUpdateIssueAsync(IIssue issue, UserId? initiatedByUserId, IssueSeverity? newSeverity = null, string? newSummary = null, string? newUserSummary = null, string? newDescription = null, bool? newManuallyPromoted = null, UserId? newOwnerId = null, UserId? newNominatedById = null, bool? newAcknowledged = null, UserId? newDeclinedById = null, int? newFixChange = null, UserId? newResolvedById = null, List<ObjectId>? newExcludeSpanIds = null, DateTime? newLastSeenAt = null, string? newExternaIssueKey = null, UserId? newQuarantinedById = null, UserId? newForceClosedById = null, Uri? newWorkflowThreadUrl = null, CancellationToken cancellationToken = default)
+		{
+			Issue issueDocument = (Issue)issue;
+
+			if (newDeclinedById != null && newDeclinedById == issueDocument.OwnerId)
+			{
+				newOwnerId = UserId.Empty;
+			}
+
+			if (issue.ResolvedById == null && (newForceClosedById != null && newResolvedById == null))
+			{
+				newResolvedById = newForceClosedById;
+			}
+
+			DateTime utcNow = DateTime.UtcNow;
+
+			List<UpdateDefinition<Issue>> updates = new List<UpdateDefinition<Issue>>();
+			if (newSeverity != null)
+			{
+				updates.Add(Builders<Issue>.Update.Set(x => x.Severity, newSeverity.Value));
+			}
+			if (newSummary != null)
+			{
+				updates.Add(Builders<Issue>.Update.Set(x => x.Summary, newSummary));
+			}
+			if (newUserSummary != null)
+			{
+				if (newUserSummary.Length == 0)
+				{
+					updates.Add(Builders<Issue>.Update.Unset(x => x.UserSummary!));
+				}
+				else
+				{
+					updates.Add(Builders<Issue>.Update.Set(x => x.UserSummary, newUserSummary));
+				}
+			}
+			if (newDescription != null)
+			{
+				if (newDescription.Length == 0)
+				{
+					updates.Add(Builders<Issue>.Update.Unset(x => x.Description));
+				}
+				else
+				{
+					updates.Add(Builders<Issue>.Update.Set(x => x.Description, newDescription));
+				}
+			}
+			if (newManuallyPromoted != null)
+			{
+				updates.Add(Builders<Issue>.Update.Set(x => x.ManuallyPromoted, newManuallyPromoted.Value));
+			}
+			if (newResolvedById != null)
+			{
+				newOwnerId ??= newResolvedById;
+				newAcknowledged ??= true;
+			}
+			if (newOwnerId != null)
+			{
+				if (newOwnerId.Value == UserId.Empty)
+				{
+					updates.Add(Builders<Issue>.Update.Unset(x => x.OwnerId!));
+					updates.Add(Builders<Issue>.Update.Unset(x => x.NominatedAt!));
+					updates.Add(Builders<Issue>.Update.Unset(x => x.NominatedById!));
+				}
+				else
+				{
+					updates.Add(Builders<Issue>.Update.Set(x => x.OwnerId!, newOwnerId.Value));
+
+					updates.Add(Builders<Issue>.Update.Set(x => x.NominatedAt, DateTime.UtcNow));
+					if (newNominatedById == null)
+					{
+						updates.Add(Builders<Issue>.Update.Unset(x => x.NominatedById!));
+					}
+					else
+					{
+						updates.Add(Builders<Issue>.Update.Set(x => x.NominatedById, newNominatedById.Value));
+					}
+					newAcknowledged ??= false;
+				}
+			}
+			if (newAcknowledged != null)
+			{
+				if (newAcknowledged.Value)
+				{
+					if (issueDocument.AcknowledgedAt == null)
+					{
+						updates.Add(Builders<Issue>.Update.Set(x => x.AcknowledgedAt, utcNow));
+					}
+				}
+				else
+				{
+					if (issueDocument.AcknowledgedAt != null)
+					{
+						updates.Add(Builders<Issue>.Update.Unset(x => x.AcknowledgedAt!));
+					}
+				}
+			}
+			if (newFixChange != null)
+			{
+				if (newFixChange == 0)
+				{
+					updates.Add(Builders<Issue>.Update.Unset(x => x.FixChange!));
+				}
+				else
+				{
+					updates.Add(Builders<Issue>.Update.Set(x => x.FixChange, newFixChange));
+				}
+			}
+			if (newResolvedById != null)
+			{
+				if (newResolvedById.Value != UserId.Empty)
+				{
+					if (issueDocument.ResolvedAt == null || issueDocument.ResolvedById != newResolvedById)
+					{
+						updates.Add(Builders<Issue>.Update.Set(x => x.ResolvedAt, utcNow));
+						updates.Add(Builders<Issue>.Update.Set(x => x.ResolvedById, newResolvedById.Value));
+					}
+				}
+				else
+				{
+					if (issueDocument.ResolvedAt != null)
+					{
+						updates.Add(Builders<Issue>.Update.Unset(x => x.ResolvedAt!));
+					}
+					if (issueDocument.ResolvedById != null)
+					{
+						updates.Add(Builders<Issue>.Update.Unset(x => x.ResolvedById!));
+					}
+				}
+			}
+			if (newExcludeSpanIds != null)
+			{
+				List<ObjectId> newCombinedExcludeSpanIds = newExcludeSpanIds;
+				if (issue.ExcludeSpans != null)
+				{
+					newCombinedExcludeSpanIds = newCombinedExcludeSpanIds.Union(issue.ExcludeSpans).ToList();
+				}
+				updates.Add(Builders<Issue>.Update.Set(x => x.ExcludeSpans, newCombinedExcludeSpanIds));
+			}
+			if (newLastSeenAt != null)
+			{
+				updates.Add(Builders<Issue>.Update.Set(x => x.LastSeenAt, newLastSeenAt.Value));
+			}
+
+			if (newDeclinedById != null)
+			{
+				GetLogger(issue.Id).LogInformation("Declined by {UserId}", newDeclinedById.Value);
+				await _issueSuspects.UpdateManyAsync(x => x.IssueId == issue.Id && x.AuthorId == newDeclinedById.Value, Builders<IssueSuspect>.Update.Set(x => x.DeclinedAt, DateTime.UtcNow), null, cancellationToken);
+			}
+			if (newQuarantinedById != null)
+			{
+				if (newQuarantinedById.Value == UserId.Empty)
+				{
+					updates.Add(Builders<Issue>.Update.Unset(x => x.QuarantinedByUserId));
+					updates.Add(Builders<Issue>.Update.Unset(x => x.QuarantineTimeUtc));
+				}
+				else
+				{
+					updates.Add(Builders<Issue>.Update.Set(x => x.QuarantinedByUserId!, newQuarantinedById.Value));
+					updates.Add(Builders<Issue>.Update.Set(x => x.QuarantineTimeUtc, DateTime.UtcNow));
+				}
+			}
+			else if ((newResolvedById != null && newResolvedById.Value != UserId.Empty) || (newForceClosedById != null && newForceClosedById.Value != UserId.Empty))
+			{
+				// Clear quarantine if being resolved or if being force closed
+				updates.Add(Builders<Issue>.Update.Unset(x => x.QuarantinedByUserId));
+				updates.Add(Builders<Issue>.Update.Unset(x => x.QuarantineTimeUtc));
+			}
+
+			if (newForceClosedById != null)
+			{
+				if (newForceClosedById.Value == UserId.Empty)
+				{
+					updates.Add(Builders<Issue>.Update.Unset(x => x.ForceClosedByUserId));
+				}
+				else
+				{
+					updates.Add(Builders<Issue>.Update.Set(x => x.ForceClosedByUserId, newForceClosedById.Value));
+				}
+			}
+
+			if (newExternaIssueKey != null)
+			{
+				updates.Add(Builders<Issue>.Update.Set(x => x.ExternalIssueKey, newExternaIssueKey.Length == 0 ? null : newExternaIssueKey));
+			}
+
+			if (newWorkflowThreadUrl != null)
+			{
+				updates.Add(Builders<Issue>.Update.Set(x => x.WorkflowThreadUrl, newWorkflowThreadUrl.ToString().Length == 0 ? null : newWorkflowThreadUrl));
+			}
+
+			if (updates.Count == 0)
+			{
+				return issueDocument;
+			}
+
+			Issue? newIssue = await TryUpdateIssueAsync(issue, Builders<Issue>.Update.Combine(updates), cancellationToken);
+			if (newIssue == null)
+			{
+				return null;
+			}
+
+			await LogIssueChangesAsync(initiatedByUserId, issueDocument, newIssue, cancellationToken);
+			return newIssue;
+		}
+
+		/// <inheritdoc/>
+		public async Task<IIssue?> TryUpdateIssueDerivedDataAsync(IIssue issue, string newSummary, IssueSeverity newSeverity, List<NewIssueFingerprint> newFingerprints, List<NewIssueStream> newStreams, List<NewIssueSuspectData> newSuspects, DateTime? newResolvedAt, DateTime? newVerifiedAt, DateTime newLastSeenAt, CancellationToken cancellationToken)
+		{
+			Issue issueImpl = (Issue)issue;
+
+			// Update all the suspects for this issue
+			IReadOnlyList<IssueSuspect> oldSuspectImpls = await _issueSuspects.Find(x => x.IssueId == issue.Id).ToListAsync(cancellationToken);
+			List<IssueSuspect> newSuspectImpls = await UpdateIssueSuspectsAsync(issue.Id, oldSuspectImpls, newSuspects, newResolvedAt, cancellationToken);
+
+			// Find the spans for this issue
+			List<IssueSpan> newSpans = await _issueSpans.Find(x => x.IssueId == issue.Id).ToListAsync(cancellationToken);
+
+			// Update the resolved time on any issues
+			List<ObjectId> updateSpanIds = newSpans.Where(x => x.ResolvedAt != newResolvedAt).Select(x => x.Id).ToList();
+			if (updateSpanIds.Count > 0)
+			{
+				FilterDefinition<IssueSpan> filter = Builders<IssueSpan>.Filter.In(x => x.Id, updateSpanIds);
+				await _issueSpans.UpdateManyAsync(filter, Builders<IssueSpan>.Update.Set(x => x.ResolvedAt, newResolvedAt), null, cancellationToken);
+			}
+
+			// Figure out if this issue should be promoted
+			bool newPromoted;
+			if (issueImpl.ManuallyPromoted.HasValue)
+			{
+				newPromoted = issueImpl.ManuallyPromoted.Value;
+			}
+			else if (issueImpl.ManuallyPromotedDeprecated.HasValue)
+			{
+				newPromoted = issueImpl.ManuallyPromotedDeprecated.Value;
+			}
+			else
+			{
+				newPromoted = newSpans.Any(x => ((IIssueSpan)x).PromoteByDefault);
+			}
+
+			// Find the default owner
+			UserId? newDefaultOwnerId = null;
+			string? autoAssignToUser = newSpans.Select(x => x.LastFailure.Annotations?.AutoAssignToUser).Where(x => x != null).FirstOrDefault();
+			if (autoAssignToUser != null)
+			{
+				IUser? user = await _userCollection.FindUserByLoginAsync(autoAssignToUser, cancellationToken);
+				if (user != null)
+				{
+					newDefaultOwnerId = user.Id;
+				}
+			}
+			else
+			{
+				// Figure out if we can auto-assign an owner
+				bool canAutoAssign = newPromoted || newSpans.Any(x => x.LastFailure.Annotations?.AutoAssign ?? false);
+				if (canAutoAssign && newSuspectImpls.Count > 0)
+				{
+					UserId possibleOwnerId = newSuspectImpls[0].AuthorId;
+					if (newSuspectImpls.All(x => x.AuthorId == possibleOwnerId) && newSuspectImpls.Any(x => x.DeclinedAt == null))
+					{
+						newDefaultOwnerId = possibleOwnerId;
+					}
+				}
+			}
+
+			// Get the range of suspect changes
+			int newMinSuspectChange = (newSuspects.Count > 0) ? newSuspects.Min(x => x.Change) : 0;
+			int newMaxSuspectChange = (newSuspects.Count > 0) ? newSuspects.Min(x => x.Change) : 0;
+
+			// Perform the actual update with this data
+			List<UpdateDefinition<Issue>> updates = new List<UpdateDefinition<Issue>>();
+			if (!String.Equals(issue.Summary, newSummary, StringComparison.Ordinal))
+			{
+				updates.Add(Builders<Issue>.Update.Set(x => x.Summary, newSummary));
+			}
+			if (issue.Severity != newSeverity)
+			{
+				updates.Add(Builders<Issue>.Update.Set(x => x.Severity, newSeverity));
+			}
+			if (issue.Promoted != newPromoted)
+			{
+				updates.Add(Builders<Issue>.Update.Set(x => x.Promoted, newPromoted));
+			}
+			if (issue.Fingerprints.Count != newFingerprints.Count || !newFingerprints.Zip(issue.Fingerprints).All(x => x.First.Equals(x.Second)))
+			{
+				updates.Add(Builders<Issue>.Update.Set(x => x.Fingerprints, newFingerprints.Select(x => new IssueFingerprint(x))));
+			}
+			if (issue.Streams.Count != newStreams.Count || !newStreams.Zip(issue.Streams).All(x => x.First.StreamId == x.Second.StreamId && x.First.ContainsFix == x.Second.ContainsFix))
+			{
+				updates.Add(Builders<Issue>.Update.Set(x => x.Streams, newStreams.Select(x => new IssueStream(x))));
+			}
+			if (issueImpl.MinSuspectChange != newMinSuspectChange)
+			{
+				updates.Add(Builders<Issue>.Update.Set(x => x.MinSuspectChange, newMinSuspectChange));
+			}
+			if (issueImpl.MaxSuspectChange != newMaxSuspectChange)
+			{
+				updates.Add(Builders<Issue>.Update.Set(x => x.MaxSuspectChange, newMaxSuspectChange));
+			}
+			if (issueImpl.DefaultOwnerId != newDefaultOwnerId)
+			{
+				updates.Add(Builders<Issue>.Update.Set(x => x.DefaultOwnerId, newDefaultOwnerId));
+			}
+			if (issue.ResolvedAt != newResolvedAt)
+			{
+				updates.Add(Builders<Issue>.Update.SetOrUnsetNull(x => x.ResolvedAt, newResolvedAt));
+			}
+			if (newResolvedAt == null && issue.ResolvedById != null)
+			{
+				updates.Add(Builders<Issue>.Update.Unset(x => x.ResolvedById));
+			}
+			if (issue.VerifiedAt != newVerifiedAt)
+			{
+				updates.Add(Builders<Issue>.Update.SetOrUnsetNull(x => x.VerifiedAt, newVerifiedAt));
+			}
+			if (issue.LastSeenAt != newLastSeenAt)
+			{
+				updates.Add(Builders<Issue>.Update.Set(x => x.LastSeenAt, newLastSeenAt));
+			}
+
+			Issue? newIssue = await TryUpdateIssueAsync(issue, Builders<Issue>.Update.Combine(updates), cancellationToken);
+			if (newIssue != null)
+			{
+				await LogIssueChangesAsync(null, issueImpl, newIssue, cancellationToken);
+				await LogIssueSuspectChangesAsync(GetLogger(issue.Id), oldSuspectImpls, newSuspectImpls, cancellationToken);
+				return newIssue;
+			}
+			return null;
+		}
+
+		async Task<List<IssueSuspect>> UpdateIssueSuspectsAsync(int issueId, IReadOnlyList<IssueSuspect> oldSuspectImpls, List<NewIssueSuspectData> newSuspects, DateTime? resolvedAt, CancellationToken cancellationToken)
+		{
+			List<IssueSuspect> newSuspectImpls = new List<IssueSuspect>(oldSuspectImpls);
+
+			// Find the current list of suspects
+			HashSet<(UserId, int)> curSuspectKeys = new HashSet<(UserId, int)>(oldSuspectImpls.Select(x => (x.AuthorId, x.Change)));
+			List<IssueSuspect> createSuspects = newSuspects.Where(x => !curSuspectKeys.Contains((x.AuthorId, x.Change))).Select(x => new IssueSuspect(issueId, x, resolvedAt)).ToList();
+
+			HashSet<(UserId, int)> newSuspectKeys = new HashSet<(UserId, int)>(newSuspects.Select(x => (x.AuthorId, x.Change)));
+			List<IssueSuspect> deleteSuspects = oldSuspectImpls.Where(x => !newSuspectKeys.Contains((x.AuthorId, x.Change))).ToList();
+
+			// Apply the suspect changes
+			if (createSuspects.Count > 0)
+			{
+				await _issueSuspects.InsertManyIgnoreDuplicatesAsync(createSuspects, cancellationToken);
+				newSuspectImpls.AddRange(createSuspects);
+			}
+			if (deleteSuspects.Count > 0)
+			{
+				await _issueSuspects.DeleteManyAsync(Builders<IssueSuspect>.Filter.In(x => x.Id, deleteSuspects.Select(y => y.Id)), cancellationToken);
+				newSuspectImpls.RemoveAll(x => !newSuspectKeys.Contains((x.AuthorId, x.Change)));
+			}
+
+			// Make sure all the remaining suspects have the correct resolved time
+			if (newSuspectImpls.Any(x => x.ResolvedAt != resolvedAt))
+			{
+				await _issueSuspects.UpdateManyAsync(Builders<IssueSuspect>.Filter.Eq(x => x.IssueId, issueId), Builders<IssueSuspect>.Update.Set(x => x.ResolvedAt, resolvedAt), null, cancellationToken);
+			}
+			return newSuspectImpls;
+		}
+
+		#endregion
+
+		#region Spans
+
+		/// <inheritdoc/>
+		public async Task<IIssueSpan> AddSpanAsync(int issueId, NewIssueSpanData newSpan, CancellationToken cancellationToken)
+		{
+			IssueSpan span = new IssueSpan(issueId, newSpan);
+			await _issueSpans.InsertOneAsync(span, (InsertOneOptions?)null, cancellationToken);
+			SendTelemetry(span);
+			return span;
+		}
+
+		void SendTelemetry(IIssueSpan issueSpan)
+		{
+			if (_globalConfig.CurrentValue.TryGetStream(issueSpan.StreamId, out StreamConfig? streamConfig) && !streamConfig.TelemetryStoreId.IsEmpty)
+			{
+				_telemetrySink.SendEvent(streamConfig.TelemetryStoreId, TelemetryRecordMeta.CurrentHordeInstance, new
+				{
+					EventName = "State.IssueSpan",
+					Id = issueSpan.Id,
+					IssueId = issueSpan.IssueId,
+					Fingerprint = new { Type = issueSpan.Fingerprint.Type, Keys = issueSpan.Fingerprint.Keys },
+					FirstFailure = new { JobId = issueSpan.FirstFailure.JobId, JobName = issueSpan.FirstFailure.JobName, Change = issueSpan.FirstFailure.Change, StepId = issueSpan.FirstFailure.SpanId },
+					LastFailure = (issueSpan.LastFailure != null) ? new { JobId = issueSpan.LastFailure.JobId, JobName = issueSpan.LastFailure.JobName, Change = issueSpan.LastFailure.Change, StepId = issueSpan.LastFailure.SpanId } : null,
+					StreamId = issueSpan.StreamId,
+					StreamName = issueSpan.StreamName,
+					TemplateRefId = issueSpan.TemplateRefId
+				});
+			}
+		}
+
+		/// <inheritdoc/>
+		public async Task<IIssueSpan?> GetSpanAsync(ObjectId spanId, CancellationToken cancellationToken)
+		{
+			return await _issueSpans.Find(Builders<IssueSpan>.Filter.Eq(x => x.Id, spanId)).FirstOrDefaultAsync(cancellationToken);
+		}
+
+		/// <inheritdoc/>
+		public async Task<IIssueSpan?> TryUpdateSpanAsync(IIssueSpan span, NewIssueStepData? newLastSuccess = null, NewIssueStepData? newFailure = null, NewIssueStepData? newNextSuccess = null, List<NewIssueSpanSuspectData>? newSuspects = null, int? newIssueId = null, CancellationToken cancellationToken = default)
+		{
+			List<UpdateDefinition<IssueSpan>> updates = new List<UpdateDefinition<IssueSpan>>();
+			if (newLastSuccess != null)
+			{
+				updates.Add(Builders<IssueSpan>.Update.Set(x => x.MinChange, newLastSuccess.Change));
+				updates.Add(Builders<IssueSpan>.Update.Set(x => x.LastSuccess, new IssueStep(span.Id, newLastSuccess)));
+			}
+			if (newFailure != null)
+			{
+				if (newFailure.Change < span.FirstFailure.Change)
+				{
+					updates.Add(Builders<IssueSpan>.Update.Set(x => x.FirstFailure, new IssueStep(span.Id, newFailure)));
+				}
+				if (newFailure.Change >= span.LastFailure.Change)
+				{
+					updates.Add(Builders<IssueSpan>.Update.Set(x => x.LastFailure, new IssueStep(span.Id, newFailure)));
+				}
+				if (newFailure.PromoteByDefault != span.PromoteByDefault && newFailure.Change >= span.LastFailure.Change)
+				{
+					updates.Add(Builders<IssueSpan>.Update.Set(x => x.PromoteByDefault, newFailure.PromoteByDefault));
+				}
+			}
+			if (newNextSuccess != null)
+			{
+				updates.Add(Builders<IssueSpan>.Update.Set(x => x.MaxChange, newNextSuccess.Change));
+				updates.Add(Builders<IssueSpan>.Update.Set(x => x.NextSuccess, new IssueStep(span.Id, newNextSuccess)));
+			}
+			if (newSuspects != null)
+			{
+				updates.Add(Builders<IssueSpan>.Update.Set(x => x.Suspects, newSuspects.ConvertAll(x => new IssueSpanSuspect(x))));
+			}
+			if (newIssueId != null)
+			{
+				updates.Add(Builders<IssueSpan>.Update.Set(x => x.IssueId, newIssueId.Value));
+			}
+
+			if (updates.Count == 0)
+			{
+				return span;
+			}
+
+			IssueSpan? newSpan = await TryUpdateSpanAsync(span, Builders<IssueSpan>.Update.Combine(updates), cancellationToken);
+			if (newSpan != null)
+			{
+				ILogger logger = GetLogger(newSpan.IssueId);
+				if (newLastSuccess != null)
+				{
+					logger.LogInformation("Set last success for span {SpanId} to job {JobId} at CL {Change}", newSpan.Id, newLastSuccess.JobId, newLastSuccess.Change);
+				}
+				if (newNextSuccess != null)
+				{
+					logger.LogInformation("Set next success for span {SpanId} to job {JobId} at CL {Change}", newSpan.Id, newNextSuccess.JobId, newNextSuccess.Change);
+				}
+				if (newFailure != null)
+				{
+					logger.LogInformation("Added failure for span {SpanId} in job {JobId} at CL {Change}", newSpan.Id, newFailure.JobId, newFailure.Change);
+				}
+				SendTelemetry(newSpan);
+			}
+			return newSpan;
+		}
+
+		/// <inheritdoc/>
+		public async Task<IReadOnlyList<IIssueSpan>> FindSpansAsync(int issueId, CancellationToken cancellationToken)
+		{
+			return await _issueSpans.Find(x => x.IssueId == issueId).ToListAsync(cancellationToken);
+		}
+
+		/// <inheritdoc/>
+		public async Task<IReadOnlyList<IIssueSpan>> FindSpansAsync(IEnumerable<ObjectId> spanIds, CancellationToken cancellationToken)
+		{
+			return await _issueSpans.Find(Builders<IssueSpan>.Filter.In(x => x.Id, spanIds)).ToListAsync(cancellationToken);
+		}
+
+		/// <inheritdoc/>
+		public async Task<IReadOnlyList<IIssueSpan>> FindOpenSpansAsync(StreamId streamId, TemplateId templateId, string nodeName, int change, CancellationToken cancellationToken)
+		{
+			return await _issueSpans.Find(x => x.StreamId == streamId && x.TemplateRefId == templateId && x.NodeName == nodeName && change >= x.MinChange && change <= x.MaxChange).ToListAsync(cancellationToken);
+		}
+
+		/// <inheritdoc/>
+		public async Task<IReadOnlyList<IIssueSpan>> FindSpansAsync(IEnumerable<ObjectId>? spanIds, IEnumerable<int>? issueIds, StreamId? streamId, int? minChange, int? maxChange, bool? resolved, int? index, int? count, CancellationToken cancellationToken)
+		{
+			FilterDefinition<IssueSpan> filter = FilterDefinition<IssueSpan>.Empty;
+
+			if (spanIds != null)
+			{
+				filter &= Builders<IssueSpan>.Filter.In(x => x.Id, spanIds);
+			}
+
+			if (streamId != null)
+			{
+				filter &= Builders<IssueSpan>.Filter.Eq(x => x.StreamId, streamId);
+			}
+
+			if (issueIds != null)
+			{
+				filter &= Builders<IssueSpan>.Filter.In(x => x.IssueId, issueIds.Select<int, int?>(x => x));
+			}
+			if (minChange != null)
+			{
+				filter &= Builders<IssueSpan>.Filter.Not(Builders<IssueSpan>.Filter.Lt(x => x.MaxChange, minChange.Value));
+			}
+			if (maxChange != null)
+			{
+				filter &= Builders<IssueSpan>.Filter.Not(Builders<IssueSpan>.Filter.Gt(x => x.MinChange, maxChange.Value));
+			}
+			if (resolved != null)
+			{
+				if (resolved.Value)
+				{
+					filter &= Builders<IssueSpan>.Filter.Ne(x => x.ResolvedAt, null);
+				}
+				else
+				{
+					filter &= Builders<IssueSpan>.Filter.Eq(x => x.ResolvedAt, null);
+				}
+			}
+
+			return await _issueSpans.Find(filter).Range(index, count).ToListAsync(cancellationToken);
+		}
+
+		#endregion
+
+		#region Steps
+
+		/// <inheritdoc/>
+		public async Task<IIssueStep> AddStepAsync(ObjectId spanId, NewIssueStepData newStep, CancellationToken cancellationToken)
+		{
+			IssueStep step = new IssueStep(spanId, newStep);
+			await _issueSteps.InsertOneAsync(step, (InsertOneOptions?)null, cancellationToken);
+			return step;
+		}
+
+		/// <inheritdoc/>
+		public async Task<IReadOnlyList<IIssueStep>> FindStepsAsync(IEnumerable<ObjectId> spanIds, CancellationToken cancellationToken)
+		{
+			FilterDefinition<IssueStep> filter = Builders<IssueStep>.Filter.In(x => x.SpanId, spanIds);
+			return await _issueSteps.Find(filter).ToListAsync(cancellationToken);
+		}
+
+		/// <inheritdoc/>
+		public async Task<IReadOnlyList<IIssueStep>> FindStepsAsync(JobId jobId, JobStepBatchId? batchId, JobStepId? stepId, CancellationToken cancellationToken)
+		{
+			FilterDefinition<IssueStep> filter = Builders<IssueStep>.Filter.Eq(x => x.JobId, jobId);
+			if (batchId != null)
+			{
+				filter &= Builders<IssueStep>.Filter.Eq(x => x.BatchId, batchId.Value);
+			}
+			if (stepId != null)
+			{
+				filter &= Builders<IssueStep>.Filter.Eq(x => x.StepId, stepId.Value);
+			}
+			return await _issueSteps.Find(filter).ToListAsync(cancellationToken);
+		}
+
+		#endregion
+
+		/// <inheritdoc/>
+		public IAuditLogChannel<int> GetLogger(int issueId)
+		{
+			return _auditLog[issueId];
+		}
+	}
+}
diff --git a/Engine/Source/Programs/Horde/HordeServer/Jobs/Artifacts/ArtifactCollectionV1.cs b/Engine/Source/Programs/Horde/HordeServer/Jobs/Artifacts/ArtifactCollectionV1.cs
new file mode 100644
--- /dev/null
+++ b/Engine/Source/Programs/Horde/HordeServer/Jobs/Artifacts/ArtifactCollectionV1.cs
@@ -0,0 +1,247 @@
+// Copyright Epic Games, Inc. All Rights Reserved.
+
+using System;
+using System.Collections.Generic;
+using System.Linq;
+using System.Threading;
+using System.Threading.Tasks;
+using EpicGames.Core;
+using EpicGames.Horde.Jobs;
+using EpicGames.Horde.Storage;
+using Horde.Server.Server;
+using MongoDB.Bson;
+using MongoDB.Bson.Serialization.Attributes;
+using MongoDB.Driver;
+
+namespace Horde.Server.Jobs.Artifacts
+{
+	/// <summary>
+	/// Wraps functionality for manipulating artifacts
+	/// </summary>
+	public class ArtifactCollectionV1 : IArtifactCollectionV1
+	{
+		class Artifact : IArtifactV1
+		{
+			[BsonRequired, BsonId]
+			public ObjectId Id { get; set; }
+
+			[BsonRequired]
+			public JobId JobId { get; set; }
+
+			[BsonRequired]
+			public string Name { get; set; }
+
+			public JobStepId? StepId { get; set; }
+
+			[BsonRequired]
+			public long Length { get; set; }
+
+			public string MimeType { get; set; }
+
+			[BsonRequired]
+			public int UpdateIndex { get; set; }
+
+			[BsonConstructor]
+			private Artifact()
+			{
+				Name = null!;
+				MimeType = null!;
+			}
+
+			public Artifact(JobId jobId, JobStepId? stepId, string name, long length, string mimeType)
+			{
+				Id = ObjectId.GenerateNewId();
+				JobId = jobId;
+				StepId = stepId;
+				Name = name;
+				Length = length;
+				MimeType = mimeType;
+			}
+		}
+
+		private readonly IObjectStore _objectStore;
+		private readonly IMongoCollection<Artifact> _artifacts;
+
+		/// <summary>
+		/// Constructor
+		/// </summary>
+		public ArtifactCollectionV1(MongoService mongoService, IObjectStore<ArtifactCollectionV1> objectStore)
+		{
+			_objectStore = objectStore;
+
+			// Initialize Artifacts table
+			_artifacts = mongoService.GetCollection<Artifact>("Artifacts", keys => keys.Ascending(x => x.JobId));
+		}
+
+		/// <inheritdoc/>
+		public async Task<IArtifactV1> CreateArtifactAsync(JobId jobId, JobStepId? stepId, string name, string mimeType, System.IO.Stream data, CancellationToken cancellationToken)
+		{
+			// upload first
+			string artifactName = ValidateName(name);
+			await _objectStore.WriteAsync(GetObjectKey(jobId, stepId, artifactName), data, cancellationToken);
+
+			// then create entry
+			Artifact newArtifact = new Artifact(jobId, stepId, artifactName, data.Length, mimeType);
+			await _artifacts.InsertOneAsync(newArtifact, (InsertOneOptions?)null, cancellationToken);
+			return newArtifact;
+		}
+
+		/// <inheritdoc/>
+		public async Task<IReadOnlyList<IArtifactV1>> GetArtifactsAsync(JobId? jobId, JobStepId? stepId, string? name, CancellationToken cancellationToken)
+		{
+			FilterDefinitionBuilder<Artifact> builder = Builders<Artifact>.Filter;
+
+			FilterDefinition<Artifact> filter = FilterDefinition<Artifact>.Empty;
+			if (jobId != null)
+			{
+				filter &= builder.Eq(x => x.JobId, jobId.Value);
+			}
+			if (stepId != null)
+			{
+				filter &= builder.Eq(x => x.StepId, stepId.Value);
+			}
+			if (name != null)
+			{
+				filter &= builder.Eq(x => x.Name, name);
+			}
+
+			return await _artifacts.Find(filter).ToListAsync(cancellationToken);
+		}
+
+		/// <inheritdoc/>
+		public async Task<IReadOnlyList<IArtifactV1>> GetArtifactsAsync(IEnumerable<ObjectId> artifactIds, CancellationToken cancellationToken)
+		{
+			FilterDefinitionBuilder<Artifact> builder = Builders<Artifact>.Filter;
+
+			FilterDefinition<Artifact> filter = FilterDefinition<Artifact>.Empty;
+			filter &= builder.In(x => x.Id, artifactIds);
+
+			return await _artifacts.Find(filter).ToListAsync(cancellationToken);
+		}
+
+		/// <inheritdoc/>
+		public async Task<IArtifactV1?> GetArtifactAsync(ObjectId artifactId, CancellationToken cancellationToken)
+		{
+			return await _artifacts.Find<Artifact>(x => x.Id == artifactId).FirstOrDefaultAsync(cancellationToken);
+		}
+
+		/// <summary>
+		/// Attempts to update a single artifact document, if the update counter is valid
+		/// </summary>
+		/// <param name="current">The current artifact document</param>
+		/// <param name="update">The update definition</param>
+		/// <param name="cancellationToken">Cancellation token for the operation</param>
+		/// <returns>True if the document was updated, false if another writer updated the document first</returns>
+		private async Task<bool> TryUpdateArtifactAsync(Artifact current, UpdateDefinition<Artifact> update, CancellationToken cancellationToken)
+		{
+			UpdateResult result = await _artifacts.UpdateOneAsync<Artifact>(x => x.Id == current.Id && x.UpdateIndex == current.UpdateIndex, update.Set(x => x.UpdateIndex, current.UpdateIndex + 1), cancellationToken: cancellationToken);
+			return result.ModifiedCount == 1;
+		}
+
+		/// <summary>
+		/// Updates an artifact
+		/// </summary>
+		/// <param name="artifact">The artifact</param>
+		/// <param name="newMimeType">New mime type</param>
+		/// <param name="newData">New data</param>
+		/// <param name="cancellationToken"></param>
+		/// <returns>Async task</returns>
+		public async Task<bool> UpdateArtifactAsync(IArtifactV1? artifact, string newMimeType, System.IO.Stream newData, CancellationToken cancellationToken)
+		{
+			while (artifact != null)
+			{
+				UpdateDefinitionBuilder<Artifact> updateBuilder = Builders<Artifact>.Update;
+
+				List<UpdateDefinition<Artifact>> updates = new List<UpdateDefinition<Artifact>>();
+				updates.Add(updateBuilder.Set(x => x.MimeType, newMimeType));
+				updates.Add(updateBuilder.Set(x => x.Length, newData.Length));
+
+				// re-upload the data to external
+				string artifactName = ValidateName(artifact.Name);
+				await _objectStore.WriteAsync(GetObjectKey(artifact.JobId, artifact.StepId, artifactName), newData, cancellationToken);
+
+				if (await TryUpdateArtifactAsync((Artifact)artifact, updateBuilder.Combine(updates), cancellationToken))
+				{
+					return true;
+				}
+
+				artifact = await GetArtifactAsync(artifact.Id, cancellationToken);
+			}
+			return false;
+		}
+
+		/// <summary>
+		/// gets artifact data
+		/// </summary>
+		/// <param name="artifact">The artifact</param>
+		/// <param name="cancellationToken">Cancellation token for the operation</param>
+		/// <returns>The chunk data</returns>
+		public async Task<System.IO.Stream> OpenArtifactReadStreamAsync(IArtifactV1 artifact, CancellationToken cancellationToken)
+		{
+			System.IO.Stream stream = await _objectStore.OpenAsync(GetObjectKey(artifact.JobId, artifact.StepId, artifact.Name), cancellationToken);
+			return stream;
+		}
+
+		/// <summary>
+		/// Get the path for an artifact
+		/// </summary>
+		/// <param name="jobId"></param>
+		/// <param name="stepId"></param>
+		/// <param name="name"></param>
+		/// <returns></returns>
+		private static ObjectKey GetObjectKey(JobId jobId, JobStepId? stepId, string name)
+		{
+			if (stepId == null)
+			{
+				return new ObjectKey(new Utf8String($"{jobId}/{name}.blob"), ObjectKey.Validate.None);
+			}
+			else
+			{
+				return new ObjectKey(new Utf8String($"{jobId}/{stepId.Value}/{name}.blob"), ObjectKey.Validate.None);
+			}
+		}
+
+		/// <summary>
+		/// Checks that the name given is valid for an artifact
+		/// </summary>
+		/// <param name="name"></param>
+		private static string ValidateName(string name)
+		{
+			string newName = name.Replace('\\', '/');
+			if (newName.Length == 0)
+			{
+				throw new ArgumentException("Artifact name is empty");
+			}
+			else if (newName.StartsWith('/'))
+			{
+				throw new ArgumentException($"Artifact has an absolute path ({newName})");
+			}
+			else if (newName.EndsWith('/'))
+			{
+				throw new ArgumentException($"Artifact does not have a file name ({newName})");
+			}
+			else if (newName.Contains("//", StringComparison.Ordinal))
+			{
+				throw new ArgumentException($"Artifact name contains an invalid directory ({newName})");
+			}
+
+			string invalidChars = ":<>|\"";
+			for (int idx = 0; idx < newName.Length; idx++)
+			{
+				int minIdx = idx;
+				for (; idx < newName.Length && newName[idx] != '/'; idx++)
+				{
+					if (invalidChars.Contains(newName[idx], StringComparison.Ordinal))
+					{
+						throw new ArgumentException($"Invalid character in artifact name ({newName})");
+					}
+				}
+				if ((idx == minIdx + 1 && newName[minIdx] == '.') || (idx == minIdx + 2 && newName[minIdx] == '.' && newName[minIdx + 1] == '.'))
+				{
+					throw new ArgumentException($"Artifact may not contain symbolic directory names ({newName})");
+				}
+			}
+			return newName;
+		}
+	}
+}
\ No newline at end of file
