Horde: Rename projects before splitting into plugins (Horde.Agent -> HordeAgent, Horde.Server -> HordeServer, Horde -> HordeCmd).

diff --git a/Engine/Source/Programs/Horde/HordeServer/Jobs/Artifacts/ArtifactsControllerV1.cs b/Engine/Source/Programs/Horde/HordeServer/Jobs/Artifacts/ArtifactsControllerV1.cs
new file mode 100644
--- /dev/null
+++ b/Engine/Source/Programs/Horde/HordeServer/Jobs/Artifacts/ArtifactsControllerV1.cs
@@ -0,0 +1,418 @@
+// Copyright Epic Games, Inc. All Rights Reserved.
+
+using System;
+using System.Collections.Generic;
+using System.IdentityModel.Tokens.Jwt;
+using System.IO;
+using System.IO.Compression;
+using System.Linq;
+using System.Security.Claims;
+using System.Threading;
+using System.Threading.Tasks;
+using EpicGames.Core;
+using EpicGames.Horde.Jobs;
+using Horde.Server.Acls;
+using Horde.Server.Artifacts;
+using Horde.Server.Jobs.Graphs;
+using Horde.Server.Server;
+using Horde.Server.Streams;
+using Horde.Server.Utilities;
+using Microsoft.AspNetCore.Authorization;
+using Microsoft.AspNetCore.Http;
+using Microsoft.AspNetCore.Mvc;
+using Microsoft.Extensions.Options;
+using Microsoft.IdentityModel.Tokens;
+using MongoDB.Bson;
+
+namespace Horde.Server.Jobs.Artifacts
+{
+	/// <summary>
+	/// Controller for the /api/artifacts endpoint
+	/// </summary>
+	[ApiController]
+	[Obsolete("Use /api/v2/artifacts instead")]
+	[Route("[controller]")]
+	public class ArtifactsControllerV1 : ControllerBase
+	{
+		private readonly GlobalsService _globalsService;
+		private readonly IArtifactCollectionV1 _artifactCollection;
+		private readonly AclService _aclService;
+		private readonly JobService _jobService;
+		private readonly IOptionsSnapshot<GlobalConfig> _globalConfig;
+
+		/// <summary>
+		/// Constructor
+		/// </summary>
+		public ArtifactsControllerV1(GlobalsService globalsService, IArtifactCollectionV1 artifactCollection, AclService aclService, JobService jobService, IOptionsSnapshot<GlobalConfig> globalConfig)
+		{
+			_globalsService = globalsService;
+			_artifactCollection = artifactCollection;
+			_aclService = aclService;
+			_jobService = jobService;
+			_globalConfig = globalConfig;
+		}
+
+		/// <summary>
+		/// Creates an artifact
+		/// </summary>
+		/// <param name="jobId">BatchId</param>
+		/// <param name="stepId">StepId</param>
+		/// <param name="file">The file contents</param>
+		/// <param name="cancellationToken">Cancellation token for the operation</param>
+		/// <returns>Http result code</returns>
+		[HttpPost]
+		[Authorize]
+		[Route("/api/v1/artifacts")]
+		public async Task<ActionResult<CreateJobArtifactResponseV1>> CreateArtifactAsync([FromQuery] JobId jobId, [FromQuery] JobStepId? stepId, IFormFile file, CancellationToken cancellationToken)
+		{
+			IJob? job = await _jobService.GetJobAsync(jobId, cancellationToken);
+			if (job == null)
+			{
+				return NotFound(jobId);
+			}
+
+			StreamConfig? streamConfig;
+			if (!_globalConfig.Value.TryGetStream(job.StreamId, out streamConfig))
+			{
+				return NotFound(job.StreamId);
+			}
+			if (!streamConfig.Authorize(ArtifactAclAction.UploadArtifact, User))
+			{
+				return Forbid();
+			}
+
+			IJobStep? step = null;
+			if (stepId != null)
+			{
+				foreach (IJobStepBatch batch in job.Batches)
+				{
+					if (batch.TryGetStep(stepId.Value, out step))
+					{
+						break;
+					}
+				}
+				if (step == null)
+				{
+					// if the step doesn't exist in any of the batches, not found
+					return NotFound();
+				}
+			}
+
+			IArtifactV1 newArtifact = await _artifactCollection.CreateArtifactAsync(job.Id, step?.Id, file.FileName, file.ContentType ?? "horde-mime/unknown", file.OpenReadStream(), cancellationToken);
+			return new CreateJobArtifactResponseV1(newArtifact.Id.ToString());
+		}
+
+		/// <summary>
+		/// Updates an artifact
+		/// </summary>
+		/// <param name="artifactId">JobId</param>
+		/// <param name="file">The file contents</param>
+		/// <param name="cancellationToken">Cancellation token for the operation</param>
+		/// <returns>Http result code</returns>
+		[HttpPut]
+		[Authorize]
+		[Route("/api/v1/artifacts/{artifactId}")]
+		public async Task<ActionResult<CreateJobArtifactResponseV1>> UpdateArtifactAsync(string artifactId, IFormFile file, CancellationToken cancellationToken)
+		{
+			IArtifactV1? artifact = await _artifactCollection.GetArtifactAsync(ObjectId.Parse(artifactId), cancellationToken);
+			if (artifact == null)
+			{
+				return NotFound();
+			}
+			if (!await _jobService.AuthorizeAsync(artifact.JobId, ArtifactAclAction.UploadArtifact, User, _globalConfig.Value, cancellationToken))
+			{
+				return Forbid();
+			}
+
+			await _artifactCollection.UpdateArtifactAsync(artifact, file.ContentType ?? "horde-mime/unknown", file.OpenReadStream(), cancellationToken);
+			return Ok();
+		}
+
+		/// <summary>
+		/// Query artifacts for a job step
+		/// </summary>
+		/// <param name="jobId">Optional JobId to filter by</param>
+		/// <param name="stepId">Optional StepId to filter by</param>
+		/// <param name="code">Whether to generate a direct download code</param>
+		/// <param name="filter">Filter for the properties to return</param>
+		/// <param name="cancellationToken">Cancellation token for the operation</param>
+		/// <returns>Information about all the artifacts</returns>
+		[HttpGet]
+		[Authorize]
+		[Route("/api/v1/artifacts")]
+		[ProducesResponseType(typeof(List<GetJobArtifactResponseV1>), 200)]
+		public async Task<ActionResult<List<object>>> GetArtifactsAsync([FromQuery] JobId jobId, [FromQuery] JobStepId? stepId = null, [FromQuery] bool code = false, [FromQuery] PropertyFilter? filter = null, CancellationToken cancellationToken = default)
+		{
+			if (!await _jobService.AuthorizeAsync(jobId, ArtifactAclAction.DownloadArtifact, User, _globalConfig.Value, cancellationToken))
+			{
+				return Forbid();
+			}
+
+			string? downloadCode = code ? (string?)await GetDirectDownloadCodeForJobAsync(jobId, cancellationToken) : null;
+
+			IReadOnlyList<IArtifactV1> artifacts = await _artifactCollection.GetArtifactsAsync(jobId, stepId, null, cancellationToken);
+			return artifacts.ConvertAll(x => new GetJobArtifactResponseV1(x, downloadCode).ApplyFilter(filter));
+		}
+
+		/// <summary>
+		/// Gets the claim required to download artifacts for a particular job
+		/// </summary>
+		/// <param name="jobId">The job id</param>
+		/// <returns>The required claim</returns>
+		static Claim GetDirectDownloadClaim(JobId jobId)
+		{
+			return new Claim(HordeClaimTypes.JobArtifacts, jobId.ToString());
+		}
+
+		/// <summary>
+		/// Get a download code for the artifacts of a job
+		/// </summary>
+		/// <param name="jobId">The job id</param>
+		/// <param name="cancellationToken">Cancellation token for the operation</param>
+		/// <returns>The download code</returns>
+		async ValueTask<string> GetDirectDownloadCodeForJobAsync(JobId jobId, CancellationToken cancellationToken)
+		{
+			Claim downloadClaim = GetDirectDownloadClaim(jobId);
+			return await _aclService.IssueBearerTokenAsync(new[] { downloadClaim }, TimeSpan.FromHours(4.0), cancellationToken);
+		}
+
+		/// <summary>
+		/// Retrieve metadata about a specific artifact
+		/// </summary>
+		/// <param name="artifactId">Id of the artifact to get information about</param>
+		/// <param name="code">Whether to generate a direct download code</param>
+		/// <param name="filter">Filter for the properties to return</param>
+		/// <param name="cancellationToken">Cancellation token for the operation</param>
+		/// <returns>Information about the requested project</returns>
+		[HttpGet]
+		[Authorize]
+		[Route("/api/v1/artifacts/{artifactId}")]
+		[ProducesResponseType(typeof(GetJobArtifactResponseV1), 200)]
+		public async Task<ActionResult<object>> GetArtifactAsync(string artifactId, bool code = false, [FromQuery] PropertyFilter? filter = null, CancellationToken cancellationToken = default)
+		{
+			IArtifactV1? artifact = await _artifactCollection.GetArtifactAsync(ObjectId.Parse(artifactId), cancellationToken);
+			if (artifact == null)
+			{
+				return NotFound();
+			}
+			if (!await _jobService.AuthorizeAsync(artifact.JobId, ArtifactAclAction.DownloadArtifact, User, _globalConfig.Value, cancellationToken))
+			{
+				return Forbid();
+			}
+
+			string? downloadCode = code ? (string?)await GetDirectDownloadCodeForJobAsync(artifact.JobId, cancellationToken) : null;
+			return new GetJobArtifactResponseV1(artifact, downloadCode).ApplyFilter(filter);
+		}
+
+		/// <summary>
+		/// Retrieve raw data for an artifact
+		/// </summary>
+		/// <param name="artifactId">Id of the artifact to get information about</param>
+		/// <param name="cancellationToken">Cancellation token for the operation</param>
+		/// <returns>Raw artifact data</returns>
+		[HttpGet]
+		[Authorize]
+		[Route("/api/v1/artifacts/{artifactId}/data")]
+		public async Task<ActionResult> GetArtifactDataAsync(string artifactId, CancellationToken cancellationToken)
+		{
+			// Catch case clients are sending an undefined artifact id
+			if (artifactId == "undefined")
+			{
+				return NotFound();
+			}
+
+			IArtifactV1? artifact = await _artifactCollection.GetArtifactAsync(ObjectId.Parse(artifactId), cancellationToken);
+			if (artifact == null)
+			{
+				return NotFound();
+			}
+			if (!await _jobService.AuthorizeAsync(artifact.JobId, ArtifactAclAction.DownloadArtifact, User, _globalConfig.Value, cancellationToken))
+			{
+				return Forbid();
+			}
+
+			// Fun, filestream result automatically closes the stream!
+			return new FileStreamResult(await _artifactCollection.OpenArtifactReadStreamAsync(artifact, cancellationToken), artifact.MimeType);
+		}
+
+		/// <summary>
+		/// Retrieve raw data for an artifact by filename
+		/// </summary>
+		/// <param name="jobId">Unique id for the job</param>
+		/// <param name="stepId">Unique id for the step</param>
+		/// <param name="filename">Filename of artifact from step</param>
+		/// <param name="cancellationToken">Cancellation token for the operation</param>
+		/// <returns>Raw artifact data</returns>
+		[HttpGet]
+		[Route("/api/v1/jobs/{jobId}/steps/{stepId}/artifacts/{filename}/data")]
+		public async Task<ActionResult<object>> GetArtifactDataByFilenameAsync(JobId jobId, JobStepId stepId, string filename, CancellationToken cancellationToken)
+		{
+			if (!await _jobService.AuthorizeAsync(jobId, ArtifactAclAction.DownloadArtifact, User, _globalConfig.Value, cancellationToken))
+			{
+				return Forbid();
+			}
+
+			IReadOnlyList<IArtifactV1> artifacts = await _artifactCollection.GetArtifactsAsync(jobId, stepId, filename, cancellationToken);
+			if (artifacts.Count == 0)
+			{
+				return NotFound();
+			}
+
+			IArtifactV1 artifact = artifacts[0];
+			return new FileStreamResult(await _artifactCollection.OpenArtifactReadStreamAsync(artifact, cancellationToken), artifact.MimeType);
+		}
+
+		/// <summary>
+		/// Retrieve raw data for an artifact
+		/// </summary>
+		/// <param name="artifactId">Id of the artifact to get information about</param>
+		/// <param name="code">The authorization code for this resource</param>
+		/// <param name="cancellationToken">Cancellation token for the operation</param>
+		/// <returns>Raw artifact data</returns>
+		[HttpGet]
+		[AllowAnonymous]
+		[Route("/api/v1/artifacts/{artifactId}/download")]
+		public async Task<ActionResult> DownloadArtifactAsync(string artifactId, [FromQuery] string code, CancellationToken cancellationToken)
+		{
+			IGlobals globals = await _globalsService.GetAsync(cancellationToken);
+
+			TokenValidationParameters parameters = new TokenValidationParameters();
+			parameters.ValidateAudience = false;
+			parameters.RequireExpirationTime = true;
+			parameters.ValidateLifetime = true;
+			parameters.ValidIssuer = globals.JwtIssuer;
+			parameters.ValidateIssuer = true;
+			parameters.ValidateIssuerSigningKey = true;
+			parameters.IssuerSigningKey = globals.JwtSigningKey;
+
+			JwtSecurityTokenHandler handler = new JwtSecurityTokenHandler();
+			TokenValidationResult tokenResult = await handler.ValidateTokenAsync(code, parameters);
+
+			IArtifactV1? artifact = await _artifactCollection.GetArtifactAsync(ObjectId.Parse(artifactId), cancellationToken);
+			if (artifact == null)
+			{
+				return NotFound();
+			}
+
+			Claim directDownloadClaim = GetDirectDownloadClaim(artifact.JobId);
+			if (!tokenResult.ClaimsIdentity.HasClaim(directDownloadClaim.Type, directDownloadClaim.Value))
+			{
+				return Forbid();
+			}
+
+			return new InlineFileStreamResult(await _artifactCollection.OpenArtifactReadStreamAsync(artifact, cancellationToken), artifact.MimeType, Path.GetFileName(artifact.Name));
+		}
+
+		/// <summary>
+		/// Returns a zip archive of many artifacts
+		/// </summary>
+		/// <param name="artifactZipRequest">Artifact request params</param>
+		/// <param name="cancellationToken">Cancellation token for the operation</param>
+		/// <returns>Zip of many artifacts</returns>
+		[HttpPost]
+		[Authorize]
+		[Route("/api/v1/artifacts/zip")]
+		public async Task<ActionResult> ZipArtifactsAsync(GetJobArtifactZipRequestV1 artifactZipRequest, CancellationToken cancellationToken)
+		{
+			if (artifactZipRequest.JobId == null)
+			{
+				return BadRequest("Must specify a JobId");
+			}
+
+			IJob? job = await _jobService.GetJobAsync(artifactZipRequest.JobId.Value, cancellationToken);
+			if (job == null)
+			{
+				return NotFound();
+			}
+			if (!_globalConfig.Value.Authorize(job, ArtifactAclAction.DownloadArtifact, User))
+			{
+				return Forbid();
+			}
+
+			IReadOnlyList<IArtifactV1> artifacts = await _artifactCollection.GetArtifactsAsync(job.Id, artifactZipRequest.StepId, null, cancellationToken);
+
+			Dictionary<ObjectId, IArtifactV1> idToArtifact = artifacts.ToDictionary(x => x.Id, x => x);
+
+			IReadOnlyList<IArtifactV1> zipArtifacts;
+			if (artifactZipRequest.ArtifactIds == null)
+			{
+				zipArtifacts = artifacts;
+			}
+			else
+			{
+				List<IArtifactV1> filteredZipArtifacts = new List<IArtifactV1>();
+				foreach (string artifactId in artifactZipRequest.ArtifactIds)
+				{
+					IArtifactV1? artifact;
+					if (idToArtifact.TryGetValue(ObjectId.Parse(artifactId), out artifact))
+					{
+						filteredZipArtifacts.Add(artifact);
+					}
+					else
+					{
+						return NotFound();
+					}
+				}
+				zipArtifacts = filteredZipArtifacts;
+			}
+
+			IGraph graph = await _jobService.GetGraphAsync(job, cancellationToken);
+
+			return new CustomFileCallbackResult("Artifacts.zip", "application/octet-stream", false, async (outputStream, context) =>
+			{
+				// Make an unseekable MemoryStream for the ZipArchive. We have to do this because the ZipEntry stream falls back to a synchronous write to it's own stream wrappers.
+				using (CustomBufferStream zipOutputStream = new CustomBufferStream())
+				{
+					// Keep the stream open after dispose so we can write the EOF bits.
+					using (ZipArchive zipArchive = new ZipArchive(zipOutputStream, ZipArchiveMode.Create, true))
+					{
+						foreach (IArtifactV1 artifact in zipArtifacts)
+						{
+							await using (System.IO.Stream artifactStream = await _artifactCollection.OpenArtifactReadStreamAsync(artifact, cancellationToken))
+							{
+								// tack on the step name into the directory if it exists
+								string stepName = String.Empty;
+								if (artifact.StepId.HasValue)
+								{
+									foreach (IJobStepBatch batch in job.Batches)
+									{
+										IJobStep? step;
+										if (batch.TryGetStep(artifact.StepId.Value, out step))
+										{
+											stepName = graph.Groups[batch.GroupIdx].Nodes[step.NodeIdx].Name;
+											break;
+										}
+									}
+								}
+
+								ZipArchiveEntry zipEntry = zipArchive.CreateEntry(artifact.Name);
+								using (System.IO.Stream entryStream = zipEntry.Open())
+								{
+									byte[] buffer = new byte[4096];
+									int totalBytesRead = 0;
+									while (totalBytesRead < artifact.Length)
+									{
+										int bytesRead = await artifactStream.ReadAsync(buffer, 0, buffer.Length);
+
+										// Write bytes to the entry stream.  Also advances the MemStream pos.
+										await entryStream.WriteAsync(buffer, 0, bytesRead);
+										// Dump what we have to the output stream
+										await outputStream.WriteAsync(zipOutputStream.GetBuffer(), 0, (int)zipOutputStream.Position);
+
+										// Reset everything.
+										zipOutputStream.Position = 0;
+										zipOutputStream.SetLength(0);
+										totalBytesRead += bytesRead;
+									}
+								}
+							}
+						}
+					}
+					// Write out the EOF stuff
+					zipOutputStream.Position = 0;
+					await zipOutputStream.CopyToAsync(outputStream);
+				}
+			});
+		}
+	}
+}
diff --git a/Engine/Source/Programs/Horde/HordeServer/Jobs/Artifacts/IArtifactCollectionV1.cs b/Engine/Source/Programs/Horde/HordeServer/Jobs/Artifacts/IArtifactCollectionV1.cs
new file mode 100644
--- /dev/null
+++ b/Engine/Source/Programs/Horde/HordeServer/Jobs/Artifacts/IArtifactCollectionV1.cs
@@ -0,0 +1,72 @@
+// Copyright Epic Games, Inc. All Rights Reserved.
+
+using System.Collections.Generic;
+using System.Threading;
+using System.Threading.Tasks;
+using EpicGames.Horde.Jobs;
+using MongoDB.Bson;
+
+namespace Horde.Server.Jobs.Artifacts
+{
+	/// <summary>
+	/// Interface for a collection of artifacts
+	/// </summary>
+	public interface IArtifactCollectionV1
+	{
+		/// <summary>
+		/// Creates a new artifact
+		/// </summary>
+		/// <param name="jobId">Unique id of the job that owns this artifact</param>
+		/// <param name="stepId">Optional Step id</param>
+		/// <param name="name">Name of artifact</param>
+		/// <param name="mimeType">Type of artifact</param>
+		/// <param name="data">The data to write</param>
+		/// <param name="cancellationToken">Cancellation token for the operation</param>
+		/// <returns>The new log file document</returns>
+		Task<IArtifactV1> CreateArtifactAsync(JobId jobId, JobStepId? stepId, string name, string mimeType, System.IO.Stream data, CancellationToken cancellationToken = default);
+
+		/// <summary>
+		/// Gets all the available artifacts for a job
+		/// </summary>
+		/// <param name="jobId">Unique id of the job to query</param>
+		/// <param name="stepId">Unique id of the Step to query</param>
+		/// <param name="name">Name of the artifact</param>
+		/// <param name="cancellationToken">Cancellation token for the operation</param>
+		/// <returns>List of artifact documents</returns>
+		Task<IReadOnlyList<IArtifactV1>> GetArtifactsAsync(JobId? jobId, JobStepId? stepId, string? name, CancellationToken cancellationToken = default);
+
+		/// <summary>
+		/// Gets a specific list of artifacts based on id
+		/// </summary>
+		/// <param name="artifactIds">The list of artifact Ids</param>
+		/// <param name="cancellationToken">Cancellation token for the operation</param>
+		/// <returns>List of artifact documents</returns>
+		Task<IReadOnlyList<IArtifactV1>> GetArtifactsAsync(IEnumerable<ObjectId> artifactIds, CancellationToken cancellationToken = default);
+
+		/// <summary>
+		/// Gets an artifact by ID
+		/// </summary>
+		/// <param name="artifactId">Unique id of the artifact</param>
+		/// <param name="cancellationToken">Cancellation token for the operation</param>
+		/// <returns>The artifact document</returns>
+		Task<IArtifactV1?> GetArtifactAsync(ObjectId artifactId, CancellationToken cancellationToken = default);
+
+		/// <summary>
+		/// Updates an artifact
+		/// </summary>
+		/// <param name="artifact">The artifact</param>
+		/// <param name="newMimeType">New mime type</param>
+		/// <param name="newData">New data</param>
+		/// <param name="cancellationToken">Cancellation token for the operation</param>
+		/// <returns>Async task</returns>
+		Task<bool> UpdateArtifactAsync(IArtifactV1? artifact, string newMimeType, System.IO.Stream newData, CancellationToken cancellationToken = default);
+
+		/// <summary>
+		/// gets artifact data
+		/// </summary>
+		/// <param name="artifact">The artifact</param>
+		/// <param name="cancellationToken">Cancellation token for the operation</param>
+		/// <returns>The chunk data</returns>
+		Task<System.IO.Stream> OpenArtifactReadStreamAsync(IArtifactV1 artifact, CancellationToken cancellationToken = default);
+	}
+}
diff --git a/Engine/Source/Programs/Horde/HordeServer/Jobs/JobService.cs b/Engine/Source/Programs/Horde/HordeServer/Jobs/JobService.cs
new file mode 100644
--- /dev/null
+++ b/Engine/Source/Programs/Horde/HordeServer/Jobs/JobService.cs
@@ -0,0 +1,1469 @@
+// Copyright Epic Games, Inc. All Rights Reserved.
+
+using System;
+using System.Collections.Generic;
+using System.Globalization;
+using System.Linq;
+using System.Security.Claims;
+using System.Text.RegularExpressions;
+using System.Threading;
+using System.Threading.Tasks;
+using EpicGames.Core;
+using EpicGames.Horde.Acls;
+using EpicGames.Horde.Agents;
+using EpicGames.Horde.Agents.Leases;
+using EpicGames.Horde.Agents.Sessions;
+using EpicGames.Horde.Commits;
+using EpicGames.Horde.Jobs;
+using EpicGames.Horde.Jobs.Templates;
+using EpicGames.Horde.Logs;
+using EpicGames.Horde.Streams;
+using EpicGames.Horde.Users;
+using Horde.Server.Agents;
+using Horde.Server.Commits;
+using Horde.Server.Issues;
+using Horde.Server.Jobs.Bisect;
+using Horde.Server.Jobs.Graphs;
+using Horde.Server.Jobs.Templates;
+using Horde.Server.Jobs.Timing;
+using Horde.Server.Notifications;
+using Horde.Server.Perforce;
+using Horde.Server.Server;
+using Horde.Server.Streams;
+using Horde.Server.Users;
+using Horde.Server.Utilities;
+using Microsoft.Extensions.Logging;
+using MongoDB.Bson;
+using OpenTelemetry.Trace;
+
+namespace Horde.Server.Jobs
+{
+	/// <summary>
+	/// Exception thrown when attempting to retry executing a node that does not allow retries
+	/// </summary>
+	public class RetryNotAllowedException : Exception
+	{
+		/// <summary>
+		/// The node name
+		/// </summary>
+		public string NodeName { get; private set; }
+
+		/// <summary>
+		/// Constructor
+		/// </summary>
+		/// <param name="nodeName">Name of the node that does not permit retries</param>
+		public RetryNotAllowedException(string nodeName)
+			: base($"Node '{nodeName}' does not permit retries")
+		{
+			NodeName = nodeName;
+		}
+	}
+
+	/// <summary>
+	/// Cache of information about job ACLs
+	/// </summary>
+	public class JobPermissionsCache
+	{
+		/// <summary>
+		/// Map of job id to permissions for that job
+		/// </summary>
+		public Dictionary<JobId, StreamId> Jobs { get; set; } = new Dictionary<JobId, StreamId>();
+	}
+
+	/// <summary>
+	/// Wraps funtionality for manipulating jobs, jobsteps, and jobstep runs
+	/// </summary>
+	public class JobService
+	{
+		readonly IJobCollection _jobs;
+		readonly IGraphCollection _graphs;
+		readonly IAgentCollection _agents;
+		readonly IJobStepRefCollection _jobStepRefs;
+		readonly IBisectTaskCollection _bisectTasks;
+		readonly IJobTimingCollection _jobTimings;
+		readonly IUserCollection _userCollection;
+		readonly INotificationTriggerCollection _triggerCollection;
+		readonly JobTaskSource _jobTaskSource;
+		readonly IStreamCollection _streamCollection;
+		readonly ITemplateCollection _templateCollection;
+		readonly IssueService? _issueService;
+		readonly IPerforceService _perforceService;
+		readonly Tracer _tracer;
+		readonly ILogger _logger;
+
+		/// <summary>
+		/// Delegate for label update events
+		/// </summary>
+		public delegate void LabelUpdateEvent(IJob job, IReadOnlyList<(LabelState, LabelOutcome)> oldLabelStates, IReadOnlyList<(LabelState, LabelOutcome)> newLabelStates);
+
+		/// <summary>
+		/// Event triggered when a label state changes
+		/// </summary>
+		public event LabelUpdateEvent? OnLabelUpdate;
+
+		/// <summary>
+		/// Delegate for job step complete events
+		/// </summary>
+		public delegate void JobStepCompleteEvent(IJob job, IGraph graph, JobStepBatchId batchId, JobStepId stepId);
+
+		/// <summary>
+		/// Event triggered when a job step completes
+		/// </summary>
+		public event JobStepCompleteEvent? OnJobStepComplete;
+
+		/// <summary>
+		/// Event triggered when a job is scheduled on the underlying JobTaskSource
+		///
+		/// Exposed and duplicated here to avoid dependency on the exact JobTaskSource
+		/// </summary>
+		public event JobTaskSource.JobScheduleEvent? OnJobScheduled;
+
+		/// <summary>
+		/// Constructor
+		/// </summary>
+		public JobService(IJobCollection jobs, IGraphCollection graphs, IAgentCollection agents, IJobStepRefCollection jobStepRefs, IBisectTaskCollection bisectTasks, IJobTimingCollection jobTimings, IUserCollection userCollection, INotificationTriggerCollection triggerCollection, JobTaskSource jobTaskSource, IStreamCollection streamCollection, ITemplateCollection templateCollection, IssueService issueService, IPerforceService perforceService, Tracer tracer, ILogger<JobService> logger)
+		{
+			_jobs = jobs;
+			_graphs = graphs;
+			_agents = agents;
+			_jobStepRefs = jobStepRefs;
+			_bisectTasks = bisectTasks;
+			_jobTimings = jobTimings;
+			_userCollection = userCollection;
+			_triggerCollection = triggerCollection;
+			_jobTaskSource = jobTaskSource;
+			_streamCollection = streamCollection;
+			_templateCollection = templateCollection;
+			_issueService = issueService;
+			_perforceService = perforceService;
+			_tracer = tracer;
+			_logger = logger;
+
+			_jobTaskSource.OnJobScheduled += (pool, numAgentsOnline, job, graph, batchId) =>
+			{
+				OnJobScheduled?.Invoke(pool, numAgentsOnline, job, graph, batchId);
+			};
+		}
+
+#pragma warning disable IDE0060
+		static bool ShouldClonePreflightChange(StreamId streamId)
+		{
+			return false; //			return StreamId == new StreamId("ue5-main");
+		}
+#pragma warning restore IDE0060
+
+		/// <summary>
+		/// Creates a new job
+		/// </summary>
+		/// <param name="jobId">A requested job id</param>
+		/// <param name="streamConfig">The stream that this job belongs to</param>
+		/// <param name="templateRefId">Name of the template ref</param>
+		/// <param name="templateHash">Template for this job</param>
+		/// <param name="graph">The graph for the new job</param>
+		/// <param name="name">Name of the job</param>
+		/// <param name="change">The change to build</param>
+		/// <param name="codeChange">The corresponding code changelist</param>
+		/// <param name="options">Options for the new job</param>
+		/// <param name="cancellationToken">Cancellation token for the operation</param>
+		/// <returns>Unique id representing the job</returns>
+		public async Task<IJob> CreateJobAsync(JobId? jobId, StreamConfig streamConfig, TemplateId templateRefId, ContentHash templateHash, IGraph graph, string name, int change, int codeChange, CreateJobOptions options, CancellationToken cancellationToken = default)
+		{
+			using TelemetrySpan span = _tracer.StartActiveSpan($"{nameof(JobService)}.{nameof(CreateJobAsync)}");
+			span.SetAttribute("JobId", jobId?.ToString());
+			span.SetAttribute("Stream", streamConfig.Name);
+			span.SetAttribute("TemplateRefId", templateRefId);
+			span.SetAttribute("TemplateHash", templateHash);
+			span.SetAttribute("GraphId", graph.Id);
+			span.SetAttribute("Name", name);
+			span.SetAttribute("Change", change);
+			span.SetAttribute("CodeChange", codeChange);
+			span.SetAttribute("PreflightChange", options.PreflightChange);
+			span.SetAttribute("ClonedPreflightChange", options.ClonedPreflightChange);
+			span.SetAttribute("StartedByUserId", options.StartedByUserId.ToString());
+			span.SetAttribute("Priority", options.Priority.ToString());
+			span.SetAttribute("ShowUgsBadges", options.ShowUgsBadges);
+			span.SetAttribute("NotificationChannel", options.NotificationChannel ?? "null");
+			span.SetAttribute("NotificationChannelFilter", options.NotificationChannelFilter ?? "null");
+			span.SetAttribute("Arguments", String.Join(',', options.Arguments));
+
+			if (options.AutoSubmit != null)
+			{
+				span.SetAttribute("AutoSubmit", options.AutoSubmit.Value);
+			}
+			if (options.UpdateIssues != null)
+			{
+				span.SetAttribute("UpdateIssues", options.UpdateIssues.Value);
+			}
+			if (options.JobTriggers != null)
+			{
+				span.SetAttribute("JobTriggers.Count", options.JobTriggers.Count);
+			}
+
+			JobId jobIdValue = jobId ?? JobIdUtils.GenerateNewId();
+			using IDisposable? scope = _logger.BeginScope("CreateJobAsync({JobId})", jobIdValue);
+
+			if (options.PreflightChange != null && ShouldClonePreflightChange(streamConfig.Id))
+			{
+				options.ClonedPreflightChange = await CloneShelvedChangeAsync(streamConfig.ClusterName, options.ClonedPreflightChange ?? options.PreflightChange.Value, cancellationToken);
+			}
+
+			_logger.LogInformation("Creating job at CL {Change}, code CL {CodeChange}, preflight CL {PreflightChange}, cloned CL {ClonedPreflightChange}", change, codeChange, options.PreflightChange, options.ClonedPreflightChange);
+
+			Dictionary<string, string> properties = new Dictionary<string, string>();
+			properties["Change"] = change.ToString(CultureInfo.InvariantCulture);
+			properties["CodeChange"] = codeChange.ToString(CultureInfo.InvariantCulture);
+			properties["PreflightChange"] = options.PreflightChange?.ToString(CultureInfo.InvariantCulture) ?? String.Empty;
+			properties["ClonedPreflightChange"] = options.ClonedPreflightChange?.ToString(CultureInfo.InvariantCulture) ?? String.Empty;
+			properties["StreamId"] = streamConfig.Id.ToString();
+			properties["TemplateId"] = templateRefId.ToString();
+			properties["JobId"] = jobIdValue.ToString();
+
+			for (int idx = 0; idx < options.Arguments.Count; idx++)
+			{
+				options.Arguments[idx] = StringUtils.ExpandProperties(options.Arguments[idx], properties);
+			}
+
+			name = StringUtils.ExpandProperties(name, properties);
+
+			IJob newJob = await _jobs.AddAsync(jobIdValue, streamConfig.Id, templateRefId, templateHash, graph, name, change, codeChange, options, cancellationToken);
+			_jobTaskSource.UpdateQueuedJob(newJob, graph);
+
+			await _jobTaskSource.UpdateUgsBadgesAsync(newJob, graph, new List<(LabelState, LabelOutcome)>(), cancellationToken);
+
+			if (options.StartedByUserId != null)
+			{
+				await _userCollection.UpdateSettingsAsync(options.StartedByUserId.Value, addPinnedJobIds: new[] { newJob.Id }, templateOptions: new UpdateUserJobTemplateOptions { StreamId = streamConfig.Id, TemplateId = templateRefId, TemplateHash = templateHash.ToString(), Arguments = options.Arguments }, cancellationToken: cancellationToken);
+			}
+
+			await AbortAnyDuplicateJobsAsync(newJob, cancellationToken);
+
+			return newJob;
+		}
+
+		private async Task AbortAnyDuplicateJobsAsync(IJob newJob, CancellationToken cancellationToken)
+		{
+			using TelemetrySpan span = _tracer.StartActiveSpan($"{nameof(JobService)}.{nameof(AbortAnyDuplicateJobsAsync)}");
+			span.SetAttribute("JobId", newJob.Id.ToString());
+			span.SetAttribute("JobName", newJob.Name);
+
+			IReadOnlyList<IJob> jobsToAbort = new List<IJob>();
+			if (newJob.PreflightChange > 0)
+			{
+				jobsToAbort = await _jobs.FindAsync(preflightChange: newJob.PreflightChange, cancellationToken: cancellationToken);
+			}
+
+			foreach (IJob job in jobsToAbort)
+			{
+				if (job.GetState() == JobState.Complete)
+				{
+					continue;
+				}
+				if (job.Id == newJob.Id)
+				{
+					continue; // Don't remove the new job
+				}
+				if (job.TemplateId != newJob.TemplateId)
+				{
+					continue;
+				}
+				if (job.TemplateHash != newJob.TemplateHash)
+				{
+					continue;
+				}
+				if (job.StartedByUserId != newJob.StartedByUserId)
+				{
+					continue;
+				}
+				if (String.Join(",", job.Arguments) != String.Join(",", newJob.Arguments))
+				{
+					continue;
+				}
+
+				IJob? updatedJob = await UpdateJobAsync(job, null, null, null, KnownUsers.System, null, null, null, cancellationToken: cancellationToken);
+				if (updatedJob == null)
+				{
+					_logger.LogError("Failed marking duplicate job as aborted! Job ID: {JobId}", job.Id);
+				}
+
+				IJob? updatedJob2 = await GetJobAsync(job.Id, cancellationToken);
+				if (updatedJob2?.AbortedByUserId != updatedJob?.AbortedByUserId)
+				{
+					throw new NotImplementedException();
+				}
+			}
+		}
+
+		/// <summary>
+		/// Deletes a job
+		/// </summary>
+		/// <param name="job">The job to delete</param>
+		/// <param name="cancellationToken">Cancellation token for the operation</param>
+		public async Task<bool> DeleteJobAsync(IJob job, CancellationToken cancellationToken)
+		{
+			using TelemetrySpan span = _tracer.StartActiveSpan($"{nameof(JobService)}.{nameof(DeleteJobAsync)}");
+			span.SetAttribute("JobId", job.Id.ToString());
+			span.SetAttribute("JobName", job.Name);
+
+			using IDisposable? scope = _logger.BeginScope("DeleteJobAsync({JobId})", job.Id);
+
+			// Delete the job
+			while (!await job.TryDeleteAsync(cancellationToken))
+			{
+				IJob? newJob = await job.RefreshAsync(cancellationToken);
+				if (newJob == null)
+				{
+					return false;
+				}
+				job = newJob;
+			}
+
+			// Remove all the triggers from it
+			List<ObjectId> triggerIds = new List<ObjectId>();
+			if (job.NotificationTriggerId != null)
+			{
+				triggerIds.Add(job.NotificationTriggerId.Value);
+			}
+			foreach (IJobStep step in job.Batches.SelectMany(x => x.Steps))
+			{
+				if (step.NotificationTriggerId != null)
+				{
+					triggerIds.Add(step.NotificationTriggerId.Value);
+				}
+			}
+			await _triggerCollection.DeleteAsync(triggerIds, cancellationToken);
+
+			return true;
+		}
+
+		/// <summary>
+		/// Delete all the jobs for a stream
+		/// </summary>
+		/// <param name="streamId">Unique id of the stream</param>
+		/// <param name="cancellationToken">Cancellation token for the operation</param>
+		/// <returns>Async task</returns>
+		public async Task DeleteJobsForStreamAsync(StreamId streamId, CancellationToken cancellationToken)
+		{
+			await _jobs.RemoveStreamAsync(streamId, cancellationToken);
+		}
+
+		/// <summary>
+		/// Updates a new job
+		/// </summary>
+		/// <param name="job">The job document to update</param>
+		/// <param name="name">Name of the job</param>
+		/// <param name="priority">Priority of the job</param>
+		/// <param name="autoSubmit">Whether to automatically submit the preflighted change on completion</param>
+		/// <param name="abortedByUserId">Name of the user that aborted this job</param>
+		/// <param name="onCompleteTriggerId">Object id for a notification trigger</param>
+		/// <param name="reports">New reports to add</param>
+		/// <param name="arguments">New arguments for the job</param>
+		/// <param name="labelIdxToTriggerId">New trigger ID for a label in the job</param>
+		/// <param name="cancellationReason">Reason the job was canceled</param>
+		/// <param name="cancellationToken">Cancellation token for the operation</param>
+		public async Task<IJob?> UpdateJobAsync(IJob job, string? name = null, Priority? priority = null, bool? autoSubmit = null, UserId? abortedByUserId = null, ObjectId? onCompleteTriggerId = null, List<JobReport>? reports = null, List<string>? arguments = null, KeyValuePair<int, ObjectId>? labelIdxToTriggerId = null, string? cancellationReason = null, CancellationToken cancellationToken = default)
+		{
+			using TelemetrySpan span = _tracer.StartActiveSpan($"{nameof(JobService)}.{nameof(UpdateJobAsync)}");
+			span.SetAttribute("JobId", job.Id.ToString());
+			span.SetAttribute("Name", name);
+
+			using IDisposable? scope = _logger.BeginScope("UpdateJobAsync({JobId})", job.Id);
+			for (IJob? newJob = job; newJob != null; newJob = await GetJobAsync(job.Id, cancellationToken))
+			{
+				IGraph graph = await GetGraphAsync(newJob, cancellationToken);
+
+				// Capture the previous label states
+				IReadOnlyList<(LabelState, LabelOutcome)> oldLabelStates = newJob.GetLabelStates(graph);
+
+				// Update the new list of job steps
+				newJob = await newJob.TryUpdateJobAsync(name, priority, autoSubmit, null, null, abortedByUserId, onCompleteTriggerId, reports, arguments, labelIdxToTriggerId, cancellationReason: cancellationReason, cancellationToken: cancellationToken);
+				if (newJob != null)
+				{
+					// Update any badges that have been modified
+					await _jobTaskSource.UpdateUgsBadgesAsync(newJob, graph, oldLabelStates, cancellationToken);
+
+					// Cancel any leases which are no longer required
+					foreach (IJobStepBatch batch in newJob.Batches)
+					{
+						if ((batch.Error == JobStepBatchError.Cancelled || batch.Error == JobStepBatchError.NoLongerNeeded) && (batch.State == JobStepBatchState.Starting || batch.State == JobStepBatchState.Running) && batch.AgentId != null && batch.LeaseId != null)
+						{
+							await CancelLeaseAsync(batch.AgentId.Value, batch.LeaseId.Value, cancellationToken);
+						}
+					}
+					return newJob;
+				}
+			}
+			return null;
+		}
+
+		/// <summary>
+		/// Evaluate a change query to determine which CL to run a job at
+		/// </summary>
+		/// <returns></returns>
+		public async Task<int?> EvaluateChangeQueriesAsync(StreamId streamId, List<ChangeQueryConfig> queries, List<CommitTag>? commitTags, ICommitCollection commits, CancellationToken cancellationToken)
+		{
+			foreach (ChangeQueryConfig query in queries)
+			{
+				int? change = await EvaluateChangeQueryAsync(streamId, query, commitTags, commits, cancellationToken);
+				if (change != null)
+				{
+					return change;
+				}
+			}
+			return null;
+		}
+
+		/// <summary>
+		/// Evaluate a change query to determine which CL to run a job at
+		/// </summary>
+		/// <returns></returns>
+		public async ValueTask<int?> EvaluateChangeQueryAsync(StreamId streamId, ChangeQueryConfig query, List<CommitTag>? commitTags, ICommitCollection commits, CancellationToken cancellationToken)
+		{
+			if (query.Condition == null || query.Condition.Evaluate(propertyName => GetTagPropertyValues(propertyName, commitTags)))
+			{
+				// Find the last change with the listed tag
+				if (query.CommitTag != null)
+				{
+					ICommit? taggedCommit = await commits.FindAsync(null, null, 1, new[] { query.CommitTag.Value }, cancellationToken).FirstOrDefaultAsync(cancellationToken);
+					if (taggedCommit != null)
+					{
+						_logger.LogInformation("Last commit with tag '{Tag}' was {Change}", query.CommitTag.Value, taggedCommit.Number);
+						return taggedCommit.Number;
+					}
+				}
+
+				// Find the job with the given result
+				if (query.TemplateId != null)
+				{
+					List<JobStepOutcome> outcomes = query.Outcomes ?? new List<JobStepOutcome> { JobStepOutcome.Success };
+
+					IReadOnlyList<IJob> jobs = await FindJobsAsync(streamId: streamId, templates: new[] { query.TemplateId.Value }, target: query.Target, state: new[] { JobStepState.Completed }, outcome: outcomes.ToArray(), count: 1, excludeUserJobs: true, excludeCancelled: true, cancellationToken: cancellationToken);
+					if (jobs.Count > 0)
+					{
+						_logger.LogInformation("Last successful build of {TemplateId} target {Target} was job {JobId} at change {Change}", query.TemplateId, query.Target, jobs[0].Id, jobs[0].Change);
+						return jobs[0].Change;
+					}
+				}
+			}
+			return null;
+		}
+
+		static IEnumerable<string> GetTagPropertyValues(string propertyName, List<CommitTag>? commitTags)
+		{
+			const string TagPrefix = "tag.";
+			if (propertyName.StartsWith(TagPrefix, StringComparison.Ordinal) && commitTags != null)
+			{
+				string tagName = propertyName.Substring(TagPrefix.Length);
+				foreach (CommitTag commitTag in commitTags)
+				{
+					if (tagName.Equals(commitTag.Text, StringComparison.OrdinalIgnoreCase))
+					{
+						yield return "1";
+					}
+				}
+			}
+		}
+
+		/// <summary>
+		/// Cancel an active lease
+		/// </summary>
+		/// <param name="agentId">The agent to retreive</param>
+		/// <param name="leaseId">The lease id to update</param>
+		/// <param name="cancellationToken">Cancellation token for the operation</param>
+		/// <returns></returns>
+		async Task CancelLeaseAsync(AgentId agentId, LeaseId leaseId, CancellationToken cancellationToken)
+		{
+			using TelemetrySpan span = _tracer.StartActiveSpan($"{nameof(JobService)}.{nameof(CancelLeaseAsync)}");
+			span.SetAttribute("AgentId", agentId.ToString());
+			span.SetAttribute("LeaseId", leaseId.ToString());
+
+			for (; ; )
+			{
+				IAgent? agent = await _agents.GetAsync(agentId, cancellationToken);
+				if (agent == null)
+				{
+					break;
+				}
+
+				int index = 0;
+				while (index < agent.Leases.Count && agent.Leases[index].Id != leaseId)
+				{
+					index++;
+				}
+				if (index == agent.Leases.Count)
+				{
+					break;
+				}
+
+				AgentLease lease = agent.Leases[index];
+				if (lease.State != LeaseState.Active)
+				{
+					break;
+				}
+
+				IAgent? newAgent = await agent.TryCancelLeaseAsync(index, cancellationToken);
+				if (newAgent != null)
+				{
+					_jobTaskSource.CancelLongPollForAgent(agent.Id);
+					break;
+				}
+			}
+		}
+
+		/// <summary>
+		/// Gets a job with the given unique id
+		/// </summary>
+		/// <param name="jobId">Job id to search for</param>
+		/// <param name="cancellationToken">Cancellation token for the operation</param>
+		/// <returns>Information about the given job</returns>
+		public Task<IJob?> GetJobAsync(JobId jobId, CancellationToken cancellationToken = default)
+		{
+			return _jobs.GetAsync(jobId, cancellationToken);
+		}
+
+		/// <summary>
+		/// Gets the graph for a job
+		/// </summary>
+		/// <param name="job">Job to retrieve the graph for</param>
+		/// <param name="cancellationToken">Cancellation token for the operation</param>
+		/// <returns>The graph for this job</returns>
+		public Task<IGraph> GetGraphAsync(IJob job, CancellationToken cancellationToken = default)
+		{
+			return _graphs.GetAsync(job.GraphHash, cancellationToken);
+		}
+
+		/// <summary>
+		/// Searches for jobs matching the given criteria
+		/// </summary>
+		/// <param name="jobIds">List of job ids to return</param>
+		/// <param name="streamId">The stream containing the job</param>
+		/// <param name="name">Name of the job</param>
+		/// <param name="templates">Templates to look for</param>
+		/// <param name="minChange">The minimum changelist number</param>
+		/// <param name="maxChange">The maximum changelist number</param>		
+		/// <param name="preflightChange">The preflight change to look for</param>
+		/// <param name="preflightOnly">Whether to only include preflights</param>
+		/// <param name="preflightStartedByUser">User for which to include preflight jobs</param>		
+		/// <param name="startedByUser">User for which to include jobs</param>
+		/// <param name="minCreateTime">The minimum creation time</param>
+		/// <param name="maxCreateTime">The maximum creation time</param>
+		/// <param name="target">The target to query</param>
+		/// <param name="batchState">One or more batches matches this state</param>
+		/// <param name="state">State to query</param>
+		/// <param name="outcome">Outcomes to return</param>
+		/// <param name="modifiedBefore">Filter the results by last modified time</param>
+		/// <param name="modifiedAfter">Filter the results by last modified time</param>
+		/// <param name="index">Index of the first result to return</param>
+		/// <param name="count">Number of results to return</param>
+		/// <param name="consistentRead">If the database read should be made to the replica server</param>
+		/// <param name="excludeUserJobs">Whether to exclude user jobs from the find</param>
+		/// <param name="excludeCancelled">Whether to exclude cancelled jobs</param>
+		/// <param name="cancellationToken">Cancellation token for the operation</param>
+		/// <returns>List of jobs matching the given criteria</returns>
+		public async Task<IReadOnlyList<IJob>> FindJobsAsync(JobId[]? jobIds = null, StreamId? streamId = null, string? name = null, TemplateId[]? templates = null, int? minChange = null, int? maxChange = null, int? preflightChange = null, bool? preflightOnly = null, UserId? preflightStartedByUser = null, UserId? startedByUser = null, DateTimeOffset? minCreateTime = null, DateTimeOffset? maxCreateTime = null, string? target = null, JobStepBatchState? batchState = null, JobStepState[]? state = null, JobStepOutcome[]? outcome = null, DateTimeOffset? modifiedBefore = null, DateTimeOffset? modifiedAfter = null, int? index = null, int? count = null, bool consistentRead = true, bool? excludeUserJobs = null, bool? excludeCancelled = null, CancellationToken cancellationToken = default)
+		{
+			using TelemetrySpan span = _tracer.StartActiveSpan($"{nameof(JobService)}.{nameof(FindJobsAsync)}");
+			span.SetAttribute("JobIds", (jobIds == null) ? null : String.Join(',', jobIds));
+			span.SetAttribute("StreamId", streamId);
+			span.SetAttribute("Name", name);
+			span.SetAttribute("Templates", templates);
+			span.SetAttribute("MinChange", minChange);
+			span.SetAttribute("MaxChange", maxChange);
+			span.SetAttribute("PreflightChange", preflightChange);
+			span.SetAttribute("PreflightStartedByUser", preflightStartedByUser?.ToString());
+			span.SetAttribute("StartedByUser", startedByUser?.ToString());
+			span.SetAttribute("MinCreateTime", minCreateTime);
+			span.SetAttribute("MaxCreateTime", maxCreateTime);
+			span.SetAttribute("Target", target);
+			span.SetAttribute("State", state?.ToString());
+			span.SetAttribute("Outcome", outcome?.ToString());
+			span.SetAttribute("ModifiedBefore", modifiedBefore);
+			span.SetAttribute("ModifiedAfter", modifiedAfter);
+			span.SetAttribute("Index", index);
+			span.SetAttribute("Count", count);
+
+			if (target == null && (state == null || state.Length == 0) && (outcome == null || outcome.Length == 0))
+			{
+				return await _jobs.FindAsync(jobIds, streamId, name, templates, minChange, maxChange, preflightChange, preflightOnly, preflightStartedByUser, startedByUser, minCreateTime, maxCreateTime, modifiedBefore, modifiedAfter, batchState, index, count, consistentRead, null, excludeUserJobs, cancellationToken);
+			}
+			else
+			{
+				List<IJob> results = new List<IJob>();
+				_logger.LogInformation("Performing scan for job with ");
+
+				int maxCount = (count ?? 1);
+				while (results.Count < maxCount)
+				{
+					IReadOnlyList<IJob> scanJobs = await _jobs.FindAsync(jobIds, streamId, name, templates, minChange, maxChange, preflightChange, preflightOnly, preflightStartedByUser, startedByUser, minCreateTime, maxCreateTime, modifiedBefore, modifiedAfter, batchState, 0, 5, consistentRead, null, excludeUserJobs, cancellationToken);
+					if (scanJobs.Count == 0)
+					{
+						break;
+					}
+
+					foreach (IJob job in scanJobs.OrderByDescending(x => x.Change))
+					{
+						if (excludeCancelled != null && excludeCancelled.Value && WasCancelled(job))
+						{
+							continue;
+						}
+
+						(JobStepState, JobStepOutcome)? result;
+						if (target == null)
+						{
+							result = job.GetTargetState();
+						}
+						else
+						{
+							result = job.GetTargetState(await GetGraphAsync(job, cancellationToken), target);
+						}
+
+						if (result != null)
+						{
+							(JobStepState jobState, JobStepOutcome jobOutcome) = result.Value;
+							if ((state == null || state.Length == 0 || state.Contains(jobState)) && (outcome == null || outcome.Length == 0 || outcome.Contains(jobOutcome)))
+							{
+								results.Add(job);
+								if (results.Count == maxCount)
+								{
+									break;
+								}
+							}
+						}
+					}
+
+					maxChange = scanJobs.Min(x => x.Change) - 1;
+				}
+
+				return results;
+			}
+		}
+
+		/// <summary>
+		/// Test whether a job was cancelled
+		/// </summary>
+		static bool WasCancelled(IJob job)
+		{
+			return job.AbortedByUserId != null || job.Batches.Any(x => x.Steps.Any(y => y.AbortedByUserId != null));
+		}
+
+		/// <summary>
+		/// Searches for jobs matching a stream with given templates
+		/// </summary>
+		/// <param name="streamId">The stream containing the job</param>
+		/// <param name="templates">Templates to look for</param>
+		/// <param name="preflightStartedByUser">User for which to include preflight jobs</param>
+		/// <param name="maxCreateTime">The maximum creation time</param>
+		/// <param name="modifiedAfter">Filter the results by last modified time</param>	
+		/// <param name="index">Index of the first result to return</param>
+		/// <param name="count">Number of results to return</param>
+		/// <param name="consistentRead">If the database read should be made to the replica server</param>
+		/// <param name="cancellationToken">Cancellation token for the operation</param>
+		/// <returns>List of jobs matching the given criteria</returns>
+		public async Task<IReadOnlyList<IJob>> FindJobsByStreamWithTemplatesAsync(StreamId streamId, TemplateId[] templates, UserId? preflightStartedByUser = null, DateTimeOffset? maxCreateTime = null, DateTimeOffset? modifiedAfter = null, int? index = null, int? count = null, bool consistentRead = true, CancellationToken cancellationToken = default)
+		{
+			using TelemetrySpan span = _tracer.StartActiveSpan($"{nameof(JobService)}.{nameof(FindJobsByStreamWithTemplatesAsync)}");
+			span.SetAttribute("StreamId", streamId);
+			span.SetAttribute("Templates", templates);
+			span.SetAttribute("PreflightStartedByUser", preflightStartedByUser?.ToString());
+			span.SetAttribute("MaxCreateTime", maxCreateTime);
+			span.SetAttribute("ModifiedAfter", modifiedAfter);
+			span.SetAttribute("Index", index);
+			span.SetAttribute("Count", count);
+
+			return await _jobs.FindLatestByStreamWithTemplatesAsync(streamId, templates, preflightStartedByUser, maxCreateTime, modifiedAfter, index, count, consistentRead, cancellationToken);
+		}
+
+		/// <summary>
+		/// Attempts to update the node groups to be executed for a job. Fails if another write happens in the meantime.
+		/// </summary>
+		/// <param name="job">The job to update</param>
+		/// <param name="newGraph">New graph for this job</param>
+		/// <param name="cancellationToken">Cancellation token for the operation</param>
+		/// <returns>True if the groups were updated to the given list. False if another write happened first.</returns>
+		public async Task<IJob?> TryUpdateGraphAsync(IJob job, IGraph newGraph, CancellationToken cancellationToken)
+		{
+			using TelemetrySpan span = _tracer.StartActiveSpan($"{nameof(JobService)}.{nameof(TryUpdateGraphAsync)}");
+			span.SetAttribute("Job", job.Id.ToString());
+			span.SetAttribute("NewGraph", newGraph.Id);
+
+			using IDisposable? scope = _logger.BeginScope("TryUpdateGraphAsync({JobId})", job.Id);
+
+			IReadOnlyList<(LabelState, LabelOutcome)> oldLabelStates = job.GetLabelStates(newGraph);
+
+			IJob? newJob = await job.TryUpdateGraphAsync(newGraph, cancellationToken);
+			if (newJob != null)
+			{
+				await _jobTaskSource.UpdateUgsBadgesAsync(newJob, newGraph, oldLabelStates, cancellationToken);
+				_jobTaskSource.UpdateQueuedJob(newJob, newGraph);
+			}
+			return newJob;
+		}
+
+		/// <summary>
+		/// Gets the timing info for a particular job.
+		/// </summary>
+		/// <param name="job">The job to get timing info for</param>
+		/// <param name="cancellationToken">Cancellation token for the operation</param>
+		/// <returns>Timing info for the given job</returns>
+		public async Task<IJobTiming> GetJobTimingAsync(IJob job, CancellationToken cancellationToken)
+		{
+			using TelemetrySpan span = _tracer.StartActiveSpan($"{nameof(JobService)}.{nameof(GetJobTimingAsync)}");
+			span.SetAttribute("Job", job.Id.ToString());
+
+			using IDisposable? scope = _logger.BeginScope("GetJobTimingAsync({JobId})", job.Id);
+
+			IGraph graph = await _graphs.GetAsync(job.GraphHash, cancellationToken);
+
+			Dictionary<string, JobStepTimingData> cachedNewSteps = new Dictionary<string, JobStepTimingData>();
+			for (; ; )
+			{
+				// Try to get the current timing document
+				IJobTiming? jobTiming = await _jobTimings.TryGetAsync(job.Id);
+
+				// Calculate timing info for any missing steps
+				List<JobStepTimingData> newSteps = new List<JobStepTimingData>();
+				foreach (IJobStepBatch batch in job.Batches)
+				{
+					INodeGroup group = graph.Groups[batch.GroupIdx];
+					foreach (IJobStep step in batch.Steps)
+					{
+						INode node = group.Nodes[step.NodeIdx];
+						if (jobTiming == null || !jobTiming.TryGetStepTiming(node.Name, out _))
+						{
+							JobStepTimingData? stepTimingData;
+							if (!cachedNewSteps.TryGetValue(node.Name, out stepTimingData))
+							{
+								stepTimingData = await GetStepTimingInfoAsync(job.StreamId, job.TemplateId, node.Name, job.Change, cancellationToken);
+							}
+							newSteps.Add(stepTimingData);
+						}
+					}
+				}
+
+				// Try to add or update the timing document
+				if (jobTiming == null)
+				{
+					IJobTiming? newJobTiming = await _jobTimings.TryAddAsync(job.Id, newSteps);
+					if (newJobTiming != null)
+					{
+						return newJobTiming;
+					}
+				}
+				else if (newSteps.Count == 0)
+				{
+					return jobTiming;
+				}
+				else
+				{
+					IJobTiming? newJobTiming = await _jobTimings.TryAddStepsAsync(jobTiming, newSteps);
+					if (newJobTiming != null)
+					{
+						return newJobTiming;
+					}
+				}
+
+				// Update the cached list of steps for the next iteration
+				foreach (JobStepTimingData newStep in newSteps)
+				{
+					cachedNewSteps[newStep.Name] = newStep;
+				}
+			}
+		}
+
+		/// <summary>
+		/// Finds the average duration for the given step
+		/// </summary>
+		/// <param name="streamId">The stream to search</param>
+		/// <param name="templateId">The template id</param>
+		/// <param name="nodeName">Name of the node</param>
+		/// <param name="change">Maximum changelist to consider</param>
+		/// <param name="cancellationToken">Cancellation token for the operation</param>
+		/// <returns>Expected duration for the given step</returns>
+		async Task<JobStepTimingData> GetStepTimingInfoAsync(StreamId streamId, TemplateId templateId, string nodeName, int? change, CancellationToken cancellationToken)
+		{
+			using TelemetrySpan span = _tracer.StartActiveSpan($"{nameof(JobService)}.{nameof(GetStepTimingInfoAsync)}");
+			span.SetAttribute("StreamId", streamId);
+			span.SetAttribute("TemplateId", templateId);
+			span.SetAttribute("NodeName", nodeName);
+			span.SetAttribute("Change", change);
+
+			// Find all the steps matching the given criteria
+			List<IJobStepRef> steps = await _jobStepRefs.GetStepsForNodeAsync(streamId, templateId, nodeName, change, false, 10, cancellationToken);
+
+			// Sum up all the durations and wait times
+			int count = 0;
+			float waitTimeSum = 0;
+			float initTimeSum = 0;
+			float durationSum = 0;
+
+			foreach (IJobStepRef step in steps)
+			{
+				if (step.FinishTimeUtc != null)
+				{
+					waitTimeSum += step.BatchWaitTime;
+					initTimeSum += step.BatchInitTime;
+					durationSum += (float)(step.FinishTimeUtc.Value - step.StartTimeUtc).TotalSeconds;
+					count++;
+				}
+			}
+
+			// Compute the averages
+			if (count == 0)
+			{
+				return new JobStepTimingData(nodeName, null, null, null);
+			}
+			else
+			{
+				return new JobStepTimingData(nodeName, waitTimeSum / count, initTimeSum / count, durationSum / count);
+			}
+		}
+
+		/// <summary>
+		/// Updates the state of a batch
+		/// </summary>
+		/// <param name="job">Job to update</param>
+		/// <param name="batchId">Unique id of the batch to update</param>
+		/// <param name="streamConfig">The current stream config</param>
+		/// <param name="newLogId">The new log file id</param>
+		/// <param name="newState">New state of the jobstep</param>
+		/// <param name="cancellationToken">Cancellation token for the operation</param>
+		/// <returns>True if the job was updated, false if it was deleted</returns>
+		public async Task<IJob?> UpdateBatchAsync(IJob job, JobStepBatchId batchId, StreamConfig streamConfig, LogId? newLogId = null, JobStepBatchState? newState = null, CancellationToken cancellationToken = default)
+		{
+			using TelemetrySpan span = _tracer.StartActiveSpan($"{nameof(JobService)}.{nameof(UpdateBatchAsync)}");
+			span.SetAttribute("Job", job.Id.ToString());
+			span.SetAttribute("BatchId", batchId.ToString());
+			span.SetAttribute("NewLogId", newLogId?.ToString());
+			span.SetAttribute("NewState", newState.ToString());
+
+			using IDisposable? scope = _logger.BeginScope("UpdateBatchAsync({JobId})", job.Id);
+
+			JobStepBatchError? error = null;
+
+			bool checkForBadAgent = true;
+			for (; ; )
+			{
+				// Find the index of the appropriate batch
+				int batchIdx = job.Batches.FindIndex(x => x.Id == batchId);
+				if (batchIdx == -1)
+				{
+					return null;
+				}
+
+				IJobStepBatch batch = job.Batches[batchIdx];
+
+				// If the batch has already been marked as complete, error out
+				if (batch.State == JobStepBatchState.Complete)
+				{
+					return null;
+				}
+
+				// If we're marking the batch as complete before the agent has run everything, mark it to conform
+				JobStepBatchError? newError = null;
+				if (newState == JobStepBatchState.Complete && batch.AgentId != null)
+				{
+					if (batch.Steps.Any(x => x.State == JobStepState.Waiting || x.State == JobStepState.Ready || x.State == JobStepState.Running))
+					{
+						// Check if the job is valid. If not, we will fail with a specific error code for it.
+						error ??= await CheckJobAsync(job, streamConfig, cancellationToken);
+
+						// Find the agent and set the conform flag
+						if (error == null && checkForBadAgent)
+						{
+							for (; ; )
+							{
+								IAgent? agent = await _agents.GetAsync(batch.AgentId.Value, cancellationToken);
+								if (agent == null || agent.RequestConform)
+								{
+									break;
+								}
+								if (await agent.TryUpdateAsync(new UpdateAgentOptions { RequestConform = true }, cancellationToken: cancellationToken) != null)
+								{
+									_logger.LogError("Agent {AgentId} did not complete lease; marking for conform", agent.Id);
+									break;
+								}
+							}
+							checkForBadAgent = false;
+						}
+
+						// If it failed during setup, use a dedicated error code.
+						if (error != null)
+						{
+							newError = error.Value;
+						}
+						else if (!batch.Steps.Any(x => x.StartTimeUtc != null))
+						{
+							newError = JobStepBatchError.SyncingFailed;
+						}
+						else
+						{
+							newError = JobStepBatchError.Incomplete;
+						}
+					}
+				}
+
+				// Update the batch state
+				IJob? newJob = await TryUpdateBatchAsync(job, batchId, newLogId, newState, newError, cancellationToken);
+				if (newJob != null)
+				{
+					return newJob;
+				}
+
+				// Update the job
+				newJob = await GetJobAsync(job.Id, cancellationToken);
+				if (newJob == null)
+				{
+					return null;
+				}
+
+				job = newJob;
+			}
+		}
+
+		async Task<JobStepBatchError?> CheckJobAsync(IJob job, StreamConfig streamConfig, CancellationToken cancellationToken)
+		{
+			try
+			{
+				if (job.PreflightChange != 0)
+				{
+					(CheckShelfResult result, _) = await _perforceService.CheckShelfAsync(streamConfig, job.PreflightChange, cancellationToken);
+					if (result != CheckShelfResult.Ok)
+					{
+						_logger.LogWarning("Job {JobId} is no longer valid - check shelf returned {Result}", job.Id, result);
+						return JobStepBatchError.UnknownShelf;
+					}
+				}
+				return null;
+			}
+			catch (Exception ex)
+			{
+				_logger.LogWarning(ex, "Job {JobId} is no longer valid.", job.Id);
+				return JobStepBatchError.ExecutionError;
+			}
+		}
+
+		/// <summary>
+		/// Attempts to update the state of a batch
+		/// </summary>
+		/// <param name="job">Job to update</param>
+		/// <param name="batchId">Unique id of the batch to update</param>
+		/// <param name="newLogId">The new log file id</param>
+		/// <param name="newState">New state of the jobstep</param>
+		/// <param name="newError">New error state</param>
+		/// <param name="cancellationToken">Cancellation token for the operation</param>
+		/// <returns>The updated job, otherwise null</returns>
+		public async Task<IJob?> TryUpdateBatchAsync(IJob job, JobStepBatchId batchId, LogId? newLogId = null, JobStepBatchState? newState = null, JobStepBatchError? newError = null, CancellationToken cancellationToken = default)
+		{
+			return await job.TryUpdateBatchAsync(batchId, newLogId, newState, newError, cancellationToken);
+		}
+
+		/// <summary>
+		/// Update a jobstep state
+		/// </summary>
+		/// <param name="job">Job to update</param>
+		/// <param name="batchId">Unique id of the batch containing the step</param>
+		/// <param name="stepId">Unique id of the step to update</param>
+		/// <param name="streamConfig">Current global config</param>
+		/// <param name="newState">New state of the jobstep</param>
+		/// <param name="newOutcome">New outcome of the jobstep</param>
+		/// <param name="newError">New error for the step</param>
+		/// <param name="newAbortRequested">New state of abort request</param>
+		/// <param name="newAbortByUserId">New user that requested the abort</param>
+		/// <param name="newLogId">New log id for the jobstep</param>
+		/// <param name="newNotificationTriggerId">New notification trigger id for the jobstep</param>
+		/// <param name="newRetryByUserId">Whether the step should be retried</param>
+		/// <param name="newPriority">New priority for this step</param>
+		/// <param name="newReports">New list of reports</param>
+		/// <param name="newProperties">Property changes. Any properties with a null value will be removed.</param>
+		/// <param name="newCancellationReason">The reason the step was canceled</param>
+		/// <param name="cancellationToken">Cancellation token for the operation</param>
+		/// <returns>True if the job was updated, false if it was deleted in the meantime</returns>
+		public async Task<IJob?> UpdateStepAsync(IJob job, JobStepBatchId batchId, JobStepId stepId, StreamConfig streamConfig, JobStepState newState = JobStepState.Unspecified, JobStepOutcome newOutcome = JobStepOutcome.Unspecified, JobStepError? newError = null, bool? newAbortRequested = null, UserId? newAbortByUserId = null, LogId? newLogId = null, ObjectId? newNotificationTriggerId = null, UserId? newRetryByUserId = null, Priority? newPriority = null, List<JobReport>? newReports = null, Dictionary<string, string?>? newProperties = null, string? newCancellationReason = null, CancellationToken cancellationToken = default)
+		{
+			using TelemetrySpan span = _tracer.StartActiveSpan($"{nameof(JobService)}.{nameof(UpdateStepAsync)}");
+			span.SetAttribute("Job", job.Id.ToString());
+			span.SetAttribute("BatchId", batchId.ToString());
+			span.SetAttribute("StepId", stepId.ToString());
+
+			using IDisposable? scope = _logger.BeginScope("UpdateStepAsync({JobId}:{BatchId}:{StepId})", job.Id, batchId, stepId);
+			for (; ; )
+			{
+				IJob? newJob = await TryUpdateStepAsync(job, batchId, stepId, streamConfig, newState, newOutcome, newError, newAbortRequested, newAbortByUserId, newLogId, newNotificationTriggerId, newRetryByUserId, newPriority, newReports, newProperties, newCancellationReason, cancellationToken);
+				if (newJob != null)
+				{
+					return newJob;
+				}
+
+				newJob = await GetJobAsync(job.Id, cancellationToken);
+				if (newJob == null)
+				{
+					return null;
+				}
+
+				job = newJob;
+			}
+		}
+
+		/// <summary>
+		/// Update a jobstep state
+		/// </summary>
+		/// <param name="job">Job to update</param>
+		/// <param name="batchId">Unique id of the batch containing the step</param>
+		/// <param name="stepId">Unique id of the step to update</param>
+		/// <param name="streamConfig">The current stream config</param>
+		/// <param name="newState">New state of the jobstep</param>
+		/// <param name="newOutcome">New outcome of the jobstep</param>
+		/// <param name="newError">New error for the step</param>
+		/// <param name="newAbortRequested">New state for request abort</param>
+		/// <param name="newAbortByUserId">New name of user that requested the abort</param>
+		/// <param name="newLogId">New log id for the jobstep</param>
+		/// <param name="newTriggerId">New trigger id for the jobstep</param>
+		/// <param name="newRetryByUserId">Whether the step should be retried</param>
+		/// <param name="newPriority">New priority for this step</param>
+		/// <param name="newReports">New reports</param>
+		/// <param name="newProperties">Property changes. Any properties with a null value will be removed.</param>
+		/// <param name="newCancellationReason">Property changes. Any properties with a null value will be removed.</param>
+		/// <param name="cancellationToken">Cancellation token for the operation</param>
+		/// <returns>True if the job was updated, false if it was deleted in the meantime</returns>
+		public async Task<IJob?> TryUpdateStepAsync(IJob job, JobStepBatchId batchId, JobStepId stepId, StreamConfig streamConfig, JobStepState newState = JobStepState.Unspecified, JobStepOutcome newOutcome = JobStepOutcome.Unspecified, JobStepError? newError = null, bool? newAbortRequested = null, UserId? newAbortByUserId = null, LogId? newLogId = null, ObjectId? newTriggerId = null, UserId? newRetryByUserId = null, Priority? newPriority = null, List<JobReport>? newReports = null, Dictionary<string, string?>? newProperties = null, string? newCancellationReason = null, CancellationToken cancellationToken = default)
+		{
+			using TelemetrySpan span = _tracer.StartActiveSpan($"{nameof(JobService)}.{nameof(TryUpdateStepAsync)}");
+			span.SetAttribute("Job", job.Id.ToString());
+			span.SetAttribute("BatchId", batchId.ToString());
+			span.SetAttribute("StepId", stepId.ToString());
+
+			using IDisposable? scope = _logger.BeginScope("TryUpdateStepAsync({JobId}:{BatchId}:{StepId})", job.Id, batchId, stepId);
+
+			// Get the graph for this job
+			IGraph graph = await GetGraphAsync(job, cancellationToken);
+
+			// Make sure the job state is set to unspecified
+			JobStepState oldState = JobStepState.Unspecified;
+			JobStepOutcome oldOutcome = JobStepOutcome.Unspecified;
+			if (newState != JobStepState.Unspecified)
+			{
+				if (job.TryGetStep(batchId, stepId, out IJobStep? step))
+				{
+					oldState = step.State;
+					oldOutcome = step.Outcome;
+				}
+			}
+
+			// If we're changing the state of a step, capture the label states beforehand so we can update any that change
+			IReadOnlyList<(LabelState, LabelOutcome)>? oldLabelStates = null;
+			if (oldState != newState || oldOutcome != newOutcome)
+			{
+				oldLabelStates = job.GetLabelStates(graph);
+			}
+
+			// Update the step
+			IJob? newJob = await job.TryUpdateStepAsync(batchId, stepId, newState, newOutcome, newError, newAbortRequested, newAbortByUserId, newLogId, newTriggerId, newRetryByUserId, newPriority, newReports, newProperties, newCancellationReason, cancellationToken);
+			if (newJob != null)
+			{
+				job = newJob;
+
+				using TelemetrySpan updateSpan = _tracer.StartActiveSpan($"{nameof(JobService)}.{nameof(TryUpdateStepAsync)}.Update");
+				updateSpan.SetAttribute("JobId", job.Id.ToString());
+				updateSpan.SetAttribute("BatchId", batchId.ToString());
+				updateSpan.SetAttribute("StepId", stepId.ToString());
+
+				if (oldState != newState || oldOutcome != newOutcome)
+				{
+					_logger.LogInformation("Transitioned job {JobId}, batch {BatchId}, step {StepId} from {OldState} to {NewState}", job.Id, batchId, stepId, oldState, newState);
+
+					// Send any updates for modified badges
+					if (oldLabelStates != null)
+					{
+						using TelemetrySpan _ = _tracer.StartActiveSpan($"{nameof(JobService)}.{nameof(TryUpdateStepAsync)}.SendBadgeUpdates");
+						IReadOnlyList<(LabelState, LabelOutcome)> newLabelStates = job.GetLabelStates(graph);
+						OnLabelUpdate?.Invoke(job, oldLabelStates, newLabelStates);
+						await _jobTaskSource.UpdateUgsBadgesAsync(job, graph, oldLabelStates, newLabelStates, cancellationToken);
+					}
+
+					// Submit the change if auto-submit is enabled
+					if (job.PreflightChange != 0 && newState == JobStepState.Completed)
+					{
+						(JobStepState state, JobStepOutcome outcome) = job.GetTargetState();
+						if (state == JobStepState.Completed)
+						{
+							if (job.AutoSubmit && outcome == JobStepOutcome.Success && job.AbortedByUserId == null)
+							{
+								job = await AutoSubmitChangeAsync(streamConfig, job, graph, cancellationToken);
+							}
+							else if (job.ClonedPreflightChange == 0 && job.StartedByUserId.HasValue && outcome == JobStepOutcome.Success && job.AbortedByUserId == null)
+							{
+								IUserSettings settings = await _userCollection.GetSettingsAsync(job.StartedByUserId.Value, cancellationToken);
+								if (settings.AlwaysTagPreflightCL)
+								{
+									_logger.LogInformation("Updating description for {PreflightChange} for {UserId} user settings", job.PreflightChange, job.StartedByUserId.Value);
+									await _perforceService.UpdateChangelistDescriptionAsync(streamConfig.ClusterName, job.PreflightChange, x => x.TrimEnd() + $"\n#preflight {job.Id}", cancellationToken);
+								}
+							}
+							else if (job.ClonedPreflightChange != 0)
+							{
+								await DeleteShelvedChangeAsync(streamConfig.ClusterName, job.ClonedPreflightChange);
+							}
+						}
+					}
+
+					// Notify subscribers
+					if (newState == JobStepState.Completed)
+					{
+						using TelemetrySpan _ = _tracer.StartActiveSpan($"{nameof(JobService)}.{nameof(TryUpdateStepAsync)}.NotifySubscribers");
+						OnJobStepComplete?.Invoke(job, graph, batchId, stepId);
+					}
+
+					// Create any downstream jobs
+					if (newState == JobStepState.Completed && newOutcome != JobStepOutcome.Failure)
+					{
+						using TelemetrySpan _ = _tracer.StartActiveSpan($"{nameof(JobService)}.{nameof(TryUpdateStepAsync)}.CreateDownstreamJobs");
+						for (int idx = 0; idx < job.ChainedJobs.Count; idx++)
+						{
+							IChainedJob jobTrigger = job.ChainedJobs[idx];
+							if (jobTrigger.JobId == null)
+							{
+								// Check the state of the trigger's target here; the step state above may not be the thing the trigger is waiting on.
+								(JobStepState, JobStepOutcome)? state = job.GetTargetState(graph, jobTrigger.Target);
+								if (state != null && state.Value.Item1 == JobStepState.Completed)
+								{
+									JobStepOutcome jobTriggerOutcome = state.Value.Item2;
+									if (jobTriggerOutcome == JobStepOutcome.Success || jobTriggerOutcome == JobStepOutcome.Warnings)
+									{
+										job = await FireJobTriggerAsync(job, graph, jobTrigger, streamConfig) ?? job;
+									}
+								}
+							}
+						}
+					}
+
+					// Update the jobstep ref if it completed
+					if (newState == JobStepState.Running || newState == JobStepState.Completed || newState == JobStepState.Aborted)
+					{
+						using TelemetrySpan _ = _tracer.StartActiveSpan($"{nameof(JobService)}.{nameof(TryUpdateStepAsync)}.UpdateJobStepRef");
+						if (job.TryGetBatch(batchId, out IJobStepBatch? batch) && batch.TryGetStep(stepId, out IJobStep? step) && step.StartTimeUtc != null)
+						{
+							await _jobStepRefs.UpdateAsync(job, batch, step, graph);
+							await _bisectTasks.UpdateAsync(job, batch, step, graph, cancellationToken: cancellationToken);
+						}
+					}
+
+					// Update any issues that depend on this step
+					if (newState == JobStepState.Completed && job.UpdateIssues)
+					{
+						if (_issueService != null)
+						{
+							using TelemetrySpan _ = _tracer.StartActiveSpan($"{nameof(JobService)}.{nameof(TryUpdateStepAsync)}.UpdateIssuesV2");
+							try
+							{
+								await _issueService.UpdateCompleteStepAsync(job, graph, batchId, stepId, cancellationToken);
+							}
+							catch (Exception ex)
+							{
+								_logger.LogError(ex, "Exception while updating issue service for job {JobId}:{BatchId}:{StepId}: {Message}", job.Id, batchId, stepId, ex.Message);
+							}
+						}
+					}
+				}
+
+				using (TelemetrySpan _ = _tracer.StartActiveSpan($"{nameof(JobService)}.{nameof(TryUpdateStepAsync)}.UpdateQueuedJobs"))
+				{
+					// Notify the dispatch service that the job has changed
+					_jobTaskSource.UpdateQueuedJob(job, graph);
+				}
+
+				return job;
+			}
+
+			return null;
+		}
+
+		/// <summary>
+		/// Submit the given change for a preflight
+		/// </summary>
+		/// <param name="streamConfig">The current stream config</param>
+		/// <param name="job">The job being run</param>
+		/// <param name="graph">Graph for the job</param>
+		/// <param name="cancellationToken">Cancellation token for the operation</param>
+		/// <returns></returns>
+		private async Task<IJob> AutoSubmitChangeAsync(StreamConfig streamConfig, IJob job, IGraph graph, CancellationToken cancellationToken)
+		{
+			using TelemetrySpan span = _tracer.StartActiveSpan($"{nameof(JobService)}.{nameof(AutoSubmitChangeAsync)}");
+			span.SetAttribute("Job", job.Id.ToString());
+			span.SetAttribute("Graph", graph.Id);
+
+			int? change;
+			string message;
+			try
+			{
+				int clonedPreflightChange = job.ClonedPreflightChange;
+				if (clonedPreflightChange == 0)
+				{
+					if (ShouldClonePreflightChange(job.StreamId))
+					{
+						clonedPreflightChange = await CloneShelvedChangeAsync(streamConfig.ClusterName, job.PreflightChange, cancellationToken);
+					}
+					else
+					{
+						clonedPreflightChange = job.PreflightChange;
+					}
+				}
+
+				_logger.LogInformation("Updating description for {ClonedPreflightChange}", clonedPreflightChange);
+
+				await _perforceService.UpdateChangelistDescriptionAsync(streamConfig.ClusterName, clonedPreflightChange, x => x.TrimEnd() + $"\n#preflight {job.Id}", cancellationToken);
+
+				_logger.LogInformation("Submitting change {Change} (through {ChangeCopy}) after successful completion of {JobId}", job.PreflightChange, clonedPreflightChange, job.Id);
+				(change, message) = await _perforceService.SubmitShelvedChangeAsync(streamConfig, clonedPreflightChange, job.PreflightChange, cancellationToken);
+
+				_logger.LogInformation("Attempt to submit {Change} (through {ChangeCopy}): {Message}", job.PreflightChange, clonedPreflightChange, message);
+
+				if (!String.IsNullOrEmpty(message))
+				{
+					message = Regex.Replace(message, @"^Submit validation failed.*\n(?:\s*'[^']*' validation failed:\s*\n)?", "");
+				}
+
+				if (ShouldClonePreflightChange(job.StreamId))
+				{
+					if (change != null && job.ClonedPreflightChange != 0)
+					{
+						await DeleteShelvedChangeAsync(streamConfig.ClusterName, job.PreflightChange);
+					}
+				}
+				else
+				{
+					if (change != null && job.PreflightChange != 0)
+					{
+						await DeleteShelvedChangeAsync(streamConfig.ClusterName, job.PreflightChange);
+					}
+				}
+			}
+			catch (Exception ex)
+			{
+				(change, message) = ((int?)null, "Internal error");
+				_logger.LogError(ex, "Unable to submit shelved change");
+			}
+
+			for (; ; )
+			{
+				IJob? newJob = await job.TryUpdateJobAsync(autoSubmitChange: change, autoSubmitMessage: message, cancellationToken: cancellationToken);
+				if (newJob != null)
+				{
+					return newJob;
+				}
+
+				newJob = await GetJobAsync(job.Id, cancellationToken);
+				if (newJob == null)
+				{
+					return job;
+				}
+
+				job = newJob;
+			}
+		}
+
+		/// <summary>
+		/// Clone a shelved changelist for running a preflight
+		/// </summary>
+		/// <param name="clusterName">Name of the Perforce cluster</param>
+		/// <param name="change">The changelist to clone</param>
+		/// <param name="cancellationToken">Cancellation token for the operation</param>
+		/// <returns></returns>
+		private async Task<int> CloneShelvedChangeAsync(string clusterName, int change, CancellationToken cancellationToken = default)
+		{
+			int clonedChange;
+			try
+			{
+				clonedChange = await _perforceService.DuplicateShelvedChangeAsync(clusterName, change, cancellationToken);
+				_logger.LogInformation("CL {Change} was duplicated into {ClonedChange}", change, clonedChange);
+			}
+			catch (Exception ex)
+			{
+				_logger.LogError(ex, "Unable to CL {Change} for preflight: {Message}", change, ex.Message);
+				throw;
+			}
+			return clonedChange;
+		}
+
+		/// <summary>
+		/// Deletes a shelved changelist, catching and logging any exceptions
+		/// </summary>
+		/// <param name="clusterName"></param>
+		/// <param name="change">The changelist to delete</param>
+		/// <returns>True if the change was deleted successfully, false otherwise</returns>
+		private Task<bool> DeleteShelvedChangeAsync(string clusterName, int change)
+		{
+			_ = clusterName;
+			_logger.LogInformation("Leaving shelved change {Change}", change);
+			return Task.FromResult<bool>(true);
+			/*
+						_logger.LogInformation("Removing shelf {Change}", change);
+						try
+						{
+							await _perforceService.DeleteShelvedChangeAsync(clusterName, change);
+							return true;
+						}
+						catch (Exception ex)
+						{
+							_logger.LogError(ex, "Unable to delete shelved change {Change}", change);
+							return false;
+						}*/
+		}
+
+		/// <summary>
+		/// Fires a trigger for a chained job
+		/// </summary>
+		/// <param name="job">The job containing the trigger</param>
+		/// <param name="graph">Graph for the job containing the trigger</param>
+		/// <param name="jobTrigger">The trigger object to fire</param>
+		/// <param name="streamConfig">Current config for the stream containing the job</param>
+		/// <returns>New job instance</returns>
+		private async Task<IJob?> FireJobTriggerAsync(IJob job, IGraph graph, IChainedJob jobTrigger, StreamConfig streamConfig)
+		{
+			CancellationToken cancellationToken = CancellationToken.None;
+
+			using TelemetrySpan span = _tracer.StartActiveSpan($"{nameof(JobService)}.{nameof(FireJobTriggerAsync)}");
+			span.SetAttribute("Job", job.Id.ToString());
+			span.SetAttribute("Graph", graph.Id);
+			span.SetAttribute("JobTrigger.Target", jobTrigger.Target);
+			span.SetAttribute("JobTrigger.JobId", jobTrigger.JobId?.ToString());
+			span.SetAttribute("JobTrigger.TemplateRefId", jobTrigger.TemplateRefId);
+
+			for (; ; )
+			{
+				// Update the job
+				JobId chainedJobId = JobIdUtils.GenerateNewId();
+
+				IJob? newJob = await job.TryUpdateJobAsync(jobTrigger: new KeyValuePair<TemplateId, JobId>(jobTrigger.TemplateRefId, chainedJobId));
+				if (newJob != null)
+				{
+					TemplateRefConfig? templateRefConfig;
+					if (!streamConfig.TryGetTemplate(jobTrigger.TemplateRefId, out templateRefConfig))
+					{
+						_logger.LogWarning("Cannot find template {TemplateId} in stream {StreamId}", jobTrigger.TemplateRefId, newJob.StreamId);
+						break;
+					}
+
+					ITemplate template = await _templateCollection.GetOrAddAsync(templateRefConfig);
+
+					IGraph triggerGraph = await _graphs.AddAsync(template, streamConfig.InitialAgentType);
+
+					int? change = null;
+					int codeChange;
+					if (jobTrigger.UseDefaultChangeForTemplate)
+					{
+						ICommitCollection commits = _perforceService.GetCommits(streamConfig);
+
+						if (templateRefConfig.DefaultChange != null)
+						{
+							change = await EvaluateChangeQueriesAsync(streamConfig.Id, templateRefConfig.DefaultChange, null, commits, cancellationToken);
+						}
+						if (change == null)
+						{
+							change = await commits.GetLatestNumberAsync(cancellationToken);
+						}
+
+						ICommit? commit = await commits.GetLastCodeChangeAsync(change, cancellationToken);
+						codeChange = commit?.Number ?? change.Value;
+					}
+					else
+					{
+						change = job.Change;
+						codeChange = job.CodeChange;
+					}
+
+					_logger.LogInformation("Creating downstream job {ChainedJobId} from job {JobId} at change {Change}", chainedJobId, newJob.Id, change);
+
+					CreateJobOptions options = new CreateJobOptions(templateRefConfig);
+					options.PreflightChange = newJob.PreflightChange;
+					options.PreflightDescription = newJob.PreflightDescription;
+					template.GetDefaultParameters(options.Parameters, true);
+					template.GetArgumentsForParameters(options.Parameters, options.Arguments);
+
+					await CreateJobAsync(chainedJobId, streamConfig, jobTrigger.TemplateRefId, template.Hash, triggerGraph, templateRefConfig.Name, change.Value, codeChange, options);
+					return newJob;
+				}
+
+				// Fetch the job again
+				newJob = await _jobs.GetAsync(job.Id);
+				if (newJob == null)
+				{
+					break;
+				}
+
+				job = newJob;
+			}
+			return null;
+		}
+
+		/// <summary>
+		/// Determines if the user is authorized to perform an action on a particular stream
+		/// </summary>
+		/// <param name="jobId">The job to check</param>
+		/// <param name="action">The action being performed</param>
+		/// <param name="user">The principal to authorize</param>
+		/// <param name="globalConfig">Current global config instance</param>
+		/// <param name="cancellationToken">Cancellation token for the operation</param>
+		/// <returns>True if the action is authorized</returns>
+		public async Task<bool> AuthorizeAsync(JobId jobId, AclAction action, ClaimsPrincipal user, GlobalConfig globalConfig, CancellationToken cancellationToken)
+		{
+			IJob? job = await GetJobAsync(jobId, cancellationToken);
+			return globalConfig.Authorize(job, action, user);
+		}
+
+		/// <summary>
+		/// Determines if the user is authorized to perform an action on a particular stream
+		/// </summary>
+		/// <param name="job">The job to check</param>
+		/// <param name="user">The principal to authorize</param>
+		/// <returns>True if the action is authorized</returns>
+		public static bool AuthorizeSession(IJob job, ClaimsPrincipal user)
+		{
+			return job.Batches.Any(x => AuthorizeSession(x, user));
+		}
+
+		/// <summary>
+		/// Determines if the user is authorized to perform an action on a particular stream
+		/// </summary>
+		/// <param name="batch">The batch being executed</param>
+		/// <param name="user">The principal to authorize</param>
+		/// <returns>True if the action is authorized</returns>
+		public static bool AuthorizeSession(IJobStepBatch batch, ClaimsPrincipal user)
+		{
+			if (batch.SessionId != null)
+			{
+				foreach (Claim claim in user.Claims)
+				{
+					if (claim.Type == HordeClaimTypes.AgentSessionId)
+					{
+						SessionId sessionIdValue;
+						if (SessionId.TryParse(claim.Value, out sessionIdValue) && sessionIdValue == batch.SessionId.Value)
+						{
+							return true;
+						}
+					}
+				}
+			}
+			if (batch.LeaseId != null)
+			{
+				foreach (Claim claim in user.Claims)
+				{
+					LeaseId leaseIdValue;
+					if (LeaseId.TryParse(claim.Value, out leaseIdValue) && leaseIdValue == batch.LeaseId.Value)
+					{
+						return true;
+					}
+				}
+			}
+			return false;
+		}
+
+		/// <summary>
+		/// Get a list of each batch's session ID formatted as a string for debugging purposes
+		/// </summary>
+		/// <param name="job">The job to list</param>
+		/// <returns>List of session IDs</returns>
+		public static string GetAllBatchSessionIds(IJob job)
+		{
+			return String.Join(",", job.Batches.Select(b => b.SessionId));
+		}
+	}
+}
diff --git a/Engine/Source/Programs/Horde/HordeServer/Jobs/JobStepRefCollection.cs b/Engine/Source/Programs/Horde/HordeServer/Jobs/JobStepRefCollection.cs
new file mode 100644
--- /dev/null
+++ b/Engine/Source/Programs/Horde/HordeServer/Jobs/JobStepRefCollection.cs
@@ -0,0 +1,282 @@
+// Copyright Epic Games, Inc. All Rights Reserved.
+
+using System;
+using System.Collections.Generic;
+using System.Linq;
+using System.Threading;
+using System.Threading.Tasks;
+using EpicGames.Horde.Agents;
+using EpicGames.Horde.Agents.Pools;
+using EpicGames.Horde.Jobs;
+using EpicGames.Horde.Jobs.Templates;
+using EpicGames.Horde.Logs;
+using EpicGames.Horde.Streams;
+using Horde.Server.Server;
+using Horde.Server.Streams;
+using Horde.Server.Telemetry;
+using Horde.Server.Utilities;
+using Microsoft.Extensions.Options;
+using MongoDB.Bson.Serialization.Attributes;
+using MongoDB.Driver;
+
+namespace Horde.Server.Jobs
+{
+	/// <summary>
+	/// Collection of JobStepRef documents
+	/// </summary>
+	public class JobStepRefCollection : IJobStepRefCollection
+	{
+		class JobStepRef : IJobStepRef
+		{
+			[BsonId]
+			public JobStepRefId Id { get; set; }
+
+			public string JobName { get; set; } = "Unknown";
+
+			public string Name { get; set; }
+
+			public StreamId StreamId { get; set; }
+			public TemplateId TemplateId { get; set; }
+			public int Change { get; set; }
+			public LogId? LogId { get; set; }
+			public PoolId? PoolId { get; set; }
+			public AgentId? AgentId { get; set; }
+			public JobStepState? State { get; set; }
+			public JobStepOutcome? Outcome { get; set; }
+
+			[BsonIgnoreIfNull]
+			public int? LastSuccess { get; }
+
+			[BsonIgnoreIfNull]
+			public int? LastWarning { get; }
+
+			public float BatchWaitTime { get; set; }
+			public float BatchInitTime { get; set; }
+
+			public DateTime JobStartTimeUtc { get; set; }
+
+			[BsonIgnoreIfNull]
+			public DateTimeOffset? StartTime { get; set; }
+
+			[BsonIgnoreIfNull]
+			public DateTime? StartTimeUtc { get; set; }
+
+			[BsonIgnoreIfNull]
+			public DateTimeOffset? FinishTime { get; set; }
+
+			[BsonIgnoreIfNull]
+			public DateTime? FinishTimeUtc { get; set; }
+
+			[BsonIgnoreIfNull]
+			public bool? UpdateIssues { get; set; }
+
+			[BsonIgnoreIfNull]
+			public List<int>? IssueIds { get; set; }
+
+			DateTime IJobStepRef.StartTimeUtc => StartTimeUtc ?? StartTime?.UtcDateTime ?? default;
+			DateTime? IJobStepRef.FinishTimeUtc => FinishTimeUtc ?? FinishTime?.UtcDateTime;
+			string IJobStepRef.NodeName => Name;
+			bool IJobStepRef.UpdateIssues => UpdateIssues ?? false;
+			IReadOnlyList<int>? IJobStepRef.IssueIds => IssueIds;
+
+			public JobStepRef(JobStepRefId id, string jobName, string nodeName, StreamId streamId, TemplateId templateId, int change, LogId? logId, PoolId? poolId, AgentId? agentId, JobStepState? state, JobStepOutcome? outcome, bool updateIssues, int? lastSuccess, int? lastWarning, float batchWaitTime, float batchInitTime, DateTime jobStartTimeUtc, DateTime startTimeUtc, DateTime? finishTimeUtc)
+			{
+				Id = id;
+				JobName = jobName;
+				Name = nodeName;
+				StreamId = streamId;
+				TemplateId = templateId;
+				Change = change;
+				LogId = logId;
+				PoolId = poolId;
+				AgentId = agentId;
+				State = state;
+				Outcome = outcome;
+				UpdateIssues = updateIssues;
+				LastSuccess = lastSuccess;
+				LastWarning = lastWarning;
+				BatchWaitTime = batchWaitTime;
+				BatchInitTime = batchInitTime;
+				JobStartTimeUtc = jobStartTimeUtc;
+				StartTimeUtc = startTimeUtc;
+				FinishTimeUtc = finishTimeUtc;
+			}
+		}
+
+		readonly IMongoCollection<JobStepRef> _jobStepRefs;
+		readonly ITelemetrySink _telemetrySink;
+		readonly IOptionsMonitor<GlobalConfig> _globalConfig;
+
+		/// <summary>
+		/// Constructor
+		/// </summary>
+		public JobStepRefCollection(MongoService mongoService, ITelemetrySink telemetrySink, IOptionsMonitor<GlobalConfig> globalConfig)
+		{
+			List<MongoIndex<JobStepRef>> indexes = new List<MongoIndex<JobStepRef>>();
+			indexes.Add(keys => keys.Ascending(x => x.StreamId).Ascending(x => x.TemplateId).Ascending(x => x.Name).Descending(x => x.Change));
+
+			_jobStepRefs = mongoService.GetCollection<JobStepRef>("JobStepRefs", indexes);
+			_telemetrySink = telemetrySink;
+			_globalConfig = globalConfig;
+		}
+
+		/// <inheritdoc/>
+		public async Task<IJobStepRef> InsertOrReplaceAsync(JobStepRefId id, string jobName, string stepName, StreamId streamId, TemplateId templateId, int change, LogId? logId, PoolId? poolId, AgentId? agentId, JobStepState? state, JobStepOutcome? outcome, bool updateIssues, int? lastSuccess, int? lastWarning, float waitTime, float initTime, DateTime jobStartTimeUtc, DateTime startTimeUtc, DateTime? finishTimeUtc)
+		{
+			JobStepRef newJobStepRef = new JobStepRef(id, jobName, stepName, streamId, templateId, change, logId, poolId, agentId, state, outcome, updateIssues, lastSuccess, lastWarning, waitTime, initTime, jobStartTimeUtc, startTimeUtc, finishTimeUtc);
+			await _jobStepRefs.ReplaceOneAsync(Builders<JobStepRef>.Filter.Eq(x => x.Id, newJobStepRef.Id), newJobStepRef, new ReplaceOptions { IsUpsert = true });
+
+			if (_globalConfig.CurrentValue.TryGetStream(streamId, out StreamConfig? streamConfig) && !streamConfig.TelemetryStoreId.IsEmpty)
+			{
+				_telemetrySink.SendEvent(streamConfig.TelemetryStoreId, TelemetryRecordMeta.CurrentHordeInstance, new
+				{
+					EventName = "State.JobStepRef",
+					Id = id.ToString(),
+					JobId = id.JobId.ToString(),
+					BatchId = id.BatchId.ToString(),
+					StepId = id.StepId.ToString(),
+					AgentId = agentId,
+					BatchInitTime = initTime,
+					BatchWaitTime = waitTime,
+					Change = change,
+					FinishTime = finishTimeUtc,
+					JobName = jobName,
+					JobStartTime = jobStartTimeUtc,
+					StepName = stepName,
+					State = state,
+					Outcome = outcome,
+					PoolId = poolId,
+					StartTime = startTimeUtc,
+					StreamId = streamId,
+					TemplateId = templateId,
+					UpdateIssues = updateIssues,
+					Duration = (finishTimeUtc != null) ? (double?)(finishTimeUtc.Value - startTimeUtc).TotalSeconds : null
+				});
+			}
+
+			return newJobStepRef;
+		}
+
+		/// <inheritdoc/>
+		public async Task<IJobStepRef?> UpdateAsync(JobId jobId, JobStepBatchId batchId, JobStepId stepId, List<int>? issueIds)
+		{
+			UpdateDefinitionBuilder<JobStepRef> updateBuilder = Builders<JobStepRef>.Update;
+			List<UpdateDefinition<JobStepRef>> updates = new List<UpdateDefinition<JobStepRef>>();
+
+			if (issueIds != null)
+			{
+				updates.Add(updateBuilder.Set(x => x.IssueIds, issueIds));
+			}
+
+			if (updates.Count == 0)
+			{
+				return await FindAsync(jobId, batchId, stepId);
+			}
+
+			JobStepRefId id = new JobStepRefId(jobId, batchId, stepId);
+			return await _jobStepRefs.FindOneAndUpdateAsync(x => x.Id.Equals(id), updateBuilder.Combine(updates));
+		}
+
+		/// <inheritdoc/>
+		public async Task<IJobStepRef?> FindAsync(JobId jobId, JobStepBatchId batchId, JobStepId stepId)
+		{
+			JobStepRefId id = new JobStepRefId(jobId, batchId, stepId);
+			return await _jobStepRefs.Find(x => x.Id.Equals(id)).FirstOrDefaultAsync();
+		}
+
+		/// <inheritdoc/>
+		public async Task<List<IJobStepRef>> FindAsync(JobStepRefId[] ids, CancellationToken cancellationToken)
+		{
+			FilterDefinitionBuilder<JobStepRef> filterBuilder = Builders<JobStepRef>.Filter;
+			FilterDefinition<JobStepRef> filter = FilterDefinition<JobStepRef>.Empty;
+			filter &= filterBuilder.In(x => x.Id, ids);
+
+			List<JobStepRef> steps = await _jobStepRefs.Find(filter).SortByDescending(x => x.Change).ToListAsync(cancellationToken);
+			return steps.ConvertAll<IJobStepRef>(x => x);
+		}
+
+		/// <inheritdoc/>
+		public async Task<List<IJobStepRef>> GetStepsForNodeAsync(StreamId streamId, TemplateId templateId, string nodeName, int? change, bool includeFailed, int maxCount, CancellationToken cancellationToken)
+		{
+			// Find all the steps matching the given criteria
+			FilterDefinitionBuilder<JobStepRef> filterBuilder = Builders<JobStepRef>.Filter;
+
+			FilterDefinition<JobStepRef> filter = FilterDefinition<JobStepRef>.Empty;
+			filter &= filterBuilder.Eq(x => x.StreamId, streamId);
+			filter &= filterBuilder.Eq(x => x.TemplateId, templateId);
+			filter &= filterBuilder.Eq(x => x.Name, nodeName);
+			if (change != null)
+			{
+				filter &= filterBuilder.Lte(x => x.Change, change.Value);
+			}
+			if (!includeFailed)
+			{
+				filter &= filterBuilder.Ne(x => x.Outcome, JobStepOutcome.Failure);
+			}
+
+			List<JobStepRef> steps = await _jobStepRefs.Find(filter).SortByDescending(x => x.Change).ThenByDescending(x => x.StartTimeUtc).Limit(maxCount).ToListAsync(cancellationToken);
+			return steps.ConvertAll<IJobStepRef>(x => x);
+		}
+
+		/// <inheritdoc/>
+		public async Task<IJobStepRef?> GetPrevStepForNodeAsync(StreamId streamId, TemplateId templateId, string nodeName, int change, JobStepOutcome? outcome = null, bool? updateIssues = null, IEnumerable<JobId>? excludeJobIds = null)
+		{
+			FilterDefinitionBuilder<JobStepRef> filterBuilder = Builders<JobStepRef>.Filter;
+
+			FilterDefinition<JobStepRef> filter = FilterDefinition<JobStepRef>.Empty;
+			filter &= filterBuilder.Eq(x => x.StreamId, streamId);
+			filter &= filterBuilder.Eq(x => x.TemplateId, templateId);
+			filter &= filterBuilder.Eq(x => x.Name, nodeName);
+			filter &= filterBuilder.Lt(x => x.Change, change);
+
+			if (outcome != null)
+			{
+				filter &= filterBuilder.Eq(x => x.Outcome, outcome);
+			}
+			else
+			{
+				filter &= filterBuilder.Ne(x => x.Outcome, null);
+			}
+
+			if (updateIssues != null)
+			{
+				filter &= filterBuilder.Ne(x => x.UpdateIssues, false);
+			}
+
+			if (excludeJobIds != null && excludeJobIds.Any())
+			{
+				filter &= filterBuilder.Nin(x => x.Id.JobId, excludeJobIds);
+			}
+
+			return await _jobStepRefs.Find(filter).SortByDescending(x => x.Change).FirstOrDefaultAsync();
+		}
+
+		/// <inheritdoc/>
+		public async Task<IJobStepRef?> GetNextStepForNodeAsync(StreamId streamId, TemplateId templateId, string nodeName, int change, JobStepOutcome? outcome = null, bool? updateIssues = null)
+		{
+			FilterDefinitionBuilder<JobStepRef> filterBuilder = Builders<JobStepRef>.Filter;
+
+			FilterDefinition<JobStepRef> filter = FilterDefinition<JobStepRef>.Empty;
+			filter &= filterBuilder.Eq(x => x.StreamId, streamId);
+			filter &= filterBuilder.Eq(x => x.TemplateId, templateId);
+			filter &= filterBuilder.Eq(x => x.Name, nodeName);
+			filter &= filterBuilder.Gt(x => x.Change, change);
+
+			if (outcome != null)
+			{
+				filter &= filterBuilder.Eq(x => x.Outcome, outcome);
+			}
+			else
+			{
+				filter &= filterBuilder.Ne(x => x.Outcome, null);
+			}
+
+			if (updateIssues != null)
+			{
+				filter &= filterBuilder.Ne(x => x.UpdateIssues, false);
+			}
+
+			return await _jobStepRefs.Find(filter).SortBy(x => x.Change).FirstOrDefaultAsync();
+		}
+	}
+}
