Horde: Rename projects before splitting into plugins (Horde.Agent -> HordeAgent, Horde.Server -> HordeServer, Horde -> HordeCmd).

diff --git a/Engine/Source/Programs/Horde/HordeServer/Jobs/Schedules/ScheduleService.cs b/Engine/Source/Programs/Horde/HordeServer/Jobs/Schedules/ScheduleService.cs
new file mode 100644
--- /dev/null
+++ b/Engine/Source/Programs/Horde/HordeServer/Jobs/Schedules/ScheduleService.cs
@@ -0,0 +1,551 @@
+// Copyright Epic Games, Inc. All Rights Reserved.
+
+using System;
+using System.Collections.Generic;
+using System.Diagnostics;
+using System.Linq;
+using System.Text.RegularExpressions;
+using System.Threading;
+using System.Threading.Tasks;
+using EpicGames.Core;
+using EpicGames.Horde.Commits;
+using EpicGames.Horde.Jobs;
+using EpicGames.Horde.Jobs.Templates;
+using EpicGames.Horde.Streams;
+using EpicGames.Redis;
+using EpicGames.Redis.Utility;
+using EpicGames.Serialization;
+using Horde.Server.Commits;
+using Horde.Server.Jobs.Graphs;
+using Horde.Server.Jobs.Templates;
+using Horde.Server.Server;
+using Horde.Server.Streams;
+using Horde.Server.Utilities;
+using HordeCommon;
+using Microsoft.Extensions.Hosting;
+using Microsoft.Extensions.Logging;
+using Microsoft.Extensions.Options;
+using OpenTelemetry.Trace;
+using StackExchange.Redis;
+
+namespace Horde.Server.Jobs.Schedules
+{
+	/// <summary>
+	/// Manipulates schedule instances
+	/// </summary>
+	public sealed class ScheduleService : IHostedService, IAsyncDisposable
+	{
+		[RedisConverter(typeof(RedisCbConverter<QueueItem>))]
+		class QueueItem
+		{
+			[CbField("sid")]
+			public StreamId StreamId { get; set; }
+
+			[CbField("tid")]
+			public TemplateId TemplateId { get; set; }
+
+			public QueueItem()
+			{
+			}
+
+			public QueueItem(StreamId streamId, TemplateId templateId)
+			{
+				StreamId = streamId;
+				TemplateId = templateId;
+			}
+
+			public static DateTime GetTimeFromScore(double score) => DateTime.UnixEpoch + TimeSpan.FromSeconds(score);
+			public static double GetScoreFromTime(DateTime time) => (time.ToUniversalTime() - DateTime.UnixEpoch).TotalSeconds;
+		}
+
+		readonly IGraphCollection _graphs;
+		readonly ICommitService _commitService;
+		readonly IJobCollection _jobCollection;
+		readonly IDowntimeService _downtimeService;
+		readonly JobService _jobService;
+		readonly IStreamCollection _streamCollection;
+		readonly ITemplateCollection _templateCollection;
+		readonly IClock _clock;
+		readonly Tracer _tracer;
+		readonly ILogger _logger;
+		readonly RedisService _redis;
+		static readonly RedisKey s_baseLockKey = "scheduler/locks";
+		static readonly RedisKey s_tickLockKey = s_baseLockKey.Append("/tick"); // Lock to tick the queue
+		static readonly RedisSortedSetKey<QueueItem> s_queueKey = "scheduler/queue"; // Items to tick, ordered by time
+		readonly IOptionsMonitor<GlobalConfig> _globalConfig;
+		readonly ITicker _ticker;
+
+		/// <summary>
+		/// Constructor
+		/// </summary>
+		public ScheduleService(RedisService redis, IGraphCollection graphs, ICommitService commitService, IJobCollection jobCollection, JobService jobService, IDowntimeService downtimeService, IStreamCollection streamCollection, ITemplateCollection templateCollection, MongoService mongoService, IClock clock, IOptionsMonitor<GlobalConfig> globalConfig, Tracer tracer, ILogger<ScheduleService> logger)
+		{
+			_graphs = graphs;
+			_commitService = commitService;
+			_jobCollection = jobCollection;
+			_jobService = jobService;
+			_downtimeService = downtimeService;
+			_streamCollection = streamCollection;
+			_templateCollection = templateCollection;
+			_clock = clock;
+			_globalConfig = globalConfig;
+			_tracer = tracer;
+			_logger = logger;
+
+			_redis = redis;
+			if (mongoService.ReadOnlyMode)
+			{
+				_ticker = new NullTicker();
+			}
+			else
+			{
+				_ticker = clock.AddTicker<ScheduleService>(TimeSpan.FromMinutes(1.0), TickAsync, logger);
+			}
+		}
+
+		/// <inheritdoc/>
+		public async Task StartAsync(CancellationToken cancellationToken)
+		{
+			await _ticker.StartAsync();
+		}
+
+		/// <inheritdoc/>
+		public async Task StopAsync(CancellationToken cancellationToken)
+		{
+			await _ticker.StopAsync();
+		}
+
+		/// <inheritdoc/>
+		public async ValueTask DisposeAsync()
+		{
+			await _ticker.DisposeAsync();
+		}
+
+		async ValueTask TickAsync(CancellationToken cancellationToken)
+		{
+			DateTime utcNow = _clock.UtcNow;
+
+			// Don't start any new jobs during scheduled downtime
+			if (_downtimeService.IsDowntimeActive)
+			{
+				return;
+			}
+
+			// Update the current queue
+			await using (RedisLock sharedLock = new(_redis.GetDatabase(), s_tickLockKey))
+			{
+				if (await sharedLock.AcquireAsync(TimeSpan.FromMinutes(1.0), false))
+				{
+					await UpdateQueueAsync(utcNow, cancellationToken);
+				}
+			}
+
+			// Keep updating schedules
+			while (!cancellationToken.IsCancellationRequested)
+			{
+				// Get the item with the lowest score (ie. the one that hasn't been updated in the longest time)
+				QueueItem? item = await PopQueueItemAsync();
+				if (item == null)
+				{
+					break;
+				}
+
+				// Acquire the lock for this schedule and update it
+				await using (RedisLock sharedLock = new RedisLock<QueueItem>(_redis.GetDatabase(), s_baseLockKey, item))
+				{
+					if (await sharedLock.AcquireAsync(TimeSpan.FromMinutes(1.0)))
+					{
+						try
+						{
+							await TriggerAsync(item.StreamId, item.TemplateId, utcNow, cancellationToken);
+						}
+						catch (OperationCanceledException)
+						{
+							throw;
+						}
+						catch (Exception ex)
+						{
+							_logger.LogError(ex, "Error while updating schedule for {StreamId}/{TemplateId}", item.StreamId, item.TemplateId);
+						}
+					}
+				}
+			}
+		}
+
+		async Task<QueueItem?> PopQueueItemAsync()
+		{
+			IDatabaseAsync target = _redis.GetDatabase();
+			for (; ; )
+			{
+				QueueItem[] items = await target.SortedSetRangeByRankAsync(s_queueKey, 0, 0);
+				if (items.Length == 0)
+				{
+					return null;
+				}
+				if (await target.SortedSetRemoveAsync(s_queueKey, items[0]))
+				{
+					return items[0];
+				}
+			}
+		}
+
+		internal async Task ResetAsync()
+		{
+			IDatabase redis = _redis.GetDatabase();
+			await redis.KeyDeleteAsync(s_queueKey);
+			await redis.KeyDeleteAsync(s_tickLockKey);
+		}
+
+		internal async Task TickForTestingAsync()
+		{
+			await UpdateQueueAsync(_clock.UtcNow, CancellationToken.None);
+			await TickAsync(CancellationToken.None);
+		}
+
+		/// <summary>
+		/// Get the current set of streams and ensure there's an entry for each item
+		/// </summary>
+		public async Task UpdateQueueAsync(DateTime utcNow, CancellationToken cancellationToken)
+		{
+			List<SortedSetEntry<QueueItem>> queueItems = new List<SortedSetEntry<QueueItem>>();
+
+			GlobalConfig globalConfig = _globalConfig.CurrentValue;
+			foreach (StreamConfig streamConfig in globalConfig.Streams)
+			{
+				IStream stream = await _streamCollection.GetAsync(streamConfig, cancellationToken);
+				foreach ((TemplateId templateId, ITemplateRef templateRef) in stream.Templates)
+				{
+					if (templateRef.Schedule != null)
+					{
+						DateTime? nextTriggerTimeUtc = templateRef.Schedule.GetNextTriggerTimeUtc(_clock.TimeZone);
+						if (nextTriggerTimeUtc != null)
+						{
+							if (utcNow > nextTriggerTimeUtc.Value)
+							{
+								double score = QueueItem.GetScoreFromTime(nextTriggerTimeUtc.Value);
+								queueItems.Add(new SortedSetEntry<QueueItem>(new QueueItem(stream.Id, templateId), score));
+
+								await _streamCollection.UpdateScheduleTriggerAsync(stream, templateId, utcNow, cancellationToken: cancellationToken);
+							}
+						}
+					}
+				}
+			}
+
+			await _redis.GetDatabase().SortedSetAddAsync(s_queueKey, queueItems.ToArray());
+		}
+
+		/// <summary>
+		/// Trigger a schedule to run
+		/// </summary>
+		/// <param name="streamId">Stream for the schedule</param>
+		/// <param name="templateId"></param>
+		/// <param name="utcNow"></param>
+		/// <param name="cancellationToken"></param>
+		/// <returns>Async task</returns>
+		internal async Task<bool> TriggerAsync(StreamId streamId, TemplateId templateId, DateTime utcNow, CancellationToken cancellationToken)
+		{
+			GlobalConfig globalConfig = _globalConfig.CurrentValue;
+			if (!globalConfig.TryGetStream(streamId, out StreamConfig? streamConfig))
+			{
+				return false;
+			}
+
+			IStream? stream = await _streamCollection.GetAsync(streamConfig, cancellationToken);
+			if (stream == null || !stream.Templates.TryGetValue(templateId, out ITemplateRef? templateRef))
+			{
+				return false;
+			}
+
+			ITemplateSchedule? schedule = templateRef.Schedule;
+			if (schedule == null || !schedule.Config.Enabled)
+			{
+				return false;
+			}
+
+			using TelemetrySpan span = _tracer.StartActiveSpan($"{nameof(ScheduleService)}.{nameof(TriggerAsync)}");
+			span.SetAttribute("streamId", stream.Id);
+			span.SetAttribute("templateId", templateId);
+
+			Stopwatch stopwatch = Stopwatch.StartNew();
+			_logger.LogInformation("Updating schedule for {StreamId} template {TemplateId}", stream.Id, templateId);
+
+			// Get a list of jobs that we need to remove
+			List<JobId> removeJobIds = new List<JobId>();
+			foreach (JobId activeJobId in schedule.ActiveJobs)
+			{
+				IJob? job = await _jobService.GetJobAsync(activeJobId, cancellationToken);
+				if (job == null || job.Batches.All(x => x.State == JobStepBatchState.Complete))
+				{
+					_logger.LogInformation("Removing active job {JobId}", activeJobId);
+					removeJobIds.Add(activeJobId);
+				}
+			}
+			await _streamCollection.UpdateScheduleTriggerAsync(stream, templateId, removeJobs: removeJobIds, cancellationToken: cancellationToken);
+
+			// If the stream is paused, bail out
+			if (stream.IsPaused(utcNow))
+			{
+				_logger.LogDebug("Skipping schedule update for stream {StreamId}. It has been paused until {PausedUntil} with comment '{PauseComment}'.", stream.Id, stream.PausedUntil, stream.PauseComment);
+				return false;
+			}
+
+			// Trigger this schedule
+			try
+			{
+				await TriggerAsync(stream, templateId, templateRef, schedule, schedule.ActiveJobs.Count - removeJobIds.Count, utcNow, cancellationToken);
+			}
+			catch (Exception ex)
+			{
+				_logger.LogError(ex, "Failed to start schedule {StreamId}/{TemplateId}", stream.Id, templateId);
+			}
+
+			// Print some timing info
+			stopwatch.Stop();
+			_logger.LogInformation("Updated schedule for {StreamId} template {TemplateId} in {TimeSeconds}ms", stream.Id, templateId, (long)stopwatch.Elapsed.TotalMilliseconds);
+			return true;
+		}
+
+		/// <summary>
+		/// Trigger a schedule to run
+		/// </summary>
+		/// <param name="stream">Stream for the schedule</param>
+		/// <param name="templateId"></param>
+		/// <param name="templateRef"></param>
+		/// <param name="schedule"></param>
+		/// <param name="numActiveJobs"></param>
+		/// <param name="utcNow">The current time</param>
+		/// <param name="cancellationToken"></param>
+		/// <returns>Async task</returns>
+		private async Task TriggerAsync(IStream stream, TemplateId templateId, ITemplateRef templateRef, ITemplateSchedule schedule, int numActiveJobs, DateTime utcNow, CancellationToken cancellationToken)
+		{
+			// Check we're not already at the maximum number of allowed jobs
+			if (schedule.Config.MaxActive != 0 && numActiveJobs >= schedule.Config.MaxActive)
+			{
+				_logger.LogInformation("Skipping trigger of {StreamId} template {TemplateId} - already have maximum number of jobs running ({NumJobs})", stream.Id, templateId, schedule.Config.MaxActive);
+				foreach (JobId jobId in schedule.ActiveJobs)
+				{
+					_logger.LogInformation("Active job for {StreamId} template {TemplateId}: {JobId}", stream.Id, templateId, jobId);
+				}
+				return;
+			}
+
+			// Minimum changelist number, inclusive
+			int minChangeNumber = schedule.LastTriggerChange;
+			if (minChangeNumber > 0 && !schedule.Config.RequireSubmittedChange)
+			{
+				minChangeNumber--;
+			}
+
+			// Maximum changelist number, exclusive
+			int? maxChangeNumber = null;
+
+			// Get the maximum number of changes to trigger
+			int maxNewChanges = 1;
+			if (schedule.Config.MaxChanges != 0)
+			{
+				maxNewChanges = schedule.Config.MaxChanges;
+			}
+			if (schedule.Config.MaxActive != 0)
+			{
+				maxNewChanges = Math.Min(maxNewChanges, schedule.Config.MaxActive - numActiveJobs);
+			}
+
+			// Create a timer to limit the amount we look back through P4 history
+			Stopwatch timer = Stopwatch.StartNew();
+
+			// Create a file filter
+			FileFilter? fileFilter = null;
+			if (schedule.Config.Files != null)
+			{
+				fileFilter = new FileFilter(schedule.Config.Files);
+			}
+
+			// Cache the Perforce history as we're iterating through changes to improve query performance
+			ICommitCollection commits = _commitService.GetCollection(stream.Config);
+			IAsyncEnumerable<ICommit> commitEnumerable = commits.FindAsync(minChangeNumber, null, null, schedule.Config.GetAllCommitTags(), cancellationToken);
+			await using IAsyncEnumerator<ICommit> commitEnumerator = commitEnumerable.GetAsyncEnumerator(cancellationToken);
+
+			// Start as many jobs as possible
+			List<(int Change, int CodeChange)> triggerChanges = new List<(int, int)>();
+			while (triggerChanges.Count < maxNewChanges)
+			{
+				cancellationToken.ThrowIfCancellationRequested();
+
+				// Get the next valid change
+				int change;
+				ICommit? commit;
+
+				if (schedule.Config.Gate != null)
+				{
+					change = await GetNextChangeForGateAsync(stream.Id, templateId, schedule.Config.Gate, minChangeNumber, maxChangeNumber, cancellationToken);
+					commit = await commits.FindAsync(change, change, 1, null, cancellationToken).FirstOrDefaultAsync(cancellationToken); // May be a change in a different stream
+				}
+				else if (await commitEnumerator.MoveNextAsync(cancellationToken))
+				{
+					commit = commitEnumerator.Current;
+					change = commit.Number;
+				}
+				else
+				{
+					commit = null;
+					change = 0;
+				}
+
+				// Quit if we didn't find anything
+				if (change <= 0)
+				{
+					break;
+				}
+				if (change < minChangeNumber)
+				{
+					break;
+				}
+				if (change == minChangeNumber && (schedule.Config.RequireSubmittedChange || triggerChanges.Count > 0))
+				{
+					break;
+				}
+
+				// Adjust the changelist for the desired filter
+				if (commit == null || await ShouldBuildChangeAsync(commit, schedule.Config.GetAllCommitTags(), fileFilter, cancellationToken))
+				{
+					int codeChange = change;
+
+					ICommit? lastCodeCommit = await commits.GetLastCodeChangeAsync(change, cancellationToken);
+					if (lastCodeCommit != null)
+					{
+						codeChange = lastCodeCommit.Number;
+					}
+					else
+					{
+						_logger.LogWarning("Unable to find code change for CL {Change}", change);
+					}
+
+					triggerChanges.Add((change, codeChange));
+				}
+
+				// Check we haven't exceeded the time limit
+				if (timer.Elapsed > TimeSpan.FromMinutes(2.0))
+				{
+					_logger.LogError("Querying for changes to trigger for {StreamId} template {TemplateId} has taken {Time}. Aborting.", stream.Id, templateId, timer.Elapsed);
+					break;
+				}
+
+				// Update the remaining range of changes to check for
+				maxChangeNumber = change - 1;
+				if (maxChangeNumber < minChangeNumber)
+				{
+					break;
+				}
+			}
+
+			// Early out if there's nothing to do
+			if (triggerChanges.Count == 0)
+			{
+				_logger.LogInformation("Skipping trigger of {StreamName} template {TemplateId} - no candidate changes after CL {LastTriggerChange}", stream.Id, templateId, schedule.LastTriggerChange);
+				return;
+			}
+
+			// Get the matching template
+			ITemplate template = await _templateCollection.GetOrAddAsync(templateRef.Config);
+
+			// Register the graph for it
+			IGraph graph = await _graphs.AddAsync(template, stream.Config.InitialAgentType, cancellationToken);
+
+			// We may need to submit a new change for any new jobs. This only makes sense if there's one change.
+			if (template.SubmitNewChange != null)
+			{
+				int newChange = await commits.CreateNewAsync(template, cancellationToken);
+				ICommit? newCodeChange = await commits.GetLastCodeChangeAsync(newChange, cancellationToken);
+				triggerChanges = new List<(int, int)> { (newChange, newCodeChange?.Number ?? newChange) };
+			}
+
+			// Try to start all the new jobs
+			_logger.LogInformation("Starting {NumJobs} new jobs for {StreamId} template {TemplateId} (active: {NumActive}, max new: {MaxNewJobs})", triggerChanges.Count, stream.Id, templateId, numActiveJobs, maxNewChanges);
+			foreach ((int change, int codeChange) in triggerChanges.OrderBy(x => x.Change))
+			{
+				cancellationToken.ThrowIfCancellationRequested();
+
+				CreateJobOptions options = new CreateJobOptions(templateRef.Config);
+				options.Priority = template.Priority;
+				template.GetDefaultParameters(options.Parameters, true);
+				template.GetArgumentsForParameters(options.Parameters, options.Arguments);
+
+				IJob newJob = await _jobService.CreateJobAsync(null, stream.Config, templateId, template.Hash, graph, template.Name, change, codeChange, options, cancellationToken);
+				_logger.LogInformation("Started new job for {StreamId} template {TemplateId} at CL {Change} (Code CL {CodeChange}): {JobId}", stream.Id, templateId, change, codeChange, newJob.Id);
+				await _streamCollection.UpdateScheduleTriggerAsync(stream, templateId, utcNow, change, new List<JobId> { newJob.Id }, new List<JobId>(), cancellationToken);
+			}
+		}
+
+		/// <summary>
+		/// Tests whether a schedule should build a particular change, based on its requested change filters
+		/// </summary>
+		/// <param name="commit">The commit details</param>
+		/// <param name="filterTags">Filter for the tags to trigger a build</param>
+		/// <param name="fileFilter">Filter for the files to trigger a build</param>
+		/// <param name="cancellationToken">Cancellation token for the operation</param>
+		/// <returns></returns>
+		private async ValueTask<bool> ShouldBuildChangeAsync(ICommit commit, List<CommitTag>? filterTags, FileFilter? fileFilter, CancellationToken cancellationToken)
+		{
+			if (Regex.IsMatch(commit.Description, @"^\s*#\s*skipci", RegexOptions.Multiline))
+			{
+				return false;
+			}
+			if (filterTags != null && filterTags.Count > 0)
+			{
+				IReadOnlyList<CommitTag> commitTags = await commit.GetTagsAsync(cancellationToken);
+				if (!commitTags.Any(x => filterTags.Contains(x)))
+				{
+					_logger.LogDebug("Not building change {Change} ({ChangeTags}) due to filter tags ({FilterTags})", commit.Number, String.Join(", ", commitTags.Select(x => x.ToString())), String.Join(", ", filterTags.Select(x => x.ToString())));
+					return false;
+				}
+			}
+			if (fileFilter != null)
+			{
+				if (!await commit.MatchesFilterAsync(fileFilter, cancellationToken))
+				{
+					_logger.LogDebug("Not building change {Change} due to file filter", commit.Number);
+					return false;
+				}
+			}
+			return true;
+		}
+
+		/// <summary>
+		/// Gets the next change to build for a schedule on a gate
+		/// </summary>
+		/// <returns></returns>
+		private async Task<int> GetNextChangeForGateAsync(StreamId streamId, TemplateId templateRefId, ScheduleGateConfig gate, int? minChange, int? maxChange, CancellationToken cancellationToken)
+		{
+			for (; ; )
+			{
+				cancellationToken.ThrowIfCancellationRequested();
+
+				IReadOnlyList<IJob> jobs = await _jobCollection.FindAsync(streamId: streamId, templates: new[] { gate.TemplateId }, minChange: minChange, maxChange: maxChange, count: 1, cancellationToken: cancellationToken);
+				if (jobs.Count == 0)
+				{
+					return 0;
+				}
+
+				IJob job = jobs[0];
+
+				IGraph? graph = await _graphs.GetAsync(job.GraphHash, cancellationToken);
+				if (graph != null)
+				{
+					(JobStepState, JobStepOutcome)? state = job.GetTargetState(graph, gate.Target);
+					if (state != null && state.Value.Item1 == JobStepState.Completed)
+					{
+						JobStepOutcome outcome = state.Value.Item2;
+						if (outcome == JobStepOutcome.Success || outcome == JobStepOutcome.Warnings)
+						{
+							return job.Change;
+						}
+						_logger.LogInformation("Skipping trigger of {StreamName} template {TemplateId} - last {OtherTemplateRefId} job ({JobId}) ended with errors", streamId, templateRefId, gate.TemplateId, job.Id);
+					}
+				}
+
+				maxChange = job.Change - 1;
+			}
+		}
+	}
+}
diff --git a/Engine/Source/Programs/Horde/HordeServer/Jobs/Schedules/Schedules.cs b/Engine/Source/Programs/Horde/HordeServer/Jobs/Schedules/Schedules.cs
new file mode 100644
--- /dev/null
+++ b/Engine/Source/Programs/Horde/HordeServer/Jobs/Schedules/Schedules.cs
@@ -0,0 +1,139 @@
+// Copyright Epic Games, Inc. All Rights Reserved.
+
+using System;
+using System.Collections.Generic;
+using EpicGames.Horde.Commits;
+using EpicGames.Horde.Jobs;
+using Horde.Server.Streams;
+
+namespace Horde.Server.Jobs.Schedules
+{
+	/// <summary>
+	/// Response describing a schedule
+	/// </summary>
+	public class GetScheduleResponse
+	{
+		/// <summary>
+		/// Whether the schedule is currently enabled
+		/// </summary>
+		public bool Enabled { get; set; }
+
+		/// <summary>
+		/// Maximum number of scheduled jobs at once
+		/// </summary>
+		public int MaxActive { get; set; }
+
+		/// <summary>
+		/// Maximum number of changes the schedule can fall behind head revision. If greater than zero, builds will be triggered for every submitted changelist until the backlog is this size.
+		/// </summary>
+		public int MaxChanges { get; set; }
+
+		/// <summary>
+		/// Whether the build requires a change to be submitted
+		/// </summary>
+		public bool RequireSubmittedChange { get; set; }
+
+		/// <summary>
+		/// Gate for this schedule to trigger
+		/// </summary>
+		public ScheduleGateConfig? Gate { get; set; }
+
+		/// <summary>
+		/// Which commits to run this job for
+		/// </summary>
+		public List<CommitTag>? Commits { get; set; }
+
+		/// <summary>
+		/// The types of changes to run for
+		/// </summary>
+		[Obsolete("Use Commits instead")]
+		public List<ChangeContentFlags>? Filter { get; set; }
+
+		/// <summary>
+		/// Parameters for the template
+		/// </summary>
+		public Dictionary<string, string> TemplateParameters { get; set; }
+
+		/// <summary>
+		/// New patterns for the schedule
+		/// </summary>
+		public List<SchedulePatternConfig> Patterns { get; set; }
+
+		/// <summary>
+		/// Last changelist number that this was triggered for
+		/// </summary>
+		public int LastTriggerChange { get; set; }
+
+		/// <summary>
+		/// Last time that the schedule was triggered
+		/// </summary>
+		public DateTimeOffset LastTriggerTime { get; set; }
+
+		/// <summary>
+		/// Next trigger times for schedule
+		/// </summary>
+		public List<DateTime> NextTriggerTimesUTC { get; set; }
+
+		/// <summary>
+		/// List of active jobs
+		/// </summary>
+		public List<JobId> ActiveJobs { get; set; }
+
+		/// <summary>
+		/// Constructor
+		/// </summary>
+		/// <param name="schedule">Schedule to construct from</param>
+		/// <param name="schedulerTimeZone">The scheduler time zone</param>
+		public GetScheduleResponse(ITemplateSchedule schedule, TimeZoneInfo schedulerTimeZone)
+		{
+			Enabled = schedule.Config.Enabled;
+			MaxActive = schedule.Config.MaxActive;
+			MaxChanges = schedule.Config.MaxChanges;
+			RequireSubmittedChange = schedule.Config.RequireSubmittedChange;
+			Gate = schedule.Config.Gate;
+			Commits = schedule.Config.Commits;
+#pragma warning disable CS0618 // Type or member is obsolete
+			Filter = schedule.Config.Filter;
+#pragma warning restore CS0618 // Type or member is obsolete
+			TemplateParameters = schedule.Config.TemplateParameters;
+			Patterns = schedule.Config.Patterns;
+			LastTriggerChange = schedule.LastTriggerChange;
+			LastTriggerTime = schedule.LastTriggerTimeUtc;
+			ActiveJobs = new List<JobId>(schedule.ActiveJobs);
+
+			DateTime curTime = schedule.LastTriggerTimeUtc;
+			NextTriggerTimesUTC = new List<DateTime>();
+			for (int i = 0; i < 16; i++)
+			{
+				DateTime? nextTime = schedule.Config.GetNextTriggerTimeUtc(curTime, schedulerTimeZone);
+				if (nextTime == null)
+				{
+					break;
+				}
+
+				curTime = nextTime.Value;
+				NextTriggerTimesUTC.Add(curTime);
+			}
+		}
+	}
+
+	/// <summary>
+	/// Response describing when a schedule is expected to trigger
+	/// </summary>
+	public class GetScheduleForecastResponse
+	{
+		/// <summary>
+		/// Next trigger times
+		/// </summary>
+		public List<DateTime> Times { get; set; }
+
+		/// <summary>
+		/// Constructor
+		/// </summary>
+		/// <param name="times">List of trigger times</param>
+		public GetScheduleForecastResponse(List<DateTime> times)
+		{
+			Times = times;
+		}
+	}
+}
diff --git a/Engine/Source/Programs/Horde/HordeServer/Jobs/Templates/ITemplateCollection.cs b/Engine/Source/Programs/Horde/HordeServer/Jobs/Templates/ITemplateCollection.cs
new file mode 100644
--- /dev/null
+++ b/Engine/Source/Programs/Horde/HordeServer/Jobs/Templates/ITemplateCollection.cs
@@ -0,0 +1,27 @@
+// Copyright Epic Games, Inc. All Rights Reserved.
+
+using System.Threading.Tasks;
+using EpicGames.Core;
+
+namespace Horde.Server.Jobs.Templates
+{
+	/// <summary>
+	/// Interface for a collection of template documents
+	/// </summary>
+	public interface ITemplateCollection
+	{
+		/// <summary>
+		/// Gets a template by ID
+		/// </summary>
+		/// <param name="templateId">Unique id of the template</param>
+		/// <returns>The template document</returns>
+		Task<ITemplate?> GetAsync(ContentHash templateId);
+
+		/// <summary>
+		/// Gets a template instance from its configuration
+		/// </summary>
+		/// <param name="templateConfig">The template configuration</param>
+		/// <returns>Template instance</returns>
+		Task<ITemplate> GetOrAddAsync(TemplateConfig templateConfig);
+	}
+}
diff --git a/Engine/Source/Programs/Horde/HordeServer/Logs/Data/LogChunkData.cs b/Engine/Source/Programs/Horde/HordeServer/Logs/Data/LogChunkData.cs
new file mode 100644
--- /dev/null
+++ b/Engine/Source/Programs/Horde/HordeServer/Logs/Data/LogChunkData.cs
@@ -0,0 +1,278 @@
+// Copyright Epic Games, Inc. All Rights Reserved.
+
+using System;
+using System.Collections.Generic;
+using System.Linq;
+using EpicGames.Core;
+using EpicGames.Horde.Logs;
+using Microsoft.Extensions.Logging;
+
+namespace Horde.Server.Logs.Data
+{
+	/// <summary>
+	/// Data for a log chunk
+	/// </summary>
+	public class LogChunkData
+	{
+		/// <summary>
+		/// Offset of this chunk
+		/// </summary>
+		public long Offset { get; }
+
+		/// <summary>
+		/// Length of this chunk
+		/// </summary>
+		public int Length { get; }
+
+		/// <summary>
+		/// Line index of this chunk
+		/// </summary>
+		public int LineIndex { get; }
+
+		/// <summary>
+		/// Number of lines in this chunk
+		/// </summary>
+		public int LineCount { get; }
+
+		/// <summary>
+		/// List of sub-chunks
+		/// </summary>
+		public IReadOnlyList<LogSubChunkData> SubChunks { get; }
+
+		/// <summary>
+		/// Offset of each sub-chunk within this chunk
+		/// </summary>
+		public int[] SubChunkOffset { get; }
+
+		/// <summary>
+		/// Total line count of the chunk after each sub-chunk
+		/// </summary>
+		public int[] SubChunkLineIndex { get; }
+
+		/// <summary>
+		/// Constructor
+		/// </summary>
+		/// <param name="offset">Offset of this chunk</param>
+		/// <param name="lineIndex">First line index of this chunk</param>
+		/// <param name="subChunks">Sub-chunks for this chunk</param>
+		public LogChunkData(long offset, int lineIndex, IReadOnlyList<LogSubChunkData> subChunks)
+		{
+			Offset = offset;
+			LineIndex = lineIndex;
+			SubChunks = subChunks;
+			SubChunkOffset = CreateSumLookup(subChunks, x => x.Length);
+			SubChunkLineIndex = CreateSumLookup(subChunks, x => x.LineCount);
+
+			if (subChunks.Count > 0)
+			{
+				int lastSubChunkIdx = subChunks.Count - 1;
+				Length = SubChunkOffset[lastSubChunkIdx] + subChunks[lastSubChunkIdx].Length;
+				LineCount = SubChunkLineIndex[lastSubChunkIdx] + subChunks[lastSubChunkIdx].LineCount;
+			}
+		}
+
+		/// <summary>
+		/// Creates a lookup table by summing a field across a list of subchunks
+		/// </summary>
+		/// <param name="subChunks"></param>
+		/// <param name="field"></param>
+		/// <returns></returns>
+		static int[] CreateSumLookup(IReadOnlyList<LogSubChunkData> subChunks, Func<LogSubChunkData, int> field)
+		{
+			int[] total = new int[subChunks.Count];
+
+			int value = 0;
+			for (int idx = 0; idx + 1 < subChunks.Count; idx++)
+			{
+				value += field(subChunks[idx]);
+				total[idx + 1] = value;
+			}
+
+			return total;
+		}
+
+		/// <summary>
+		/// Gets the offset of a line within the file
+		/// </summary>
+		/// <param name="lineIdx">The line index</param>
+		/// <returns>Offset of the line within the file</returns>
+		public long GetLineOffsetWithinChunk(int lineIdx)
+		{
+			int subChunkIdx = SubChunkLineIndex.BinarySearch(lineIdx);
+			if (subChunkIdx < 0)
+			{
+				subChunkIdx = ~subChunkIdx - 1;
+			}
+			return SubChunkOffset[subChunkIdx] + SubChunks[subChunkIdx].InflateText().LineOffsets[lineIdx - SubChunkLineIndex[subChunkIdx]];
+		}
+
+		/// <summary>
+		/// Gets the sub chunk index for the given line
+		/// </summary>
+		/// <param name="chunkLineIdx">Line index within the chunk</param>
+		/// <returns>Subchunk line index</returns>
+		public int GetSubChunkForLine(int chunkLineIdx)
+		{
+			int subChunkIdx = SubChunkLineIndex.BinarySearch(chunkLineIdx);
+			if (subChunkIdx < 0)
+			{
+				subChunkIdx = ~subChunkIdx - 1;
+			}
+			return subChunkIdx;
+		}
+
+		/// <summary>
+		/// Gets the index of the sub-chunk containing the given offset
+		/// </summary>
+		/// <param name="offset">Offset to search for</param>
+		/// <returns>Index of the sub-chunk</returns>
+		public int GetSubChunkForOffsetWithinChunk(int offset)
+		{
+			int subChunkIdx = SubChunkOffset.BinarySearch(offset);
+			if (subChunkIdx < 0)
+			{
+				subChunkIdx = ~subChunkIdx - 1;
+			}
+			return subChunkIdx;
+		}
+
+		/// <summary>
+		/// Signature for output data
+		/// </summary>
+		const int CurrentSignature = 'L' | ('C' << 8) | ('D' << 16);
+
+		/// <summary>
+		/// Read a log chunk from the given stream
+		/// </summary>
+		/// <param name="reader">The reader to read from</param>
+		/// <param name="offset">Offset of this chunk within the file</param>
+		/// <param name="lineIndex">Line index of this chunk</param>
+		/// <returns>New log chunk data</returns>
+		public static LogChunkData Read(MemoryReader reader, long offset, int lineIndex)
+		{
+			int signature = reader.ReadInt32();
+			if ((signature & 0xffffff) != CurrentSignature)
+			{
+				List<LogSubChunkData> subChunks = new List<LogSubChunkData>();
+				subChunks.Add(new LogSubChunkData(LogType.Json, offset, lineIndex, new ReadOnlyLogText(reader.RemainingMemory)));
+				reader.Advance(reader.RemainingMemory.Length);
+				return new LogChunkData(offset, lineIndex, subChunks);
+			}
+
+			int version = signature >> 24;
+			if (version == 0)
+			{
+				List<LogSubChunkData> subChunks = ReadSubChunkList(reader, offset, lineIndex);
+				reader.ReadVariableLengthBytesWithInt32Length();
+				return new LogChunkData(offset, lineIndex, subChunks);
+			}
+			else
+			{
+				List<LogSubChunkData> subChunks = ReadSubChunkList(reader, offset, lineIndex);
+				return new LogChunkData(offset, lineIndex, subChunks);
+			}
+		}
+
+		/// <summary>
+		/// Read a list of sub-chunks from the given stream
+		/// </summary>
+		/// <param name="reader">The reader to read from</param>
+		/// <param name="subChunkOffset">Offset of this chunk within the file</param>
+		/// <param name="subChunkLineIndex">Line index of this chunk</param>
+		/// <returns>List of sub-chunks</returns>
+		static List<LogSubChunkData> ReadSubChunkList(MemoryReader reader, long subChunkOffset, int subChunkLineIndex)
+		{
+			int numSubChunks = reader.ReadInt32();
+
+			List<LogSubChunkData> subChunks = new List<LogSubChunkData>();
+			for (int idx = 0; idx < numSubChunks; idx++)
+			{
+				LogSubChunkData subChunkData = reader.ReadLogSubChunkData(subChunkOffset, subChunkLineIndex);
+				subChunkOffset += subChunkData.Length;
+				subChunkLineIndex += subChunkData.LineCount;
+				subChunks.Add(subChunkData);
+			}
+
+			return subChunks;
+		}
+
+		/// <summary>
+		/// Construct an object from flat memory buffer
+		/// </summary>
+		/// <param name="memory">Memory buffer</param>
+		/// <param name="offset">Offset of this chunk within the file</param>
+		/// <param name="lineIndex">Line index of this chunk</param>
+		/// <returns>Log chunk data</returns>
+		public static LogChunkData FromMemory(ReadOnlyMemory<byte> memory, long offset, int lineIndex)
+		{
+			MemoryReader reader = new MemoryReader(memory);
+			LogChunkData chunkData = Read(reader, offset, lineIndex);
+			reader.CheckEmpty();
+			return chunkData;
+		}
+
+		/// <summary>
+		/// Write the chunk data to a stream
+		/// </summary>
+		/// <returns>Serialized data</returns>
+		public void Write(MemoryWriter writer, ILogger logger)
+		{
+			writer.WriteInt32(CurrentSignature | (1 << 24));
+
+			writer.WriteInt32(SubChunks.Count);
+			foreach (LogSubChunkData subChunk in SubChunks)
+			{
+				writer.WriteLogSubChunkData(subChunk, logger);
+			}
+		}
+
+		/// <summary>
+		/// Construct an object from flat memory buffer
+		/// </summary>
+		/// <param name="logger">Logger for output</param>
+		/// <returns>Log chunk data</returns>
+		public byte[] ToByteArray(ILogger logger)
+		{
+			byte[] data = new byte[GetSerializedSize(logger)];
+			MemoryWriter writer = new MemoryWriter(data);
+			Write(writer, logger);
+			writer.CheckEmpty();
+			return data;
+		}
+
+		/// <summary>
+		/// Determines the size of the serialized buffer
+		/// </summary>
+		public int GetSerializedSize(ILogger logger)
+		{
+			return sizeof(int) + (sizeof(int) + SubChunks.Sum(x => x.GetSerializedSize(logger)));
+		}
+	}
+
+	/// <summary>
+	/// Extensions for the log chunk data
+	/// </summary>
+	static class LogChunkDataExtensions
+	{
+		/// <summary>
+		/// Read a log chunk from the given stream
+		/// </summary>
+		/// <param name="reader">The reader to read from</param>
+		/// <param name="offset">Offset of this chunk within the file</param>
+		/// <param name="lineIndex">Line index of this chunk</param>
+		/// <returns>New log chunk data</returns>
+		public static LogChunkData ReadLogChunkData(this MemoryReader reader, long offset, int lineIndex)
+		{
+			return LogChunkData.Read(reader, offset, lineIndex);
+		}
+
+		/// <summary>
+		/// Write the chunk data to a stream
+		/// </summary>
+		/// <returns>Serialized data</returns>
+		public static void WriteLogChunkData(this MemoryWriter writer, LogChunkData chunkData, ILogger logger)
+		{
+			chunkData.Write(writer, logger);
+		}
+	}
+}
diff --git a/Engine/Source/Programs/Horde/HordeServer/Logs/Data/LogIndexData.cs b/Engine/Source/Programs/Horde/HordeServer/Logs/Data/LogIndexData.cs
new file mode 100644
--- /dev/null
+++ b/Engine/Source/Programs/Horde/HordeServer/Logs/Data/LogIndexData.cs
@@ -0,0 +1,603 @@
+// Copyright Epic Games, Inc. All Rights Reserved.
+
+using System;
+using System.Collections.Generic;
+using System.IO;
+using System.Linq;
+using EpicGames.Core;
+using EpicGames.Horde.Logs;
+using Horde.Server.Utilities;
+using OpenTelemetry.Trace;
+
+namespace Horde.Server.Logs.Data
+{
+	/// <summary>
+	/// Contains a source block of text to be indexed
+	/// </summary>
+	public class LogIndexBlock
+	{
+		/// <summary>
+		/// First line within the file
+		/// </summary>
+		public int LineIndex { get; }
+
+		/// <summary>
+		/// Number of lines in this block
+		/// </summary>
+		public int LineCount { get; }
+
+		/// <summary>
+		/// The plain text for this block
+		/// </summary>
+		public ILogText? CachedPlainText { get; private set; }
+
+		/// <summary>
+		/// The compressed plain text data
+		/// </summary>
+		public ReadOnlyMemory<byte> CompressedPlainText { get; }
+
+		/// <summary>
+		/// Constructor
+		/// </summary>
+		/// <param name="lineIndex">Index of the first line within this block</param>
+		/// <param name="lineCount">Number of lines in the block</param>
+		/// <param name="cachedPlainText">The decompressed plain text</param>
+		/// <param name="compressedPlainText">The compressed text data</param>
+		public LogIndexBlock(int lineIndex, int lineCount, ILogText? cachedPlainText, ReadOnlyMemory<byte> compressedPlainText)
+		{
+			LineIndex = lineIndex;
+			LineCount = lineCount;
+			CachedPlainText = cachedPlainText;
+			CompressedPlainText = compressedPlainText;
+		}
+
+		/// <summary>
+		/// Accessor for the decompressed plaintext
+		/// </summary>
+		public ILogText InflatePlainText()
+		{
+			CachedPlainText ??= new ReadOnlyLogText(CompressedPlainText.DecompressBzip2());
+			return CachedPlainText;
+		}
+	}
+
+	/// <summary>
+	/// Extension methods for <see cref="LogIndexBlock"/>
+	/// </summary>
+	public static class LogIndexBlockExtensions
+	{
+		/// <summary>
+		/// Deserialize an index block
+		/// </summary>
+		/// <param name="reader">Reader to deserialize from</param>
+		/// <returns>The new index block</returns>
+		public static LogIndexBlock ReadLogIndexBlock(this MemoryReader reader)
+		{
+			int lineIndex = reader.ReadInt32();
+			int lineCount = reader.ReadInt32();
+			ReadOnlyMemory<byte> compressedPlainText = reader.ReadVariableLengthBytesWithInt32Length();
+			return new LogIndexBlock(lineIndex, lineCount, null, compressedPlainText);
+		}
+
+		/// <summary>
+		/// Serialize an index block
+		/// </summary>
+		/// <param name="writer">Writer to serialize to</param>
+		/// <param name="block">The block to serialize</param>
+		public static void WriteLogIndexBlock(this MemoryWriter writer, LogIndexBlock block)
+		{
+			writer.WriteInt32(block.LineIndex);
+			writer.WriteInt32(block.LineCount);
+			writer.WriteVariableLengthBytesWithInt32Length(block.CompressedPlainText.Span);
+		}
+
+		/// <summary>
+		/// Gets the serialized size of a block
+		/// </summary>
+		/// <param name="block">The block to serialize</param>
+		public static int GetSerializedSize(this LogIndexBlock block)
+		{
+			return sizeof(int) + sizeof(int) + (sizeof(int) + block.CompressedPlainText.Length);
+		}
+	}
+
+	/// <summary>
+	/// Data for a log chunk
+	/// </summary>
+	public class LogIndexData
+	{
+		/// <summary>
+		/// Index for tokens into the block list
+		/// </summary>
+		ReadOnlyTrie? _cachedTrie;
+
+		/// <summary>
+		/// Number of bits in the index devoted to the block index
+		/// </summary>
+		readonly int _numBlockBits;
+
+		/// <summary>
+		/// List of text blocks
+		/// </summary>
+		readonly LogIndexBlock[] _blocks;
+
+		/// <summary>
+		/// Empty index data
+		/// </summary>
+		public static LogIndexData Empty { get; } = new LogIndexData(ReadOnlyTrie.Empty, 0, Array.Empty<LogIndexBlock>());
+
+		/// <summary>
+		/// Number of lines covered by the index
+		/// </summary>
+		public int LineCount => (_blocks.Length > 0) ? (_blocks[^1].LineIndex + _blocks[^1].LineCount) : 0;
+
+		/// <summary>
+		/// Constructor
+		/// </summary>
+		/// <param name="cachedTrie">Index into the text blocks</param>
+		/// <param name="numBlockBits">Number of bits devoted to the block index</param>
+		/// <param name="blocks">Bloom filters for this log file</param>
+		public LogIndexData(ReadOnlyTrie? cachedTrie, int numBlockBits, LogIndexBlock[] blocks)
+		{
+			_cachedTrie = cachedTrie;
+			_numBlockBits = numBlockBits;
+			_blocks = blocks;
+		}
+
+		/// <summary>
+		/// Updates the blocks in this index relative to a give line index
+		/// </summary>
+		/// <param name="lineIndex"></param>
+		public void SetBaseLineIndex(int lineIndex)
+		{
+			for (int blockIdx = 0; blockIdx < _blocks.Length; blockIdx++)
+			{
+				LogIndexBlock block = _blocks[blockIdx];
+				_blocks[blockIdx] = new LogIndexBlock(lineIndex, block.LineCount, block.CachedPlainText, block.CompressedPlainText);
+				lineIndex += block.LineCount;
+			}
+		}
+
+		/// <summary>
+		/// Builds a trie for this index
+		/// </summary>
+		/// <returns>The trie for this data</returns>
+		ReadOnlyTrie BuildTrie()
+		{
+			if (_cachedTrie == null)
+			{
+				ReadOnlyTrieBuilder builder = new ReadOnlyTrieBuilder();
+				for (int blockIdx = 0; blockIdx < _blocks.Length; blockIdx++)
+				{
+					LogIndexBlock block = _blocks[blockIdx];
+					LogToken.GetTokens(block.InflatePlainText().Data.Span, token => builder.Add((token << _numBlockBits) | (uint)blockIdx));
+				}
+				_cachedTrie = builder.Build();
+			}
+			return _cachedTrie;
+		}
+
+		/// <summary>
+		/// Create an index from an array of blocks
+		/// </summary>
+		/// <param name="indexes">List of indexes to merge</param>
+		/// <returns>Index data</returns>
+		public static LogIndexData Merge(IEnumerable<LogIndexData> indexes)
+		{
+			using TelemetrySpan span = OpenTelemetryTracers.Horde.StartActiveSpan($"{nameof(LogIndexData)}.{nameof(Merge)}");
+
+			// Create the combined block list
+			LogIndexBlock[] newBlocks = indexes.SelectMany(x => x._blocks).ToArray();
+
+			// Figure out how many bits to devote to the block size
+			int newNumBlockBits = 0;
+			while (newBlocks.Length > (1 << newNumBlockBits))
+			{
+				newNumBlockBits += 4;
+			}
+
+			// Add all the blocks into a combined index
+			ReadOnlyTrieBuilder newTrieBuilder = new ReadOnlyTrieBuilder();
+
+			int blockCount = 0;
+			foreach (LogIndexData index in indexes)
+			{
+				ulong blockMask = ((1UL << index._numBlockBits) - 1);
+				foreach (ulong value in index.BuildTrie())
+				{
+					ulong token = value >> index._numBlockBits;
+					int blockIdx = (int)(value & blockMask);
+
+					ulong newToken = (token << newNumBlockBits) | (ulong)(long)(blockCount + blockIdx);
+					newTrieBuilder.Add(newToken);
+				}
+				blockCount += index._blocks.Length;
+			}
+
+			// Construct the new index data
+			ReadOnlyTrie newTrie = newTrieBuilder.Build();
+			return new LogIndexData(newTrie, newNumBlockBits, newBlocks);
+		}
+
+		/// <summary>
+		/// Search for the given text in the index
+		/// </summary>
+		/// <param name="firstLineIndex">First line index to search from</param>
+		/// <param name="text">Text to search for</param>
+		/// <param name="stats">Receives stats for the search</param>
+		/// <returns>List of line numbers for the text</returns>
+		public IEnumerable<int> Search(int firstLineIndex, SearchText text, SearchStats stats)
+		{
+			int lastBlockCount = 0;
+			foreach (int blockIdx in EnumeratePossibleBlocks(text.Bytes, firstLineIndex))
+			{
+				LogIndexBlock block = _blocks[blockIdx];
+
+				stats.NumScannedBlocks++;
+				stats.NumDecompressedBlocks += (block.CachedPlainText == null) ? 1 : 0;
+
+				stats.NumSkippedBlocks += blockIdx - lastBlockCount;
+				lastBlockCount = blockIdx + 1;
+
+				// Decompress the text
+				ILogText blockText = block.InflatePlainText();
+
+				// Find the initial offset within this block
+				int offset = 0;
+				if (firstLineIndex > block.LineIndex)
+				{
+					int lineIndexWithinBlock = firstLineIndex - block.LineIndex;
+					offset = blockText.LineOffsets[lineIndexWithinBlock];
+				}
+
+				// Search within this block
+				for (; ; )
+				{
+					// Find the next offset
+					int nextOffset = blockText.Data.Span.FindNextOcurrence(offset, text);
+					if (nextOffset == -1)
+					{
+						stats.NumScannedBytes += blockText.Length - offset;
+						break;
+					}
+
+					// Update the stats
+					stats.NumScannedBytes += nextOffset - offset;
+					offset = nextOffset;
+
+					// Check it's not another match within the same line
+					int lineIndexWithinBlock = blockText.GetLineIndexForOffset(offset);
+					yield return block.LineIndex + lineIndexWithinBlock;
+
+					// Move to the next line
+					offset = blockText.LineOffsets[lineIndexWithinBlock + 1];
+				}
+
+				// If the last scanned bytes is zero, we didn't have any matches from this chunk
+				if (offset == 0 && _cachedTrie != null)
+				{
+					stats.NumFalsePositiveBlocks++;
+				}
+			}
+			stats.NumSkippedBlocks += _blocks.Length - lastBlockCount;
+		}
+
+		/// <summary>
+		/// Search for the given text in the index
+		/// </summary>
+		/// <param name="text">Text to search for</param>
+		/// <param name="lineIndex">The first </param>
+		/// <returns>List of line numbers for the text</returns>
+		IEnumerable<int> EnumeratePossibleBlocks(ReadOnlyMemory<byte> text, int lineIndex)
+		{
+			// Find the starting block index
+			int blockIdx = _blocks.BinarySearch(x => x.LineIndex, lineIndex);
+			if (blockIdx < 0)
+			{
+				blockIdx = Math.Max(~blockIdx - 1, 0);
+			}
+
+			// Make sure we're not starting out of range
+			if (blockIdx < _blocks.Length)
+			{
+				// In order to be considered a possible block for a positive match, the block must contain an exact list of tokens parsed from the
+				// search text, along with a set of partial matches (for substrings that are not guaranteed to start on token boundaries). The former
+				// is relatively cheap to compute, and done by stepping through all enumerators for each token in a wave, only returning blocks
+				// which all contain the token. At the very least, this must contain blocks in the index from BlockIdx onwards.
+				List<IEnumerator<int>> enumerators = new List<IEnumerator<int>>();
+				enumerators.Add(Enumerable.Range(blockIdx, _blocks.Length - blockIdx).GetEnumerator());
+
+				// The second aspect of the search requires a more expensive search through the trie. This is done by running a set of arbitrary 
+				// delegates to filter the matches returned from the enumerators.
+				List<Predicate<int>> predicates = new List<Predicate<int>>();
+
+				// If the index has a trie, tokenize the input and generate a list of enumerators and predicates.
+				if (_cachedTrie != null)
+				{
+					// Find a list of filters for matching blocks
+					HashSet<ulong> tokens = new HashSet<ulong>();
+					for (int tokenPos = 0; tokenPos < text.Length;)
+					{
+						ReadOnlySpan<byte> token = LogToken.GetTokenText(text.Span, tokenPos);
+						if (tokenPos == 0)
+						{
+							GetUnalignedTokenPredicate(token, token.Length == text.Length, predicates);
+						}
+						else
+						{
+							GetAlignedTokenPredicate(token, tokenPos + token.Length == text.Length, tokens, predicates);
+						}
+						tokenPos += token.Length;
+					}
+
+					// Create an enumerator for each token
+					foreach (ulong token in tokens)
+					{
+						ulong minValue = (token << _numBlockBits) | (ulong)(long)blockIdx;
+						ulong maxValue = minValue + (1UL << _numBlockBits) - 1;
+
+						IEnumerator<int> enumerator = _cachedTrie.EnumerateRange(minValue, maxValue).Select(x => (int)(x - minValue)).GetEnumerator();
+						if (!enumerator.MoveNext())
+						{
+							yield break;
+						}
+
+						enumerators.Add(enumerator);
+					}
+				}
+
+				// Enumerate the matches
+				for (; ; )
+				{
+					// Advance all the enumerators that are behind the current block index. If they are all equal, we have a match.
+					bool match = true;
+					foreach (IEnumerator<int> enumerator in enumerators)
+					{
+						while (enumerator.Current < blockIdx)
+						{
+							if (!enumerator.MoveNext())
+							{
+								yield break;
+							}
+						}
+
+						if (enumerator.Current > blockIdx)
+						{
+							blockIdx = enumerator.Current;
+							match = false;
+						}
+					}
+
+					// Return the match and move to the next block
+					if (match)
+					{
+						if (predicates.All(predicate => predicate(blockIdx)))
+						{
+							yield return blockIdx;
+						}
+						blockIdx++;
+					}
+				}
+			}
+		}
+
+		/// <summary>
+		/// Gets predicates for matching a token that starts 
+		/// </summary>
+		/// <param name="text">The token text</param>
+		/// <param name="allowPartialMatch">Whether to allow a partial match of the token</param>
+		/// <param name="tokens">Set of aligned tokens that are required</param>
+		/// <param name="predicates">List of predicates for the search</param>
+		void GetAlignedTokenPredicate(ReadOnlySpan<byte> text, bool allowPartialMatch, HashSet<ulong> tokens, List<Predicate<int>> predicates)
+		{
+			for (int offset = 0; offset < text.Length; offset += LogToken.MaxTokenBytes)
+			{
+				ulong token = LogToken.GetWindowedTokenValue(text, offset);
+				if (offset + LogToken.MaxTokenBytes > text.Length && allowPartialMatch)
+				{
+					ulong tokenMask = LogToken.GetWindowedTokenMask(text, offset, true);
+					predicates.Add(blockIdx => BlockContainsToken(blockIdx, token, tokenMask));
+					break;
+				}
+				tokens.Add(token);
+			}
+		}
+
+		/// <summary>
+		/// Generates a predicate for matching a token which may or may not start on a regular token boundary
+		/// </summary>
+		/// <param name="text">The token text</param>
+		/// <param name="allowPartialMatch">Whether to allow a partial match of the token</param>
+		/// <param name="predicates">List of predicates for the search</param>
+		void GetUnalignedTokenPredicate(ReadOnlySpan<byte> text, bool allowPartialMatch, List<Predicate<int>> predicates)
+		{
+			byte[] textCopy = text.ToArray();
+
+			Lazy<HashSet<int>> blocks = new Lazy<HashSet<int>>(() =>
+			{
+				HashSet<int> union = new HashSet<int>();
+				for (int shift = 0; shift < LogToken.MaxTokenBytes; shift++)
+				{
+					HashSet<int> blocks = new HashSet<int>(BlocksContainingToken(textCopy.AsSpan(), -shift, allowPartialMatch));
+					for (int offset = -shift + LogToken.MaxTokenBytes; offset < textCopy.Length && blocks.Count > 0; offset += LogToken.MaxTokenBytes)
+					{
+						blocks.IntersectWith(BlocksContainingToken(textCopy.AsSpan(), offset, allowPartialMatch));
+					}
+					if (blocks.Count > 0)
+					{
+						union.UnionWith(blocks);
+					}
+				}
+				return union;
+			});
+
+			predicates.Add(blockIdx => blocks.Value.Contains(blockIdx));
+		}
+
+		/// <summary>
+		/// Tests whether a block contains a particular token
+		/// </summary>
+		/// <param name="blockIdx">Index of the block to search</param>
+		/// <param name="token">The token to test</param>
+		/// <param name="tokenMask">Mask of which bits in the token are valid</param>
+		/// <returns>True if the given block contains a token</returns>
+		bool BlockContainsToken(int blockIdx, ulong token, ulong tokenMask)
+		{
+			token = (token << _numBlockBits) | (uint)blockIdx;
+			tokenMask = (tokenMask << _numBlockBits) | ((1UL << _numBlockBits) - 1);
+			return _cachedTrie!.EnumerateValues((value, valueMask) => (value & tokenMask) == (token & valueMask)).Any();
+		}
+
+		/// <summary>
+		/// Tests whether a block contains a particular token
+		/// </summary>
+		/// <param name="text">The token to test</param>
+		/// <param name="offset">Offset of the window into the token to test</param>
+		/// <param name="allowPartialMatch">Whether to allow a partial match of the token</param>
+		/// <returns>True if the given block contains a token</returns>
+		IEnumerable<int> BlocksContainingToken(ReadOnlySpan<byte> text, int offset, bool allowPartialMatch)
+		{
+			ulong token = LogToken.GetWindowedTokenValue(text, offset) << _numBlockBits;
+			ulong tokenMask = LogToken.GetWindowedTokenMask(text, offset, allowPartialMatch) << _numBlockBits;
+			ulong blockMask = (1UL << _numBlockBits) - 1;
+			return _cachedTrie!.EnumerateValues((value, valueMask) => (value & tokenMask) == (token & valueMask)).Select(x => (int)(x & blockMask)).Distinct();
+		}
+
+		/// <summary>
+		/// Deserialize the index from memory
+		/// </summary>
+		/// <param name="reader">Reader to deserialize from</param>
+		/// <returns>Index data</returns>
+		public static LogIndexData Read(MemoryReader reader)
+		{
+			int version = reader.ReadInt32();
+			if (version == 0)
+			{
+				return LogIndexData.Empty;
+			}
+			else if (version == 1)
+			{
+				ReadOnlyTrie index = reader.ReadTrie();
+				int numBlockBits = reader.ReadInt32();
+
+				LogIndexBlock[] blocks = new LogIndexBlock[reader.ReadInt32()];
+				for (int idx = 0; idx < blocks.Length; idx++)
+				{
+					int lineIndex = (idx > 0) ? (blocks[idx - 1].LineIndex + blocks[idx - 1].LineCount) : 0;
+					int lineCount = reader.ReadInt32();
+					ReadOnlyMemory<byte> compressedPlainText = reader.ReadVariableLengthBytesWithInt32Length();
+					reader.ReadTrie();
+					blocks[idx] = new LogIndexBlock(lineIndex, lineCount, null, compressedPlainText);
+				}
+
+				return new LogIndexData(index, numBlockBits, blocks);
+			}
+			else if (version == 2)
+			{
+				ReadOnlyTrie index = reader.ReadTrie();
+				int numBlockBits = reader.ReadInt32();
+
+				LogIndexBlock[] blocks = new LogIndexBlock[reader.ReadInt32()];
+				for (int idx = 0; idx < blocks.Length; idx++)
+				{
+					int lineIndex = reader.ReadInt32();
+					int lineCount = reader.ReadInt32();
+					ReadOnlyMemory<byte> compressedPlainText = reader.ReadVariableLengthBytesWithInt32Length();
+					reader.ReadTrie();
+					blocks[idx] = new LogIndexBlock(lineIndex, lineCount, null, compressedPlainText);
+				}
+
+				return new LogIndexData(index, numBlockBits, blocks);
+			}
+			else if (version == 3)
+			{
+				ReadOnlyTrie index = reader.ReadTrie();
+				int numBlockBits = reader.ReadInt32();
+				LogIndexBlock[] blocks = reader.ReadVariableLengthArrayWithInt32Length(() => reader.ReadLogIndexBlock());
+				for (int idx = 1; idx < blocks.Length; idx++)
+				{
+					if (blocks[idx].LineIndex == 0)
+					{
+						int lineIndex = blocks[idx - 1].LineIndex + blocks[idx - 1].LineCount;
+						blocks[idx] = new LogIndexBlock(lineIndex, blocks[idx].LineCount, blocks[idx].CachedPlainText, blocks[idx].CompressedPlainText);
+					}
+				}
+				return new LogIndexData(index, numBlockBits, blocks);
+			}
+			else
+			{
+				throw new InvalidDataException($"Invalid index version number {version}");
+			}
+		}
+
+		/// <summary>
+		/// Serialize an index into memory
+		/// </summary>
+		/// <param name="writer">Writer to serialize to</param>
+		public void Write(MemoryWriter writer)
+		{
+			writer.WriteInt32(3);
+			writer.WriteTrie(BuildTrie());
+			writer.WriteInt32(_numBlockBits);
+			writer.WriteVariableLengthArrayWithInt32Length(_blocks, x => writer.WriteLogIndexBlock(x));
+		}
+
+		/// <summary>
+		/// Deserialize the index from memory
+		/// </summary>
+		/// <param name="memory">Memory to deserialize from</param>
+		/// <returns>Index data</returns>
+		public static LogIndexData FromMemory(ReadOnlyMemory<byte> memory)
+		{
+			MemoryReader reader = new MemoryReader(memory);
+			return Read(reader);
+		}
+
+		/// <summary>
+		/// Serializes the index
+		/// </summary>
+		/// <returns>Index data</returns>
+		public byte[] ToByteArray()
+		{
+			byte[] buffer = new byte[GetSerializedSize()];
+
+			MemoryWriter writer = new MemoryWriter(buffer);
+			Write(writer);
+			writer.CheckEmpty();
+
+			return buffer;
+		}
+
+		/// <summary>
+		/// Gets the serialized size of this index
+		/// </summary>
+		/// <returns>The serialized size</returns>
+		public int GetSerializedSize()
+		{
+			return sizeof(int) + BuildTrie().GetSerializedSize() + sizeof(int) + (sizeof(int) + _blocks.Sum(x => x.GetSerializedSize()));
+		}
+	}
+
+	static class LogIndexDataExtensions
+	{
+		/// <summary>
+		/// Deserialize the index from memory
+		/// </summary>
+		/// <param name="reader">Reader to deserialize from</param>
+		/// <returns>Index data</returns>
+		public static LogIndexData ReadLogIndexData(this MemoryReader reader)
+		{
+			return LogIndexData.Read(reader);
+		}
+
+		/// <summary>
+		/// Serialize an index into memory
+		/// </summary>
+		/// <param name="writer">Writer to serialize to</param>
+		/// <param name="index">The index to write</param>
+		public static void WriteLogIndexData(this MemoryWriter writer, LogIndexData index)
+		{
+			index.Write(writer);
+		}
+	}
+}
diff --git a/Engine/Source/Programs/Horde/HordeServer/Logs/Data/LogSubChunkData.cs b/Engine/Source/Programs/Horde/HordeServer/Logs/Data/LogSubChunkData.cs
new file mode 100644
--- /dev/null
+++ b/Engine/Source/Programs/Horde/HordeServer/Logs/Data/LogSubChunkData.cs
@@ -0,0 +1,247 @@
+// Copyright Epic Games, Inc. All Rights Reserved.
+
+using System;
+using EpicGames.Core;
+using EpicGames.Horde.Logs;
+using Horde.Server.Utilities;
+using Microsoft.Extensions.Logging;
+
+namespace Horde.Server.Logs.Data
+{
+	/// <summary>
+	/// Pending data for a sub-chunk
+	/// </summary>
+	public class LogSubChunkData
+	{
+		/// <summary>
+		/// Type of data stored in this subchunk
+		/// </summary>
+		public LogType Type { get; }
+
+		/// <summary>
+		/// Offset within the file of this sub-chunk
+		/// </summary>
+		public long Offset { get; }
+
+		/// <summary>
+		/// Length of this sub-chunk
+		/// </summary>
+		public int Length { get; }
+
+		/// <summary>
+		/// Index of the first line in this sub-chunk
+		/// </summary>
+		public int LineIndex { get; }
+
+		/// <summary>
+		/// Gets the number of lines in this sub-chunk
+		/// </summary>
+		public int LineCount { get; }
+
+		/// <summary>
+		/// Text data
+		/// </summary>
+		ILogText? _textInternal;
+
+		/// <summary>
+		/// Compressed text data
+		/// </summary>
+		ReadOnlyMemory<byte> _compressedTextInternal;
+
+		/// <summary>
+		/// The index for this sub-chunk
+		/// </summary>
+		LogIndexData? _indexInternal;
+
+		/// <summary>
+		/// The log text
+		/// </summary>
+		public ILogText InflateText()
+		{
+			_textInternal ??= new ReadOnlyLogText(_compressedTextInternal.DecompressBzip2());
+			return _textInternal;
+		}
+
+		/// <summary>
+		/// The compressed log text
+		/// </summary>
+		public ReadOnlyMemory<byte> DeflateText()
+		{
+			if (_compressedTextInternal.IsEmpty)
+			{
+				_compressedTextInternal = _textInternal!.Data.CompressBzip2();
+			}
+			return _compressedTextInternal;
+		}
+
+		/// <summary>
+		/// Index for tokens in this chunk
+		/// </summary>
+		public LogIndexData BuildIndex(ILogger logger)
+		{
+			if (_indexInternal == null)
+			{
+				ILogText plainText = InflateText();
+				if (Type != LogType.Text)
+				{
+					plainText = plainText.ToPlainText(logger);
+				}
+
+				LogIndexBlock[] blocks = new LogIndexBlock[1];
+				blocks[0] = new LogIndexBlock(LineIndex, LineCount, plainText, plainText.Data.CompressBzip2());
+				_indexInternal = new LogIndexData(null, 0, blocks);
+			}
+			return _indexInternal;
+		}
+
+		/// <summary>
+		/// Constructor
+		/// </summary>
+		/// <param name="type">Type of data stored in this subchunk</param>
+		/// <param name="offset">Offset within the file of this sub-chunk</param>
+		/// <param name="lineIndex">Index of the first line in this sub-chunk</param>
+		/// <param name="text">Text to add</param>
+		public LogSubChunkData(LogType type, long offset, int lineIndex, ILogText text)
+		{
+			Type = type;
+			Offset = offset;
+			Length = text.Data.Length;
+			LineIndex = lineIndex;
+			LineCount = text.LineCount;
+			_textInternal = text;
+		}
+
+		/// <summary>
+		/// Constructor for raw data
+		/// </summary>
+		/// <param name="type">Type of data stored in this subchunk</param>
+		/// <param name="offset">Offset within the file of this sub-chunk</param>
+		/// <param name="length">Length of the uncompressed data</param>
+		/// <param name="lineIndex">Index of the first line</param>
+		/// <param name="lineCount">Number of lines in the uncompressed text</param>
+		/// <param name="compressedText">Compressed text data</param>
+		/// <param name="index">Index data</param>
+		public LogSubChunkData(LogType type, long offset, int length, int lineIndex, int lineCount, ReadOnlyMemory<byte> compressedText, LogIndexData? index)
+		{
+			Type = type;
+			Offset = offset;
+			Length = length;
+			LineIndex = lineIndex;
+			LineCount = lineCount;
+			_compressedTextInternal = compressedText;
+			_indexInternal = index;
+		}
+
+		/// <summary>
+		/// Constructs a sub-chunk from a block of memory. Uses slices of the given memory buffer rather than copying the data.
+		/// </summary>
+		/// <param name="reader">The reader to read from</param>
+		/// <param name="offset">Offset of the sub-chunk within this file</param>
+		/// <param name="lineIndex">Index of the first line in this subchunk</param>
+		/// <returns>New subchunk data</returns>
+		public static LogSubChunkData Read(MemoryReader reader, long offset, int lineIndex)
+		{
+			int version = reader.ReadInt32(); // Version placeholder
+			if (version == 0)
+			{
+				LogType type = (LogType)reader.ReadInt32();
+				int length = reader.ReadInt32();
+				int lineCount = reader.ReadInt32();
+				ReadOnlyMemory<byte> compressedText = reader.ReadVariableLengthBytesWithInt32Length();
+				reader.ReadVariableLengthBytesWithInt32Length();
+				return new LogSubChunkData(type, offset, length, lineIndex, lineCount, compressedText, null);
+			}
+			else if (version == 1)
+			{
+				LogType type = (LogType)reader.ReadInt32();
+				int length = reader.ReadInt32();
+				int lineCount = reader.ReadInt32();
+				ReadOnlyMemory<byte> compressedText = reader.ReadVariableLengthBytesWithInt32Length();
+				ReadOnlyMemory<byte> compressedPlainText = reader.ReadVariableLengthBytesWithInt32Length();
+				ReadOnlyTrie trie = reader.ReadTrie();
+				LogIndexBlock indexBlock = new LogIndexBlock(lineIndex, lineCount, null, compressedPlainText);
+				LogIndexData index = new LogIndexData(trie, 0, new[] { indexBlock });
+				return new LogSubChunkData(type, offset, length, lineIndex, lineCount, compressedText, index);
+			}
+			else
+			{
+				LogType type = (LogType)reader.ReadInt32();
+				int length = reader.ReadInt32();
+				int lineCount = reader.ReadInt32();
+				ReadOnlyMemory<byte> compressedText = reader.ReadVariableLengthBytesWithInt32Length();
+				LogIndexData index = reader.ReadLogIndexData();
+				index.SetBaseLineIndex(lineIndex); // Fix for incorrectly saved data
+				return new LogSubChunkData(type, offset, length, lineIndex, lineCount, compressedText, index);
+			}
+		}
+
+		/// <summary>
+		/// Serializes the sub-chunk to a stream
+		/// </summary>
+		/// <param name="writer">Writer to output to</param>
+		/// <param name="logger">Logger for output</param>
+		public void Write(MemoryWriter writer, ILogger logger)
+		{
+			writer.WriteInt32(2); // Version placeholder
+
+			writer.WriteInt32((int)Type);
+			writer.WriteInt32(Length);
+			writer.WriteInt32(LineCount);
+
+			writer.WriteVariableLengthBytesWithInt32Length(DeflateText().Span);
+			writer.WriteLogIndexData(BuildIndex(logger));
+		}
+
+		/// <summary>
+		/// Serializes this object to a byte array
+		/// </summary>
+		/// <returns>Byte array</returns>
+		public byte[] ToByteArray(ILogger logger)
+		{
+			byte[] data = new byte[GetSerializedSize(logger)];
+
+			MemoryWriter writer = new MemoryWriter(data);
+			Write(writer, logger);
+			writer.CheckEmpty();
+
+			return data;
+		}
+
+		/// <summary>
+		/// Determines the size of serialized data
+		/// </summary>
+		public int GetSerializedSize(ILogger logger)
+		{
+			return sizeof(int) + sizeof(int) + sizeof(int) + sizeof(int) + (sizeof(int) + DeflateText().Length) + BuildIndex(logger).GetSerializedSize();
+		}
+	}
+
+	/// <summary>
+	/// Pending data for a sub-chunk
+	/// </summary>
+	public static class LogSubChunkDataExtensions
+	{
+		/// <summary>
+		/// Constructs a sub-chunk from a block of memory. Uses slices of the given memory buffer rather than copying the data.
+		/// </summary>
+		/// <param name="reader">The reader to read from</param>
+		/// <param name="offset">Offset of this sub-chunk within the file</param>
+		/// <param name="lineIndex">Index of the first line within this chunk</param>
+		/// <returns>New subchunk data</returns>
+		public static LogSubChunkData ReadLogSubChunkData(this MemoryReader reader, long offset, int lineIndex)
+		{
+			return LogSubChunkData.Read(reader, offset, lineIndex);
+		}
+
+		/// <summary>
+		/// Serializes the sub-chunk to a stream
+		/// </summary>
+		/// <param name="writer">Writer to output to</param>
+		/// <param name="subChunkData">The sub-chunk data to write</param>
+		/// <param name="logger">Logger for output</param>
+		public static void WriteLogSubChunkData(this MemoryWriter writer, LogSubChunkData subChunkData, ILogger logger)
+		{
+			subChunkData.Write(writer, logger);
+		}
+	}
+}
diff --git a/Engine/Source/Programs/Horde/HordeServer/Logs/Data/LogText.cs b/Engine/Source/Programs/Horde/HordeServer/Logs/Data/LogText.cs
new file mode 100644
--- /dev/null
+++ b/Engine/Source/Programs/Horde/HordeServer/Logs/Data/LogText.cs
@@ -0,0 +1,502 @@
+// Copyright Epic Games, Inc. All Rights Reserved.
+
+using System;
+using System.Collections.Generic;
+using System.Diagnostics;
+using System.IO;
+using System.Text;
+using System.Text.Json;
+using EpicGames.Core;
+using Microsoft.Extensions.Logging;
+
+namespace Horde.Server.Logs.Data
+{
+	/// <summary>
+	/// Exception thrown trying to parse log output
+	/// </summary>
+	public class LogTextParseException : Exception
+	{
+		/// <summary>
+		/// Constructor
+		/// </summary>
+		public LogTextParseException(string message, Exception? ex)
+			: base(message, ex)
+		{
+		}
+	}
+
+	/// <summary>
+	/// Interface for log data
+	/// </summary>
+	public interface ILogText
+	{
+		/// <summary>
+		/// The raw text data. Contains a complete set of lines followed by newline characters.
+		/// </summary>
+		public ReadOnlyMemory<byte> Data { get; }
+
+		/// <summary>
+		/// Length of the data in this block
+		/// </summary>
+		public int Length { get; }
+
+		/// <summary>
+		/// Offsets of lines within the data object, including a sentinel for the end of the data (LineCount + 1 entries).
+		/// </summary>
+		public IReadOnlyList<int> LineOffsets { get; }
+
+		/// <summary>
+		/// Number of lines in the 
+		/// </summary>
+		public int LineCount { get; }
+
+		/// <summary>
+		/// Computes the number of bytes used by this object
+		/// </summary>
+		public int AllocatedSize { get; }
+	}
+
+	/// <summary>
+	/// Helper methods for parsing log text
+	/// </summary>
+	public static class LogTextHelpers
+	{
+		/// <summary>
+		/// Finds the number of newline characters in the given span
+		/// </summary>
+		/// <param name="span">Span of utf-8 characters to scan</param>
+		/// <param name="offset">Starting offset within the span</param>
+		/// <returns>Number of lines in the given span</returns>
+		public static int GetLineCount(ReadOnlySpan<byte> span, int offset)
+		{
+			int lineCount = 0;
+			for (; offset < span.Length; offset++)
+			{
+				if (span[offset] == (byte)'\n')
+				{
+					lineCount++;
+				}
+			}
+			return lineCount;
+		}
+
+		/// <summary>
+		/// Find the offsets of the start of each line within the given span
+		/// </summary>
+		/// <param name="span">Span to search</param>
+		/// <param name="offset">Offset within the span to start searching</param>
+		/// <param name="lineOffsets">Receives the list of line offsets</param>
+		/// <param name="lineCount">The current line count</param>
+		public static void FindLineOffsets(ReadOnlySpan<byte> span, int offset, Span<int> lineOffsets, int lineCount)
+		{
+			Debug.Assert(span.Length == 0 || span[^1] == (byte)'\n');
+			for (; ; )
+			{
+				// Store the start of this line
+				lineOffsets[lineCount++] = offset;
+
+				// Check if we're at the end
+				if (offset >= span.Length)
+				{
+					break;
+				}
+
+				// Move to the end of this line
+				while (span[offset] != '\n')
+				{
+					offset++;
+				}
+
+				// Move to the start of the next line
+				offset++;
+			}
+		}
+	}
+
+	/// <summary>
+	/// Immutable block of log text
+	/// </summary>
+	public class ReadOnlyLogText : ILogText
+	{
+		/// <inheritdoc/>
+		public ReadOnlyMemory<byte> Data { get; private set; }
+
+		/// <inheritdoc/>
+		public int Length => Data.Length;
+
+		/// <inheritdoc/>
+		public IReadOnlyList<int> LineOffsets { get; private set; }
+
+		/// <inheritdoc/>
+		public int LineCount => LineOffsets.Count - 1;
+
+		/// <inheritdoc/>
+		public int AllocatedSize => Data.Length + (LineOffsets.Count * sizeof(int));
+
+		/// <summary>
+		/// Constructor
+		/// </summary>
+		/// <param name="data">Data to construct from</param>
+		public ReadOnlyLogText(ReadOnlyMemory<byte> data)
+		{
+			ReadOnlySpan<byte> span = data.Span;
+
+			int lineCount = LogTextHelpers.GetLineCount(span, 0);
+
+			int[] lineOffsets = new int[lineCount + 1];
+			LogTextHelpers.FindLineOffsets(span, 0, lineOffsets, 0);
+			Debug.Assert(lineOffsets[lineCount] == data.Length);
+
+			Data = data;
+			LineOffsets = lineOffsets;
+		}
+	}
+
+	/// <summary>
+	/// Mutable log text object
+	/// </summary>
+	public class LogText : ILogText
+	{
+		/// <summary>
+		/// Accessor for Data
+		/// </summary>
+		byte[] _internalData;
+
+		/// <inheritdoc/>
+		public ReadOnlyMemory<byte> Data => _internalData.AsMemory(0, Length);
+
+		/// <summary>
+		/// Hacky accessor for InternalData, for serializing to Redis
+		/// </summary>
+		public byte[] InternalDataHack => _internalData;
+
+		/// <inheritdoc/>
+		public int Length
+		{
+			get;
+			private set;
+		}
+
+		/// <summary>
+		/// Offsets of the start of each line within the data
+		/// </summary>
+		readonly List<int> _internalLineOffsets = new List<int> { 0 };
+
+		/// <inheritdoc/>
+		public IReadOnlyList<int> LineOffsets => _internalLineOffsets;
+
+		/// <inheritdoc/>
+		public int LineCount => _internalLineOffsets.Count - 1;
+
+		/// <summary>
+		/// Chunk Size
+		/// </summary>
+		public int MaxLength => _internalData.Length;
+
+		/// <summary>
+		/// Computes the size of data used by this object
+		/// </summary>
+		public int AllocatedSize => _internalData.Length + (_internalLineOffsets.Capacity * sizeof(int));
+
+		/// <summary>
+		/// Constructor
+		/// </summary>
+		public LogText()
+		{
+			_internalData = Array.Empty<byte>();
+		}
+
+		/// <summary>
+		/// Constructor
+		/// </summary>
+		/// <param name="data">Data to initialize this chunk with. Ownership of this array is transferred to the chunk, and its length determines the chunk size.</param>
+		/// <param name="length">Number of valid bytes within the initial data array</param>
+		public LogText(byte[] data, int length)
+		{
+			_internalData = data;
+			Length = 0; // Updated below
+
+			UpdateLength(length);
+		}
+
+		/// <summary>
+		/// Create a new chunk data object with the given data appended. The internal buffers are reused, with the assumption that
+		/// there is no contention over writing to the same location in the chunk.
+		/// </summary>
+		/// <param name="textData">The data to append</param>
+		/// <returns>New chunk data object</returns>
+		public void Append(ReadOnlySpan<byte> textData)
+		{
+			textData.CopyTo(_internalData.AsSpan(Length, textData.Length));
+			UpdateLength(Length + textData.Length);
+		}
+
+		/// <summary>
+		/// Updates the plain text representation of this chunk
+		/// </summary>
+		public void AppendPlainText(ILogText srcText, int srcLineIndex, int srcLineCount, ILogger logger)
+		{
+			// Convert all the lines to plain text
+			for (int idx = 0; idx < srcLineCount; idx++)
+			{
+				int lineOffset = srcText.LineOffsets[srcLineIndex + idx];
+				int nextLineOffset = srcText.LineOffsets[srcLineIndex + idx + 1];
+				ReadOnlySpan<byte> inputLine = srcText.Data.Slice(lineOffset, nextLineOffset - lineOffset).Span;
+
+				// Make sure the output buffer is large enough
+				int requiredSpace = nextLineOffset - lineOffset;
+				if (requiredSpace > _internalData.Length - Length)
+				{
+					int maxLineOffset = srcText.LineOffsets[srcLineIndex + srcLineCount];
+					Array.Resize(ref _internalData, _internalData.Length + requiredSpace + (maxLineOffset - nextLineOffset) / 2);
+				}
+
+				// Convert the line to plain text
+				try
+				{
+					Length = ConvertToPlainText(inputLine, _internalData, Length);
+				}
+				catch (Exception ex)
+				{
+					inputLine.CopyTo(_internalData.AsSpan(Length));
+					logger.LogWarning(ex, "Exception while attempting to parse log text as JSON. Line: \"{Line}\"", Encoding.UTF8.GetString(inputLine).Trim());
+				}
+				_internalLineOffsets.Add(Length);
+			}
+		}
+
+		/// <summary>
+		/// Determines if the given line is empty
+		/// </summary>
+		/// <param name="input">The input data</param>
+		/// <returns>True if the given text is empty</returns>
+		public static bool IsEmptyOrWhitespace(ReadOnlySpan<byte> input)
+		{
+			for (int idx = 0; idx < input.Length; idx++)
+			{
+				byte v = input[idx];
+				if (v != (byte)'\n' && v != '\r' && v != ' ')
+				{
+					return false;
+				}
+			}
+			return true;
+		}
+
+		/// <summary>
+		/// Converts a JSON log line to plain text
+		/// </summary>
+		/// <param name="input">The JSON data</param>
+		/// <param name="output">Output buffer for the converted line</param>
+		/// <param name="outputOffset">Offset within the buffer to write the converted data</param>
+		/// <returns></returns>
+		public static int ConvertToPlainText(ReadOnlySpan<byte> input, byte[] output, int outputOffset)
+		{
+			if (IsEmptyOrWhitespace(input))
+			{
+				output[outputOffset] = (byte)'\n';
+				return outputOffset + 1;
+			}
+
+			Utf8JsonReader reader = new Utf8JsonReader(input);
+			if (reader.Read() && reader.TokenType == JsonTokenType.StartObject)
+			{
+				while (reader.Read() && reader.TokenType == JsonTokenType.PropertyName)
+				{
+					if (!reader.ValueTextEquals("message"))
+					{
+						reader.Skip();
+						continue;
+					}
+					if (!reader.Read() || reader.TokenType != JsonTokenType.String)
+					{
+						reader.Skip();
+						continue;
+					}
+
+					int unescapedLength = UnescapeUtf8(reader.ValueSpan, output.AsSpan(outputOffset));
+					outputOffset += unescapedLength;
+
+					output[outputOffset] = (byte)'\n';
+					outputOffset++;
+
+					break;
+				}
+			}
+			return outputOffset;
+		}
+
+		/// <summary>
+		/// Unescape a json utf8 string
+		/// </summary>
+		/// <param name="source">Source span of bytes</param>
+		/// <param name="target">Target span of bytes</param>
+		/// <returns>Length of the converted data</returns>
+		static int UnescapeUtf8(ReadOnlySpan<byte> source, Span<byte> target)
+		{
+			int length = 0;
+			for (; ; )
+			{
+				// Copy up to the next backslash
+				int backslash = source.IndexOf((byte)'\\');
+				if (backslash == -1)
+				{
+					source.CopyTo(target);
+					length += source.Length;
+					break;
+				}
+				else if (backslash > 0)
+				{
+					source.Slice(0, backslash).CopyTo(target);
+					source = source.Slice(backslash);
+					target = target.Slice(backslash);
+					length += backslash;
+				}
+
+				// Check what the escape code is
+				if (source[1] == 'u')
+				{
+					char[] chars = { (char)((StringUtils.ParseHexByte(source, 2) << 8) | StringUtils.ParseHexByte(source, 4)) };
+					int encodedLength = Encoding.UTF8.GetBytes(chars.AsSpan(), target);
+					source = source.Slice(6);
+					target = target.Slice(encodedLength);
+					length += encodedLength;
+				}
+				else
+				{
+					target[0] = source[1] switch
+					{
+						(byte)'\"' => (byte)'\"',
+						(byte)'\\' => (byte)'\\',
+						(byte)'b' => (byte)'\b',
+						(byte)'f' => (byte)'\f',
+						(byte)'n' => (byte)'\n',
+						(byte)'r' => (byte)'\r',
+						(byte)'t' => (byte)'\t',
+						_ => source[1]
+					};
+					source = source.Slice(2);
+					target = target.Slice(1);
+					length++;
+				}
+			}
+			return length;
+		}
+
+		/// <summary>
+		/// Generates placeholder data for a missing span
+		/// </summary>
+		/// <param name="chunkIdx">Index of the chunk</param>
+		/// <param name="hostName">Host name of the machine that was holding the data</param>
+		/// <param name="targetLength">Desired length of the buffer</param>
+		/// <param name="targetLineCount">Desired line count for the buffer</param>
+		public void AppendMissingDataInfo(int chunkIdx, string? hostName, int targetLength, int targetLineCount)
+		{
+			int newLength = Length;
+
+			if (_internalData.Length < targetLength)
+			{
+				Array.Resize(ref _internalData, targetLength);
+			}
+
+			string suffix = $" (server: {hostName})";
+			for (int newLineCount = LineCount; newLineCount < targetLineCount; newLineCount++)
+			{
+				byte[] errorBytes = Encoding.ASCII.GetBytes($"{{ \"level\":\"Error\",\"message\":\"[Missing data at chunk {chunkIdx}, line {newLineCount}{suffix}]\" }}\n");
+				if (newLength + errorBytes.Length > targetLength)
+				{
+					break;
+				}
+
+				errorBytes.AsSpan().CopyTo(_internalData.AsSpan(newLength));
+				newLength += errorBytes.Length;
+
+				suffix = String.Empty;
+			}
+
+			if (newLength < targetLength)
+			{
+				newLength = Math.Max(newLength - 1, 0);
+				_internalData.AsSpan(newLength, (int)((targetLength - 1) - newLength)).Fill((byte)' ');
+				_internalData[targetLength - 1] = (byte)'\n';
+			}
+
+			UpdateLength((int)targetLength);
+		}
+
+		/// <summary>
+		/// Shrinks the data allocated by this chunk to the minimum required
+		/// </summary>
+		public void Shrink()
+		{
+			if (Length < _internalData.Length)
+			{
+				Array.Resize(ref _internalData, Length);
+			}
+		}
+
+		/// <summary>
+		/// Updates the length of this chunk, computing all the newline offsets
+		/// </summary>
+		/// <param name="newLength">New length of the chunk</param>
+		private void UpdateLength(int newLength)
+		{
+			if (newLength > Length)
+			{
+				// Make sure the data ends with a newline
+				if (_internalData[newLength - 1] != '\n')
+				{
+					throw new InvalidDataException("Chunk data must end with a newline");
+				}
+
+				// Calculate the new number of newlines
+				ReadOnlySpan<byte> newData = _internalData.AsSpan(0, newLength);
+				for (int idx = Length; idx < newLength; idx++)
+				{
+					if (newData[idx] == '\n')
+					{
+						_internalLineOffsets.Add(idx + 1);
+					}
+				}
+
+				// Update the size of the buffer
+				Length = newLength;
+			}
+		}
+	}
+
+	/// <summary>
+	/// Extension methods for ILogText
+	/// </summary>
+	public static class LogTextExtensions
+	{
+		/// <summary>
+		/// Find the line index for a particular offset
+		/// </summary>
+		/// <param name="logText">The text to search</param>
+		/// <param name="offset">Offset within the text</param>
+		/// <returns>The line index</returns>
+		public static int GetLineIndexForOffset(this ILogText logText, int offset)
+		{
+			int lineIdx = logText.LineOffsets.BinarySearch(offset);
+			if (lineIdx < 0)
+			{
+				lineIdx = ~lineIdx - 1;
+			}
+			return lineIdx;
+		}
+
+		/// <summary>
+		/// Converts a log text instance to plain text
+		/// </summary>
+		/// <param name="logText">The text to convert</param>
+		/// <param name="logger">Logger for conversion warnings</param>
+		/// <returns>The plain text instance</returns>
+		public static ILogText ToPlainText(this ILogText logText, ILogger logger)
+		{
+			LogText other = new LogText();
+			other.AppendPlainText(logText, 0, logText.LineCount, logger);
+			return other;
+		}
+	}
+}
diff --git a/Engine/Source/Programs/Horde/HordeServer/Logs/Data/LogToken.cs b/Engine/Source/Programs/Horde/HordeServer/Logs/Data/LogToken.cs
new file mode 100644
--- /dev/null
+++ b/Engine/Source/Programs/Horde/HordeServer/Logs/Data/LogToken.cs
@@ -0,0 +1,205 @@
+// Copyright Epic Games, Inc. All Rights Reserved.
+
+using System;
+using System.Collections.Generic;
+
+namespace Horde.Server.Logs.Data
+{
+	/// <summary>
+	/// Functionality for decomposing log text into tokens
+	/// </summary>
+	public static class LogToken
+	{
+		/// <summary>
+		/// Maximum number of bytes in each token
+		/// </summary>
+		public const int MaxTokenBytes = 4;
+
+		/// <summary>
+		/// Number of bits in each token
+		/// </summary>
+		public const int MaxTokenBits = MaxTokenBytes * 8;
+
+		/// <summary>
+		/// Lookup from input byte to token type
+		/// </summary>
+		static readonly byte[] s_tokenTypes = GetTokenTypes();
+
+		/// <summary>
+		/// Lookup from input byte to token char
+		/// </summary>
+		static readonly byte[] s_tokenChars = GetTokenChars();
+
+		/// <summary>
+		/// Get a set of tokens in a trie
+		/// </summary>
+		/// <param name="text">Text to scan</param>
+		/// <returns></returns>
+		public static ReadOnlyTrie GetTrie(ReadOnlySpan<byte> text)
+		{
+			ReadOnlyTrieBuilder builder = new ReadOnlyTrieBuilder();
+			GetTokens(text, builder.Add);
+			return builder.Build();
+		}
+
+		/// <summary>
+		/// Gets a single token
+		/// </summary>
+		/// <param name="text">The text to parse</param>
+		/// <returns>The token value</returns>
+		public static ulong GetToken(ReadOnlySpan<byte> text)
+		{
+			ulong token = 0;
+			GetTokens(text, x => token = x);
+			return token;
+		}
+
+		/// <summary>
+		/// Get a set of tokens in a trie
+		/// </summary>
+		/// <param name="text">Text to scan</param>
+		/// <returns></returns>
+		public static HashSet<ulong> GetTokenSet(ReadOnlySpan<byte> text)
+		{
+			HashSet<ulong> tokens = new HashSet<ulong>();
+			GetTokens(text, x => tokens.Add(x));
+			return tokens;
+		}
+
+		/// <summary>
+		/// Decompose a span of text into tokens
+		/// </summary>
+		/// <param name="text">Text to scan</param>
+		/// <param name="addToken">Receives a set of tokens</param>
+		public static void GetTokens(ReadOnlySpan<byte> text, Action<ulong> addToken)
+		{
+			if (text.Length > 0)
+			{
+				int type = s_tokenTypes[text[0]];
+				int numTokenBits = 8;
+				ulong token = s_tokenChars[text[0]];
+
+				for (int textIdx = 1; textIdx < text.Length; textIdx++)
+				{
+					byte nextChar = s_tokenChars[text[textIdx]];
+					int nextType = s_tokenTypes[nextChar];
+					if (type != nextType || numTokenBits + 8 > MaxTokenBits)
+					{
+						addToken(token << (MaxTokenBits - numTokenBits));
+						token = 0;
+						numTokenBits = 0;
+						type = nextType;
+					}
+					token = (token << 8) | nextChar;
+					numTokenBits += 8;
+				}
+
+				addToken(token << (MaxTokenBits - numTokenBits));
+			}
+		}
+
+		/// <summary>
+		/// Gets the length of the first token in the given span
+		/// </summary>
+		/// <param name="text">The text to search</param>
+		/// <param name="pos">Start position for the search</param>
+		/// <returns>Length of the first token</returns>
+		public static ReadOnlySpan<byte> GetTokenText(ReadOnlySpan<byte> text, int pos)
+		{
+			int type = s_tokenTypes[text[pos]];
+			for (int end = pos + 1; ; end++)
+			{
+				if (end == text.Length || s_tokenTypes[text[end]] != type)
+				{
+					return text.Slice(pos, end - pos);
+				}
+			}
+		}
+
+		/// <summary>
+		/// Gets the length of the first token in the given span
+		/// </summary>
+		/// <param name="text">The text to search</param>
+		/// <param name="offset">Offset of the window to read from the token</param>
+		/// <returns>Length of the first token</returns>
+		public static ulong GetWindowedTokenValue(ReadOnlySpan<byte> text, int offset)
+		{
+			ulong token = 0;
+			for (int idx = 0; idx < MaxTokenBytes; idx++)
+			{
+				token <<= 8;
+				if (offset >= 0 && offset < text.Length)
+				{
+					token |= s_tokenChars[text[offset]];
+				}
+				offset++;
+			}
+			return token;
+		}
+
+		/// <summary>
+		/// Gets the length of the first token in the given span
+		/// </summary>
+		/// <param name="text">The text to search</param>
+		/// <param name="offset">Offset of the window to read from the token</param>
+		/// <param name="allowPartialMatch">Whether to allow only matching the start of the string</param>
+		/// <returns>Length of the first token</returns>
+		public static ulong GetWindowedTokenMask(ReadOnlySpan<byte> text, int offset, bool allowPartialMatch)
+		{
+			ulong token = 0;
+			for (int idx = 0; idx < MaxTokenBytes; idx++)
+			{
+				token <<= 8;
+				if (offset >= 0 && (offset < text.Length || !allowPartialMatch))
+				{
+					token |= 0xff;
+				}
+				offset++;
+			}
+			return token;
+		}
+
+		/// <summary>
+		/// Build the lookup table for token types
+		/// </summary>
+		/// <returns>Array whose elements map from an input byte to token type</returns>
+		static byte[] GetTokenTypes()
+		{
+			byte[] charTypes = new byte[256];
+			for (int idx = 'a'; idx <= 'z'; idx++)
+			{
+				charTypes[idx] = 1;
+			}
+			for (int idx = 'A'; idx <= 'Z'; idx++)
+			{
+				charTypes[idx] = 1;
+			}
+			for (int idx = '0'; idx <= '9'; idx++)
+			{
+				charTypes[idx] = 2;
+			}
+			charTypes[' '] = 3;
+			charTypes['\t'] = 3;
+			charTypes['\n'] = 4;
+			return charTypes;
+		}
+
+		/// <summary>
+		/// Build the lookup table for token types
+		/// </summary>
+		/// <returns>Array whose elements map from an input byte to token type</returns>
+		static byte[] GetTokenChars()
+		{
+			byte[] chars = new byte[256];
+			for (int idx = 0; idx < 256; idx++)
+			{
+				chars[idx] = (byte)idx;
+			}
+			for (int idx = 'A'; idx <= 'Z'; idx++)
+			{
+				chars[idx] = (byte)('a' + idx - 'A');
+			}
+			return chars;
+		}
+	}
+}
diff --git a/Engine/Source/Programs/Horde/HordeServer/Logs/ILogStorage.cs b/Engine/Source/Programs/Horde/HordeServer/Logs/ILogStorage.cs
new file mode 100644
--- /dev/null
+++ b/Engine/Source/Programs/Horde/HordeServer/Logs/ILogStorage.cs
@@ -0,0 +1,32 @@
+// Copyright Epic Games, Inc. All Rights Reserved.
+
+using System;
+using System.Threading.Tasks;
+using EpicGames.Horde.Logs;
+using Horde.Server.Logs.Data;
+
+namespace Horde.Server.Logs
+{
+	/// <summary>
+	/// Legacy interface for accessing log data. This is not used for new log files; the agent writes log chunks directly to the storage client.
+	/// </summary>
+	public interface ILogStorage : IDisposable
+	{
+		/// <summary>
+		/// Attempts to read an index for the given log file
+		/// </summary>
+		/// <param name="logId">Unique id of the log file</param>
+		/// <param name="length">Length of the file covered by the index</param>
+		/// <returns>Index for the log file</returns>
+		Task<LogIndexData?> ReadIndexAsync(LogId logId, long length);
+
+		/// <summary>
+		/// Retrieves an item from the cache
+		/// </summary>
+		/// <param name="logId">Unique id of the log file</param>
+		/// <param name="offset">Offset of the chunk to read</param>
+		/// <param name="lineIndex">First line of the chunk</param>
+		/// <returns>Data for the given key, or null if it's not present</returns>
+		Task<LogChunkData?> ReadChunkAsync(LogId logId, long offset, int lineIndex);
+	}
+}
diff --git a/Engine/Source/Programs/Horde/HordeServer/Logs/Storage/LocalLogStorage.cs b/Engine/Source/Programs/Horde/HordeServer/Logs/Storage/LocalLogStorage.cs
new file mode 100644
--- /dev/null
+++ b/Engine/Source/Programs/Horde/HordeServer/Logs/Storage/LocalLogStorage.cs
@@ -0,0 +1,128 @@
+// Copyright Epic Games, Inc. All Rights Reserved.
+
+using System;
+using System.Threading.Tasks;
+using EpicGames.Horde.Logs;
+using Horde.Server.Logs.Data;
+using Microsoft.Extensions.Caching.Memory;
+
+namespace Horde.Server.Logs.Storage
+{
+	/// <summary>
+	/// In-memory cache for chunk and index data
+	/// </summary>
+	class LocalLogStorage : ILogStorage
+	{
+		/// <summary>
+		/// The memory cache instance
+		/// </summary>
+		readonly MemoryCache _cache;
+
+		/// <summary>
+		/// The inner storage
+		/// </summary>
+		readonly ILogStorage _inner;
+
+		/// <summary>
+		/// Constructor
+		/// </summary>
+		/// <param name="numItems">Maximum size for the memory cache</param>
+		/// <param name="inner">The inner storage provider</param>
+		public LocalLogStorage(int numItems, ILogStorage inner)
+		{
+			MemoryCacheOptions options = new MemoryCacheOptions();
+			options.SizeLimit = numItems;
+
+			_cache = new MemoryCache(options);
+			_inner = inner;
+		}
+
+		/// <inheritdoc/>
+		public void Dispose()
+		{
+			_cache.Dispose();
+			_inner.Dispose();
+		}
+
+		/// <summary>
+		/// Gets the cache key for a particular index
+		/// </summary>
+		/// <param name="logId">The log file to retrieve an index for</param>
+		/// <param name="length">Length of the file covered by the index</param>
+		/// <returns>Cache key for the index</returns>
+		static string IndexKey(LogId logId, long length)
+		{
+			return $"{logId}/index-{length}";
+		}
+
+		/// <summary>
+		/// Gets the cache key for a particular chunk
+		/// </summary>
+		/// <param name="logId">The log file to retrieve an index for</param>
+		/// <param name="offset">The chunk offset</param>
+		/// <returns>Cache key for the chunk</returns>
+		static string ChunkKey(LogId logId, long offset)
+		{
+			return $"{logId}/chunk-{offset}";
+		}
+
+		/// <summary>
+		/// Adds an entry to the cache
+		/// </summary>
+		/// <param name="key">The cache key</param>
+		/// <param name="value">Value to store</param>
+		void AddEntry(string key, object? value)
+		{
+			using (ICacheEntry entry = _cache.CreateEntry(key))
+			{
+				if (value == null)
+				{
+					entry.SetAbsoluteExpiration(TimeSpan.FromSeconds(30.0));
+					entry.SetSize(0);
+				}
+				else
+				{
+					entry.SetSlidingExpiration(TimeSpan.FromHours(4));
+					entry.SetSize(1);
+				}
+				entry.SetValue(value);
+			}
+		}
+
+		/// <summary>
+		/// Reads a value from the cache, or from the inner storage provider, adding the result to the cache
+		/// </summary>
+		/// <typeparam name="T">Type of object to read</typeparam>
+		/// <param name="key">The cache key</param>
+		/// <param name="readInner">Delegate to read from the inner storage provider</param>
+		/// <returns>The retrieved value</returns>
+		async Task<T?> ReadValueAsync<T>(string key, Func<Task<T?>> readInner) where T : class
+		{
+			T? value;
+			if (!_cache.TryGetValue(key, out value))
+			{
+				try
+				{
+					value = await readInner();
+				}
+				finally
+				{
+					AddEntry(key, value);
+				}
+			}
+			return value;
+		}
+
+		/// <inheritdoc/>
+		public Task<LogIndexData?> ReadIndexAsync(LogId logId, long length)
+		{
+			return ReadValueAsync(IndexKey(logId, length), () => _inner.ReadIndexAsync(logId, length));
+		}
+
+		/// <inheritdoc/>
+		public Task<LogChunkData?> ReadChunkAsync(LogId logId, long offset, int lineIndex)
+		{
+			return ReadValueAsync(ChunkKey(logId, offset), () => _inner.ReadChunkAsync(logId, offset, lineIndex));
+		}
+	}
+}
diff --git a/Engine/Source/Programs/Horde/HordeServer/Logs/Storage/NullLogStorage.cs b/Engine/Source/Programs/Horde/HordeServer/Logs/Storage/NullLogStorage.cs
new file mode 100644
--- /dev/null
+++ b/Engine/Source/Programs/Horde/HordeServer/Logs/Storage/NullLogStorage.cs
@@ -0,0 +1,31 @@
+// Copyright Epic Games, Inc. All Rights Reserved.
+
+using System.Threading.Tasks;
+using EpicGames.Horde.Logs;
+using Horde.Server.Logs.Data;
+
+namespace Horde.Server.Logs.Storage
+{
+	/// <summary>
+	/// Empty implementation of log storage
+	/// </summary>
+	public sealed class NullLogStorage : ILogStorage
+	{
+		/// <inheritdoc/>
+		public void Dispose()
+		{
+		}
+
+		/// <inheritdoc/>
+		public Task<LogIndexData?> ReadIndexAsync(LogId logId, long length)
+		{
+			return Task.FromResult<LogIndexData?>(null);
+		}
+
+		/// <inheritdoc/>
+		public Task<LogChunkData?> ReadChunkAsync(LogId logId, long offset, int lineIndex)
+		{
+			return Task.FromResult<LogChunkData?>(null);
+		}
+	}
+}
diff --git a/Engine/Source/Programs/Horde/HordeServer/Logs/Storage/PersistentLogStorage.cs b/Engine/Source/Programs/Horde/HordeServer/Logs/Storage/PersistentLogStorage.cs
new file mode 100644
--- /dev/null
+++ b/Engine/Source/Programs/Horde/HordeServer/Logs/Storage/PersistentLogStorage.cs
@@ -0,0 +1,66 @@
+// Copyright Epic Games, Inc. All Rights Reserved.
+
+using System;
+using System.Threading.Tasks;
+using EpicGames.Core;
+using EpicGames.Horde.Logs;
+using EpicGames.Horde.Storage;
+using Horde.Server.Logs.Data;
+using Microsoft.Extensions.Logging;
+
+namespace Horde.Server.Logs.Storage
+{
+	/// <summary>
+	/// Bulk storage for log file data
+	/// </summary>
+	class PersistentLogStorage : ILogStorage
+	{
+		readonly IObjectStore _objectStore;
+		readonly ILogger _logger;
+
+		/// <summary>
+		/// Constructor
+		/// </summary>
+		/// <param name="objectStore">The object store</param>
+		/// <param name="logger">Logging provider</param>
+		public PersistentLogStorage(IObjectStore<PersistentLogStorage> objectStore, ILogger<PersistentLogStorage> logger)
+		{
+			_objectStore = objectStore;
+			_logger = logger;
+		}
+
+		/// <inheritdoc/>
+		public void Dispose()
+		{
+		}
+
+		/// <inheritdoc/>
+		public async Task<LogIndexData?> ReadIndexAsync(LogId logId, long length)
+		{
+			_logger.LogDebug("Reading log {LogId} index length {Length} from persistent storage", logId, length);
+
+			ObjectKey key = new ObjectKey($"{logId}/index_{length}");
+			IReadOnlyMemoryOwner<byte> data = await _objectStore.ReadAsync(key);
+			return LogIndexData.FromMemory(data.Memory.ToArray());
+		}
+
+		/// <inheritdoc/>
+		public async Task<LogChunkData?> ReadChunkAsync(LogId logId, long offset, int lineIndex)
+		{
+			_logger.LogDebug("Reading log {LogId} chunk offset {Offset} from persistent storage", logId, offset);
+
+			ObjectKey key = new ObjectKey($"{logId}/offset_{offset}");
+			IReadOnlyMemoryOwner<byte> data = await _objectStore.ReadAsync(key);
+
+			MemoryReader reader = new MemoryReader(data.Memory.ToArray());
+			LogChunkData chunkData = reader.ReadLogChunkData(offset, lineIndex);
+
+			if (reader.RemainingMemory.Length > 0)
+			{
+				throw new Exception($"Serialization of persistent chunk {key} is not at expected offset ({reader.RemainingMemory.Length} bytes remaining)");
+			}
+
+			return chunkData;
+		}
+	}
+}
diff --git a/Engine/Source/Programs/Horde/HordeServer/Logs/Storage/SequencedLogStorage.cs b/Engine/Source/Programs/Horde/HordeServer/Logs/Storage/SequencedLogStorage.cs
new file mode 100644
--- /dev/null
+++ b/Engine/Source/Programs/Horde/HordeServer/Logs/Storage/SequencedLogStorage.cs
@@ -0,0 +1,112 @@
+// Copyright Epic Games, Inc. All Rights Reserved.
+
+using System.Collections.Generic;
+using System.Threading.Tasks;
+using EpicGames.Horde.Logs;
+using Horde.Server.Logs.Data;
+
+namespace Horde.Server.Logs.Storage
+{
+	/// <summary>
+	/// Storage layer which caches pending read tasks, to avoid fetching the same item more than once
+	/// </summary>
+	class SequencedLogStorage : ILogStorage
+	{
+		/// <summary>
+		/// Inner storage implementation
+		/// </summary>
+		readonly ILogStorage _inner;
+
+		/// <summary>
+		/// Pending index reads
+		/// </summary>
+		readonly Dictionary<(LogId, long), Task<LogIndexData?>> _indexReadTasks = new Dictionary<(LogId, long), Task<LogIndexData?>>();
+
+		/// <summary>
+		/// Pending chunk reads
+		/// </summary>
+		readonly Dictionary<(LogId, long), Task<LogChunkData?>> _chunkReadTasks = new Dictionary<(LogId, long), Task<LogChunkData?>>();
+
+		/// <summary>
+		/// Constructor
+		/// </summary>
+		/// <param name="inner">The inner storage provider</param>
+		public SequencedLogStorage(ILogStorage inner)
+		{
+			_inner = inner;
+		}
+
+		/// <inheritdoc/>
+		public void Dispose()
+		{
+			_inner.Dispose();
+		}
+
+		/// <inheritdoc/>
+		public Task<LogIndexData?> ReadIndexAsync(LogId logId, long length)
+		{
+			Task<LogIndexData?>? task;
+			lock (_indexReadTasks)
+			{
+				if (!_indexReadTasks.TryGetValue((logId, length), out task))
+				{
+					task = InnerReadIndexAsync(logId, length);
+					_indexReadTasks.Add((logId, length), task);
+				}
+			}
+			return task;
+		}
+
+		/// <inheritdoc/>
+		public Task<LogChunkData?> ReadChunkAsync(LogId logId, long offset, int lineIndex)
+		{
+			Task<LogChunkData?>? task;
+			lock (_chunkReadTasks)
+			{
+				if (!_chunkReadTasks.TryGetValue((logId, offset), out task))
+				{
+					task = InnerReadChunkAsync(logId, offset, lineIndex);
+					_chunkReadTasks[(logId, offset)] = task;
+				}
+			}
+			return task;
+		}
+
+		/// <summary>
+		/// Wrapper for reading an index from the inner storage provider
+		/// </summary>
+		/// <param name="logId">The log file to read the index for</param>
+		/// <param name="length">Length of the file that's indexed</param>
+		/// <returns>The index data</returns>
+		async Task<LogIndexData?> InnerReadIndexAsync(LogId logId, long length)
+		{
+			await Task.Yield();
+
+			LogIndexData? indexData = await _inner.ReadIndexAsync(logId, length);
+			lock (_indexReadTasks)
+			{
+				_indexReadTasks.Remove((logId, length));
+			}
+			return indexData;
+		}
+
+		/// <summary>
+		/// Wrapper for reading a chunk from the inner storage provider
+		/// </summary>
+		/// <param name="logId">The log file to read the index for</param>
+		/// <param name="offset">Offset of the chunk to read</param>
+		/// <param name="lineIndex">Index of the first line in this chunk</param>
+		/// <returns>The index data</returns>
+		async Task<LogChunkData?> InnerReadChunkAsync(LogId logId, long offset, int lineIndex)
+		{
+			await Task.Yield();
+
+			LogChunkData? chunkData = await _inner.ReadChunkAsync(logId, offset, lineIndex);
+			lock (_chunkReadTasks)
+			{
+				_chunkReadTasks.Remove((logId, offset));
+			}
+			return chunkData;
+		}
+	}
+}
