Horde: Rename projects before splitting into plugins (Horde.Agent -> HordeAgent, Horde.Server -> HordeServer, Horde -> HordeCmd).

diff --git a/Engine/Source/Programs/Horde/HordeServer/Startup.cs b/Engine/Source/Programs/Horde/HordeServer/Startup.cs
new file mode 100644
--- /dev/null
+++ b/Engine/Source/Programs/Horde/HordeServer/Startup.cs
@@ -0,0 +1,1507 @@
+// Copyright Epic Games, Inc. All Rights Reserved.
+
+using System;
+using System.Collections.Generic;
+using System.Diagnostics;
+using System.Globalization;
+using System.IO;
+using System.Linq;
+using System.Net.Mime;
+using System.Reflection;
+using System.Runtime.InteropServices;
+using System.Text.Json;
+using System.Text.Json.Serialization;
+using System.Threading;
+using System.Threading.Tasks;
+using Amazon;
+using Amazon.AutoScaling;
+using Amazon.CloudWatch;
+using Amazon.EC2;
+using Amazon.Extensions.NETCore.Setup;
+using Amazon.SQS;
+using EpicGames.AspNet;
+using EpicGames.Core;
+using EpicGames.Horde;
+using EpicGames.Horde.Acls;
+using EpicGames.Horde.Agents;
+using EpicGames.Horde.Agents.Leases;
+using EpicGames.Horde.Commits;
+using EpicGames.Horde.Jobs;
+using EpicGames.Horde.Logs;
+using EpicGames.Horde.Projects;
+using EpicGames.Horde.Server;
+using EpicGames.Horde.Storage;
+using EpicGames.Horde.Storage.Bundles;
+using EpicGames.Horde.Storage.ObjectStores;
+using EpicGames.Horde.Streams;
+using EpicGames.Horde.Users;
+using EpicGames.Redis;
+using Grpc.Core;
+using Grpc.Core.Interceptors;
+using Horde.Server.Accounts;
+using Horde.Server.Acls;
+using Horde.Server.Agents;
+using Horde.Server.Agents.Enrollment;
+using Horde.Server.Agents.Fleet;
+using Horde.Server.Agents.Leases;
+using Horde.Server.Agents.Pools;
+using Horde.Server.Agents.Relay;
+using Horde.Server.Agents.Sessions;
+using Horde.Server.Agents.Telemetry;
+using Horde.Server.Agents.Utilization;
+using Horde.Server.Artifacts;
+using Horde.Server.Auditing;
+using Horde.Server.Authentication;
+using Horde.Server.Aws;
+using Horde.Server.Commits;
+using Horde.Server.Compute;
+using Horde.Server.Configuration;
+using Horde.Server.Dashboard;
+using Horde.Server.Ddc;
+using Horde.Server.Devices;
+using Horde.Server.Issues;
+using Horde.Server.Issues.External;
+using Horde.Server.Jobs;
+using Horde.Server.Jobs.Artifacts;
+using Horde.Server.Jobs.Bisect;
+using Horde.Server.Jobs.Graphs;
+using Horde.Server.Jobs.Schedules;
+using Horde.Server.Jobs.Templates;
+using Horde.Server.Jobs.TestData;
+using Horde.Server.Jobs.Timing;
+using Horde.Server.Logs;
+using Horde.Server.Logs.Storage;
+using Horde.Server.Notifications;
+using Horde.Server.Notifications.Sinks;
+using Horde.Server.Perforce;
+using Horde.Server.Replicators;
+using Horde.Server.Secrets;
+using Horde.Server.Secrets.Providers;
+using Horde.Server.Server;
+using Horde.Server.Server.Notices;
+using Horde.Server.ServiceAccounts;
+using Horde.Server.Storage;
+using Horde.Server.Storage.ObjectStores;
+using Horde.Server.Streams;
+using Horde.Server.Tasks;
+using Horde.Server.Telemetry;
+using Horde.Server.Telemetry.Metrics;
+using Horde.Server.Telemetry.Sinks;
+using Horde.Server.Tools;
+using Horde.Server.Ugs;
+using Horde.Server.Users;
+using Horde.Server.Utilities;
+using HordeCommon;
+using Microsoft.AspNetCore.Authentication;
+using Microsoft.AspNetCore.Authentication.Cookies;
+using Microsoft.AspNetCore.Authentication.JwtBearer;
+using Microsoft.AspNetCore.Authentication.OpenIdConnect;
+using Microsoft.AspNetCore.Authorization;
+using Microsoft.AspNetCore.Builder;
+using Microsoft.AspNetCore.DataProtection;
+using Microsoft.AspNetCore.Hosting;
+using Microsoft.AspNetCore.Hosting.Server.Features;
+using Microsoft.AspNetCore.Http;
+using Microsoft.AspNetCore.Http.Features;
+using Microsoft.AspNetCore.HttpOverrides;
+using Microsoft.AspNetCore.Mvc;
+using Microsoft.AspNetCore.Mvc.Controllers;
+using Microsoft.AspNetCore.Mvc.Filters;
+using Microsoft.AspNetCore.Mvc.ModelBinding;
+using Microsoft.AspNetCore.ResponseCompression;
+using Microsoft.Extensions.Configuration;
+using Microsoft.Extensions.DependencyInjection;
+using Microsoft.Extensions.Diagnostics.HealthChecks;
+using Microsoft.Extensions.Hosting;
+using Microsoft.Extensions.Logging;
+using Microsoft.Extensions.Logging.Abstractions;
+using Microsoft.Extensions.Options;
+using Microsoft.Extensions.Primitives;
+using Microsoft.OpenApi.Models;
+using MongoDB.Bson;
+using MongoDB.Bson.Serialization;
+using MongoDB.Bson.Serialization.Conventions;
+using MongoDB.Bson.Serialization.Serializers;
+using Serilog;
+using Serilog.Events;
+using StackExchange.Redis;
+using Swashbuckle.AspNetCore.SwaggerGen;
+using Status = Grpc.Core.Status;
+using StatusCode = Grpc.Core.StatusCode;
+
+#pragma warning disable CA1505 // 'ConfigureServices' has a maintainability index of '9'. Rewrite or refactor the code to increase its maintainability index (MI) above '9'.
+
+namespace Horde.Server
+{
+	using ContentHash = EpicGames.Core.ContentHash;
+	using ILogger = Microsoft.Extensions.Logging.ILogger;
+	using ReferenceResolver = Horde.Server.Ddc.ReferenceResolver;
+
+	class Startup
+	{
+		static Startup()
+		{
+			ProtoBuf.Meta.RuntimeTypeModel.Default[typeof(ProjectId)].SetSurrogate(typeof(StringIdProto<ProjectId, ProjectIdConverter>));
+			ProtoBuf.Meta.RuntimeTypeModel.Default[typeof(StreamId)].SetSurrogate(typeof(StringIdProto<StreamId, StreamIdConverter>));
+		}
+
+		class GrpcExceptionInterceptor : Interceptor
+		{
+			readonly ILogger<GrpcExceptionInterceptor> _logger;
+
+			public GrpcExceptionInterceptor(ILogger<GrpcExceptionInterceptor> logger)
+			{
+				_logger = logger;
+			}
+
+			public override Task<TResponse> UnaryServerHandler<TRequest, TResponse>(TRequest request, ServerCallContext context, UnaryServerMethod<TRequest, TResponse> continuation)
+			{
+				return GuardAsync(context, () => base.UnaryServerHandler(request, context, continuation));
+			}
+
+			public override Task<TResponse> ClientStreamingServerHandler<TRequest, TResponse>(IAsyncStreamReader<TRequest> requestStream, ServerCallContext context, ClientStreamingServerMethod<TRequest, TResponse> continuation) where TRequest : class where TResponse : class
+			{
+				return GuardAsync(context, () => base.ClientStreamingServerHandler(requestStream, context, continuation));
+			}
+
+			public override Task ServerStreamingServerHandler<TRequest, TResponse>(TRequest request, IServerStreamWriter<TResponse> responseStream, ServerCallContext context, ServerStreamingServerMethod<TRequest, TResponse> continuation) where TRequest : class where TResponse : class
+			{
+				return GuardAsync(context, () => base.ServerStreamingServerHandler(request, responseStream, context, continuation));
+			}
+
+			public override Task DuplexStreamingServerHandler<TRequest, TResponse>(IAsyncStreamReader<TRequest> requestStream, IServerStreamWriter<TResponse> responseStream, ServerCallContext context, DuplexStreamingServerMethod<TRequest, TResponse> continuation) where TRequest : class where TResponse : class
+			{
+				return GuardAsync(context, () => base.DuplexStreamingServerHandler(requestStream, responseStream, context, continuation));
+			}
+
+			async Task<T> GuardAsync<T>(ServerCallContext context, Func<Task<T>> callFunc) where T : class
+			{
+				T result = null!;
+				await GuardAsync(context, async () => { result = await callFunc(); });
+				return result;
+			}
+
+			async Task GuardAsync(ServerCallContext context, Func<Task> callFunc)
+			{
+				HttpContext httpContext = context.GetHttpContext();
+
+				AgentId? agentId = AclService.GetAgentId(httpContext.User);
+				if (agentId != null)
+				{
+					using IDisposable? scope = _logger.BeginScope("Agent: {AgentId}, RemoteIP: {RemoteIP}, Method: {Method}", agentId.Value, httpContext.Connection.RemoteIpAddress, context.Method);
+					await GuardInnerAsync(context, callFunc);
+				}
+				else
+				{
+					using IDisposable? scope = _logger.BeginScope("RemoteIP: {RemoteIP}, Method: {Method}", httpContext.Connection.RemoteIpAddress, context.Method);
+					await GuardInnerAsync(context, callFunc);
+				}
+			}
+
+			async Task GuardInnerAsync(ServerCallContext context, Func<Task> callFunc)
+			{
+				try
+				{
+					await callFunc();
+				}
+				catch (StructuredRpcException ex)
+				{
+#pragma warning disable CA2254 // Template should be a static expression
+					_logger.LogError(ex, ex.Format, ex.Args);
+#pragma warning restore CA2254 // Template should be a static expression
+					throw;
+				}
+				catch (Exception ex)
+				{
+					if (context.CancellationToken.IsCancellationRequested)
+					{
+						_logger.LogInformation(ex, "Call to method {Method} was cancelled", context.Method);
+						throw;
+					}
+					else
+					{
+						_logger.LogError(ex, "Exception in call to {Method}", context.Method);
+						throw new RpcException(new Status(StatusCode.Internal, $"An exception was thrown on the server: {ex}"));
+					}
+				}
+			}
+		}
+
+		class BsonSerializationProvider : IBsonSerializationProvider
+		{
+			public IBsonSerializer? GetSerializer(Type type)
+			{
+				if (type == typeof(ContentHash))
+				{
+					return new ContentHashSerializer();
+				}
+				if (type == typeof(DateTimeOffset))
+				{
+					return new DateTimeOffsetStringSerializer();
+				}
+				return null;
+			}
+		}
+
+		class JsonObjectIdConverter : JsonConverter<ObjectId>
+		{
+			public override ObjectId Read(ref Utf8JsonReader reader, Type typeToConvert, JsonSerializerOptions options)
+			{
+				string? str = reader.GetString();
+				if (str == null)
+				{
+					throw new InvalidDataException("Unable to parse object id");
+				}
+				if (str.Length == 0)
+				{
+					return ObjectId.Empty;
+				}
+				return ObjectId.Parse(str);
+			}
+
+			public override void Write(Utf8JsonWriter writer, ObjectId objectId, JsonSerializerOptions options)
+			{
+				writer.WriteStringValue(objectId.ToString());
+			}
+		}
+
+		class JsonDateTimeConverter : JsonConverter<DateTime>
+		{
+			public override DateTime Read(ref Utf8JsonReader reader, Type typeToConvert, JsonSerializerOptions options)
+			{
+				Debug.Assert(typeToConvert == typeof(DateTime));
+
+				string? str = reader.GetString();
+				if (str == null)
+				{
+					throw new InvalidDataException("Unable to parse DateTime");
+				}
+				return DateTime.Parse(str, CultureInfo.CurrentCulture);
+			}
+
+			public override void Write(Utf8JsonWriter writer, DateTime dateTime, JsonSerializerOptions options)
+			{
+				writer.WriteStringValue(dateTime.ToUniversalTime().ToString("yyyy'-'MM'-'dd'T'HH':'mm':'ssZ", CultureInfo.CurrentCulture));
+			}
+		}
+
+		class JsonTimeSpanConverter : JsonConverter<TimeSpan>
+		{
+			public override TimeSpan Read(ref Utf8JsonReader reader, Type typeToConvert, JsonSerializerOptions options)
+			{
+				string? str = reader.GetString();
+				if (str == null)
+				{
+					throw new InvalidDataException("Unable to parse TimeSpan");
+				}
+				return TimeSpan.Parse(str, CultureInfo.CurrentCulture);
+			}
+
+			public override void Write(Utf8JsonWriter writer, TimeSpan timeSpan, JsonSerializerOptions options)
+			{
+				writer.WriteStringValue(timeSpan.ToString("c"));
+			}
+		}
+
+		class ObsoleteLoggingFilter : IActionFilter
+		{
+			public void OnActionExecuted(ActionExecutedContext context)
+			{
+			}
+
+			public void OnActionExecuting(ActionExecutingContext context)
+			{
+				ControllerActionDescriptor? actionDescriptor = context.ActionDescriptor as ControllerActionDescriptor;
+				if (actionDescriptor?.MethodInfo.GetCustomAttribute<ObsoleteAttribute>() != null
+					|| actionDescriptor?.MethodInfo?.DeclaringType?.GetCustomAttribute<ObsoleteAttribute>() != null)
+				{
+					ILogger? logger = context.HttpContext.RequestServices.GetService<ILogger<ObsoleteLoggingFilter>>();
+					logger?.LogDebug("Using obsolete endpoint: {RequestPath} (Source: {RemoteIp})", context.HttpContext.Request.Path, context.HttpContext.Connection.RemoteIpAddress);
+				}
+			}
+		}
+
+		static IObjectStore CreateObjectStore(IServiceProvider sp, StorageBackendOptions options)
+		{
+			switch (options.Type ?? StorageBackendType.FileSystem)
+			{
+				case StorageBackendType.FileSystem:
+					return sp.GetRequiredService<FileObjectStoreFactory>().CreateStore(DirectoryReference.Combine(ServerApp.DataDir, options.BaseDir ?? "Storage"));
+				case StorageBackendType.Aws:
+					return sp.GetRequiredService<AwsObjectStoreFactory>().CreateStore(options);
+				case StorageBackendType.Memory:
+					return new MemoryObjectStore();
+				default:
+					throw new NotImplementedException();
+			}
+		}
+
+		public Startup(IConfiguration configuration)
+		{
+			Configuration = configuration;
+		}
+
+		public IConfiguration Configuration { get; }
+
+		/// <summary>
+		/// Bind config to concrete subclasses of BaseTelemetryConfig
+		/// </summary>
+		/// <param name="settings">Settings object being updated</param>
+		/// <param name="telemetrySection">Telemetry config section</param>
+		/// <exception cref="ConfigurationException"></exception>
+		private static void BindTelemetrySettings(ServerSettings settings, IConfiguration telemetrySection)
+		{
+			List<BaseTelemetryConfig> telemetryConfigs = new();
+			foreach (IConfigurationSection child in telemetrySection.GetChildren())
+			{
+				string? typeStr = child.GetValue<string>("Type");
+				if (typeStr == null || !Enum.TryParse(typeStr, true, out TelemetrySinkType sinkType))
+				{
+					throw new ConfigurationException($"Unable to parse sink type '{typeStr}'");
+				}
+
+				switch (sinkType)
+				{
+					case TelemetrySinkType.Epic:
+						EpicTelemetryConfig epic = new();
+						child.Bind(epic);
+						telemetryConfigs.Add(epic);
+						break;
+
+					case TelemetrySinkType.ClickHouse:
+						ClickHouseTelemetryConfig clickHouse = new();
+						child.Bind(clickHouse);
+						telemetryConfigs.Add(clickHouse);
+						break;
+
+					case TelemetrySinkType.Mongo:
+						MongoTelemetryConfig mongo = new();
+						child.Bind(mongo);
+						telemetryConfigs.Add(mongo);
+						break;
+
+					case TelemetrySinkType.None:
+					default:
+						break;
+				}
+			}
+
+			settings.Telemetry = telemetryConfigs;
+		}
+
+		public static void BindServerSettings(IConfiguration configuration, ServerSettings settings)
+		{
+			IConfigurationSection hordeSection = configuration.GetSection("Horde");
+			hordeSection.Bind(settings);
+			BindTelemetrySettings(settings, hordeSection.GetSection("Telemetry"));
+			settings.Validate();
+		}
+
+		/// <summary>
+		/// Converter to/from Redis values
+		/// </summary>
+		sealed class AgentIdRedisConverter : IRedisConverter<AgentId>
+		{
+			/// <inheritdoc/>
+			public AgentId FromRedisValue(RedisValue value) => new AgentId((string)value!);
+
+			/// <inheritdoc/>
+			public RedisValue ToRedisValue(AgentId value) => value.ToString();
+		}
+
+		/// <summary>
+		/// Configure custom types for the public API documentation
+		/// </summary>
+		class SwaggerSchemaFilter : ISchemaFilter
+		{
+			public void Apply(OpenApiSchema schema, SchemaFilterContext context)
+			{
+				if (context.Type == typeof(AgentId))
+				{
+				}
+				if (context.Type == typeof(Utf8String) || context.Type.GetCustomAttribute<JsonSchemaStringAttribute>() != null)
+				{
+					schema.Type = "string";
+					schema.Properties.Clear();
+				}
+			}
+		}
+
+		// This method gets called *multiple times* by the runtime. Use this method to add services to the container.
+		public void ConfigureServices(IServiceCollection services)
+		{
+			// IOptionsMonitor pattern for live updating of configuration settings
+			services.Configure<ServerSettings>(x => BindServerSettings(Configuration, x));
+
+			// We may upload a large number of references with posts to storage endpoints. Increase the max number of form values to allow this (the default is 1024).
+			services.Configure<FormOptions>(options =>
+			{
+				options.ValueCountLimit = 100_000;
+			});
+
+			// Bind the settings again for local variable access in this method
+			ServerSettings settings = new();
+			BindServerSettings(Configuration, settings);
+
+			OpenTelemetryHelper.Configure(services, settings.OpenTelemetry);
+
+			if (settings.GlobalThreadPoolMinSize != null)
+			{
+				// Min thread pool size is set to combat timeouts seen with the Redis client.
+				// See comments for <see cref="ServerSettings.GlobalThreadPoolMinSize" /> and
+				// [REDACTED_PROFILE_URL]/StackExchange.Redis/issues/1680
+				int min = settings.GlobalThreadPoolMinSize.Value;
+				ThreadPool.SetMinThreads(min, min);
+			}
+
+			RedisSerializer.RegisterConverter<AgentId, AgentIdRedisConverter>();
+
+#pragma warning disable CA2000 // Dispose objects before losing scope
+			RedisService redisService;
+			using (Serilog.Extensions.Logging.SerilogLoggerFactory loggerFactory = new Serilog.Extensions.Logging.SerilogLoggerFactory(Serilog.Log.Logger))
+			{
+				redisService = new RedisService(Options.Create(settings), loggerFactory.CreateLogger<RedisService>());
+			}
+#pragma warning restore CA2000 // Dispose objects before losing scope
+			services.AddSingleton<RedisService>(sp => redisService);
+			services.AddDataProtection().PersistKeysToStackExchangeRedis(() => redisService.GetDatabase(), "aspnet-data-protection");
+			
+			services.AddResponseCompression(options =>
+			{
+				options.EnableForHttps = true;
+				options.Providers.Add<GzipCompressionProvider>();
+				options.Providers.Add<BrotliCompressionProvider>();
+				options.Providers.Add<ZstdCompressionProvider>();
+			});
+			
+			if (settings.CorsEnabled)
+			{
+				services.AddCors(options =>
+				{
+					options.AddPolicy("CorsPolicy",
+						builder => builder.WithOrigins(settings.CorsOrigin.Split(";"))
+						.AllowAnyMethod()
+						.AllowAnyHeader()
+						.AllowCredentials());
+				});
+			}
+
+			services.AddGrpc(options =>
+			{
+				options.EnableDetailedErrors = true;
+				options.MaxReceiveMessageSize = 200 * 1024 * 1024; // 100 MB (packaged builds of Horde agent can be large) 
+				options.Interceptors.Add(typeof(GrpcExceptionInterceptor));
+			});
+			services.AddGrpcReflection();
+
+			services.AddSingleton<IAccountCollection, AccountCollection>();
+			services.AddSingleton<IAgentCollection, AgentCollection>();
+			services.AddSingleton<IArtifactCollection, ArtifactCollection>();
+			services.AddSingleton<IArtifactCollectionV1, ArtifactCollectionV1>();
+			services.AddSingleton<IGraphCollection, GraphCollection>();
+			services.AddSingleton<IIssueCollection, IssueCollection>();
+			services.AddSingleton<IJobCollection, JobCollection>();
+			services.AddSingleton<IJobStepRefCollection, JobStepRefCollection>();
+			services.AddSingleton<IJobTimingCollection, JobTimingCollection>();
+			services.AddSingleton<ILeaseCollection, LeaseCollection>();
+			services.AddSingleton<ILogCollection, LogCollection>();
+			services.AddSingleton<INotificationTriggerCollection, NotificationTriggerCollection>();
+			services.AddSingleton<IPoolCollection, PoolCollection>();
+			services.AddSingleton<IBisectTaskCollection, BisectTaskCollection>();
+			services.AddSingleton<IReplicatorCollection, ReplicatorCollection>();
+			services.AddSingleton<IServiceAccountCollection, ServiceAccountCollection>();
+			services.AddSingleton<ISessionCollection, SessionCollection>();
+			services.AddSingleton<ISubscriptionCollection, SubscriptionCollection>();
+			services.AddSingleton<IStreamCollection, StreamCollection>();
+			services.AddSingleton<ITemplateCollection, TemplateCollection>();
+			services.AddSingleton<ITestDataCollection, TestDataCollection>();
+			services.AddSingleton<IUtilizationDataCollection, UtilizationDataCollection>();
+			services.AddSingleton<ITemplateCollection, TemplateCollection>();
+			services.AddSingleton<IUgsMetadataCollection, UgsMetadataCollection>();
+			services.AddSingleton<IUserCollection, UserCollectionV2>();
+			services.AddSingleton<IDeviceCollection, DeviceCollection>();
+			services.AddSingleton<INoticeCollection, NoticeCollection>();
+			services.AddSingleton<IDashboardPreviewCollection, DashboardPreviewCollection>();
+
+			services.AddSingleton<IToolCollection, ToolCollection>();
+
+			services.AddSingleton<IConfigSource, InMemoryConfigSource>();
+			services.AddSingleton<IConfigSource, FileConfigSource>();
+			services.AddSingleton<IConfigSource, PerforceConfigSource>();
+
+			services.AddSingleton<ConfigService>();
+			services.AddSingleton<IOptionsFactory<GlobalConfig>>(sp => sp.GetRequiredService<ConfigService>());
+			services.AddSingleton<IOptionsChangeTokenSource<GlobalConfig>>(sp => sp.GetRequiredService<ConfigService>());
+
+			// Always run the hosted config service, regardless of the server role, so we receive updates on the latest config values.
+			services.AddHostedService(provider => provider.GetRequiredService<ConfigService>());
+
+			// Auditing
+			services.AddSingleton<IAuditLog<AgentId>>(sp => sp.GetRequiredService<IAuditLogFactory<AgentId>>().Create("Agents.Log", "AgentId"));
+
+			services.AddSingleton(typeof(IAuditLogFactory<>), typeof(AuditLogFactory<>));
+			services.AddSingleton(typeof(ISingletonDocument<>), typeof(SingletonDocument<>));
+
+			services.AddSingleton<AwsAutoScalingLifecycleService>();
+			services.AddSingleton<FleetService>();
+			services.AddSingleton<IFleetManagerFactory, FleetManagerFactory>();
+			services.AddSingleton<IPoolSizeStrategyFactory, NoOpPoolSizeStrategyFactory>();
+			services.AddSingleton<IPoolSizeStrategyFactory, JobQueueStrategyFactory>();
+			services.AddSingleton<IPoolSizeStrategyFactory, LeaseUtilizationStrategyFactory>();
+
+			// Associate IFleetManager interface with the default implementation from config for convenience
+			// Though most fleet managers are created on a per-pool basis
+			services.AddSingleton<IFleetManager>(ctx => ctx.GetRequiredService<IFleetManagerFactory>().CreateFleetManager(FleetManagerType.Default));
+
+			// Run the tunnel service for all run modes
+			services.AddSingleton<TunnelService>();
+			services.AddHostedService<TunnelService>(sp => sp.GetRequiredService<TunnelService>());
+
+			// Runs the agent relay service for all run modes to notify long-polling requests
+			services.AddSingleton<AgentRelayService>();
+			services.AddHostedService(provider => provider.GetRequiredService<AgentRelayService>());
+
+			services.AddSingleton<AclService>();
+			services.AddSingleton<AgentService>();
+			services.AddSingleton(provider => new Lazy<AgentService>(provider.GetRequiredService<AgentService>));
+			services.AddHostedService<ArtifactExpirationService>();
+			services.AddSingleton<ConsistencyService>();
+			services.AddSingleton<RequestTrackerService>();
+			services.AddSingleton<ComputeService>();
+			services.AddSingleton<MongoCommandTracer>();
+			services.AddSingleton<MongoService>();
+			services.AddSingleton<GlobalsService>();
+			services.AddSingleton<ICommitService, CommitService>();
+			services.AddSingleton<IClock, Clock>();
+			services.AddSingleton<IDowntimeService, DowntimeService>();
+			services.AddSingleton<IssueService>();
+			services.AddSingleton<JobService>();
+			services.AddSingleton<LifetimeService>();
+			services.AddHostedService(provider => provider.GetRequiredService<LifetimeService>());
+			services.AddSingleton(typeof(IHealthMonitor<>), typeof(HealthMonitor<>));
+			services.AddSingleton<EnrollmentService>();
+			services.AddSingleton<ServerStatusService>();
+			services.AddHostedService(provider => provider.GetRequiredService<ServerStatusService>());
+			services.AddSingleton<LogTailService>();
+			services.AddHostedService(provider => provider.GetRequiredService<LogTailService>()); // NB: Runs even on workers, to receive tail notifications.
+			services.AddSingleton<INotificationService, NotificationService>();
+			services.AddSingleton<UnsyncCache>();
+
+			if (settings.Commits.ReplicateMetadata)
+			{
+				services.AddSingleton<PerforceServiceCache>();
+				services.AddSingleton<IPerforceService>(sp => sp.GetRequiredService<PerforceServiceCache>());
+			}
+			else
+			{
+				services.AddSingleton<PerforceService>();
+				services.AddSingleton<IPerforceService>(sp => sp.GetRequiredService<PerforceService>());
+			}
+			services.AddSingleton<PerforceReplicator>();
+
+			services.AddSingleton<PerforceLoadBalancer>();
+			services.AddSingleton<PoolService>();
+			services.AddSingleton<ReplicationService>();
+			services.AddSingleton<ScheduleService>();
+
+			if (settings.SlackToken != null)
+			{
+				services.AddSingleton<SlackNotificationSink>();
+				services.AddSingleton<IAvatarService, SlackNotificationSink>(sp => sp.GetRequiredService<SlackNotificationSink>());
+				services.AddSingleton<INotificationSink, SlackNotificationSink>(sp => sp.GetRequiredService<SlackNotificationSink>());
+			}
+			else
+			{
+				services.AddSingleton<IAvatarService, NullAvatarService>();
+			}
+
+			services.AddScoped<OAuthControllerFilter>();
+
+			services.AddSingleton<DeviceService>();
+			services.AddSingleton<NoticeService>();
+			services.AddSingleton<StorageService>();
+			services.AddSingleton<IStorageService>(sp => sp.GetRequiredService<StorageService>());
+			services.AddScoped(sp => sp.GetRequiredService<StorageService>().CreateStorageClientFactory(sp.GetRequiredService<IOptionsSnapshot<GlobalConfig>>().Value));
+			services.AddSingleton<IBlockCache>(sp => CreateBlockCache(sp));
+			services.AddSingleton<TestDataService>();
+			services.AddSingleton<BundleCache>();
+			services.AddSingleton<StorageBackendCache>(CreateStorageBackendCache);
+
+			if (settings.JiraUrl != null)
+			{
+				services.AddSingleton<IExternalIssueService, JiraService>();
+			}
+			else
+			{
+				services.AddSingleton<IExternalIssueService, DefaultExternalIssueService>();
+			}
+
+			// DDC
+			services.AddScoped<IRequestHelper, RequestHelper>();
+			services.AddScoped<IBlobService, BlobService>();
+			services.AddScoped<IRefService, RefService>();
+			services.AddScoped<IReferenceResolver, ReferenceResolver>();
+			services.AddScoped<IContentIdStore, ContentIdStore>();
+			services.AddSingleton<BufferedPayloadFactory>();
+			services.AddSingleton<NginxRedirectHelper>();
+			services.AddSingleton<FormatResolver>();
+			services.AddSingleton<CompressedBufferUtils>();
+
+			// Secret providers
+			services.AddSingleton<SecretCollectionInternal>();
+			services.AddScoped<ISecretCollection, SecretCollection>();
+			services.AddSingleton<ISecretProvider, AwsParameterStoreSecretProvider>();
+
+			// Storage providers
+			services.AddSingleton<IObjectStoreFactory, ObjectStoreFactory>();
+			services.AddSingleton<AwsObjectStoreFactory>();
+			services.AddSingleton<FileObjectStoreFactory>();
+			services.AddSingleton(sp => CreateObjectStore(sp, settings.LogStorage).ForType<PersistentLogStorage>());
+			services.AddSingleton(sp => CreateObjectStore(sp, settings.ArtifactStorage).ForType<ArtifactCollectionV1>());
+
+			if (settings.WithAws)
+			{
+				AWSOptions awsOptions = Configuration.GetAWSOptions();
+				services.AddDefaultAWSOptions(awsOptions);
+				if (awsOptions.Region == null && Environment.GetEnvironmentVariable("AWS_REGION") == null)
+				{
+					awsOptions.Region = RegionEndpoint.USEast1;
+				}
+
+				services.AddAWSService<IAmazonCloudWatch>();
+				services.AddAWSService<IAmazonAutoScaling>();
+				services.AddAWSService<IAmazonSQS>();
+				services.AddAWSService<IAmazonEC2>();
+
+				services.AddSingleton<AwsCloudWatchMetricExporter>();
+
+				services.AddSingleton<IPoolSizeStrategyFactory, LeaseUtilizationAwsMetricStrategyFactory>();
+				services.AddSingleton<IPoolSizeStrategyFactory, ComputeQueueAwsMetricStrategyFactory>();
+			}
+
+			ConfigureLogStorage(services);
+
+			AuthenticationBuilder authBuilder = services.AddAuthentication(options =>
+				{
+					switch (settings.AuthMethod)
+					{
+						case AuthMethod.Anonymous:
+							options.DefaultAuthenticateScheme = AnonymousAuthHandler.AuthenticationScheme;
+							options.DefaultSignInScheme = AnonymousAuthHandler.AuthenticationScheme;
+							options.DefaultChallengeScheme = AnonymousAuthHandler.AuthenticationScheme;
+							break;
+
+						case AuthMethod.Okta:
+							// If an authentication cookie is present, use it to get authentication information
+							options.DefaultAuthenticateScheme = CookieAuthenticationDefaults.AuthenticationScheme;
+							options.DefaultSignInScheme = CookieAuthenticationDefaults.AuthenticationScheme;
+
+							// If authentication is required, and no cookie is present, use OIDC to sign in
+							options.DefaultChallengeScheme = OktaAuthHandler.AuthenticationScheme;
+							break;
+
+						case AuthMethod.OpenIdConnect:
+							// If an authentication cookie is present, use it to get authentication information
+							options.DefaultAuthenticateScheme = CookieAuthenticationDefaults.AuthenticationScheme;
+							options.DefaultSignInScheme = CookieAuthenticationDefaults.AuthenticationScheme;
+
+							// If authentication is required, and no cookie is present, use OIDC to sign in
+							options.DefaultChallengeScheme = OpenIdConnectDefaults.AuthenticationScheme;
+							break;
+
+						case AuthMethod.Horde:
+							options.DefaultAuthenticateScheme = CookieAuthenticationDefaults.AuthenticationScheme;
+							options.DefaultSignInScheme = CookieAuthenticationDefaults.AuthenticationScheme;
+							options.DefaultChallengeScheme = CookieAuthenticationDefaults.AuthenticationScheme;
+							break;
+
+						default:
+							throw new ArgumentException($"Invalid auth method {settings.AuthMethod}");
+					}
+				});
+
+			List<string> schemes = new List<string>();
+
+			authBuilder.AddCookie(options =>
+			{
+				if (settings.AuthMethod == AuthMethod.Horde)
+				{
+					options.Cookie.Name = CookieAuthenticationDefaults.AuthenticationScheme;
+					options.LoginPath = "/account/login/horde";
+					options.LogoutPath = "/";
+				}
+
+				options.Events.OnValidatePrincipal = context =>
+				{
+					if (!String.Equals(context.Principal?.FindFirst(HordeClaimTypes.Version)?.Value, HordeClaimTypes.CurrentVersion, StringComparison.Ordinal))
+					{
+						context.RejectPrincipal();
+					}
+					if (context.Principal?.FindFirst(HordeClaimTypes.UserId) == null)
+					{
+						context.RejectPrincipal();
+					}
+					return Task.CompletedTask;
+				};
+
+				options.Events.OnRedirectToAccessDenied = context =>
+				{
+					context.Response.StatusCode = StatusCodes.Status403Forbidden;
+					return context.Response.CompleteAsync();
+				};
+			});
+			schemes.Add(CookieAuthenticationDefaults.AuthenticationScheme);
+
+			authBuilder.AddServiceAccounts(options => { });
+			schemes.Add(ServiceAccountAuthHandler.AuthenticationScheme);
+
+			switch (settings.AuthMethod)
+			{
+				case AuthMethod.Anonymous:
+					authBuilder.AddAnonymous(options => { });
+					schemes.Add(AnonymousAuthHandler.AuthenticationScheme);
+					break;
+
+				case AuthMethod.Okta:
+					authBuilder.AddOkta(settings, OktaAuthHandler.AuthenticationScheme, OpenIdConnectDefaults.DisplayName, options =>
+						{
+							options.Authority = settings.OidcAuthority;
+							options.ClientId = settings.OidcClientId;
+							options.Scope.Remove("groups");
+
+							if (!String.IsNullOrEmpty(settings.OidcSigninRedirect))
+							{
+								options.Events = new OpenIdConnectEvents
+								{
+									OnRedirectToIdentityProvider = async redirectContext =>
+									{
+										redirectContext.ProtocolMessage.RedirectUri = settings.OidcSigninRedirect;
+										await Task.CompletedTask;
+									}
+								};
+							}
+						});
+					schemes.Add(OktaAuthHandler.AuthenticationScheme);
+					break;
+
+				case AuthMethod.OpenIdConnect:
+					authBuilder.AddHordeOpenId(settings, OpenIdConnectDefaults.AuthenticationScheme, OpenIdConnectDefaults.DisplayName, options =>
+						{
+							options.Authority = settings.OidcAuthority;
+							options.ClientId = settings.OidcClientId;
+							options.ClientSecret = settings.OidcClientSecret;
+							foreach (string scope in settings.OidcRequestedScopes)
+							{
+								options.Scope.Add(scope);
+							}
+
+							if (!String.IsNullOrEmpty(settings.OidcSigninRedirect))
+							{
+								options.Events = new OpenIdConnectEvents
+								{
+									OnRedirectToIdentityProvider = async redirectContext =>
+									{
+										redirectContext.ProtocolMessage.RedirectUri = settings.OidcSigninRedirect;
+										await Task.CompletedTask;
+									}
+								};
+							}
+						});
+					schemes.Add(OpenIdConnectDefaults.AuthenticationScheme);
+					break;
+
+				case AuthMethod.Horde:
+					authBuilder.AddHordeOpenId(settings, OpenIdConnectDefaults.AuthenticationScheme, OpenIdConnectDefaults.DisplayName, options =>
+					{
+						options.Authority = new Uri(settings.ServerUrl, "api/v1/oauth2").ToString();
+						options.ClientId = "default";
+						if (settings.HttpsPort == 0)
+						{
+							options.RequireHttpsMetadata = false;
+						}
+						foreach (string scope in settings.OidcRequestedScopes)
+						{
+							options.Scope.Add(scope);
+						}
+					});
+					schemes.Add(OpenIdConnectDefaults.AuthenticationScheme);
+					break;
+
+				default:
+					throw new ArgumentException($"Invalid auth method {settings.AuthMethod}");
+			}
+
+			authBuilder.AddScheme<JwtBearerOptions, JwtAuthHandler>(JwtAuthHandler.AuthenticationScheme, options => { });
+			schemes.Add(JwtAuthHandler.AuthenticationScheme);
+
+			if (!String.IsNullOrEmpty(settings.OidcAuthority) && !String.IsNullOrEmpty(settings.OidcAudience))
+			{
+				ExternalJwtAuthHandler hordeJwtBearer = new(settings);
+				hordeJwtBearer.AddHordeJwtBearerConfiguration(authBuilder);
+				schemes.Add(ExternalJwtAuthHandler.AuthenticationScheme);
+			}
+
+			services.AddAuthorization(options =>
+				{
+					options.DefaultPolicy = new AuthorizationPolicyBuilder(schemes.ToArray())
+						.RequireAuthenticatedUser()
+						.Build();
+				});
+
+			services.AddSingleton<TelemetryManager>();
+			services.AddSingleton<ITelemetrySink>(sp => sp.GetRequiredService<TelemetryManager>());
+			services.AddHostedService(sp => sp.GetRequiredService<TelemetryManager>());
+			services.AddSingleton<MongoTelemetrySink>();
+			services.AddHostedService(sp => sp.GetRequiredService<MongoTelemetrySink>());
+			services.AddSingleton<MetricTelemetrySink>();
+
+			services.AddSingleton<MetricCollection>();
+			services.AddHostedService(sp => sp.GetRequiredService<MetricCollection>());
+			services.AddSingleton<IMetricCollection, MetricCollection>(sp => sp.GetRequiredService<MetricCollection>());
+
+			services.AddHttpClient(EpicTelemetrySink.HttpClientName, client => { });
+			services.AddHttpClient(ClickHouseTelemetrySink.HttpClientName, client => { });
+
+			// Create the agent telemetry collection, and register the hosted service so we can flush from any server.
+			services.AddSingleton<AgentTelemetryCollection>();
+			services.AddSingleton<IAgentTelemetryCollection>(sp => sp.GetRequiredService<AgentTelemetryCollection>());
+			services.AddHostedService(provider => provider.GetRequiredService<AgentTelemetryCollection>());
+
+			// Hosted service that needs to run no matter the run mode of the process (server vs worker)
+			services.AddHostedService(provider => (DowntimeService)provider.GetRequiredService<IDowntimeService>());
+
+			// Always run agent service too; need to be able to listen to Redis for events on any server.
+			services.AddHostedService(provider => provider.GetRequiredService<AgentService>());
+
+			// Notifications can be triggered from any instance, so always make sure we're ticking the background task.
+			services.AddHostedService(provider => (NotificationService)provider.GetRequiredService<INotificationService>());
+
+			if (settings.IsRunModeActive(RunMode.Worker) && !settings.MongoReadOnlyMode)
+			{
+				services.AddHostedService<AgentReportService>();
+				services.AddHostedService<BisectService>();
+				services.AddHostedService(provider => provider.GetRequiredService<FleetService>());
+				services.AddHostedService(provider => provider.GetRequiredService<ConsistencyService>());
+				services.AddHostedService(provider => provider.GetRequiredService<IssueService>());
+				services.AddHostedService<IssueReportService>();
+				services.AddHostedService<IssueTagService>();
+				services.AddHostedService<JobExpirationService>();
+				services.AddHostedService<MetricService>();
+				services.AddHostedService(provider => provider.GetRequiredService<PerforceLoadBalancer>());
+				services.AddHostedService<PoolUpdateService>();
+				services.AddHostedService<UtilizationDataService>();
+				services.AddHostedService(provider => provider.GetRequiredService<DeviceService>());
+				services.AddHostedService<DeviceReportService>();
+				services.AddHostedService(provider => provider.GetRequiredService<TestDataService>());
+				services.AddHostedService(provider => provider.GetRequiredService<ComputeService>());
+				services.AddHostedService(provider => provider.GetRequiredService<EnrollmentService>());
+				services.AddHostedService(provider => provider.GetRequiredService<StorageService>());
+
+				if (settings.Commits.ReplicateMetadata)
+				{
+					services.AddHostedService(provider => provider.GetRequiredService<PerforceServiceCache>());
+				}
+
+				if (settings.Commits.ReplicateContent)
+				{
+					services.AddHostedService(provider => provider.GetRequiredService<ReplicationService>());
+				}
+
+				if (!settings.DisableSchedules)
+				{
+					services.AddHostedService(provider => provider.GetRequiredService<ScheduleService>());
+				}
+
+				if (settings.SlackToken != null)
+				{
+					services.AddHostedService(provider => provider.GetRequiredService<SlackNotificationSink>());
+				}
+
+				if (settings.WithAws)
+				{
+					services.AddHostedService(provider => provider.GetRequiredService<AwsAutoScalingLifecycleService>());
+					services.AddHostedService(provider => provider.GetRequiredService<AwsCloudWatchMetricExporter>());
+				}
+			}
+
+			services.AddHostedService(provider => provider.GetRequiredService<IExternalIssueService>());
+
+			// Task sources. Order of registration is important here; it dictates the priority in which sources are served.
+			services.AddSingleton<JobTaskSource>();
+
+			if (!settings.MongoReadOnlyMode)
+			{
+				services.AddHostedService<JobTaskSource>(provider => provider.GetRequiredService<JobTaskSource>());
+				services.AddSingleton<ConformTaskSource>();
+				services.AddHostedService<ConformTaskSource>(provider => provider.GetRequiredService<ConformTaskSource>());
+				services.AddSingleton<ComputeTaskSource>();
+
+				services.AddSingleton<ITaskSource, UpgradeTaskSource>();
+				services.AddSingleton<ITaskSource, ShutdownTaskSource>();
+				services.AddSingleton<ITaskSource, RestartTaskSource>();
+				services.AddSingleton<ITaskSource, ConformTaskSource>(provider => provider.GetRequiredService<ConformTaskSource>());
+				services.AddSingleton<ITaskSource, JobTaskSource>(provider => provider.GetRequiredService<JobTaskSource>());
+				services.AddSingleton<ITaskSource, ComputeTaskSource>(provider => provider.GetRequiredService<ComputeTaskSource>());
+			}
+
+			// Allow longer to shutdown so we can debug missing cancellation tokens
+			services.Configure<HostOptions>(options =>
+			{
+				options.ShutdownTimeout = TimeSpan.FromSeconds(30.0);
+			});
+
+			// Allow forwarded headers
+			services.Configure<ForwardedHeadersOptions>(options =>
+			{
+				options.ForwardedHeaders = ForwardedHeaders.XForwardedFor | ForwardedHeaders.XForwardedProto;
+				options.KnownProxies.Clear();
+				options.KnownNetworks.Clear();
+			});
+
+			services.AddMvc().AddJsonOptions(options => ConfigureJsonSerializer(options.JsonSerializerOptions));
+			services.AddControllersWithViews(options => options.Filters.Add(new ObsoleteLoggingFilter()))
+				.AddRazorRuntimeCompilation();
+
+			services.AddControllers(options =>
+			{
+				options.InputFormatters.Add(new CbInputFormatter());
+				options.OutputFormatters.Add(new CbOutputFormatter());
+				options.OutputFormatters.Insert(0, new CbPreferredOutputFormatter());
+				options.OutputFormatters.Add(new RawOutputFormatter(NullLogger.Instance));
+				options.FormatterMappings.SetMediaTypeMappingForFormat("raw", MediaTypeNames.Application.Octet);
+				options.FormatterMappings.SetMediaTypeMappingForFormat("uecb", CustomMediaTypeNames.UnrealCompactBinary);
+				options.FormatterMappings.SetMediaTypeMappingForFormat("uecbpkg", CustomMediaTypeNames.UnrealCompactBinaryPackage);
+			}).ConfigureApiBehaviorOptions(options =>
+			{
+				options.InvalidModelStateResponseFactory = context =>
+				{
+					BadRequestObjectResult result = new BadRequestObjectResult(context.ModelState);
+					// always return errors as json objects
+					// we could allow more types here, but we do not want raw for instance
+					result.ContentTypes.Add(MediaTypeNames.Application.Json);
+
+					return result;
+				};
+			});
+
+			services.AddSwaggerGen(config =>
+			{
+				config.SchemaGeneratorOptions.SchemaFilters.Add(new SwaggerSchemaFilter());
+				config.SwaggerDoc("v1", new OpenApiInfo { Title = "Horde Server API", Version = "v1" });
+				config.IncludeXmlComments(Path.Combine(AppContext.BaseDirectory, $"{Assembly.GetExecutingAssembly().GetName().Name}.xml"));
+			});
+
+			services.Configure<ApiBehaviorOptions>(options =>
+			{
+				options.InvalidModelStateResponseFactory = context =>
+				{
+					foreach (KeyValuePair<string, ModelStateEntry> pair in context.ModelState)
+					{
+						ModelError? error = pair.Value.Errors.FirstOrDefault();
+						if (error != null)
+						{
+							string message = error.ErrorMessage;
+							if (String.IsNullOrEmpty(message))
+							{
+								message = error.Exception?.Message ?? "Invalid error object";
+							}
+							return new BadRequestObjectResult(EpicGames.Core.LogEvent.Create(LogLevel.Error, KnownLogEvents.None, error.Exception, "Invalid value for {Name}: {Message}", pair.Key, message));
+						}
+					}
+					return new BadRequestObjectResult(context.ModelState);
+				};
+			});
+
+			DirectoryReference dashboardDir = DirectoryReference.Combine(ServerApp.AppDir, "DashboardApp");
+			if (DirectoryReference.Exists(dashboardDir))
+			{
+				services.AddSpaStaticFiles(config => { config.RootPath = "DashboardApp"; });
+			}
+
+			ConfigureMongoDbClient();
+			ConfigureFormatters();
+
+			OnAddHealthChecks(services);
+		}
+
+		static BlockCache CreateBlockCache(IServiceProvider serviceProvider)
+		{
+			ServerSettings serverSettings = serviceProvider.GetRequiredService<IOptions<ServerSettings>>().Value;
+			DirectoryReference cacheDir = DirectoryReference.Combine(ServerApp.DataDir, String.IsNullOrEmpty(serverSettings.BlockCacheDir) ? "BlockCache" : serverSettings.BlockCacheDir);
+			return BlockCache.Create(cacheDir, (int)(serverSettings.BlockCacheSizeBytes / (1024 * 1024 * 1024)));
+		}
+
+		static StorageBackendCache CreateStorageBackendCache(IServiceProvider serviceProvider)
+		{
+			ServerSettings serverSettings = serviceProvider.GetRequiredService<IOptions<ServerSettings>>().Value;
+			DirectoryReference cacheDir = DirectoryReference.Combine(ServerApp.DataDir, String.IsNullOrEmpty(serverSettings.BundleCacheDir) ? "Cache" : serverSettings.BundleCacheDir);
+			return new StorageBackendCache(cacheDir, serverSettings.BundleCacheSizeBytes, serviceProvider.GetRequiredService<ILogger<StorageBackendCache>>());
+		}
+
+		public static void ConfigureFormatters()
+		{
+			LogValueFormatter.RegisterTypeAnnotation<AgentId>("AgentId");
+			LogValueFormatter.RegisterTypeAnnotation<JobId>("JobId");
+			LogValueFormatter.RegisterTypeAnnotation<LeaseId>("LeaseId");
+			LogValueFormatter.RegisterTypeAnnotation<LogId>("LogId");
+			LogValueFormatter.RegisterTypeAnnotation<UserId>("UserId");
+		}
+
+		public static JsonSerializerOptions JsonSerializerOptions { get; } = CreateJsonSerializerOptions();
+
+		static JsonSerializerOptions CreateJsonSerializerOptions()
+		{
+			JsonSerializerOptions options = new JsonSerializerOptions();
+			ConfigureJsonSerializer(options);
+			options.MakeReadOnly(true);
+			return options;
+		}
+
+		public static void ConfigureJsonSerializer(JsonSerializerOptions options)
+		{
+			HordeHttpClient.ConfigureJsonSerializer(options);
+			options.Converters.Add(new JsonObjectIdConverter());
+			options.Converters.Add(new JsonKnownTypesConverterFactory());
+			options.Converters.Add(new ObjectIdJsonConverterFactory());
+			options.Converters.Add(new SubResourceIdJsonConverterFactory());
+			options.Converters.Add(new JsonDateTimeConverter());
+			options.Converters.Add(new JsonTimeSpanConverter());
+		}
+
+		private static void ConfigureLogStorage(IServiceCollection services)
+		{
+			services.AddSingleton<PersistentLogStorage>();
+
+			services.AddSingleton<ILogStorage>(provider =>
+			{
+				ILogStorage storage = provider.GetRequiredService<PersistentLogStorage>();
+				storage = new SequencedLogStorage(storage);
+				storage = new LocalLogStorage(50, storage);
+				return storage;
+			});
+		}
+
+		public sealed class BlobLocatorBsonSerializer : SerializerBase<BlobLocator>
+		{
+			/// <inheritdoc/>
+			public override BlobLocator Deserialize(BsonDeserializationContext context, BsonDeserializationArgs args)
+			{
+				return new BlobLocator(context.Reader.ReadString());
+			}
+
+			/// <inheritdoc/>
+			public override void Serialize(BsonSerializationContext context, BsonSerializationArgs args, BlobLocator value)
+			{
+				context.Writer.WriteString(value.ToString());
+			}
+		}
+
+		public sealed class RefNameBsonSerializer : SerializerBase<RefName>
+		{
+			/// <inheritdoc/>
+			public override RefName Deserialize(BsonDeserializationContext context, BsonDeserializationArgs args)
+			{
+				return new RefName(context.Reader.ReadString());
+			}
+
+			/// <inheritdoc/>
+			public override void Serialize(BsonSerializationContext context, BsonSerializationArgs args, RefName value)
+			{
+				context.Writer.WriteString(value.ToString());
+			}
+		}
+
+		public sealed class IoHashBsonSerializer : SerializerBase<IoHash>
+		{
+			/// <inheritdoc/>
+			public override IoHash Deserialize(BsonDeserializationContext context, BsonDeserializationArgs args)
+			{
+				return IoHash.Parse(context.Reader.ReadString());
+			}
+
+			/// <inheritdoc/>
+			public override void Serialize(BsonSerializationContext context, BsonSerializationArgs args, IoHash value)
+			{
+				context.Writer.WriteString(value.ToString());
+			}
+		}
+
+		public sealed class NamespaceIdBsonSerializer : SerializerBase<NamespaceId>
+		{
+			/// <inheritdoc/>
+			public override NamespaceId Deserialize(BsonDeserializationContext context, BsonDeserializationArgs args)
+			{
+				return new NamespaceId(context.Reader.ReadString());
+			}
+
+			/// <inheritdoc/>
+			public override void Serialize(BsonSerializationContext context, BsonSerializationArgs args, NamespaceId value)
+			{
+				context.Writer.WriteString(value.ToString());
+			}
+		}
+
+		public sealed class AgentIdBsonSerializer : SerializerBase<AgentId>
+		{
+			/// <inheritdoc/>
+			public override AgentId Deserialize(BsonDeserializationContext context, BsonDeserializationArgs args)
+			{
+				string argument;
+				if (context.Reader.CurrentBsonType == MongoDB.Bson.BsonType.ObjectId)
+				{
+					argument = context.Reader.ReadObjectId().ToString();
+				}
+				else
+				{
+					argument = context.Reader.ReadString();
+				}
+				return new AgentId(argument);
+			}
+
+			/// <inheritdoc/>
+			public override void Serialize(BsonSerializationContext context, BsonSerializationArgs args, AgentId value)
+			{
+				context.Writer.WriteString(value.ToString());
+			}
+		}
+
+		public sealed class AclActionBsonSerializer : SerializerBase<AclAction>
+		{
+			/// <inheritdoc/>
+			public override AclAction Deserialize(BsonDeserializationContext context, BsonDeserializationArgs args) => new AclAction(context.Reader.ReadString());
+
+			/// <inheritdoc/>
+			public override void Serialize(BsonSerializationContext context, BsonSerializationArgs args, AclAction value) => context.Writer.WriteString(value.Name);
+		}
+
+		public sealed class AclScopeNameBsonSerializer : SerializerBase<AclScopeName>
+		{
+			/// <inheritdoc/>
+			public override AclScopeName Deserialize(BsonDeserializationContext context, BsonDeserializationArgs args) => new AclScopeName(context.Reader.ReadString());
+
+			/// <inheritdoc/>
+			public override void Serialize(BsonSerializationContext context, BsonSerializationArgs args, AclScopeName value) => context.Writer.WriteString(value.Text);
+		}
+
+		sealed class SubResourceIdBsonSerializer<TValue, TConverter> : SerializerBase<TValue> where TValue : struct where TConverter : SubResourceIdConverter<TValue>, new()
+		{
+			readonly TConverter _converter = new TConverter();
+
+			/// <inheritdoc/>
+			public override TValue Deserialize(BsonDeserializationContext context, BsonDeserializationArgs args) => _converter.FromSubResourceId(new SubResourceId((ushort)context.Reader.ReadInt32()));
+
+			/// <inheritdoc/>
+			public override void Serialize(BsonSerializationContext context, BsonSerializationArgs args, TValue value) => context.Writer.WriteInt32(_converter.ToSubResourceId(value).Value);
+		}
+
+		/// <summary>
+		/// Serializer for subresource ids
+		/// </summary>
+		public class SubResourceIdSerializer : IBsonSerializer<SubResourceId>
+		{
+			/// <inheritdoc/>
+			public Type ValueType => typeof(SubResourceId);
+
+			/// <inheritdoc/>
+			public object Deserialize(BsonDeserializationContext context, BsonDeserializationArgs args)
+			{
+				return new SubResourceId((ushort)context.Reader.ReadInt32());
+			}
+
+			/// <inheritdoc/>
+			public void Serialize(BsonSerializationContext context, BsonSerializationArgs args, object value)
+			{
+				SubResourceId id = (SubResourceId)value;
+				context.Writer.WriteInt32((int)id.Value);
+			}
+
+			/// <inheritdoc/>
+			public void Serialize(BsonSerializationContext context, BsonSerializationArgs args, SubResourceId id)
+			{
+				context.Writer.WriteInt32((int)id.Value);
+			}
+
+			/// <inheritdoc/>
+			SubResourceId IBsonSerializer<SubResourceId>.Deserialize(BsonDeserializationContext context, BsonDeserializationArgs vrgs)
+			{
+				return new SubResourceId((ushort)context.Reader.ReadInt32());
+			}
+		}
+
+		sealed class SubResourceIdBsonSerializationProvider : BsonSerializationProviderBase
+		{
+			/// <inheritdoc/>
+			public override IBsonSerializer? GetSerializer(Type type, IBsonSerializerRegistry serializerRegistry)
+			{
+				SubResourceIdConverterAttribute? attribute = type.GetCustomAttribute<SubResourceIdConverterAttribute>();
+				if (attribute == null)
+				{
+					return null;
+				}
+				return (IBsonSerializer?)Activator.CreateInstance(typeof(SubResourceIdBsonSerializer<,>).MakeGenericType(type, attribute.ConverterType));
+			}
+		}
+
+		sealed class CommitTagBsonSerializer : SerializerBase<CommitTag>
+		{
+			/// <inheritdoc/>
+			public override CommitTag Deserialize(BsonDeserializationContext context, BsonDeserializationArgs args)
+			{
+				return new CommitTag(context.Reader.ReadString());
+			}
+
+			/// <inheritdoc/>
+			public override void Serialize(BsonSerializationContext context, BsonSerializationArgs args, CommitTag value)
+			{
+				context.Writer.WriteString(value.ToString());
+			}
+		}
+
+		static int s_haveConfiguredMongoDb = 0;
+
+		public static void ConfigureMongoDbClient()
+		{
+			if (Interlocked.CompareExchange(ref s_haveConfiguredMongoDb, 1, 0) == 0)
+			{
+				// Ignore extra elements on deserialized documents
+				ConventionPack conventionPack = new ConventionPack();
+				conventionPack.Add(new IgnoreExtraElementsConvention(true));
+				conventionPack.Add(new EnumRepresentationConvention(BsonType.String));
+				ConventionRegistry.Register("Horde", conventionPack, type => true);
+
+				// Register the custom serializers
+				BsonSerializer.RegisterSerializer(new BlobLocatorBsonSerializer());
+				BsonSerializer.RegisterSerializer(new RefNameBsonSerializer());
+				BsonSerializer.RegisterSerializer(new IoHashBsonSerializer());
+				BsonSerializer.RegisterSerializer(new NamespaceIdBsonSerializer());
+				BsonSerializer.RegisterSerializer(new AgentIdBsonSerializer());
+				BsonSerializer.RegisterSerializer(new AclActionBsonSerializer());
+				BsonSerializer.RegisterSerializer(new AclScopeNameBsonSerializer());
+				BsonSerializer.RegisterSerializer(new ConditionSerializer());
+				BsonSerializer.RegisterSerializer(new SubResourceIdSerializer());
+				BsonSerializer.RegisterSerializer(new CommitTagBsonSerializer());
+				BsonSerializer.RegisterSerializationProvider(new BsonSerializationProvider());
+				BsonSerializer.RegisterSerializationProvider(new StringIdBsonSerializationProvider());
+				BsonSerializer.RegisterSerializationProvider(new BinaryIdBsonSerializationProvider());
+				BsonSerializer.RegisterSerializationProvider(new ObjectIdBsonSerializationProvider());
+				BsonSerializer.RegisterSerializationProvider(new SubResourceIdBsonSerializationProvider());
+			}
+		}
+
+		private static void OnAddHealthChecks(IServiceCollection services)
+		{
+			services.AddHealthChecks()
+				.AddCheck("self", () => HealthCheckResult.Healthy(), tags: new[] { "self" })
+				.AddCheck<RedisService>("redis")
+				.AddCheck<MongoService>("mongo");
+		}
+
+		// This method gets called by the runtime. Use this method to configure the HTTP request pipeline.
+		public static void Configure(IApplicationBuilder app, IWebHostEnvironment env, Microsoft.Extensions.Hosting.IHostApplicationLifetime lifetime, IOptions<ServerSettings> settings)
+		{
+			app.UseForwardedHeaders();
+			app.UseResponseCompression();
+
+			// Used for allowing auth cookies in combination with OpenID Connect auth (for example, Google Auth did not work with these unset)
+			app.UseCookiePolicy(new CookiePolicyOptions()
+			{
+				MinimumSameSitePolicy = SameSiteMode.None,
+				CheckConsentNeeded = _ => true
+			});
+
+			if (settings.Value.CorsEnabled)
+			{
+				app.UseCors("CorsPolicy");
+			}
+
+			// Enable middleware to serve generated Swagger as a JSON endpoint.
+			app.UseSwagger();
+
+			// Enable serilog request logging
+			app.UseSerilogRequestLogging(options =>
+			{
+				options.GetLevel = GetRequestLoggingLevel;
+				options.EnrichDiagnosticContext = (diagnosticContext, httpContext) =>
+				{
+					diagnosticContext.Set("RemoteIP", httpContext?.Connection?.RemoteIpAddress);
+
+					// Header sent by the dashboard to indicate how long a user has been inactive for a particular browser page (in seconds)
+					if (httpContext?.Request.Headers.TryGetValue("X-Horde-LastUserActivity", out StringValues values) is true)
+					{
+						string? value = values.FirstOrDefault();
+						if (!String.IsNullOrEmpty(value) && Int32.TryParse(value, out int lastUserActivity))
+						{
+							diagnosticContext.Set("LastUserActivity", lastUserActivity);
+						}
+					}
+				};
+			});
+
+			// Enable middleware to serve swagger-ui (HTML, JS, CSS, etc.),
+			// specifying the Swagger JSON endpoint.
+			app.UseSwaggerUI(c =>
+			{
+				c.SwaggerEndpoint("/swagger/v1/swagger.json", "Horde Server API");
+				c.RoutePrefix = "swagger";
+			});
+
+			if (!env.IsDevelopment())
+			{
+				app.UseMiddleware<RequestTrackerMiddleware>();
+			}
+
+			app.UseExceptionHandler("/api/v1/exception");
+
+			DirectoryReference dashboardDir = DirectoryReference.Combine(ServerApp.AppDir, "DashboardApp");
+			if (DirectoryReference.Exists(dashboardDir))
+			{
+				app.UseDefaultFiles();
+				app.UseStaticFiles();
+				app.UseWhen(IsSpaRequest, builder => builder.UseSpaStaticFiles());
+			}
+
+			app.UseRouting();
+
+			app.UseAuthentication();
+			app.UseAuthorization();
+
+			app.UseEndpoints(endpoints =>
+			{
+				endpoints.MapGrpcService<HealthService>();
+				endpoints.MapGrpcService<RpcService>();
+				endpoints.MapGrpcService<JobRpcService>();
+				endpoints.MapGrpcService<LogRpcService>();
+				endpoints.MapGrpcService<AgentRelayService>();
+				endpoints.MapGrpcService<EnrollmentRpc>();
+
+				endpoints.MapGrpcReflectionService();
+
+				endpoints.MapControllers();
+
+				endpoints.MapHealthChecks("/health/live");
+			});
+
+			if (DirectoryReference.Exists(dashboardDir))
+			{
+				app.MapWhen(IsSpaRequest, builder => builder.UseSpa(spa => spa.Options.SourcePath = "DashboardApp"));
+			}
+
+			if (settings.Value.OpenBrowser && RuntimeInformation.IsOSPlatform(OSPlatform.Windows))
+			{
+				lifetime.ApplicationStarted.Register(() => LaunchBrowser(app));
+			}
+		}
+
+		static bool IsSpaRequest(HttpContext context)
+		{
+			return !context.Request.Path.StartsWithSegments("/api", StringComparison.OrdinalIgnoreCase);
+		}
+
+		static void LaunchBrowser(IApplicationBuilder app)
+		{
+			IServerAddressesFeature? feature = app.ServerFeatures.Get<IServerAddressesFeature>();
+			if (feature != null && feature.Addresses.Count > 0)
+			{
+				// with a development cert, host will be set by default to localhost, otherwise there will be no host in address
+				string address = feature.Addresses.First().Replace("[::]", System.Net.Dns.GetHostName(), StringComparison.OrdinalIgnoreCase);
+				Process.Start(new ProcessStartInfo { FileName = address, UseShellExecute = true });
+			}
+		}
+
+		static LogEventLevel GetRequestLoggingLevel(HttpContext context, double elapsedMs, Exception? ex)
+		{
+			if (context.Request != null && context.Request.Path.HasValue)
+			{
+				string requestPath = context.Request.Path;
+
+				if (requestPath.Equals("/Horde.HordeRpc/QueryServerStateV2", StringComparison.OrdinalIgnoreCase))
+				{
+					return LogEventLevel.Verbose;
+				}
+
+				if (requestPath.Equals("/Horde.HordeRpc/UpdateSession", StringComparison.OrdinalIgnoreCase))
+				{
+					return LogEventLevel.Verbose;
+				}
+
+				if (requestPath.Equals("/Horde.HordeRpc/CreateEvents", StringComparison.OrdinalIgnoreCase))
+				{
+					return LogEventLevel.Verbose;
+				}
+
+				if (requestPath.Equals("/Horde.HordeRpc/WriteOutput", StringComparison.OrdinalIgnoreCase))
+				{
+					return LogEventLevel.Information;
+				}
+
+				if (requestPath.StartsWith("/Horde.HordeRpc", StringComparison.OrdinalIgnoreCase))
+				{
+					return LogEventLevel.Debug;
+				}
+
+				if (requestPath.Equals("/Horde.Relay.RelayRpc/GetPortMappings", StringComparison.OrdinalIgnoreCase))
+				{
+					return LogEventLevel.Verbose;
+				}
+
+				if (requestPath.StartsWith("/health", StringComparison.OrdinalIgnoreCase))
+				{
+					return LogEventLevel.Debug;
+				}
+
+				if (requestPath.StartsWith("/grpc.health", StringComparison.OrdinalIgnoreCase))
+				{
+					return LogEventLevel.Debug;
+				}
+
+				if (requestPath.StartsWith("/api/v1", StringComparison.OrdinalIgnoreCase))
+				{
+					return LogEventLevel.Debug;
+				}
+
+				if (requestPath.StartsWith("/ugs/api", StringComparison.OrdinalIgnoreCase))
+				{
+					return LogEventLevel.Verbose;
+				}
+
+				if (IsSpaRequest(context))
+				{
+					return LogEventLevel.Verbose;
+				}
+			}
+			return LogEventLevel.Information;
+		}
+
+		class HostApplicationLifetime : IHostApplicationLifetime
+		{
+			public CancellationToken ApplicationStarted => CancellationToken.None;
+			public CancellationToken ApplicationStopped => CancellationToken.None;
+			public CancellationToken ApplicationStopping => CancellationToken.None;
+
+			public void StopApplication() { }
+		}
+
+		public static void AddServices(IServiceCollection serviceCollection, IConfiguration configuration)
+		{
+			Startup startup = new Startup(configuration);
+			startup.ConfigureServices(serviceCollection);
+			serviceCollection.AddSingleton<IHostApplicationLifetime, HostApplicationLifetime>();
+		}
+
+		public static ServiceCollection CreateServiceCollection(IConfiguration configuration)
+		{
+			ServiceCollection serviceCollection = new ServiceCollection();
+			AddServices(serviceCollection, configuration);
+			return serviceCollection;
+		}
+
+		public static ServiceCollection CreateServiceCollection(IConfiguration configuration, ILoggerProvider loggerProvider)
+		{
+			ServiceCollection services = new ServiceCollection();
+			services.AddSingleton(configuration);
+			services.AddSingleton(loggerProvider);
+			AddServices(services, configuration);
+			return services;
+		}
+	}
+}
diff --git a/Engine/Source/Programs/Horde/HordeServer/Streams/IStreamCollection.cs b/Engine/Source/Programs/Horde/HordeServer/Streams/IStreamCollection.cs
new file mode 100644
--- /dev/null
+++ b/Engine/Source/Programs/Horde/HordeServer/Streams/IStreamCollection.cs
@@ -0,0 +1,143 @@
+// Copyright Epic Games, Inc. All Rights Reserved.
+
+using System;
+using System.Collections.Generic;
+using System.Linq;
+using System.Threading;
+using System.Threading.Tasks;
+using EpicGames.Horde.Jobs;
+using EpicGames.Horde.Jobs.Templates;
+
+namespace Horde.Server.Streams
+{
+	/// <summary>
+	/// Collection of stream documents
+	/// </summary>
+	public interface IStreamCollection
+	{
+		/// <summary>
+		/// Gets a stream by ID
+		/// </summary>
+		/// <param name="streamConfig">The stream config object</param>
+		/// <param name="cancellationToken">Cancellation token for the operation</param>
+		/// <returns>The stream document</returns>
+		Task<IStream> GetAsync(StreamConfig streamConfig, CancellationToken cancellationToken = default);
+
+		/// <summary>
+		/// Gets a stream by ID
+		/// </summary>
+		/// <param name="streamConfigs">The stream config object</param>
+		/// <param name="cancellationToken">Cancellation token for the operation</param>
+		/// <returns>The stream document</returns>
+		Task<IReadOnlyList<IStream>> GetAsync(IReadOnlyList<StreamConfig> streamConfigs, CancellationToken cancellationToken = default);
+
+		/// <summary>
+		/// Updates user-facing properties for an existing stream
+		/// </summary>
+		/// <param name="stream">The stream to update</param>
+		/// <param name="newPausedUntil">The new datetime for pausing builds</param>
+		/// <param name="newPauseComment">The reason for pausing</param>
+		/// <param name="cancellationToken">Cancellation token for the operation</param>
+		/// <returns>The updated stream if successful, null otherwise</returns>
+		Task<IStream?> TryUpdatePauseStateAsync(IStream stream, DateTime? newPausedUntil, string? newPauseComment, CancellationToken cancellationToken = default);
+
+		/// <summary>
+		/// Attempts to update the last trigger time for a schedule
+		/// </summary>
+		/// <param name="stream">The stream to update</param>
+		/// <param name="templateRefId">The template ref id</param>
+		/// <param name="lastTriggerTimeUtc">New last trigger time for the schedule</param>
+		/// <param name="lastTriggerChange">New last trigger changelist for the schedule</param>
+		/// <param name="newActiveJobs">New list of active jobs</param>
+		/// <param name="cancellationToken">Cancellation token for the operation</param>
+		/// <returns>The updated stream if successful, null otherwise</returns>
+		Task<IStream?> TryUpdateScheduleTriggerAsync(IStream stream, TemplateId templateRefId, DateTime? lastTriggerTimeUtc, int? lastTriggerChange, List<JobId> newActiveJobs, CancellationToken cancellationToken = default);
+
+		/// <summary>
+		/// Attempts to update a stream template ref
+		/// </summary>
+		/// <param name="streamInterface">The stream containing the template ref</param>
+		/// <param name="templateRefId">The template ref to update</param>
+		/// <param name="stepStates">The stream states to update, pass an empty list to clear all step states, otherwise will be a partial update based on included step updates</param>
+		/// <param name="cancellationToken">Cancellation token for the operation</param>
+		/// <returns></returns>
+		Task<IStream?> TryUpdateTemplateRefAsync(IStream streamInterface, TemplateId templateRefId, List<UpdateStepStateRequest>? stepStates = null, CancellationToken cancellationToken = default);
+	}
+
+	/// <summary>
+	/// Extension methods for <see cref="StreamCollection"/>
+	/// </summary>
+	public static class StreamCollectionExtensions
+	{
+		/// <summary>
+		/// Updates an existing stream
+		/// </summary>
+		/// <param name="streamCollection">The stream collection</param>
+		/// <param name="stream">The stream to update</param>
+		/// <param name="newPausedUntil">The new datetime for pausing builds</param>
+		/// <param name="newPauseComment">The reason for pausing</param>
+		/// <param name="cancellationToken">Cancellation token for the operation</param>
+		/// <returns>Async task object</returns>
+		public static async Task<IStream?> UpdatePauseStateAsync(this IStreamCollection streamCollection, IStream? stream, DateTime? newPausedUntil = null, string? newPauseComment = null, CancellationToken cancellationToken = default)
+		{
+			for (; stream != null; stream = await streamCollection.GetAsync(stream.Config, cancellationToken))
+			{
+				IStream? newStream = await streamCollection.TryUpdatePauseStateAsync(stream, newPausedUntil, newPauseComment, cancellationToken);
+				if (newStream != null)
+				{
+					return newStream;
+				}
+			}
+			return null;
+		}
+
+		/// <summary>
+		/// Attempts to update the last trigger time for a schedule
+		/// </summary>
+		/// <param name="streamCollection">The stream collection</param>
+		/// <param name="stream">The stream to update</param>
+		/// <param name="templateRefId">The template ref id</param>
+		/// <param name="lastTriggerTimeUtc"></param>
+		/// <param name="lastTriggerChange"></param>
+		/// <param name="addJobs">Jobs to add</param>
+		/// <param name="removeJobs">Jobs to remove</param>
+		/// <param name="cancellationToken">Cancellation token for the operation</param>
+		/// <returns>True if the stream was updated</returns>
+		public static async Task<IStream?> UpdateScheduleTriggerAsync(this IStreamCollection streamCollection, IStream stream, TemplateId templateRefId, DateTime? lastTriggerTimeUtc = null, int? lastTriggerChange = null, List<JobId>? addJobs = null, List<JobId>? removeJobs = null, CancellationToken cancellationToken = default)
+		{
+			IStream? newStream = stream;
+			while (newStream != null)
+			{
+				ITemplateRef? templateRef;
+				if (!newStream.Templates.TryGetValue(templateRefId, out templateRef))
+				{
+					break;
+				}
+				if (templateRef.Schedule == null)
+				{
+					break;
+				}
+
+				IEnumerable<JobId> newActiveJobs = templateRef.Schedule.ActiveJobs;
+				if (removeJobs != null)
+				{
+					newActiveJobs = newActiveJobs.Except(removeJobs);
+				}
+				if (addJobs != null)
+				{
+					newActiveJobs = newActiveJobs.Union(addJobs);
+				}
+
+				newStream = await streamCollection.TryUpdateScheduleTriggerAsync(newStream, templateRefId, lastTriggerTimeUtc, lastTriggerChange, newActiveJobs.ToList(), cancellationToken);
+
+				if (newStream != null)
+				{
+					return newStream;
+				}
+
+				newStream = await streamCollection.GetAsync(stream.Config, cancellationToken);
+			}
+			return null;
+		}
+	}
+}
diff --git a/Engine/Source/Programs/Horde/HordeServer/Streams/StreamCollection.cs b/Engine/Source/Programs/Horde/HordeServer/Streams/StreamCollection.cs
new file mode 100644
--- /dev/null
+++ b/Engine/Source/Programs/Horde/HordeServer/Streams/StreamCollection.cs
@@ -0,0 +1,511 @@
+// Copyright Epic Games, Inc. All Rights Reserved.
+
+using System;
+using System.Collections.Generic;
+using System.IO;
+using System.Linq;
+using System.Threading;
+using System.Threading.Tasks;
+using EpicGames.Core;
+using EpicGames.Horde.Jobs;
+using EpicGames.Horde.Jobs.Templates;
+using EpicGames.Horde.Streams;
+using EpicGames.Horde.Users;
+using Horde.Server.Server;
+using Horde.Server.Utilities;
+using MongoDB.Bson.Serialization.Attributes;
+using MongoDB.Driver;
+
+namespace Horde.Server.Streams
+{
+	/// <summary>
+	/// Collection of stream documents
+	/// </summary>
+	public class StreamCollection : IStreamCollection
+	{
+		/// <summary>
+		/// Information about a stream
+		/// </summary>
+		class StreamDoc : IStream
+		{
+			[BsonRequired, BsonId]
+			public StreamId Id { get; set; }
+
+			public Dictionary<TemplateId, TemplateRefDoc> Templates { get; set; } = new Dictionary<TemplateId, TemplateRefDoc>();
+			public DateTime? PausedUntil { get; set; }
+			public string? PauseComment { get; set; }
+
+			public int UpdateIndex { get; set; }
+
+			[BsonConstructor]
+			private StreamDoc()
+			{
+			}
+
+			public StreamDoc(StreamId id)
+			{
+				Id = id;
+			}
+
+			#region IStream Implementation
+
+			[BsonIgnore]
+			public StreamConfig Config { get; private set; } = null!;
+
+			[BsonIgnore]
+			IReadOnlyDictionary<TemplateId, ITemplateRef>? _cachedTemplates;
+
+			IReadOnlyDictionary<TemplateId, ITemplateRef> IStream.Templates
+			{
+				get
+				{
+					_cachedTemplates ??= Templates.ToDictionary(x => x.Key, x => (ITemplateRef)x.Value);
+					return _cachedTemplates;
+				}
+			}
+
+			public bool PostLoad(StreamConfig config, DateTime utcNow)
+			{
+				Config = config;
+
+				bool replaceDocument = false;
+
+				Dictionary<TemplateId, TemplateRefDoc> templateRefDocs = new Dictionary<TemplateId, TemplateRefDoc>(config.Templates.Count);
+				foreach (TemplateRefConfig templateRefConfig in config.Templates)
+				{
+					TemplateRefDoc? templateRefDoc;
+					if (!Templates.TryGetValue(templateRefConfig.Id, out templateRefDoc))
+					{
+						templateRefDoc = new TemplateRefDoc();
+						replaceDocument = true;
+					}
+					replaceDocument |= templateRefDoc.PostLoad(templateRefConfig, utcNow);
+
+					templateRefDocs.Add(templateRefConfig.Id, templateRefDoc);
+				}
+
+				replaceDocument |= Templates.Count != templateRefDocs.Count;
+				Templates = templateRefDocs;
+
+				return replaceDocument;
+			}
+
+			#endregion
+		}
+
+		class TemplateRefDoc : ITemplateRef
+		{
+			[BsonIgnoreIfNull]
+			public TemplateScheduleDoc? Schedule { get; set; }
+
+			[BsonIgnoreIfNull]
+			public List<TemplateStepDoc>? StepStates { get; set; }
+
+			#region ITemplateRef implementation
+
+			[BsonIgnore]
+			public TemplateId Id { get; private set; }
+
+			[BsonIgnore]
+			public TemplateRefConfig Config { get; private set; } = null!;
+
+			ITemplateSchedule? ITemplateRef.Schedule => Schedule;
+			IReadOnlyList<ITemplateStep> ITemplateRef.StepStates => (IReadOnlyList<ITemplateStep>?)StepStates ?? Array.Empty<ITemplateStep>();
+
+			public bool PostLoad(TemplateRefConfig config, DateTime utcNow)
+			{
+				Id = config.Id;
+				Config = config;
+
+				bool replaceDocument = false;
+
+				if (Config.Schedule == null)
+				{
+					if (Schedule != null)
+					{
+						Schedule = null;
+						replaceDocument = true;
+					}
+				}
+				else
+				{
+					if (Schedule == null)
+					{
+						Schedule = new TemplateScheduleDoc();
+						Schedule.LastTriggerTimeUtc = utcNow;
+						replaceDocument = true;
+					}
+				}
+
+				Schedule?.PostLoad(this);
+
+				if (StepStates != null)
+				{
+					int count = StepStates.RemoveAll(x => x.PausedByUserId == null);
+					replaceDocument |= count > 0;
+				}
+
+				return replaceDocument;
+			}
+
+			#endregion
+		}
+
+		class TemplateScheduleDoc : ITemplateSchedule
+		{
+			public int LastTriggerChange { get; set; }
+
+			[BsonIgnoreIfNull, Obsolete("Use LastTriggerTimeUtc instead")]
+			public DateTimeOffset? LastTriggerTime { get; set; }
+
+			public DateTime LastTriggerTimeUtc { get; set; }
+			public List<JobId> ActiveJobs { get; set; } = new List<JobId>();
+
+			#region ITemplateSchedule implementation
+
+			[BsonIgnore]
+			TemplateRefDoc? _owner;
+
+			[BsonIgnore]
+			public ScheduleConfig Config => _owner?.Config.Schedule!;
+
+			IReadOnlyList<JobId> ITemplateSchedule.ActiveJobs => ActiveJobs;
+
+			public void PostLoad(TemplateRefDoc owner)
+			{
+				_owner = owner;
+
+#pragma warning disable CS0618 // Type or member is obsolete
+				if (LastTriggerTime.HasValue)
+				{
+					LastTriggerTimeUtc = LastTriggerTime.Value.UtcDateTime;
+					LastTriggerTime = null;
+				}
+#pragma warning restore CS0618 // Type or member is obsolete
+			}
+
+			#endregion
+		}
+
+		class TemplateStepDoc : ITemplateStep
+		{
+			public string Name { get; set; } = String.Empty;
+			public UserId? PausedByUserId { get; set; }
+			public DateTime? PauseTimeUtc { get; set; }
+
+			UserId ITemplateStep.PausedByUserId => PausedByUserId ?? UserId.Empty;
+			DateTime ITemplateStep.PauseTimeUtc => PauseTimeUtc ?? DateTime.MinValue;
+
+			public TemplateStepDoc()
+			{
+			}
+
+			public TemplateStepDoc(string name, UserId pausedByUserId, DateTime pauseTimeUtc)
+			{
+				Name = name;
+				PausedByUserId = pausedByUserId;
+				PauseTimeUtc = pauseTimeUtc;
+			}
+		}
+
+		readonly IMongoCollection<StreamDoc> _streams;
+		readonly IClock _clock;
+
+		/// <summary>
+		/// Constructor
+		/// </summary>
+		/// <param name="mongoService">The database service instance</param>
+		/// <param name="clock"></param>
+		public StreamCollection(MongoService mongoService, IClock clock)
+		{
+			_streams = mongoService.GetCollection<StreamDoc>("Streams");
+			_clock = clock;
+		}
+
+		/// <inheritdoc/>
+		public async Task<IStream> GetAsync(StreamConfig streamConfig, CancellationToken cancellationToken)
+			=> await GetInternalAsync(streamConfig, cancellationToken);
+
+		async Task<StreamDoc> GetInternalAsync(StreamConfig streamConfig, CancellationToken cancellationToken)
+		{
+			FindOneAndUpdateOptions<StreamDoc, StreamDoc> options = new FindOneAndUpdateOptions<StreamDoc, StreamDoc>();
+			options.IsUpsert = true;
+			options.ReturnDocument = ReturnDocument.After;
+
+			for (; ; )
+			{
+				StreamDoc stream = await _streams.FindOneAndUpdateAsync<StreamDoc>(x => x.Id == streamConfig.Id, Builders<StreamDoc>.Update.SetOnInsert(x => x.UpdateIndex, 0), options, cancellationToken);
+				if (await PostLoadAsync(stream, streamConfig, cancellationToken))
+				{
+					return stream;
+				}
+			}
+		}
+
+		/// <inheritdoc/>
+		public async Task<IReadOnlyList<IStream>> GetAsync(IReadOnlyList<StreamConfig> streamConfigs, CancellationToken cancellationToken)
+			=> await GetInternalAsync(streamConfigs, cancellationToken);
+
+		async Task<List<StreamDoc>> GetInternalAsync(IReadOnlyList<StreamConfig> streamConfigs, CancellationToken cancellationToken)
+		{
+			FilterDefinition<StreamDoc> filter = Builders<StreamDoc>.Filter.In(x => x.Id, streamConfigs.Select(x => x.Id));
+			List<StreamDoc> matches = await _streams.Find(filter).ToListAsync(cancellationToken);
+
+			List<StreamDoc> results = new List<StreamDoc>(streamConfigs.Count);
+			foreach (StreamConfig streamConfig in streamConfigs)
+			{
+				StreamDoc? stream = matches.FirstOrDefault(x => x.Id == streamConfig.Id);
+				if (stream == null)
+				{
+					stream = await GetInternalAsync(streamConfig, cancellationToken);
+				}
+				else
+				{
+					stream.PostLoad(streamConfig, _clock.UtcNow);
+				}
+				results.Add(stream);
+			}
+			return results;
+		}
+
+		async ValueTask<bool> PostLoadAsync(StreamDoc streamDoc, StreamConfig streamConfig, CancellationToken cancellationToken)
+		{
+			if (streamDoc.PostLoad(streamConfig, _clock.UtcNow))
+			{
+				int updateIndex = streamDoc.UpdateIndex++;
+
+				ReplaceOneResult result = await _streams.ReplaceOneAsync(x => x.Id == streamDoc.Id && x.UpdateIndex == updateIndex, streamDoc, cancellationToken: cancellationToken);
+				if (result.MatchedCount != 1)
+				{
+					return false;
+				}
+			}
+			return true;
+		}
+
+		/// <inheritdoc/>
+		public async Task<IStream?> TryUpdateTemplateRefAsync(IStream streamInterface, TemplateId templateId, List<UpdateStepStateRequest>? stepStates = null, CancellationToken cancellationToken = default)
+		{
+			StreamDoc stream = (StreamDoc)streamInterface;
+
+			TemplateRefDoc? templateRef;
+			if (!stream.Templates.TryGetValue(templateId, out templateRef))
+			{
+				return null;
+			}
+
+			UpdateDefinitionBuilder<StreamDoc> updateBuilder = Builders<StreamDoc>.Update;
+			List<UpdateDefinition<StreamDoc>> updates = new List<UpdateDefinition<StreamDoc>>();
+
+			Dictionary<TemplateId, TemplateRefDoc> newTemplates = new Dictionary<TemplateId, TemplateRefDoc>(stream.Templates);
+
+			// clear
+			if (stepStates != null && stepStates.Count == 0)
+			{
+				bool hasUpdates = false;
+				foreach (KeyValuePair<TemplateId, TemplateRefDoc> entry in newTemplates)
+				{
+					if (entry.Value.StepStates != null)
+					{
+						hasUpdates = true;
+						entry.Value.StepStates = null;
+					}
+				}
+
+				if (hasUpdates)
+				{
+					updates.Add(updateBuilder.Set(x => x.Templates, newTemplates));
+				}
+			}
+			else if (stepStates != null)
+			{
+				// get currently valid step states
+				List<TemplateStepDoc> newStepStates = templateRef.StepStates?.ToList() ?? new List<TemplateStepDoc>();
+
+				// generate update list
+				foreach (UpdateStepStateRequest updateState in stepStates)
+				{
+					int stateIndex = newStepStates.FindIndex(x => x.Name == updateState.Name);
+
+					UserId? pausedByUserId = updateState.PausedByUserId != null ? UserId.Parse(updateState.PausedByUserId) : null;
+
+					if (stateIndex == -1)
+					{
+						// if this is a new state without anything set, ignore it
+						if (pausedByUserId != null)
+						{
+							newStepStates.Add(new TemplateStepDoc(updateState.Name, pausedByUserId.Value, _clock.UtcNow));
+						}
+					}
+					else
+					{
+						if (pausedByUserId == null)
+						{
+							newStepStates.RemoveAt(stateIndex);
+						}
+						else
+						{
+							newStepStates[stateIndex].PausedByUserId = pausedByUserId.Value;
+						}
+					}
+				}
+
+				if (newStepStates.Count == 0)
+				{
+					templateRef.StepStates = null;
+				}
+				else
+				{
+					templateRef.StepStates = newStepStates;
+				}
+
+				updates.Add(updateBuilder.Set(x => x.Templates, newTemplates));
+			}
+
+			if (updates.Count == 0)
+			{
+				return streamInterface;
+			}
+
+			return await TryUpdateStreamAsync(stream, updateBuilder.Combine(updates), cancellationToken);
+
+		}
+
+		/// <inheritdoc/>
+		public async Task<IStream?> TryUpdatePauseStateAsync(IStream streamInterface, DateTime? newPausedUntil, string? newPauseComment, CancellationToken cancellationToken)
+		{
+			StreamDoc stream = (StreamDoc)streamInterface;
+
+			UpdateDefinitionBuilder<StreamDoc> updateBuilder = Builders<StreamDoc>.Update;
+
+			List<UpdateDefinition<StreamDoc>> updates = new List<UpdateDefinition<StreamDoc>>();
+			stream.PausedUntil = newPausedUntil;
+			stream.PauseComment = newPauseComment;
+			updates.Add(updateBuilder.Set(x => x.PausedUntil, newPausedUntil));
+			updates.Add(updateBuilder.Set(x => x.PauseComment, newPauseComment));
+
+			return await TryUpdateStreamAsync(stream, updateBuilder.Combine(updates), cancellationToken);
+		}
+
+		/// <inheritdoc/>
+		public async Task<IStream?> TryUpdateScheduleTriggerAsync(IStream streamInterface, TemplateId templateId, DateTime? lastTriggerTimeUtc, int? lastTriggerChange, List<JobId> newActiveJobs, CancellationToken cancellationToken)
+		{
+			StreamDoc stream = (StreamDoc)streamInterface;
+			TemplateRefDoc template = stream.Templates[templateId];
+			TemplateScheduleDoc schedule = template.Schedule!;
+
+			// Build the updates. MongoDB driver cannot parse TemplateRefId in expression tree; need to specify field name explicitly
+			List<UpdateDefinition<StreamDoc>> updates = new List<UpdateDefinition<StreamDoc>>();
+			if (lastTriggerTimeUtc.HasValue && lastTriggerTimeUtc.Value != schedule.LastTriggerTimeUtc)
+			{
+				FieldDefinition<StreamDoc, DateTime> lastTriggerTimeField = $"{nameof(stream.Templates)}.{templateId}.{nameof(template.Schedule)}.{nameof(schedule.LastTriggerTimeUtc)}";
+				updates.Add(Builders<StreamDoc>.Update.Set(lastTriggerTimeField, lastTriggerTimeUtc.Value));
+				schedule.LastTriggerTimeUtc = lastTriggerTimeUtc.Value;
+			}
+			if (lastTriggerChange.HasValue && lastTriggerChange.Value > schedule.LastTriggerChange)
+			{
+				FieldDefinition<StreamDoc, int> lastTriggerChangeField = $"{nameof(stream.Templates)}.{templateId}.{nameof(template.Schedule)}.{nameof(schedule.LastTriggerChange)}";
+				updates.Add(Builders<StreamDoc>.Update.Set(lastTriggerChangeField, lastTriggerChange.Value));
+				schedule.LastTriggerChange = lastTriggerChange.Value;
+			}
+			if (newActiveJobs != null)
+			{
+				FieldDefinition<StreamDoc, List<JobId>> field = $"{nameof(stream.Templates)}.{templateId}.{nameof(template.Schedule)}.{nameof(schedule.ActiveJobs)}";
+				updates.Add(Builders<StreamDoc>.Update.Set(field, newActiveJobs));
+				schedule.ActiveJobs = newActiveJobs;
+			}
+
+			return (updates.Count == 0) ? streamInterface : await TryUpdateStreamAsync(stream, Builders<StreamDoc>.Update.Combine(updates), cancellationToken);
+		}
+
+		/// <summary>
+		/// Update a stream
+		/// </summary>
+		/// <param name="stream">The stream to update</param>
+		/// <param name="update">The update definition</param>
+		/// <param name="cancellationToken">Cancellation token for the operation</param>
+		/// <returns>The updated document, or null the update failed</returns>
+		private async Task<StreamDoc?> TryUpdateStreamAsync(StreamDoc stream, UpdateDefinition<StreamDoc> update, CancellationToken cancellationToken)
+		{
+			FilterDefinition<StreamDoc> filter = Builders<StreamDoc>.Filter.Expr(x => x.Id == stream.Id && x.UpdateIndex == stream.UpdateIndex);
+			update = update.Set(x => x.UpdateIndex, stream.UpdateIndex + 1);
+
+			FindOneAndUpdateOptions<StreamDoc> options = new FindOneAndUpdateOptions<StreamDoc> { ReturnDocument = ReturnDocument.After };
+
+			StreamDoc? result = await _streams.FindOneAndUpdateAsync(filter, update, options, cancellationToken);
+			if (result != null)
+			{
+				result.PostLoad(stream.Config, _clock.UtcNow);
+			}
+			return result;
+		}
+
+		/// <summary>
+		/// Checks the stream definition for consistency
+		/// </summary>
+		public static void Validate(StreamId streamId, StreamConfig config)
+		{
+			HashSet<TemplateId> remainingTemplates = new HashSet<TemplateId>(config.Templates.Select(x => x.Id));
+
+			// Check the default preflight template is valid
+			if (config.DefaultPreflight != null)
+			{
+				if (config.DefaultPreflight.TemplateId != null && !remainingTemplates.Contains(config.DefaultPreflight.TemplateId.Value))
+				{
+					throw new InvalidStreamException($"Default preflight template was listed as '{config.DefaultPreflight.TemplateId.Value}', but no template was found by that name");
+				}
+			}
+
+			// Check the chained jobs are valid
+			foreach (TemplateRefConfig templateRef in config.Templates)
+			{
+				if (templateRef.ChainedJobs != null)
+				{
+					foreach (ChainedJobTemplateConfig chainedJob in templateRef.ChainedJobs)
+					{
+						if (!remainingTemplates.Contains(chainedJob.TemplateId))
+						{
+							throw new InvalidDataException($"Invalid template ref id '{chainedJob.TemplateId}");
+						}
+					}
+				}
+			}
+
+			// Check that all the templates are referenced by a tab
+			HashSet<TemplateId> undefinedTemplates = new();
+			foreach (TabConfig tab in config.Tabs)
+			{
+				if (tab.Templates != null)
+				{
+					remainingTemplates.ExceptWith(tab.Templates);
+					foreach (TemplateId templateId in tab.Templates)
+					{
+						if (config.Templates.Find(x => x.Id == templateId) == null)
+						{
+							undefinedTemplates.Add(templateId);
+						}
+					}
+				}
+			}
+			if (remainingTemplates.Count > 0)
+			{
+				throw new InvalidStreamException(String.Join("\n", remainingTemplates.Select(x => $"Template '{x}' is not listed on any tab for {streamId}")));
+			}
+
+			if (undefinedTemplates.Count > 0)
+			{
+				throw new InvalidStreamException(String.Join("\n", undefinedTemplates.Select(x => $"Template '{x}' is not defined for {streamId}")));
+			}
+
+			// Check that all the agent types reference valid workspace names
+			foreach (KeyValuePair<string, AgentConfig> pair in config.AgentTypes)
+			{
+				string? workspaceTypeName = pair.Value.Workspace;
+				if (workspaceTypeName != null && !config.WorkspaceTypes.ContainsKey(workspaceTypeName))
+				{
+					throw new InvalidStreamException($"Agent type '{pair.Key}' references undefined workspace type '{pair.Value.Workspace}' in {streamId}");
+				}
+			}
+		}
+	}
+}
