Horde: Rename projects before splitting into plugins (Horde.Agent -> HordeAgent, Horde.Server -> HordeServer, Horde -> HordeCmd).

diff --git a/Engine/Source/Programs/Horde/HordeServer/Streams/StreamsController.cs b/Engine/Source/Programs/Horde/HordeServer/Streams/StreamsController.cs
new file mode 100644
--- /dev/null
+++ b/Engine/Source/Programs/Horde/HordeServer/Streams/StreamsController.cs
@@ -0,0 +1,450 @@
+// Copyright Epic Games, Inc. All Rights Reserved.
+
+using System;
+using System.Collections.Generic;
+using System.Linq;
+using System.Threading;
+using System.Threading.Tasks;
+using EpicGames.Horde.Commits;
+using EpicGames.Horde.Jobs;
+using EpicGames.Horde.Jobs.Templates;
+using EpicGames.Horde.Projects;
+using EpicGames.Horde.Streams;
+using EpicGames.Horde.Users;
+using Horde.Server.Commits;
+using Horde.Server.Jobs;
+using Horde.Server.Jobs.Templates;
+using Horde.Server.Projects;
+using Horde.Server.Server;
+using Horde.Server.Users;
+using Horde.Server.Utilities;
+using Microsoft.AspNetCore.Authorization;
+using Microsoft.AspNetCore.Mvc;
+using Microsoft.Extensions.Options;
+using OpenTelemetry.Trace;
+using TimeZoneConverter;
+
+namespace Horde.Server.Streams
+{
+	/// <summary>
+	/// Controller for the /api/v1/streams endpoint
+	/// </summary>
+	[ApiController]
+	[Authorize]
+	[Route("[controller]")]
+	public class StreamsController : HordeControllerBase
+	{
+		private readonly IStreamCollection _streamCollection;
+		private readonly ICommitService _commitService;
+		private readonly ITemplateCollection _templateCollection;
+		private readonly IJobStepRefCollection _jobStepRefCollection;
+		private readonly IUserCollection _userCollection;
+		private readonly Tracer _tracer;
+		private readonly IOptionsSnapshot<GlobalConfig> _globalConfig;
+		private readonly TimeZoneInfo _timeZone;
+
+		/// <summary>
+		/// Constructor
+		/// </summary>
+		public StreamsController(IStreamCollection streamCollection, ICommitService commitService, ITemplateCollection templateCollection, IJobStepRefCollection jobStepRefCollection, IUserCollection userCollection, Tracer tracer, IOptionsSnapshot<GlobalConfig> globalConfig)
+		{
+			_streamCollection = streamCollection;
+			_commitService = commitService;
+			_templateCollection = templateCollection;
+			_jobStepRefCollection = jobStepRefCollection;
+			_userCollection = userCollection;
+			_tracer = tracer;
+			_globalConfig = globalConfig;
+
+			string? timeZoneName = _globalConfig.Value.ServerSettings.ScheduleTimeZone;
+			_timeZone = (timeZoneName == null) ? TimeZoneInfo.Local : TZConvert.GetTimeZoneInfo(timeZoneName);
+		}
+
+		/// <summary>
+		/// Query all the streams for a particular project.
+		/// </summary>
+		/// <param name="projectIds">Unique id of the project to query</param>
+		/// <param name="filter">Filter for the properties to return</param>
+		/// <param name="cancellationToken">Cancellation token for the operation</param>
+		/// <returns>Information about all the projects</returns>
+		[HttpGet]
+		[Route("/api/v1/streams")]
+		[ProducesResponseType(typeof(List<GetStreamResponse>), 200)]
+		public async Task<ActionResult<List<object>>> GetStreamsAsync([FromQuery(Name = "ProjectId")] string[] projectIds, [FromQuery] PropertyFilter? filter = null, CancellationToken cancellationToken = default)
+		{
+			ProjectId[] projectIdValues = Array.ConvertAll(projectIds, x => new ProjectId(x));
+
+			List<StreamConfig> streamConfigs = new List<StreamConfig>();
+			foreach (ProjectConfig projectConfig in _globalConfig.Value.Projects)
+			{
+				if (projectIdValues.Length == 0 || projectIdValues.Contains(projectConfig.Id))
+				{
+					foreach (StreamConfig streamConfig in projectConfig.Streams)
+					{
+						if (streamConfig.Authorize(StreamAclAction.ViewStream, User))
+						{
+							streamConfigs.Add(streamConfig);
+						}
+					}
+				}
+			}
+
+			IReadOnlyList<IStream> streams = await _streamCollection.GetAsync(streamConfigs, cancellationToken);
+
+			List<GetStreamResponse> responses = new List<GetStreamResponse>();
+			foreach (IStream stream in streams)
+			{
+				GetStreamResponse response = await CreateGetStreamResponseAsync(stream, cancellationToken);
+				responses.Add(response);
+			}
+
+			return responses.OrderBy(x => x.Id).Select(x => PropertyFilter.Apply(x, filter)).ToList();
+		}
+
+		/// <summary>
+		/// Query all the streams for a particular project.
+		/// </summary>
+		/// <param name="projectIds">Unique id of the project to query</param>
+		/// <param name="filter">Filter for the properties to return</param>
+		/// <returns>Information about all the projects</returns>
+		[HttpGet]
+		[Route("/api/v2/streams")]
+		[ProducesResponseType(typeof(List<GetStreamResponse>), 200)]
+		public async Task<ActionResult<List<object>>> GetStreamsV2Async([FromQuery(Name = "ProjectId")] string[] projectIds, [FromQuery] PropertyFilter? filter = null)
+		{
+			ProjectId[] projectIdValues = Array.ConvertAll(projectIds, x => new ProjectId(x));
+
+			List<GetStreamResponseV2> responses = new List<GetStreamResponseV2>();
+			foreach (ProjectConfig projectConfig in _globalConfig.Value.Projects)
+			{
+				if (projectIdValues.Length == 0 || projectIdValues.Contains(projectConfig.Id))
+				{
+					foreach (StreamConfig streamConfig in projectConfig.Streams)
+					{
+						if (streamConfig.Authorize(StreamAclAction.ViewStream, User))
+						{
+							IStream stream = await _streamCollection.GetAsync(streamConfig);
+							GetStreamResponseV2 response = await CreateGetStreamResponseV2Async(stream, false);
+							responses.Add(response);
+						}
+					}
+				}
+			}
+
+			return responses.OrderBy(x => x.Id).Select(x => PropertyFilter.Apply(x, filter)).ToList();
+		}
+
+		/// <summary>
+		/// Retrieve information about a specific stream.
+		/// </summary>
+		/// <param name="streamId">Id of the stream to get information about</param>
+		/// <param name="filter">Filter for the properties to return</param>
+		/// <param name="cancellationToken">Cancellation token for the operation</param>
+		/// <returns>Information about the requested project</returns>
+		[HttpGet]
+		[Route("/api/v1/streams/{streamId}")]
+		[ProducesResponseType(typeof(GetStreamResponse), 200)]
+		public async Task<ActionResult<object>> GetStreamAsync(StreamId streamId, [FromQuery] PropertyFilter? filter = null, CancellationToken cancellationToken = default)
+		{
+			StreamConfig? streamConfig;
+			if (!_globalConfig.Value.TryGetStream(streamId, out streamConfig))
+			{
+				return NotFound(streamId);
+			}
+			if (!streamConfig.Authorize(StreamAclAction.ViewStream, User))
+			{
+				return Forbid(StreamAclAction.ViewStream, streamId);
+			}
+
+			IStream stream = await _streamCollection.GetAsync(streamConfig, cancellationToken);
+			return PropertyFilter.Apply(await CreateGetStreamResponseAsync(stream, cancellationToken), filter);
+		}
+
+		/// <summary>
+		/// Retrieve information about a specific stream.
+		/// </summary>
+		/// <param name="streamId">Id of the stream to get information about</param>
+		/// <param name="config">The client's cached config revision. If this matches value matches the current config version of the stream, the config object will be omitted from the response.</param>
+		/// <param name="filter">Filter for the properties to return</param>
+		/// <returns>Information about the requested project</returns>
+		[HttpGet]
+		[Route("/api/v2/streams/{streamId}/config")]
+		[ProducesResponseType(typeof(GetStreamResponseV2), 200)]
+		public async Task<ActionResult<object>> GetStreamV2Async(StreamId streamId, [FromQuery] string? config = null, [FromQuery] PropertyFilter? filter = null)
+		{
+			StreamConfig? streamConfig;
+			if (!_globalConfig.Value.TryGetStream(streamId, out streamConfig))
+			{
+				return NotFound(streamId);
+			}
+			if (!streamConfig.Authorize(StreamAclAction.ViewStream, User))
+			{
+				return Forbid(StreamAclAction.ViewStream, streamId);
+			}
+
+			bool includeConfig = !String.Equals(config, streamConfig.Revision, StringComparison.Ordinal);
+			IStream stream = await _streamCollection.GetAsync(streamConfig);
+			return PropertyFilter.Apply(await CreateGetStreamResponseV2Async(stream, includeConfig), filter);
+		}
+
+		/// <summary>
+		/// Create a stream response object, including all the templates
+		/// </summary>
+		/// <param name="stream">Stream to create response for</param>
+		/// <param name="cancellationToken">Cancellation token for the operation</param>
+		/// <returns>Response object</returns>
+		async Task<GetStreamResponse> CreateGetStreamResponseAsync(IStream stream, CancellationToken cancellationToken)
+		{
+			using TelemetrySpan span = _tracer.StartActiveSpan($"{nameof(StreamsController)}.{nameof(CreateGetStreamResponseAsync)}");
+			span.SetAttribute("streamId", stream.Id);
+
+			List<GetTemplateRefResponse> apiTemplateRefs = new List<GetTemplateRefResponse>();
+			foreach (KeyValuePair<TemplateId, ITemplateRef> pair in stream.Templates)
+			{
+				using TelemetrySpan templateScope = _tracer.StartActiveSpan($"{nameof(StreamsController)}.{nameof(CreateGetStreamResponseAsync)}.Template");
+				templateScope.SetAttribute("templateName", pair.Value.Config.Name);
+
+				ITemplateRef templateRef = pair.Value;
+				if (templateRef.Config.Authorize(StreamAclAction.ViewTemplate, User))
+				{
+					List<GetTemplateStepStateResponse>? stepStates = null;
+					if (templateRef.StepStates != null)
+					{
+						for (int i = 0; i < templateRef.StepStates.Count; i++)
+						{
+							ITemplateStep state = templateRef.StepStates[i];
+
+							stepStates ??= new List<GetTemplateStepStateResponse>();
+
+							GetThinUserInfoResponse? pausedByUserInfo = (await _userCollection.GetCachedUserAsync(state.PausedByUserId, cancellationToken))?.ToThinApiResponse();
+							stepStates.Add(new GetTemplateStepStateResponse(state, pausedByUserInfo));
+						}
+					}
+
+					ITemplate? template = await _templateCollection.GetOrAddAsync(pair.Value.Config);
+					apiTemplateRefs.Add(new GetTemplateRefResponse(pair.Key, pair.Value, template, stepStates, _timeZone));
+				}
+			}
+
+			return stream.ToApiResponse(apiTemplateRefs);
+		}
+
+		async Task<GetStreamResponseV2> CreateGetStreamResponseV2Async(IStream stream, bool includeConfig)
+		{
+			List<GetTemplateRefResponseV2> templates = new List<GetTemplateRefResponseV2>();
+			foreach (ITemplateRef templateRef in stream.Templates.Values)
+			{
+				ITemplate template = await _templateCollection.GetOrAddAsync(templateRef.Config);
+				templates.Add(new GetTemplateRefResponseV2(template, templateRef));
+			}
+			return new GetStreamResponseV2(stream, includeConfig, templates);
+		}
+
+		/// <summary>
+		/// Gets a list of changes for a stream
+		/// </summary>
+		/// <param name="streamId">The stream id</param>
+		/// <param name="min">The starting changelist number</param>
+		/// <param name="max">The ending changelist number</param>
+		/// <param name="results">Number of results to return</param>
+		/// <param name="tags">Tags to filter the changes returned</param>
+		/// <param name="filter">The filter to apply to the results</param>
+		/// <param name="cancellationToken">Cancellation token for the operation</param>
+		/// <returns>Http result code</returns>
+		[HttpGet]
+		[Route("/api/v1/streams/{streamId}/changes")]
+		[ProducesResponseType(typeof(List<GetCommitResponse>), 200)]
+		public async Task<ActionResult<List<object>>> GetChangesAsync(StreamId streamId, [FromQuery] int? min = null, [FromQuery] int? max = null, [FromQuery] int results = 50, [FromQuery] string? tags = null, PropertyFilter? filter = null, CancellationToken cancellationToken = default)
+		{
+			StreamConfig? streamConfig;
+			if (!_globalConfig.Value.TryGetStream(streamId, out streamConfig))
+			{
+				return NotFound(streamId);
+			}
+			if (!streamConfig.Authorize(StreamAclAction.ViewChanges, User))
+			{
+				return Forbid(StreamAclAction.ViewChanges, streamId);
+			}
+
+			List<CommitTag>? commitTags = null;
+			if (tags != null)
+			{
+				commitTags = tags.Split(';', StringSplitOptions.RemoveEmptyEntries).Select(x => new CommitTag(x)).ToList();
+			}
+
+			List<ICommit> commits = await _commitService.GetCollection(streamConfig).FindAsync(min, max, results, commitTags, cancellationToken).ToListAsync(cancellationToken);
+
+			List<GetCommitResponse> responses = new List<GetCommitResponse>();
+			foreach (ICommit commit in commits)
+			{
+				IUser? author = await _userCollection.GetCachedUserAsync(commit.AuthorId, cancellationToken);
+				responses.Add(CreateGetCommitResponse(commit, author!, null, null));
+			}
+			return responses.ConvertAll(x => PropertyFilter.Apply(x, filter));
+		}
+
+		static GetCommitResponse CreateGetCommitResponse(ICommit commit, IUser author, IReadOnlyList<CommitTag>? tags, IReadOnlyList<string>? files)
+		{
+			GetCommitResponse response = new GetCommitResponse(commit.Number, author.ToThinApiResponse(), commit.Description);
+			if (tags != null)
+			{
+				response.Tags = new List<CommitTag>(tags);
+			}
+			if (files != null)
+			{
+				response.Files = new List<string>(files);
+			}
+			return response;
+		}
+
+		/// <summary>
+		/// Gets a list of changes for a stream
+		/// </summary>
+		/// <param name="streamId">The stream id</param>
+		/// <param name="changeNumber">The changelist number</param>
+		/// <param name="maxFiles">Maximum number of files to return</param>
+		/// <param name="filter">The filter to apply to the results</param>
+		/// <param name="cancellationToken">Cancellation token for the operation</param>
+		/// <returns>Http result code</returns>
+		[HttpGet]
+		[Route("/api/v1/streams/{streamId}/changes/{changeNumber}")]
+		[ProducesResponseType(typeof(GetCommitResponse), 200)]
+		public async Task<ActionResult<object>> GetChangeDetailsAsync(StreamId streamId, int changeNumber, int maxFiles = 100, PropertyFilter? filter = null, CancellationToken cancellationToken = default)
+		{
+			StreamConfig? streamConfig;
+			if (!_globalConfig.Value.TryGetStream(streamId, out streamConfig))
+			{
+				return NotFound(streamId);
+			}
+			if (!streamConfig.Authorize(StreamAclAction.ViewChanges, User))
+			{
+				return Forbid(StreamAclAction.ViewChanges, streamId);
+			}
+
+			ICommit? changeDetails = await _commitService.GetCollection(streamConfig).GetAsync(changeNumber, cancellationToken);
+			if (changeDetails == null)
+			{
+				return NotFound("CL {Change} not found in stream {StreamId}", changeNumber, streamId);
+			}
+
+			IUser? author = await _userCollection.GetCachedUserAsync(changeDetails.AuthorId, cancellationToken);
+			IReadOnlyList<CommitTag> tags = await changeDetails.GetTagsAsync(cancellationToken);
+			IReadOnlyList<string> files = await changeDetails.GetFilesAsync(maxFiles, cancellationToken);
+
+			return PropertyFilter.Apply(CreateGetCommitResponse(changeDetails, author!, tags, files), filter);
+		}
+
+		/// <summary>
+		/// Gets the history of a step in the stream
+		/// </summary>
+		/// <param name="streamId">The stream id</param>
+		/// <param name="templateId"></param>
+		/// <param name="step">Name of the step to search for</param>
+		/// <param name="change">Maximum changelist number to return</param>
+		/// <param name="count">Number of results to return</param>
+		/// <param name="filter">The filter to apply to the results</param>
+		/// <returns>Http result code</returns>
+		[HttpGet]
+		[Route("/api/v1/streams/{streamId}/history")]
+		[ProducesResponseType(typeof(List<GetJobStepRefResponse>), 200)]
+		public async Task<ActionResult<List<object>>> GetStepHistoryAsync(StreamId streamId, [FromQuery] string templateId, [FromQuery] string step, [FromQuery] int? change = null, [FromQuery] int count = 10, [FromQuery] PropertyFilter? filter = null)
+		{
+			StreamConfig? streamConfig;
+			if (!_globalConfig.Value.TryGetStream(streamId, out streamConfig))
+			{
+				return NotFound(streamId);
+			}
+			if (!streamConfig.Authorize(JobAclAction.ViewJob, User))
+			{
+				return Forbid(JobAclAction.ViewJob, streamId);
+			}
+
+			TemplateId templateIdValue = new TemplateId(templateId);
+
+			List<IJobStepRef> steps = await _jobStepRefCollection.GetStepsForNodeAsync(streamId, templateIdValue, step, change, true, count);
+			return steps.ConvertAll(x => PropertyFilter.Apply(CreateGetJobStepRefResponse(x), filter));
+		}
+
+		/// <summary>
+		/// Constructor
+		/// </summary>
+		/// <param name="jobStepRef">The jobstep ref to construct from</param>
+		internal static GetJobStepRefResponse CreateGetJobStepRefResponse(IJobStepRef jobStepRef)
+		{
+			GetJobStepRefResponse response = new GetJobStepRefResponse();
+			response.JobId = jobStepRef.Id.JobId;
+			response.BatchId = jobStepRef.Id.BatchId;
+			response.StepId = jobStepRef.Id.StepId;
+			response.Change = jobStepRef.Change;
+			response.LogId = jobStepRef.LogId.ToString();
+			response.PoolId = jobStepRef.PoolId?.ToString();
+			response.AgentId = jobStepRef.AgentId?.ToString();
+			response.Outcome = jobStepRef.Outcome;
+			response.IssueIds = jobStepRef.IssueIds?.Select(id => id).ToList();
+			response.StartTime = jobStepRef.StartTimeUtc;
+			response.FinishTime = jobStepRef.FinishTimeUtc;
+			return response;
+		}
+
+		/// <summary>
+		/// Gets a template for a stream
+		/// </summary>
+		/// <param name="streamId">Unique id of the stream to query</param>
+		/// <param name="templateId">Unique id of the template to query</param>
+		/// <param name="filter">Filter for properties to return</param>
+		/// <returns>Information about all the templates</returns>
+		[HttpGet]
+		[Route("/api/v1/streams/{streamId}/templates/{templateId}")]
+		[ProducesResponseType(typeof(List<GetTemplateResponse>), 200)]
+		public async Task<ActionResult<object>> GetTemplateAsync(StreamId streamId, TemplateId templateId, [FromQuery] PropertyFilter? filter = null)
+		{
+			StreamConfig? streamConfig;
+			if (!_globalConfig.Value.TryGetStream(streamId, out streamConfig))
+			{
+				return NotFound(streamId);
+			}
+
+			TemplateRefConfig? templateConfig;
+			if (!streamConfig.TryGetTemplate(templateId, out templateConfig))
+			{
+				return NotFound(streamId, templateId);
+			}
+			if (!templateConfig.Authorize(StreamAclAction.ViewTemplate, User))
+			{
+				return Forbid(StreamAclAction.ViewTemplate, streamId);
+			}
+
+			ITemplate template = await _templateCollection.GetOrAddAsync(templateConfig);
+			return new GetTemplateResponse(template).ApplyFilter(filter);
+		}
+
+		/// <summary>
+		/// Update a stream template ref
+		/// </summary>
+		[HttpPut]
+		[Authorize]
+		[Route("/api/v1/streams/{streamId}/templates/{templateRefId}")]
+		public async Task<ActionResult> UpdateStreamTemplateRefAsync(StreamId streamId, TemplateId templateRefId, [FromBody] UpdateTemplateRefRequest update)
+		{
+			if (!_globalConfig.Value.Authorize(AdminAclAction.AdminWrite, User))
+			{
+				return Forbid(AdminAclAction.AdminWrite);
+			}
+
+			StreamConfig? streamConfig;
+			if (!_globalConfig.Value.TryGetStream(streamId, out streamConfig))
+			{
+				return NotFound(streamId);
+			}
+
+			IStream stream = await _streamCollection.GetAsync(streamConfig);
+			if (!stream.Templates.ContainsKey(templateRefId))
+			{
+				return NotFound(streamId, templateRefId);
+			}
+
+			await _streamCollection.TryUpdateTemplateRefAsync(stream, templateRefId, update.StepStates);
+			return Ok();
+		}
+	}
+}
diff --git a/Engine/Source/Programs/Horde/HordeServer/Tasks/RestartTaskSource.cs b/Engine/Source/Programs/Horde/HordeServer/Tasks/RestartTaskSource.cs
new file mode 100644
--- /dev/null
+++ b/Engine/Source/Programs/Horde/HordeServer/Tasks/RestartTaskSource.cs
@@ -0,0 +1,57 @@
+// Copyright Epic Games, Inc. All Rights Reserved.
+
+using System.Threading;
+using System.Threading.Tasks;
+using EpicGames.Horde.Agents.Leases;
+using EpicGames.Horde.Jobs;
+using EpicGames.Horde.Logs;
+using Google.Protobuf;
+using Google.Protobuf.WellKnownTypes;
+using Horde.Server.Agents;
+using Horde.Server.Logs;
+using Horde.Server.Utilities;
+using HordeCommon.Rpc.Tasks;
+
+namespace Horde.Server.Tasks
+{
+	class RestartTaskSource : TaskSourceBase<RestartTask>
+	{
+		public override string Type => "Restart";
+
+		public override TaskSourceFlags Flags => TaskSourceFlags.AllowWhenDisabled | TaskSourceFlags.AllowDuringDowntime;
+
+		readonly ILogCollection _logCollection;
+
+		public RestartTaskSource(ILogCollection logCollection)
+		{
+			_logCollection = logCollection;
+
+			OnLeaseStartedProperties.Add(nameof(RestartTask.LogId), x => LogId.Parse(x.LogId));
+		}
+
+		public override async Task<Task<AgentLease?>> AssignLeaseAsync(IAgent agent, CancellationToken cancellationToken)
+		{
+			if (!agent.RequestForceRestart)
+			{
+				if (!agent.RequestRestart)
+				{
+					return SkipAsync(cancellationToken);
+				}
+				if (agent.Leases.Count > 0)
+				{
+					return await DrainAsync(cancellationToken);
+				}
+			}
+
+			LeaseId leaseId = new LeaseId(BinaryIdUtils.CreateNew());
+			ILog log = await _logCollection.AddAsync(JobId.Empty, leaseId, agent.SessionId, LogType.Json, cancellationToken: cancellationToken);
+
+			RestartTask task = new RestartTask();
+			task.LogId = log.Id.ToString();
+
+			byte[] payload = Any.Pack(task).ToByteArray();
+
+			return LeaseAsync(new AgentLease(leaseId, null, "Restart", null, null, log.Id, LeaseState.Pending, null, true, payload));
+		}
+	}
+}
diff --git a/Engine/Source/Programs/Horde/HordeServer/Tasks/ShutdownTaskSource.cs b/Engine/Source/Programs/Horde/HordeServer/Tasks/ShutdownTaskSource.cs
new file mode 100644
--- /dev/null
+++ b/Engine/Source/Programs/Horde/HordeServer/Tasks/ShutdownTaskSource.cs
@@ -0,0 +1,53 @@
+// Copyright Epic Games, Inc. All Rights Reserved.
+
+using System.Threading;
+using System.Threading.Tasks;
+using EpicGames.Horde.Agents.Leases;
+using EpicGames.Horde.Jobs;
+using EpicGames.Horde.Logs;
+using Google.Protobuf;
+using Google.Protobuf.WellKnownTypes;
+using Horde.Server.Agents;
+using Horde.Server.Logs;
+using Horde.Server.Utilities;
+using HordeCommon.Rpc.Tasks;
+
+namespace Horde.Server.Tasks
+{
+	class ShutdownTaskSource : TaskSourceBase<ShutdownTask>
+	{
+		public override string Type => "Shutdown";
+
+		public override TaskSourceFlags Flags => TaskSourceFlags.AllowWhenDisabled | TaskSourceFlags.AllowDuringDowntime;
+
+		readonly ILogCollection _logCollection;
+
+		public ShutdownTaskSource(ILogCollection logCollection)
+		{
+			_logCollection = logCollection;
+			OnLeaseStartedProperties.Add(nameof(ShutdownTask.LogId), x => LogId.Parse(x.LogId));
+		}
+
+		public override async Task<Task<AgentLease?>> AssignLeaseAsync(IAgent agent, CancellationToken cancellationToken)
+		{
+			if (!agent.RequestShutdown)
+			{
+				return SkipAsync(cancellationToken);
+			}
+			if (agent.Leases.Count > 0)
+			{
+				return await DrainAsync(cancellationToken);
+			}
+
+			LeaseId leaseId = new LeaseId(BinaryIdUtils.CreateNew());
+			ILog log = await _logCollection.AddAsync(JobId.Empty, leaseId, agent.SessionId, LogType.Json, cancellationToken: cancellationToken);
+
+			ShutdownTask task = new ShutdownTask();
+			task.LogId = log.Id.ToString();
+
+			byte[] payload = Any.Pack(task).ToByteArray();
+
+			return LeaseAsync(new AgentLease(leaseId, null, "Shutdown", null, null, log.Id, LeaseState.Pending, null, true, payload));
+		}
+	}
+}
diff --git a/Engine/Source/Programs/Horde/HordeServer/Tasks/UpgradeTaskSource.cs b/Engine/Source/Programs/Horde/HordeServer/Tasks/UpgradeTaskSource.cs
new file mode 100644
--- /dev/null
+++ b/Engine/Source/Programs/Horde/HordeServer/Tasks/UpgradeTaskSource.cs
@@ -0,0 +1,112 @@
+// Copyright Epic Games, Inc. All Rights Reserved.
+
+using System;
+using System.Text;
+using System.Threading;
+using System.Threading.Tasks;
+using EpicGames.Core;
+using EpicGames.Horde.Agents.Leases;
+using EpicGames.Horde.Jobs;
+using EpicGames.Horde.Logs;
+using EpicGames.Horde.Tools;
+using Google.Protobuf;
+using Google.Protobuf.WellKnownTypes;
+using Horde.Server.Agents;
+using Horde.Server.Logs;
+using Horde.Server.Server;
+using Horde.Server.Tools;
+using Horde.Server.Utilities;
+using HordeCommon.Rpc.Tasks;
+using Microsoft.Extensions.Options;
+
+namespace Horde.Server.Tasks
+{
+	class UpgradeTaskSource : TaskSourceBase<UpgradeTask>
+	{
+		public override string Type => "Upgrade";
+
+		public override TaskSourceFlags Flags => TaskSourceFlags.AllowWhenDisabled | TaskSourceFlags.AllowDuringDowntime;
+
+		readonly IToolCollection _toolCollection;
+		readonly ILogCollection _logCollection;
+		readonly IOptionsMonitor<GlobalConfig> _globalConfig;
+		readonly IOptions<ServerSettings> _serverSettings;
+		readonly IClock _clock;
+
+		public UpgradeTaskSource(IToolCollection toolCollection, ILogCollection logCollection, IOptionsMonitor<GlobalConfig> globalConfig, IOptions<ServerSettings> serverSettings, IClock clock)
+		{
+			_toolCollection = toolCollection;
+			_logCollection = logCollection;
+			_globalConfig = globalConfig;
+			_serverSettings = serverSettings;
+			_clock = clock;
+
+			OnLeaseStartedProperties.Add(nameof(UpgradeTask.LogId), x => LogId.Parse(x.LogId));
+		}
+
+		public override async Task<Task<AgentLease?>> AssignLeaseAsync(IAgent agent, CancellationToken cancellationToken)
+		{
+			if (!_serverSettings.Value.EnableUpgradeTasks)
+			{
+				return SkipAsync(cancellationToken);
+			}
+
+			(ITool, IToolDeployment)? required = await GetRequiredSoftwareVersionAsync(agent, cancellationToken);
+			if (required == null)
+			{
+				return SkipAsync(cancellationToken);
+			}
+
+			(ITool tool, IToolDeployment deployment) = required.Value;
+			if (agent.Version == deployment.Version)
+			{
+				return SkipAsync(cancellationToken);
+			}
+
+			if (agent.Leases.Count > 0 || (agent.LastUpgradeTime != null && agent.LastUpgradeVersion == deployment.Version && _clock.UtcNow < agent.LastUpgradeTime.Value + TimeSpan.FromMinutes(5.0)))
+			{
+				return await DrainAsync(cancellationToken);
+			}
+
+			LeaseId leaseId = new LeaseId(BinaryIdUtils.CreateNew());
+			ILog log = await _logCollection.AddAsync(JobId.Empty, leaseId, agent.SessionId, LogType.Json, cancellationToken: cancellationToken);
+
+			UpgradeTask task = new UpgradeTask();
+			task.SoftwareId = $"{tool.Id}:{deployment.Version}";
+			task.LogId = log.Id.ToString();
+
+			byte[] payload = Any.Pack(task).ToByteArray();
+			return LeaseAsync(new AgentLease(leaseId, null, $"Upgrade to {tool.Id} {deployment.Version}", null, null, log.Id, LeaseState.Pending, null, true, payload));
+		}
+
+		/// <summary>
+		/// Determines the client software version that should be installed on an agent
+		/// </summary>
+		/// <param name="agent">The agent instance</param>
+		/// <param name="cancellationToken">Cancellation token for the operation</param>
+		/// <returns>Unique id of the client version this agent should be running</returns>
+		public async Task<(ITool, IToolDeployment)?> GetRequiredSoftwareVersionAsync(IAgent agent, CancellationToken cancellationToken)
+		{
+			GlobalConfig globalConfig = _globalConfig.CurrentValue;
+
+			ToolId toolId = agent.GetSoftwareToolId(globalConfig);
+
+			ITool? tool = await _toolCollection.GetAsync(toolId, cancellationToken);
+			if (tool == null)
+			{
+				return null;
+			}
+
+			uint value = BuzHash.Add(0, Encoding.UTF8.GetBytes(agent.Id.ToString()));
+			double phase = (value % 10000) / 10000.0;
+
+			IToolDeployment? deployment = tool.GetCurrentDeployment(phase, _clock.UtcNow);
+			if (deployment == null)
+			{
+				return null;
+			}
+
+			return (tool, deployment);
+		}
+	}
+}
diff --git a/Engine/Source/Programs/Horde/HordeServer/Utilities/ArtifactChunkStream.cs b/Engine/Source/Programs/Horde/HordeServer/Utilities/ArtifactChunkStream.cs
new file mode 100644
--- /dev/null
+++ b/Engine/Source/Programs/Horde/HordeServer/Utilities/ArtifactChunkStream.cs
@@ -0,0 +1,130 @@
+// Copyright Epic Games, Inc. All Rights Reserved.
+
+using System;
+using System.IO;
+using System.Threading;
+using System.Threading.Tasks;
+using Grpc.Core;
+using Horde.Common.Rpc;
+
+namespace Horde.Server.Utilities
+{
+	/// <summary>
+	/// Stream that takes in rpc chunks from an upload and copies them?
+	/// </summary>
+	public class ArtifactChunkStream : Stream
+	{
+		/// <summary>
+		/// The grpc client reader. Should probably be templatized
+		/// </summary>
+		readonly IAsyncStreamReader<RpcUploadArtifactRequest> _reader;
+
+		/// <summary>
+		/// Position within the stream
+		/// </summary>
+		long _streamPosition;
+
+		/// <summary>
+		/// The length of the stream
+		/// </summary>
+		readonly long _streamLength;
+
+		/// <summary>
+		/// The current request being read from
+		/// </summary>
+		RpcUploadArtifactRequest? _request;
+
+		/// <summary>
+		/// Position within the current request
+		/// </summary>
+		int _requestPos;
+
+		/// <summary>
+		/// Constructor
+		/// </summary>
+		/// <param name="reader">the grpc reader</param>
+		/// <param name="length">filesize reported by the client</param>
+		public ArtifactChunkStream(IAsyncStreamReader<RpcUploadArtifactRequest> reader, long length)
+		{
+			_reader = reader;
+			_streamLength = length;
+		}
+
+		/// <inheritdoc/>
+		public override long Position
+		{
+			get => _streamPosition;
+			set => throw new NotSupportedException();
+		}
+
+		/// <inheritdoc/>
+		public override long Length => _streamLength;
+
+		/// <inheritdoc/>
+		public override void SetLength(long value) => throw new NotSupportedException();
+
+		/// <inheritdoc/>
+		public override bool CanSeek => false;
+
+		/// <inheritdoc/>
+		public override long Seek(long offset, SeekOrigin origin) => throw new NotSupportedException();
+
+		/// <inheritdoc/>
+		public override bool CanWrite => false;
+
+		/// <inheritdoc/>
+		public override void Write(byte[] buffer, int offset, int count) => throw new NotSupportedException();
+
+		/// <inheritdoc/>
+		public override bool CanRead => true;
+
+		/// <inheritdoc/>
+		public override int Read(byte[] buffer, int offset, int count)
+		{
+#pragma warning disable VSTHRD002
+			return ReadAsync(buffer, offset, count, CancellationToken.None).Result;
+#pragma warning restore VSTHRD002
+		}
+
+		/// <inheritdoc/>
+		public override async Task<int> ReadAsync(byte[] buffer, int offset, int count, CancellationToken cancellationToken)
+		{
+			return await ReadAsync(buffer.AsMemory(offset, count), cancellationToken);
+		}
+
+		/// <inheritdoc/>
+		public override async ValueTask<int> ReadAsync(Memory<byte> buffer, CancellationToken cancellationToken)
+		{
+			int bytesRead = 0;
+			while (bytesRead < buffer.Length && _streamPosition < _streamLength)
+			{
+				if (_request == null || _requestPos == _request.Data.Length)
+				{
+					// Read the next request object
+					if (!await _reader.MoveNext())
+					{
+						throw new EndOfStreamException("Unexpected end of stream while reading artifact");
+					}
+
+					_request = _reader.Current;
+					_requestPos = 0;
+				}
+				else
+				{
+					// Copy data from the current request object
+					int numBytesToCopy = Math.Min(buffer.Length - bytesRead, _request.Data.Length - _requestPos);
+					_request.Data.Span.Slice(_requestPos, numBytesToCopy).CopyTo(buffer.Slice(bytesRead, numBytesToCopy).Span);
+					_requestPos += numBytesToCopy;
+					bytesRead += numBytesToCopy;
+					_streamPosition += numBytesToCopy;
+				}
+			}
+			return bytesRead;
+		}
+
+		/// <inheritdoc/>
+		public override void Flush()
+		{
+		}
+	}
+}
diff --git a/Engine/Source/Programs/Horde/HordeServer/Utilities/Schemas.cs b/Engine/Source/Programs/Horde/HordeServer/Utilities/Schemas.cs
new file mode 100644
--- /dev/null
+++ b/Engine/Source/Programs/Horde/HordeServer/Utilities/Schemas.cs
@@ -0,0 +1,33 @@
+// Copyright Epic Games, Inc. All Rights Reserved.
+
+using System;
+using System.Collections.Concurrent;
+using EpicGames.Core;
+
+namespace Horde.Server.Utilities
+{
+	/// <summary>
+	/// Utility functions for generating schemas
+	/// </summary>
+	static class Schemas
+	{
+		static readonly XmlDocReader s_cachedReader = new XmlDocReader();
+		static readonly ConcurrentDictionary<Type, JsonSchema> s_cachedSchemas = new ConcurrentDictionary<Type, JsonSchema>();
+
+		/// <summary>
+		/// Create a Json schema (or retrieve a cached schema)
+		/// </summary>
+		/// <param name="type"></param>
+		/// <returns></returns>
+		public static JsonSchema CreateSchema(Type type)
+		{
+			JsonSchema? schema;
+			if (!s_cachedSchemas.TryGetValue(type, out schema))
+			{
+				schema = JsonSchema.FromType(type, s_cachedReader);
+				s_cachedSchemas.TryAdd(type, schema);
+			}
+			return schema;
+		}
+	}
+}
diff --git a/Engine/Source/Programs/Horde/Scripts/test.sh b/Engine/Source/Programs/Horde/Scripts/test.sh
--- a/Engine/Source/Programs/Horde/Scripts/test.sh
+++ b/Engine/Source/Programs/Horde/Scripts/test.sh
@@ -14,7 +14,7 @@ testProjects=(
 	"Source/Programs/Shared/EpicGames.IoHash.Tests/EpicGames.IoHash.Tests.csproj"
 	"Source/Programs/Shared/EpicGames.Redis.Tests/EpicGames.Redis.Tests.csproj"
 	"Source/Programs/Shared/EpicGames.Serialization.Tests/EpicGames.Serialization.Tests.csproj"
-	"Source/Programs/Horde/Horde.Server.Tests/Horde.Server.Tests.csproj"
+	"Source/Programs/Horde/HordeServer.Tests/HordeServer.Tests.csproj"
 )
 
 for csProj in "${testProjects[@]}"; do
