Horde: Rename namespaces to match project folders.

diff --git a/Engine/Source/Programs/Horde/Docs/Internals/StructuredLogging.md b/Engine/Source/Programs/Horde/Docs/Internals/StructuredLogging.md
--- a/Engine/Source/Programs/Horde/Docs/Internals/StructuredLogging.md
+++ b/Engine/Source/Programs/Horde/Docs/Internals/StructuredLogging.md
@@ -91,4 +91,4 @@ To implement a new matcher for plain-text log output, create a class that implem
 from `EpicGames.Core`, and ensure it's registered 
 
 When adding or modifying log parsers, we strongly recommend running (and writing) tests in the `UnrealBuildTool.Tests`
-and `Horde.Agent.Tests` projects to check interaction with other handlers.
+and `HordeAgent.Tests` projects to check interaction with other handlers.
diff --git a/Engine/Source/Programs/Horde/HordeAgent.Tests/Relay/GrpcTestHelpers.cs b/Engine/Source/Programs/Horde/HordeAgent.Tests/Relay/GrpcTestHelpers.cs
--- a/Engine/Source/Programs/Horde/HordeAgent.Tests/Relay/GrpcTestHelpers.cs
+++ b/Engine/Source/Programs/Horde/HordeAgent.Tests/Relay/GrpcTestHelpers.cs
@@ -25,7 +25,7 @@ using System.Threading.Channels;
 using System.Threading.Tasks;
 using Grpc.Core;
 
-namespace Horde.Agent.Tests.Relay
+namespace HordeAgent.Tests.Relay
 {
 	public class TestServerStreamWriter<T> : IServerStreamWriter<T> where T : class
 	{
diff --git a/Engine/Source/Programs/Horde/HordeAgent.Tests/Services/AwsInstanceLifecycleServiceTests.cs b/Engine/Source/Programs/Horde/HordeAgent.Tests/Services/AwsInstanceLifecycleServiceTests.cs
--- a/Engine/Source/Programs/Horde/HordeAgent.Tests/Services/AwsInstanceLifecycleServiceTests.cs
+++ b/Engine/Source/Programs/Horde/HordeAgent.Tests/Services/AwsInstanceLifecycleServiceTests.cs
@@ -7,14 +7,14 @@ using System.Net.Http;
 using System.Threading;
 using System.Threading.Tasks;
 using EpicGames.Core;
-using Horde.Agent.Services;
+using HordeAgent.Services;
 using Microsoft.Extensions.Logging;
 using Microsoft.Extensions.Options;
 using Microsoft.VisualStudio.TestTools.UnitTesting;
 using Moq;
 using Moq.Protected;
 
-namespace Horde.Agent.Tests.Services;
+namespace HordeAgent.Tests.Services;
 
 /// <summary>
 /// Fake implementation of the AWS EC2 metadata server (IMDS)
diff --git a/Engine/Source/Programs/Horde/HordeAgent.Tests/Services/TelemetryServiceTest.cs b/Engine/Source/Programs/Horde/HordeAgent.Tests/Services/TelemetryServiceTest.cs
--- a/Engine/Source/Programs/Horde/HordeAgent.Tests/Services/TelemetryServiceTest.cs
+++ b/Engine/Source/Programs/Horde/HordeAgent.Tests/Services/TelemetryServiceTest.cs
@@ -4,12 +4,12 @@ using System;
 using System.Collections.Generic;
 using System.Threading;
 using System.Threading.Tasks;
-using Horde.Agent.Services;
+using HordeAgent.Services;
 using Microsoft.Extensions.Logging.Abstractions;
 using Microsoft.Extensions.Options;
 using Microsoft.VisualStudio.TestTools.UnitTesting;
 
-namespace Horde.Agent.Tests.Services;
+namespace HordeAgent.Tests.Services;
 
 [TestClass]
 public sealed class TelemetryServiceTest : IDisposable
diff --git a/Engine/Source/Programs/Horde/HordeAgent/AgentApp.cs b/Engine/Source/Programs/Horde/HordeAgent/AgentApp.cs
--- a/Engine/Source/Programs/Horde/HordeAgent/AgentApp.cs
+++ b/Engine/Source/Programs/Horde/HordeAgent/AgentApp.cs
@@ -11,10 +11,10 @@ using Datadog.Trace.Configuration;
 using Datadog.Trace.OpenTracing;
 using EpicGames.Core;
 using EpicGames.Horde;
-using Horde.Agent.Leases;
-using Horde.Agent.Leases.Handlers;
-using Horde.Agent.Services;
-using Horde.Agent.Utility;
+using HordeAgent.Leases;
+using HordeAgent.Leases.Handlers;
+using HordeAgent.Services;
+using HordeAgent.Utility;
 using Microsoft.Extensions.Configuration;
 using Microsoft.Extensions.DependencyInjection;
 using Microsoft.Extensions.Logging;
@@ -22,7 +22,7 @@ using Microsoft.Win32;
 using OpenTracing.Util;
 using Polly;
 
-namespace Horde.Agent
+namespace HordeAgent
 {
 	using ITracer = OpenTracing.ITracer;
 
diff --git a/Engine/Source/Programs/Horde/HordeAgent/AgentSettings.cs b/Engine/Source/Programs/Horde/HordeAgent/AgentSettings.cs
--- a/Engine/Source/Programs/Horde/HordeAgent/AgentSettings.cs
+++ b/Engine/Source/Programs/Horde/HordeAgent/AgentSettings.cs
@@ -5,7 +5,7 @@ using EpicGames.Core;
 using EpicGames.Horde;
 using Microsoft.Extensions.Configuration;
 
-namespace Horde.Agent
+namespace HordeAgent
 {
 	/// <summary>
 	/// Describes a network share to mount
diff --git a/Engine/Source/Programs/Horde/HordeAgent/Commands/Compute/AwsLambdaClient.cs b/Engine/Source/Programs/Horde/HordeAgent/Commands/Compute/AwsLambdaClient.cs
--- a/Engine/Source/Programs/Horde/HordeAgent/Commands/Compute/AwsLambdaClient.cs
+++ b/Engine/Source/Programs/Horde/HordeAgent/Commands/Compute/AwsLambdaClient.cs
@@ -4,7 +4,7 @@ using System.Net;
 using System.Text.Json;
 using Microsoft.Extensions.Logging;
 
-namespace Horde.Agent.Commands.Compute
+namespace HordeAgent.Commands.Compute
 {
 	/// <summary>
 	/// Class representing the HTTP response from next invocation request
diff --git a/Engine/Source/Programs/Horde/HordeAgent/Commands/Compute/ComputeWorkerCommand.cs b/Engine/Source/Programs/Horde/HordeAgent/Commands/Compute/ComputeWorkerCommand.cs
--- a/Engine/Source/Programs/Horde/HordeAgent/Commands/Compute/ComputeWorkerCommand.cs
+++ b/Engine/Source/Programs/Horde/HordeAgent/Commands/Compute/ComputeWorkerCommand.cs
@@ -8,7 +8,7 @@ using EpicGames.Horde.Compute;
 using EpicGames.Horde.Compute.Transports;
 using Microsoft.Extensions.Logging;
 
-namespace Horde.Agent.Commands.Compute
+namespace HordeAgent.Commands.Compute
 {
 	/// <summary>
 	/// Helper command for hosting a local compute worker in a separate process
diff --git a/Engine/Source/Programs/Horde/HordeAgent/Commands/Service/ServiceInstall.cs b/Engine/Source/Programs/Horde/HordeAgent/Commands/Service/ServiceInstall.cs
--- a/Engine/Source/Programs/Horde/HordeAgent/Commands/Service/ServiceInstall.cs
+++ b/Engine/Source/Programs/Horde/HordeAgent/Commands/Service/ServiceInstall.cs
@@ -4,10 +4,10 @@ using System.ComponentModel;
 using System.Reflection;
 using System.Text;
 using EpicGames.Core;
-using Horde.Agent.Utility;
+using HordeAgent.Utility;
 using Microsoft.Extensions.Logging;
 
-namespace Horde.Agent.Commands.Service
+namespace HordeAgent.Commands.Service
 {
 	/// <summary>
 	/// Installs the agent as a Windows service
diff --git a/Engine/Source/Programs/Horde/HordeAgent/Commands/Service/ServiceRun.cs b/Engine/Source/Programs/Horde/HordeAgent/Commands/Service/ServiceRun.cs
--- a/Engine/Source/Programs/Horde/HordeAgent/Commands/Service/ServiceRun.cs
+++ b/Engine/Source/Programs/Horde/HordeAgent/Commands/Service/ServiceRun.cs
@@ -8,7 +8,7 @@ using Microsoft.Extensions.Hosting;
 using Microsoft.Extensions.Logging;
 using Serilog.Events;
 
-namespace Horde.Agent.Commands.Service
+namespace HordeAgent.Commands.Service
 {
 	/// <summary>
 	/// Runs the agent
@@ -130,7 +130,7 @@ namespace Horde.Agent.Commands.Service
 
 			void WaitForMutex()
 			{
-				using Mutex mutex = new Mutex(false, "Horde.Agent.{C1F30772-CDD3-41E9-A6CE-42356DE7DEE3}");
+				using Mutex mutex = new Mutex(false, "HordeAgent.{C1F30772-CDD3-41E9-A6CE-42356DE7DEE3}");
 				try
 				{
 					if (!mutex.WaitOne(0))
diff --git a/Engine/Source/Programs/Horde/HordeAgent/Commands/Service/ServiceUninstall.cs b/Engine/Source/Programs/Horde/HordeAgent/Commands/Service/ServiceUninstall.cs
--- a/Engine/Source/Programs/Horde/HordeAgent/Commands/Service/ServiceUninstall.cs
+++ b/Engine/Source/Programs/Horde/HordeAgent/Commands/Service/ServiceUninstall.cs
@@ -1,10 +1,10 @@
 // Copyright Epic Games, Inc. All Rights Reserved.
 
 using EpicGames.Core;
-using Horde.Agent.Utility;
+using HordeAgent.Utility;
 using Microsoft.Extensions.Logging;
 
-namespace Horde.Agent.Commands.Service
+namespace HordeAgent.Commands.Service
 {
 	/// <summary>
 	/// Uninstalls the Windows service
diff --git a/Engine/Source/Programs/Horde/HordeAgent/Commands/Service/ServiceUpgrade.cs b/Engine/Source/Programs/Horde/HordeAgent/Commands/Service/ServiceUpgrade.cs
--- a/Engine/Source/Programs/Horde/HordeAgent/Commands/Service/ServiceUpgrade.cs
+++ b/Engine/Source/Programs/Horde/HordeAgent/Commands/Service/ServiceUpgrade.cs
@@ -10,7 +10,7 @@ using EpicGames.Core;
 using Microsoft.Extensions.Logging;
 using Microsoft.Management.Infrastructure;
 
-namespace Horde.Agent.Commands.Service
+namespace HordeAgent.Commands.Service
 {
 	/// <summary>
 	/// Upgrades a running service to the current application
diff --git a/Engine/Source/Programs/Horde/HordeAgent/Commands/Utilities/CapsCommand.cs b/Engine/Source/Programs/Horde/HordeAgent/Commands/Utilities/CapsCommand.cs
--- a/Engine/Source/Programs/Horde/HordeAgent/Commands/Utilities/CapsCommand.cs
+++ b/Engine/Source/Programs/Horde/HordeAgent/Commands/Utilities/CapsCommand.cs
@@ -1,11 +1,11 @@
 // Copyright Epic Games, Inc. All Rights Reserved.
 
 using EpicGames.Core;
-using Horde.Agent.Services;
+using HordeAgent.Services;
 using HordeCommon.Rpc.Messages;
 using Microsoft.Extensions.Logging;
 
-namespace Horde.Agent.Commands.Utilities
+namespace HordeAgent.Commands.Utilities
 {
 	/// <summary>
 	/// Shows capabilities of this agent
diff --git a/Engine/Source/Programs/Horde/HordeAgent/Commands/Utilities/CreateCertCommand.cs b/Engine/Source/Programs/Horde/HordeAgent/Commands/Utilities/CreateCertCommand.cs
--- a/Engine/Source/Programs/Horde/HordeAgent/Commands/Utilities/CreateCertCommand.cs
+++ b/Engine/Source/Programs/Horde/HordeAgent/Commands/Utilities/CreateCertCommand.cs
@@ -6,7 +6,7 @@ using EpicGames.Core;
 using Microsoft.Extensions.Configuration;
 using Microsoft.Extensions.Logging;
 
-namespace Horde.Agent.Commands.Certs
+namespace HordeAgent.Commands.Certs
 {
 	/// <summary>
 	/// Creates a certificate that can be used for server/agent SSL connections
diff --git a/Engine/Source/Programs/Horde/HordeAgent/Commands/Utilities/SetServerCommand.cs b/Engine/Source/Programs/Horde/HordeAgent/Commands/Utilities/SetServerCommand.cs
--- a/Engine/Source/Programs/Horde/HordeAgent/Commands/Utilities/SetServerCommand.cs
+++ b/Engine/Source/Programs/Horde/HordeAgent/Commands/Utilities/SetServerCommand.cs
@@ -3,10 +3,10 @@
 using System.ComponentModel;
 using System.Text.Json.Nodes;
 using EpicGames.Core;
-using Horde.Agent.Utility;
+using HordeAgent.Utility;
 using Microsoft.Extensions.Logging;
 
-namespace Horde.Agent.Commands.Utilities
+namespace HordeAgent.Commands.Utilities
 {
 	using JsonObject = System.Text.Json.Nodes.JsonObject;
 
diff --git a/Engine/Source/Programs/Horde/HordeAgent/IpcMessages.cs b/Engine/Source/Programs/Horde/HordeAgent/IpcMessages.cs
--- a/Engine/Source/Programs/Horde/HordeAgent/IpcMessages.cs
+++ b/Engine/Source/Programs/Horde/HordeAgent/IpcMessages.cs
@@ -5,7 +5,7 @@ using System.Text.Json;
 using System.Text.Json.Serialization;
 using EpicGames.Core;
 
-namespace Horde.Agent
+namespace HordeAgent
 {
 	/// <summary>
 	/// Constants
diff --git a/Engine/Source/Programs/Horde/HordeAgent/Leases/Handlers/ComputeHandler.cs b/Engine/Source/Programs/Horde/HordeAgent/Leases/Handlers/ComputeHandler.cs
--- a/Engine/Source/Programs/Horde/HordeAgent/Leases/Handlers/ComputeHandler.cs
+++ b/Engine/Source/Programs/Horde/HordeAgent/Leases/Handlers/ComputeHandler.cs
@@ -6,13 +6,13 @@ using EpicGames.Horde.Agents.Leases;
 using EpicGames.Horde.Compute;
 using EpicGames.Horde.Compute.Transports;
 using EpicGames.Horde.Logs;
-using Horde.Agent.Services;
+using HordeAgent.Services;
 using HordeCommon.Rpc.Messages;
 using HordeCommon.Rpc.Tasks;
 using Microsoft.Extensions.Logging;
 using Microsoft.Extensions.Options;
 
-namespace Horde.Agent.Leases.Handlers
+namespace HordeAgent.Leases.Handlers
 {
 	/// <summary>
 	/// Handler for compute tasks
diff --git a/Engine/Source/Programs/Horde/HordeAgent/Leases/Handlers/ConformHandler.cs b/Engine/Source/Programs/Horde/HordeAgent/Leases/Handlers/ConformHandler.cs
--- a/Engine/Source/Programs/Horde/HordeAgent/Leases/Handlers/ConformHandler.cs
+++ b/Engine/Source/Programs/Horde/HordeAgent/Leases/Handlers/ConformHandler.cs
@@ -5,12 +5,12 @@ using EpicGames.Horde;
 using EpicGames.Horde.Agents.Leases;
 using EpicGames.Horde.Logs;
 using Google.Protobuf;
-using Horde.Agent.Services;
+using HordeAgent.Services;
 using HordeCommon.Rpc.Messages;
 using HordeCommon.Rpc.Tasks;
 using Microsoft.Extensions.Logging;
 
-namespace Horde.Agent.Leases.Handlers
+namespace HordeAgent.Leases.Handlers
 {
 	class ConformHandler : LeaseHandler<ConformTask>
 	{
diff --git a/Engine/Source/Programs/Horde/HordeAgent/Leases/Handlers/RestartHandler.cs b/Engine/Source/Programs/Horde/HordeAgent/Leases/Handlers/RestartHandler.cs
--- a/Engine/Source/Programs/Horde/HordeAgent/Leases/Handlers/RestartHandler.cs
+++ b/Engine/Source/Programs/Horde/HordeAgent/Leases/Handlers/RestartHandler.cs
@@ -1,13 +1,13 @@
 // Copyright Epic Games, Inc. All Rights Reserved.
 
 using EpicGames.Horde.Agents.Leases;
-using Horde.Agent.Services;
-using Horde.Agent.Utility;
+using HordeAgent.Services;
+using HordeAgent.Utility;
 using HordeCommon.Rpc.Messages;
 using HordeCommon.Rpc.Tasks;
 using Microsoft.Extensions.Logging;
 
-namespace Horde.Agent.Leases.Handlers
+namespace HordeAgent.Leases.Handlers
 {
 	class RestartHandler : LeaseHandler<RestartTask>
 	{
diff --git a/Engine/Source/Programs/Horde/HordeAgent/Leases/Handlers/ShutdownHandler.cs b/Engine/Source/Programs/Horde/HordeAgent/Leases/Handlers/ShutdownHandler.cs
--- a/Engine/Source/Programs/Horde/HordeAgent/Leases/Handlers/ShutdownHandler.cs
+++ b/Engine/Source/Programs/Horde/HordeAgent/Leases/Handlers/ShutdownHandler.cs
@@ -1,13 +1,13 @@
 // Copyright Epic Games, Inc. All Rights Reserved.
 
 using EpicGames.Horde.Agents.Leases;
-using Horde.Agent.Services;
-using Horde.Agent.Utility;
+using HordeAgent.Services;
+using HordeAgent.Utility;
 using HordeCommon.Rpc.Messages;
 using HordeCommon.Rpc.Tasks;
 using Microsoft.Extensions.Logging;
 
-namespace Horde.Agent.Leases.Handlers
+namespace HordeAgent.Leases.Handlers
 {
 	class ShutdownHandler : LeaseHandler<ShutdownTask>
 	{
diff --git a/Engine/Source/Programs/Horde/HordeAgent/Leases/LeaseHandler.cs b/Engine/Source/Programs/Horde/HordeAgent/Leases/LeaseHandler.cs
--- a/Engine/Source/Programs/Horde/HordeAgent/Leases/LeaseHandler.cs
+++ b/Engine/Source/Programs/Horde/HordeAgent/Leases/LeaseHandler.cs
@@ -6,7 +6,7 @@ using EpicGames.Horde.Agents.Leases;
 using EpicGames.Perforce.Managed;
 using Google.Protobuf;
 using Google.Protobuf.WellKnownTypes;
-using Horde.Agent.Services;
+using HordeAgent.Services;
 using HordeCommon.Rpc;
 using HordeCommon.Rpc.Messages;
 using Microsoft.Extensions.Logging;
@@ -14,7 +14,7 @@ using OpenTracing;
 using OpenTracing.Util;
 using ByteString = Google.Protobuf.ByteString;
 
-namespace Horde.Agent.Leases
+namespace HordeAgent.Leases
 {
 	/// <summary>
 	/// Handles execution of a specific lease type
diff --git a/Engine/Source/Programs/Horde/HordeAgent/Leases/LeaseHandlerFactory.cs b/Engine/Source/Programs/Horde/HordeAgent/Leases/LeaseHandlerFactory.cs
--- a/Engine/Source/Programs/Horde/HordeAgent/Leases/LeaseHandlerFactory.cs
+++ b/Engine/Source/Programs/Horde/HordeAgent/Leases/LeaseHandlerFactory.cs
@@ -4,7 +4,7 @@ using Google.Protobuf;
 using Google.Protobuf.Reflection;
 using HordeCommon.Rpc.Messages;
 
-namespace Horde.Agent.Leases
+namespace HordeAgent.Leases
 {
 	abstract class LeaseHandlerFactory
 	{
diff --git a/Engine/Source/Programs/Horde/HordeAgent/Leases/LeaseLoggerFactory.cs b/Engine/Source/Programs/Horde/HordeAgent/Leases/LeaseLoggerFactory.cs
--- a/Engine/Source/Programs/Horde/HordeAgent/Leases/LeaseLoggerFactory.cs
+++ b/Engine/Source/Programs/Horde/HordeAgent/Leases/LeaseLoggerFactory.cs
@@ -5,7 +5,7 @@ using EpicGames.Horde.Agents.Leases;
 using Microsoft.Extensions.Logging;
 using Microsoft.Extensions.Options;
 
-namespace Horde.Agent.Leases
+namespace HordeAgent.Leases
 {
 	/// <summary>
 	/// Creates local log files for leases executed on this agent
diff --git a/Engine/Source/Programs/Horde/HordeAgent/Leases/LeaseManager.cs b/Engine/Source/Programs/Horde/HordeAgent/Leases/LeaseManager.cs
--- a/Engine/Source/Programs/Horde/HordeAgent/Leases/LeaseManager.cs
+++ b/Engine/Source/Programs/Horde/HordeAgent/Leases/LeaseManager.cs
@@ -4,13 +4,13 @@ using System.Diagnostics;
 using EpicGames.Core;
 using Google.Protobuf.WellKnownTypes;
 using Grpc.Core;
-using Horde.Agent.Services;
+using HordeAgent.Services;
 using HordeCommon.Rpc;
 using HordeCommon.Rpc.Messages;
 using Microsoft.Extensions.DependencyInjection;
 using Microsoft.Extensions.Logging;
 
-namespace Horde.Agent.Leases
+namespace HordeAgent.Leases
 {
 	/// <summary>
 	/// Implements the message handling loop for an agent. Runs asynchronously until disposed.
diff --git a/Engine/Source/Programs/Horde/HordeAgent/Services/AwsInstanceLifecycleService.cs b/Engine/Source/Programs/Horde/HordeAgent/Services/AwsInstanceLifecycleService.cs
--- a/Engine/Source/Programs/Horde/HordeAgent/Services/AwsInstanceLifecycleService.cs
+++ b/Engine/Source/Programs/Horde/HordeAgent/Services/AwsInstanceLifecycleService.cs
@@ -3,12 +3,12 @@
 using System.Net;
 using System.Text;
 using EpicGames.Core;
-using Horde.Agent.Utility;
+using HordeAgent.Utility;
 using Microsoft.Extensions.Hosting;
 using Microsoft.Extensions.Logging;
 using Microsoft.Extensions.Options;
 
-namespace Horde.Agent.Services;
+namespace HordeAgent.Services;
 
 /// <summary>
 /// Describes state of local EC2 instance the agent is running on
diff --git a/Engine/Source/Programs/Horde/HordeAgent/Services/CapabilitiesService.cs b/Engine/Source/Programs/Horde/HordeAgent/Services/CapabilitiesService.cs
--- a/Engine/Source/Programs/Horde/HordeAgent/Services/CapabilitiesService.cs
+++ b/Engine/Source/Programs/Horde/HordeAgent/Services/CapabilitiesService.cs
@@ -21,7 +21,7 @@ using OpenTracing;
 using OpenTracing.Util;
 using AddressFamily = System.Net.Sockets.AddressFamily;
 
-namespace Horde.Agent.Services
+namespace HordeAgent.Services
 {
 	/// <summary>
 	/// Implements the message handling loop for an agent. Runs asynchronously until disposed.
diff --git a/Engine/Source/Programs/Horde/HordeAgent/Services/ComputeListenerService.cs b/Engine/Source/Programs/Horde/HordeAgent/Services/ComputeListenerService.cs
--- a/Engine/Source/Programs/Horde/HordeAgent/Services/ComputeListenerService.cs
+++ b/Engine/Source/Programs/Horde/HordeAgent/Services/ComputeListenerService.cs
@@ -8,7 +8,7 @@ using Microsoft.Extensions.Hosting;
 using Microsoft.Extensions.Logging;
 using Microsoft.Extensions.Options;
 
-namespace Horde.Agent.Services
+namespace HordeAgent.Services
 {
 	/// <summary>
 	/// Service which listens for incoming connections that can be assigned for remote execution leases
diff --git a/Engine/Source/Programs/Horde/HordeAgent/Services/GrpcService.cs b/Engine/Source/Programs/Horde/HordeAgent/Services/GrpcService.cs
--- a/Engine/Source/Programs/Horde/HordeAgent/Services/GrpcService.cs
+++ b/Engine/Source/Programs/Horde/HordeAgent/Services/GrpcService.cs
@@ -7,11 +7,11 @@ using Grpc.Core;
 using Grpc.Core.Interceptors;
 using Grpc.Net.Client;
 using Grpc.Net.Client.Configuration;
-using Horde.Agent.Utility;
+using HordeAgent.Utility;
 using Microsoft.Extensions.Logging;
 using Microsoft.Extensions.Options;
 
-namespace Horde.Agent.Services
+namespace HordeAgent.Services
 {
 	/// <summary>
 	/// Service which creates a configured Grpc channel
diff --git a/Engine/Source/Programs/Horde/HordeAgent/Services/SessionFactory.cs b/Engine/Source/Programs/Horde/HordeAgent/Services/SessionFactory.cs
--- a/Engine/Source/Programs/Horde/HordeAgent/Services/SessionFactory.cs
+++ b/Engine/Source/Programs/Horde/HordeAgent/Services/SessionFactory.cs
@@ -9,14 +9,14 @@ using EpicGames.Horde.Agents;
 using EpicGames.Horde.Agents.Sessions;
 using Grpc.Core;
 using Grpc.Net.Client;
-using Horde.Agent.Utility;
+using HordeAgent.Utility;
 using Horde.Common.Rpc;
 using HordeCommon.Rpc;
 using HordeCommon.Rpc.Messages;
 using Microsoft.Extensions.Logging;
 using Microsoft.Extensions.Options;
 
-namespace Horde.Agent.Services
+namespace HordeAgent.Services
 {
 	/// <summary>
 	/// Interface for about the current session. 
diff --git a/Engine/Source/Programs/Horde/HordeAgent/Services/TelemetryService.cs b/Engine/Source/Programs/Horde/HordeAgent/Services/TelemetryService.cs
--- a/Engine/Source/Programs/Horde/HordeAgent/Services/TelemetryService.cs
+++ b/Engine/Source/Programs/Horde/HordeAgent/Services/TelemetryService.cs
@@ -18,7 +18,7 @@ using Microsoft.Extensions.Logging;
 using Microsoft.Extensions.Options;
 using Microsoft.Management.Infrastructure;
 
-namespace Horde.Agent.Services;
+namespace HordeAgent.Services;
 
 /// <summary>
 /// Metrics for CPU usage
diff --git a/Engine/Source/Programs/Horde/HordeAgent/Services/WorkerService.cs b/Engine/Source/Programs/Horde/HordeAgent/Services/WorkerService.cs
--- a/Engine/Source/Programs/Horde/HordeAgent/Services/WorkerService.cs
+++ b/Engine/Source/Programs/Horde/HordeAgent/Services/WorkerService.cs
@@ -2,13 +2,13 @@
 
 using System.Diagnostics;
 using EpicGames.Core;
-using Horde.Agent.Leases;
+using HordeAgent.Leases;
 using HordeCommon.Rpc.Messages;
 using Microsoft.Extensions.DependencyInjection;
 using Microsoft.Extensions.Hosting;
 using Microsoft.Extensions.Logging;
 
-namespace Horde.Agent.Services
+namespace HordeAgent.Services
 {
 	/// <summary>
 	/// Implements the message handling loop for an agent. Runs asynchronously until disposed.
diff --git a/Engine/Source/Programs/Horde/HordeAgent/Utility/CertificateHelper.cs b/Engine/Source/Programs/Horde/HordeAgent/Utility/CertificateHelper.cs
--- a/Engine/Source/Programs/Horde/HordeAgent/Utility/CertificateHelper.cs
+++ b/Engine/Source/Programs/Horde/HordeAgent/Utility/CertificateHelper.cs
@@ -6,7 +6,7 @@ using System.Text;
 using EpicGames.Core;
 using Microsoft.Extensions.Logging;
 
-namespace Horde.Agent.Utility
+namespace HordeAgent.Utility
 {
 	/// <summary>
 	/// Helper functions for dealing with certificates
diff --git a/Engine/Source/Programs/Horde/HordeAgent/Utility/ExceptionUtils.cs b/Engine/Source/Programs/Horde/HordeAgent/Utility/ExceptionUtils.cs
--- a/Engine/Source/Programs/Horde/HordeAgent/Utility/ExceptionUtils.cs
+++ b/Engine/Source/Programs/Horde/HordeAgent/Utility/ExceptionUtils.cs
@@ -2,7 +2,7 @@
 
 using Grpc.Core;
 
-namespace Horde.Agent.Utility
+namespace HordeAgent.Utility
 {
 	/// <summary>
 	/// Utility methods for exceptions
diff --git a/Engine/Source/Programs/Horde/HordeAgent/Utility/VersionInfo.cs.template b/Engine/Source/Programs/Horde/HordeAgent/Utility/VersionInfo.cs.template
--- a/Engine/Source/Programs/Horde/HordeAgent/Utility/VersionInfo.cs.template
+++ b/Engine/Source/Programs/Horde/HordeAgent/Utility/VersionInfo.cs.template
@@ -1,6 +1,6 @@
 // Copyright Epic Games, Inc. All Rights Reserved.
 
-namespace Horde.Agent.Utility;
+namespace HordeAgent.Utility;
 
 /// <summary>
 /// Encodes version by replacing the version pattern below
diff --git a/Engine/Source/Programs/Horde/HordeServer.Tests/Accounts/AccountCollectionTests.cs b/Engine/Source/Programs/Horde/HordeServer.Tests/Accounts/AccountCollectionTests.cs
--- a/Engine/Source/Programs/Horde/HordeServer.Tests/Accounts/AccountCollectionTests.cs
+++ b/Engine/Source/Programs/Horde/HordeServer.Tests/Accounts/AccountCollectionTests.cs
@@ -2,11 +2,11 @@
 
 using System.Threading;
 using System.Threading.Tasks;
-using Horde.Server.Accounts;
+using HordeServer.Accounts;
 using Microsoft.Extensions.DependencyInjection;
 using Microsoft.VisualStudio.TestTools.UnitTesting;
 
-namespace Horde.Server.Tests.Accounts
+namespace HordeServer.Tests.Accounts
 {
 	[TestClass]
 	public class AccountCollectionTests : TestSetup
diff --git a/Engine/Source/Programs/Horde/HordeServer.Tests/Accounts/AccountControllerTests.cs b/Engine/Source/Programs/Horde/HordeServer.Tests/Accounts/AccountControllerTests.cs
--- a/Engine/Source/Programs/Horde/HordeServer.Tests/Accounts/AccountControllerTests.cs
+++ b/Engine/Source/Programs/Horde/HordeServer.Tests/Accounts/AccountControllerTests.cs
@@ -6,8 +6,8 @@ using System.Net;
 using System.Net.Http;
 using System.Text;
 using System.Threading.Tasks;
-using Horde.Server.Accounts;
-using Horde.Server.Users;
+using HordeServer.Accounts;
+using HordeServer.Users;
 using Microsoft.AspNetCore.WebUtilities;
 using Microsoft.Extensions.DependencyInjection;
 using Microsoft.VisualStudio.TestTools.UnitTesting;
@@ -15,7 +15,7 @@ using Microsoft.VisualStudio.TestTools.UnitTesting;
 #pragma warning disable CA2234 // Pass system uri objects instead of strings
 #pragma warning disable CA1307 // Specify StringComparison for clarity
 
-namespace Horde.Server.Tests.Accounts;
+namespace HordeServer.Tests.Accounts;
 
 [TestClass]
 public class AccountControllerTest : IAsyncDisposable
diff --git a/Engine/Source/Programs/Horde/HordeServer.Tests/Acls/AclTests.cs b/Engine/Source/Programs/Horde/HordeServer.Tests/Acls/AclTests.cs
--- a/Engine/Source/Programs/Horde/HordeServer.Tests/Acls/AclTests.cs
+++ b/Engine/Source/Programs/Horde/HordeServer.Tests/Acls/AclTests.cs
@@ -5,15 +5,15 @@ using System.Security.Claims;
 using EpicGames.Horde.Acls;
 using EpicGames.Horde.Projects;
 using EpicGames.Horde.Streams;
-using Horde.Server.Acls;
-using Horde.Server.Jobs;
-using Horde.Server.Logs;
-using Horde.Server.Projects;
-using Horde.Server.Server;
-using Horde.Server.Streams;
+using HordeServer.Acls;
+using HordeServer.Jobs;
+using HordeServer.Logs;
+using HordeServer.Projects;
+using HordeServer.Server;
+using HordeServer.Streams;
 using Microsoft.VisualStudio.TestTools.UnitTesting;
 
-namespace Horde.Server.Tests.Acls
+namespace HordeServer.Tests.Acls
 {
 	[TestClass]
 	public class AclTests : TestSetup
diff --git a/Engine/Source/Programs/Horde/HordeServer.Tests/Agents/AgentCollectionTests.cs b/Engine/Source/Programs/Horde/HordeServer.Tests/Agents/AgentCollectionTests.cs
--- a/Engine/Source/Programs/Horde/HordeServer.Tests/Agents/AgentCollectionTests.cs
+++ b/Engine/Source/Programs/Horde/HordeServer.Tests/Agents/AgentCollectionTests.cs
@@ -6,11 +6,11 @@ using System.Threading.Tasks;
 using EpicGames.Horde.Agents;
 using EpicGames.Horde.Agents.Leases;
 using EpicGames.Horde.Agents.Pools;
-using Horde.Server.Agents;
-using Horde.Server.Agents.Sessions;
+using HordeServer.Agents;
+using HordeServer.Agents.Sessions;
 using Microsoft.VisualStudio.TestTools.UnitTesting;
 
-namespace Horde.Server.Tests.Agents;
+namespace HordeServer.Tests.Agents;
 
 [TestClass]
 public class AgentCollectionTests : TestSetup
diff --git a/Engine/Source/Programs/Horde/HordeServer.Tests/Agents/AgentControllerTest.cs b/Engine/Source/Programs/Horde/HordeServer.Tests/Agents/AgentControllerTest.cs
--- a/Engine/Source/Programs/Horde/HordeServer.Tests/Agents/AgentControllerTest.cs
+++ b/Engine/Source/Programs/Horde/HordeServer.Tests/Agents/AgentControllerTest.cs
@@ -2,11 +2,11 @@
 
 using System.Threading.Tasks;
 using EpicGames.Horde.Agents;
-using Horde.Server.Agents;
+using HordeServer.Agents;
 using Microsoft.AspNetCore.Mvc;
 using Microsoft.VisualStudio.TestTools.UnitTesting;
 
-namespace Horde.Server.Tests.Agents;
+namespace HordeServer.Tests.Agents;
 
 [TestClass]
 public class AgentControllerDbTest : TestSetup
diff --git a/Engine/Source/Programs/Horde/HordeServer.Tests/Agents/AgentServiceTest.cs b/Engine/Source/Programs/Horde/HordeServer.Tests/Agents/AgentServiceTest.cs
--- a/Engine/Source/Programs/Horde/HordeServer.Tests/Agents/AgentServiceTest.cs
+++ b/Engine/Source/Programs/Horde/HordeServer.Tests/Agents/AgentServiceTest.cs
@@ -10,16 +10,16 @@ using EpicGames.Core;
 using EpicGames.Horde.Agents;
 using EpicGames.Horde.Agents.Pools;
 using EpicGames.Horde.Jobs;
-using Horde.Server.Agents;
-using Horde.Server.Auditing;
-using Horde.Server.Jobs;
-using Horde.Server.Server;
-using Horde.Server.Utilities;
+using HordeServer.Agents;
+using HordeServer.Auditing;
+using HordeServer.Jobs;
+using HordeServer.Server;
+using HordeServer.Utilities;
 using HordeCommon.Rpc.Messages;
 using Microsoft.AspNetCore.Mvc;
 using Microsoft.VisualStudio.TestTools.UnitTesting;
 
-namespace Horde.Server.Tests.Agents;
+namespace HordeServer.Tests.Agents;
 
 /// <summary>
 ///     Testing the agent service
diff --git a/Engine/Source/Programs/Horde/HordeServer.Tests/Agents/Leases/LeasesControllerDbTest.cs b/Engine/Source/Programs/Horde/HordeServer.Tests/Agents/Leases/LeasesControllerDbTest.cs
--- a/Engine/Source/Programs/Horde/HordeServer.Tests/Agents/Leases/LeasesControllerDbTest.cs
+++ b/Engine/Source/Programs/Horde/HordeServer.Tests/Agents/Leases/LeasesControllerDbTest.cs
@@ -6,13 +6,13 @@ using System.Threading.Tasks;
 using EpicGames.Horde.Agents;
 using EpicGames.Horde.Agents.Leases;
 using EpicGames.Horde.Agents.Sessions;
-using Horde.Server.Agents.Leases;
-using Horde.Server.Agents.Sessions;
-using Horde.Server.Utilities;
+using HordeServer.Agents.Leases;
+using HordeServer.Agents.Sessions;
+using HordeServer.Utilities;
 using Microsoft.AspNetCore.Mvc;
 using Microsoft.VisualStudio.TestTools.UnitTesting;
 
-namespace Horde.Server.Tests.Agents.Leases
+namespace HordeServer.Tests.Agents.Leases
 {
 	/// <summary>
 	/// Database-only integration test for LeasesController
diff --git a/Engine/Source/Programs/Horde/HordeServer.Tests/Agents/Pools/PoolUpdateServiceTest.cs b/Engine/Source/Programs/Horde/HordeServer.Tests/Agents/Pools/PoolUpdateServiceTest.cs
--- a/Engine/Source/Programs/Horde/HordeServer.Tests/Agents/Pools/PoolUpdateServiceTest.cs
+++ b/Engine/Source/Programs/Horde/HordeServer.Tests/Agents/Pools/PoolUpdateServiceTest.cs
@@ -6,12 +6,12 @@ using System.Linq;
 using System.Threading;
 using System.Threading.Tasks;
 using EpicGames.Horde.Agents;
-using Horde.Server.Agents;
-using Horde.Server.Agents.Pools;
+using HordeServer.Agents;
+using HordeServer.Agents.Pools;
 using Microsoft.Extensions.Logging.Abstractions;
 using Microsoft.VisualStudio.TestTools.UnitTesting;
 
-namespace Horde.Server.Tests.Agents.Pools;
+namespace HordeServer.Tests.Agents.Pools;
 
 [TestClass]
 public class PoolUpdateServiceTest : TestSetup
diff --git a/Engine/Source/Programs/Horde/HordeServer.Tests/Agents/Pools/PoolsControllerTest.cs b/Engine/Source/Programs/Horde/HordeServer.Tests/Agents/Pools/PoolsControllerTest.cs
--- a/Engine/Source/Programs/Horde/HordeServer.Tests/Agents/Pools/PoolsControllerTest.cs
+++ b/Engine/Source/Programs/Horde/HordeServer.Tests/Agents/Pools/PoolsControllerTest.cs
@@ -3,12 +3,12 @@
 using System;
 using System.Collections.Generic;
 using System.Threading.Tasks;
-using Horde.Server.Agents.Fleet;
-using Horde.Server.Agents.Pools;
+using HordeServer.Agents.Fleet;
+using HordeServer.Agents.Pools;
 using Microsoft.AspNetCore.Mvc;
 using Microsoft.VisualStudio.TestTools.UnitTesting;
 
-namespace Horde.Server.Tests.Agents.Pools
+namespace HordeServer.Tests.Agents.Pools
 {
 	[TestClass]
 	public class PoolsControllerTest : TestSetup
diff --git a/Engine/Source/Programs/Horde/HordeServer.Tests/Agents/Registration/RegistrationServiceTests.cs b/Engine/Source/Programs/Horde/HordeServer.Tests/Agents/Registration/RegistrationServiceTests.cs
--- a/Engine/Source/Programs/Horde/HordeServer.Tests/Agents/Registration/RegistrationServiceTests.cs
+++ b/Engine/Source/Programs/Horde/HordeServer.Tests/Agents/Registration/RegistrationServiceTests.cs
@@ -6,11 +6,11 @@ using System.Security.Cryptography;
 using System.Threading;
 using System.Threading.Tasks;
 using EpicGames.Horde.Agents;
-using Horde.Server.Agents.Enrollment;
+using HordeServer.Agents.Enrollment;
 using Microsoft.Extensions.DependencyInjection;
 using Microsoft.VisualStudio.TestTools.UnitTesting;
 
-namespace Horde.Server.Tests.Agents.Registration
+namespace HordeServer.Tests.Agents.Registration
 {
 	[TestClass]
 	public class RegistrationServiceTests : TestSetup
diff --git a/Engine/Source/Programs/Horde/HordeServer.Tests/Agents/Relay/AgentRelayClientTest.cs b/Engine/Source/Programs/Horde/HordeServer.Tests/Agents/Relay/AgentRelayClientTest.cs
--- a/Engine/Source/Programs/Horde/HordeServer.Tests/Agents/Relay/AgentRelayClientTest.cs
+++ b/Engine/Source/Programs/Horde/HordeServer.Tests/Agents/Relay/AgentRelayClientTest.cs
@@ -6,13 +6,13 @@ using System.Threading;
 using System.Threading.Tasks;
 using Grpc.Core;
 using Horde.Common.Rpc;
-using Horde.Server.Agents.Relay;
-using Horde.Server.Tests.Server;
-using Horde.Server.Utilities;
+using HordeServer.Agents.Relay;
+using HordeServer.Tests.Server;
+using HordeServer.Utilities;
 using Microsoft.Extensions.Logging.Abstractions;
 using Microsoft.VisualStudio.TestTools.UnitTesting;
 
-namespace Horde.Server.Tests.Agents.Relay;
+namespace HordeServer.Tests.Agents.Relay;
 
 public class TestRelayRpcClient : RelayRpc.RelayRpcClient
 {
diff --git a/Engine/Source/Programs/Horde/HordeServer.Tests/Agents/Relay/AgentRelayServiceTests.cs b/Engine/Source/Programs/Horde/HordeServer.Tests/Agents/Relay/AgentRelayServiceTests.cs
--- a/Engine/Source/Programs/Horde/HordeServer.Tests/Agents/Relay/AgentRelayServiceTests.cs
+++ b/Engine/Source/Programs/Horde/HordeServer.Tests/Agents/Relay/AgentRelayServiceTests.cs
@@ -10,13 +10,13 @@ using EpicGames.Horde.Agents.Leases;
 using EpicGames.Horde.Compute;
 using Grpc.Core;
 using Horde.Common.Rpc;
-using Horde.Server.Agents.Relay;
-using Horde.Server.Tests.Server;
-using Horde.Server.Utilities;
+using HordeServer.Agents.Relay;
+using HordeServer.Tests.Server;
+using HordeServer.Utilities;
 using Microsoft.Extensions.Logging.Abstractions;
 using Microsoft.VisualStudio.TestTools.UnitTesting;
 
-namespace Horde.Server.Tests.Agents.Relay;
+namespace HordeServer.Tests.Agents.Relay;
 
 [TestClass]
 public class AgentRelayServiceTests : TestSetup
diff --git a/Engine/Source/Programs/Horde/HordeServer.Tests/Agents/Relay/NftablesTests.cs b/Engine/Source/Programs/Horde/HordeServer.Tests/Agents/Relay/NftablesTests.cs
--- a/Engine/Source/Programs/Horde/HordeServer.Tests/Agents/Relay/NftablesTests.cs
+++ b/Engine/Source/Programs/Horde/HordeServer.Tests/Agents/Relay/NftablesTests.cs
@@ -2,10 +2,10 @@
 
 using System.Collections.Generic;
 using Horde.Common.Rpc;
-using Horde.Server.Agents.Relay;
+using HordeServer.Agents.Relay;
 using Microsoft.VisualStudio.TestTools.UnitTesting;
 
-namespace Horde.Server.Tests.Agents.Relay;
+namespace HordeServer.Tests.Agents.Relay;
 
 [TestClass]
 public class NftablesTests
diff --git a/Engine/Source/Programs/Horde/HordeServer.Tests/Artifacts/ArtifactTests.cs b/Engine/Source/Programs/Horde/HordeServer.Tests/Artifacts/ArtifactTests.cs
--- a/Engine/Source/Programs/Horde/HordeServer.Tests/Artifacts/ArtifactTests.cs
+++ b/Engine/Source/Programs/Horde/HordeServer.Tests/Artifacts/ArtifactTests.cs
@@ -7,14 +7,14 @@ using System.Threading;
 using System.Threading.Tasks;
 using EpicGames.Horde.Artifacts;
 using EpicGames.Horde.Streams;
-using Horde.Server.Acls;
-using Horde.Server.Artifacts;
-using Horde.Server.Server;
+using HordeServer.Acls;
+using HordeServer.Artifacts;
+using HordeServer.Server;
 using HordeCommon;
 using Microsoft.Extensions.DependencyInjection;
 using Microsoft.VisualStudio.TestTools.UnitTesting;
 
-namespace Horde.Server.Tests.Artifacts
+namespace HordeServer.Tests.Artifacts
 {
 	[TestClass]
 	public class ArtifactTests : TestSetup
diff --git a/Engine/Source/Programs/Horde/HordeServer.Tests/Artifacts/BlockCacheTests.cs b/Engine/Source/Programs/Horde/HordeServer.Tests/Artifacts/BlockCacheTests.cs
--- a/Engine/Source/Programs/Horde/HordeServer.Tests/Artifacts/BlockCacheTests.cs
+++ b/Engine/Source/Programs/Horde/HordeServer.Tests/Artifacts/BlockCacheTests.cs
@@ -6,11 +6,11 @@ using System.Collections.Generic;
 using System.Linq;
 using System.Security.Cryptography;
 using System.Threading.Tasks;
-using Horde.Server.Artifacts;
+using HordeServer.Artifacts;
 using Microsoft.Extensions.Logging;
 using Microsoft.VisualStudio.TestTools.UnitTesting;
 
-namespace Horde.Server.Tests.Artifacts
+namespace HordeServer.Tests.Artifacts
 {
 	[TestClass]
 	public class BlockCacheTests
diff --git a/Engine/Source/Programs/Horde/HordeServer.Tests/Artifacts/UnsyncCacheTests.cs b/Engine/Source/Programs/Horde/HordeServer.Tests/Artifacts/UnsyncCacheTests.cs
--- a/Engine/Source/Programs/Horde/HordeServer.Tests/Artifacts/UnsyncCacheTests.cs
+++ b/Engine/Source/Programs/Horde/HordeServer.Tests/Artifacts/UnsyncCacheTests.cs
@@ -9,13 +9,13 @@ using EpicGames.Core;
 using EpicGames.Horde.Storage;
 using EpicGames.Horde.Storage.Bundles;
 using EpicGames.Horde.Storage.Nodes;
-using Horde.Server.Artifacts;
-using Horde.Server.Storage;
+using HordeServer.Artifacts;
+using HordeServer.Storage;
 using Microsoft.Extensions.Logging.Abstractions;
 using Microsoft.VisualStudio.TestTools.UnitTesting;
 using Moq;
 
-namespace Horde.Server.Tests.Artifacts
+namespace HordeServer.Tests.Artifacts
 {
 	[TestClass]
 	public class UnsyncCacheTests
diff --git a/Engine/Source/Programs/Horde/HordeServer.Tests/Authentication/ServiceAccountAuthTest.cs b/Engine/Source/Programs/Horde/HordeServer.Tests/Authentication/ServiceAccountAuthTest.cs
--- a/Engine/Source/Programs/Horde/HordeServer.Tests/Authentication/ServiceAccountAuthTest.cs
+++ b/Engine/Source/Programs/Horde/HordeServer.Tests/Authentication/ServiceAccountAuthTest.cs
@@ -4,9 +4,9 @@ using System.Collections.Generic;
 using System.Linq;
 using System.Security.Claims;
 using System.Threading.Tasks;
-using Horde.Server.Authentication;
-using Horde.Server.ServiceAccounts;
-using Horde.Server.Users;
+using HordeServer.Authentication;
+using HordeServer.ServiceAccounts;
+using HordeServer.Users;
 using Microsoft.AspNetCore.Authentication;
 using Microsoft.AspNetCore.Http;
 using Microsoft.Extensions.DependencyInjection;
@@ -16,7 +16,7 @@ using Microsoft.Extensions.WebEncoders.Testing;
 using Microsoft.Net.Http.Headers;
 using Microsoft.VisualStudio.TestTools.UnitTesting;
 
-namespace Horde.Server.Tests.Authentication
+namespace HordeServer.Tests.Authentication
 {
 	[TestClass]
 	public class ServiceAccountAuthTest : DatabaseIntegrationTest
diff --git a/Engine/Source/Programs/Horde/HordeServer.Tests/Commits/CommitTagTests.cs b/Engine/Source/Programs/Horde/HordeServer.Tests/Commits/CommitTagTests.cs
--- a/Engine/Source/Programs/Horde/HordeServer.Tests/Commits/CommitTagTests.cs
+++ b/Engine/Source/Programs/Horde/HordeServer.Tests/Commits/CommitTagTests.cs
@@ -3,10 +3,10 @@
 using System.Collections.Generic;
 using EpicGames.Core;
 using EpicGames.Horde.Commits;
-using Horde.Server.Streams;
+using HordeServer.Streams;
 using Microsoft.VisualStudio.TestTools.UnitTesting;
 
-namespace Horde.Server.Tests.Commits
+namespace HordeServer.Tests.Commits
 {
 	[TestClass]
 	public class CommitTagTests
diff --git a/Engine/Source/Programs/Horde/HordeServer.Tests/Compute/ComputeServiceTest.cs b/Engine/Source/Programs/Horde/HordeServer.Tests/Compute/ComputeServiceTest.cs
--- a/Engine/Source/Programs/Horde/HordeServer.Tests/Compute/ComputeServiceTest.cs
+++ b/Engine/Source/Programs/Horde/HordeServer.Tests/Compute/ComputeServiceTest.cs
@@ -9,14 +9,14 @@ using System.Threading;
 using System.Threading.Tasks;
 using EpicGames.Horde.Agents.Pools;
 using EpicGames.Horde.Compute;
-using Horde.Server.Agents;
-using Horde.Server.Compute;
-using Horde.Server.Server;
+using HordeServer.Agents;
+using HordeServer.Compute;
+using HordeServer.Server;
 using Microsoft.Extensions.Logging.Abstractions;
 using Microsoft.VisualStudio.TestTools.UnitTesting;
 using OpenTelemetry.Trace;
 
-namespace Horde.Server.Tests.Compute
+namespace HordeServer.Tests.Compute
 {
 	[TestClass]
 	public class ComputeServiceTest : TestSetup
diff --git a/Engine/Source/Programs/Horde/HordeServer.Tests/Compute/RedisTaskSchedulerTest.cs b/Engine/Source/Programs/Horde/HordeServer.Tests/Compute/RedisTaskSchedulerTest.cs
--- a/Engine/Source/Programs/Horde/HordeServer.Tests/Compute/RedisTaskSchedulerTest.cs
+++ b/Engine/Source/Programs/Horde/HordeServer.Tests/Compute/RedisTaskSchedulerTest.cs
@@ -3,13 +3,13 @@
 using System;
 using System.Threading;
 using System.Threading.Tasks;
-using Horde.Server.Compute;
-using Horde.Server.Server;
+using HordeServer.Compute;
+using HordeServer.Server;
 using Microsoft.Extensions.Logging;
 using Microsoft.VisualStudio.TestTools.UnitTesting;
 using StackExchange.Redis;
 
-namespace Horde.Server.Tests.Compute
+namespace HordeServer.Tests.Compute
 {
 	[TestClass]
 	public class RedisTaskSchedulerTest : TestSetup
diff --git a/Engine/Source/Programs/Horde/HordeServer.Tests/Compute/TunnelServiceTests.cs b/Engine/Source/Programs/Horde/HordeServer.Tests/Compute/TunnelServiceTests.cs
--- a/Engine/Source/Programs/Horde/HordeServer.Tests/Compute/TunnelServiceTests.cs
+++ b/Engine/Source/Programs/Horde/HordeServer.Tests/Compute/TunnelServiceTests.cs
@@ -10,11 +10,11 @@ using System.Threading;
 using System.Threading.Tasks;
 using EpicGames.Core;
 using EpicGames.Horde.Compute.Clients;
-using Horde.Server.Compute;
+using HordeServer.Compute;
 using Microsoft.Extensions.Options;
 using Microsoft.VisualStudio.TestTools.UnitTesting;
 
-namespace Horde.Server.Tests.Compute
+namespace HordeServer.Tests.Compute
 {
 	[TestClass]
 	public sealed class TunnelServiceTests
diff --git a/Engine/Source/Programs/Horde/HordeServer.Tests/Configuration/ConfigServiceTests.cs b/Engine/Source/Programs/Horde/HordeServer.Tests/Configuration/ConfigServiceTests.cs
--- a/Engine/Source/Programs/Horde/HordeServer.Tests/Configuration/ConfigServiceTests.cs
+++ b/Engine/Source/Programs/Horde/HordeServer.Tests/Configuration/ConfigServiceTests.cs
@@ -2,10 +2,10 @@
 
 using System;
 using System.Runtime.InteropServices;
-using Horde.Server.Configuration;
+using HordeServer.Configuration;
 using Microsoft.VisualStudio.TestTools.UnitTesting;
 
-namespace Horde.Server.Tests.Configuration;
+namespace HordeServer.Tests.Configuration;
 
 [TestClass]
 public class ConfigServiceTests
diff --git a/Engine/Source/Programs/Horde/HordeServer.Tests/Configuration/ConfigTests.cs b/Engine/Source/Programs/Horde/HordeServer.Tests/Configuration/ConfigTests.cs
--- a/Engine/Source/Programs/Horde/HordeServer.Tests/Configuration/ConfigTests.cs
+++ b/Engine/Source/Programs/Horde/HordeServer.Tests/Configuration/ConfigTests.cs
@@ -8,14 +8,14 @@ using System.Text.Json;
 using System.Threading;
 using System.Threading.Tasks;
 using EpicGames.Core;
-using Horde.Server.Configuration;
-using Horde.Server.Projects;
-using Horde.Server.Server;
-using Horde.Server.Streams;
+using HordeServer.Configuration;
+using HordeServer.Projects;
+using HordeServer.Server;
+using HordeServer.Streams;
 using Microsoft.Extensions.Logging.Abstractions;
 using Microsoft.VisualStudio.TestTools.UnitTesting;
 
-namespace Horde.Server.Tests.Configuration
+namespace HordeServer.Tests.Configuration
 {
 	[TestClass]
 	public class ConfigTests
diff --git a/Engine/Source/Programs/Horde/HordeServer.Tests/ControllerIntegrationTest.cs b/Engine/Source/Programs/Horde/HordeServer.Tests/ControllerIntegrationTest.cs
--- a/Engine/Source/Programs/Horde/HordeServer.Tests/ControllerIntegrationTest.cs
+++ b/Engine/Source/Programs/Horde/HordeServer.Tests/ControllerIntegrationTest.cs
@@ -5,15 +5,15 @@ using System.Collections.Generic;
 using System.Net.Http;
 using System.Threading.Tasks;
 using Amazon.CloudWatch;
-using Horde.Server.Agents;
-using Horde.Server.Configuration;
-using Horde.Server.Jobs;
-using Horde.Server.Jobs.Artifacts;
-using Horde.Server.Jobs.Graphs;
-using Horde.Server.Jobs.Templates;
-using Horde.Server.Perforce;
-using Horde.Server.Server;
-using Horde.Server.Tests.Stubs.Services;
+using HordeServer.Agents;
+using HordeServer.Configuration;
+using HordeServer.Jobs;
+using HordeServer.Jobs.Artifacts;
+using HordeServer.Jobs.Graphs;
+using HordeServer.Jobs.Templates;
+using HordeServer.Perforce;
+using HordeServer.Server;
+using HordeServer.Tests.Stubs.Services;
 using Microsoft.AspNetCore.Hosting;
 using Microsoft.AspNetCore.Mvc.Testing;
 using Microsoft.AspNetCore.TestHost;
@@ -23,7 +23,7 @@ using Microsoft.Extensions.Options;
 using Moq;
 using Serilog;
 
-namespace Horde.Server.Tests;
+namespace HordeServer.Tests;
 
 static class SerilogExtensions
 {
diff --git a/Engine/Source/Programs/Horde/HordeServer.Tests/Dashboard/DashboardTests.cs b/Engine/Source/Programs/Horde/HordeServer.Tests/Dashboard/DashboardTests.cs
--- a/Engine/Source/Programs/Horde/HordeServer.Tests/Dashboard/DashboardTests.cs
+++ b/Engine/Source/Programs/Horde/HordeServer.Tests/Dashboard/DashboardTests.cs
@@ -6,7 +6,7 @@ using System.Threading.Tasks;
 using EpicGames.Horde.Dashboard;
 using Microsoft.VisualStudio.TestTools.UnitTesting;
 
-namespace Horde.Server.Tests.Dashboard
+namespace HordeServer.Tests.Dashboard
 {
 	/// <summary>
 	/// Tests for the dashboard
diff --git a/Engine/Source/Programs/Horde/HordeServer.Tests/DatabaseIntegrationTest.cs b/Engine/Source/Programs/Horde/HordeServer.Tests/DatabaseIntegrationTest.cs
--- a/Engine/Source/Programs/Horde/HordeServer.Tests/DatabaseIntegrationTest.cs
+++ b/Engine/Source/Programs/Horde/HordeServer.Tests/DatabaseIntegrationTest.cs
@@ -6,8 +6,8 @@ using System.Net;
 using System.Runtime.InteropServices;
 using System.Threading;
 using System.Threading.Tasks;
-using Horde.Server.Server;
-using Horde.Server.Utilities;
+using HordeServer.Server;
+using HordeServer.Utilities;
 using Microsoft.AspNetCore.Mvc;
 using Microsoft.Extensions.DependencyInjection;
 using Microsoft.Extensions.Logging;
@@ -17,7 +17,7 @@ using Microsoft.VisualStudio.TestTools.UnitTesting;
 using MongoDB.Driver;
 using StackExchange.Redis;
 
-namespace Horde.Server.Tests
+namespace HordeServer.Tests
 {
 	// Stub for fulfilling IOptions interface during testing
 	public sealed class TestOptions<T> : IOptions<T> where T : class
diff --git a/Engine/Source/Programs/Horde/HordeServer.Tests/DatabaseRunner.cs b/Engine/Source/Programs/Horde/HordeServer.Tests/DatabaseRunner.cs
--- a/Engine/Source/Programs/Horde/HordeServer.Tests/DatabaseRunner.cs
+++ b/Engine/Source/Programs/Horde/HordeServer.Tests/DatabaseRunner.cs
@@ -12,7 +12,7 @@ using System.Runtime.InteropServices;
 using System.Threading.Tasks;
 using EpicGames.Core;
 
-namespace Horde.Server.Tests;
+namespace HordeServer.Tests;
 
 public abstract class DatabaseRunner : IDisposable
 {
diff --git a/Engine/Source/Programs/Horde/HordeServer.Tests/DatabaseRunnerTest.cs b/Engine/Source/Programs/Horde/HordeServer.Tests/DatabaseRunnerTest.cs
--- a/Engine/Source/Programs/Horde/HordeServer.Tests/DatabaseRunnerTest.cs
+++ b/Engine/Source/Programs/Horde/HordeServer.Tests/DatabaseRunnerTest.cs
@@ -3,11 +3,11 @@
 using System.Runtime.InteropServices;
 using System.Threading;
 using System.Threading.Tasks;
-using Horde.Server.Server;
+using HordeServer.Server;
 using Microsoft.Extensions.Logging.Abstractions;
 using Microsoft.VisualStudio.TestTools.UnitTesting;
 
-namespace Horde.Server.Tests
+namespace HordeServer.Tests
 {
 	[TestClass]
 	public class DatabaseRunnerTest
diff --git a/Engine/Source/Programs/Horde/HordeServer.Tests/Ddc/Functional/ContentId/CompressedBlobTests.cs b/Engine/Source/Programs/Horde/HordeServer.Tests/Ddc/Functional/ContentId/CompressedBlobTests.cs
--- a/Engine/Source/Programs/Horde/HordeServer.Tests/Ddc/Functional/ContentId/CompressedBlobTests.cs
+++ b/Engine/Source/Programs/Horde/HordeServer.Tests/Ddc/Functional/ContentId/CompressedBlobTests.cs
@@ -15,16 +15,16 @@ using Blake3;
 using EpicGames.AspNet;
 using EpicGames.Core;
 using EpicGames.Horde.Storage;
-using Horde.Server.Configuration;
-using Horde.Server.Ddc;
-using Horde.Server.Server;
-using Horde.Server.Storage;
+using HordeServer.Configuration;
+using HordeServer.Ddc;
+using HordeServer.Server;
+using HordeServer.Storage;
 using Microsoft.Extensions.DependencyInjection;
 using Microsoft.Extensions.Options;
 using Microsoft.VisualStudio.TestTools.UnitTesting;
 using Serilog;
 
-namespace Horde.Server.Tests.Ddc.FunctionalTests.CompressedBlobs
+namespace HordeServer.Tests.Ddc.FunctionalTests.CompressedBlobs
 {
 	[TestClass]
 	public class CompressedBlobTests : ControllerIntegrationTest
diff --git a/Engine/Source/Programs/Horde/HordeServer.Tests/Ddc/Functional/Objects/ReferencesTests.cs b/Engine/Source/Programs/Horde/HordeServer.Tests/Ddc/Functional/Objects/ReferencesTests.cs
--- a/Engine/Source/Programs/Horde/HordeServer.Tests/Ddc/Functional/Objects/ReferencesTests.cs
+++ b/Engine/Source/Programs/Horde/HordeServer.Tests/Ddc/Functional/Objects/ReferencesTests.cs
@@ -18,16 +18,16 @@ using EpicGames.AspNet;
 using EpicGames.Core;
 using EpicGames.Horde.Storage;
 using EpicGames.Serialization;
-using Horde.Server.Configuration;
-using Horde.Server.Ddc;
-using Horde.Server.Server;
-using Horde.Server.Storage;
+using HordeServer.Configuration;
+using HordeServer.Ddc;
+using HordeServer.Server;
+using HordeServer.Storage;
 using Microsoft.AspNetCore.Mvc;
 using Microsoft.Extensions.DependencyInjection;
 using Microsoft.Extensions.Options;
 using Microsoft.VisualStudio.TestTools.UnitTesting;
 
-namespace Horde.Server.Tests.Ddc.FunctionalTests.References
+namespace HordeServer.Tests.Ddc.FunctionalTests.References
 {
 	public class InsertResponse
 	{
diff --git a/Engine/Source/Programs/Horde/HordeServer.Tests/Ddc/Functional/Utils.cs b/Engine/Source/Programs/Horde/HordeServer.Tests/Ddc/Functional/Utils.cs
--- a/Engine/Source/Programs/Horde/HordeServer.Tests/Ddc/Functional/Utils.cs
+++ b/Engine/Source/Programs/Horde/HordeServer.Tests/Ddc/Functional/Utils.cs
@@ -2,7 +2,7 @@
 
 using System.Text.Json;
 
-namespace Horde.Server.Tests.Ddc.FunctionalTests
+namespace HordeServer.Tests.Ddc.FunctionalTests
 {
 	public static class JsonTestUtils
 	{
diff --git a/Engine/Source/Programs/Horde/HordeServer.Tests/Ddc/UnitTests/CompactBinaryTests.cs b/Engine/Source/Programs/Horde/HordeServer.Tests/Ddc/UnitTests/CompactBinaryTests.cs
--- a/Engine/Source/Programs/Horde/HordeServer.Tests/Ddc/UnitTests/CompactBinaryTests.cs
+++ b/Engine/Source/Programs/Horde/HordeServer.Tests/Ddc/UnitTests/CompactBinaryTests.cs
@@ -7,7 +7,7 @@ using EpicGames.Core;
 using EpicGames.Serialization;
 using Microsoft.VisualStudio.TestTools.UnitTesting;
 
-namespace Horde.Server.Tests.Ddc.UnitTests
+namespace HordeServer.Tests.Ddc.UnitTests
 {
 	[TestClass]
 	public class CompactBinaryTests
diff --git a/Engine/Source/Programs/Horde/HordeServer.Tests/Ddc/UnitTests/CompressedBufferTests.cs b/Engine/Source/Programs/Horde/HordeServer.Tests/Ddc/UnitTests/CompressedBufferTests.cs
--- a/Engine/Source/Programs/Horde/HordeServer.Tests/Ddc/UnitTests/CompressedBufferTests.cs
+++ b/Engine/Source/Programs/Horde/HordeServer.Tests/Ddc/UnitTests/CompressedBufferTests.cs
@@ -5,13 +5,13 @@ using System.Text;
 using System.Threading;
 using System.Threading.Tasks;
 using EpicGames.Core;
-using Horde.Server.Ddc;
+using HordeServer.Ddc;
 using Microsoft.Extensions.Options;
 using Microsoft.VisualStudio.TestTools.UnitTesting;
 using Moq;
 using OpenTelemetry.Trace;
 
-namespace Horde.Server.Tests.Ddc.UnitTests
+namespace HordeServer.Tests.Ddc.UnitTests
 {
 	[TestClass]
 	public class CompressedBufferTests
diff --git a/Engine/Source/Programs/Horde/HordeServer.Tests/DebugSettings.cs b/Engine/Source/Programs/Horde/HordeServer.Tests/DebugSettings.cs
--- a/Engine/Source/Programs/Horde/HordeServer.Tests/DebugSettings.cs
+++ b/Engine/Source/Programs/Horde/HordeServer.Tests/DebugSettings.cs
@@ -1,6 +1,6 @@
 // Copyright Epic Games, Inc. All Rights Reserved.
 
-namespace Horde.Server.Tests
+namespace HordeServer.Tests
 {
 	/// <summary>
 	/// Settings used for debug tests
diff --git a/Engine/Source/Programs/Horde/HordeServer.Tests/Devices/DeviceServiceTests.cs b/Engine/Source/Programs/Horde/HordeServer.Tests/Devices/DeviceServiceTests.cs
--- a/Engine/Source/Programs/Horde/HordeServer.Tests/Devices/DeviceServiceTests.cs
+++ b/Engine/Source/Programs/Horde/HordeServer.Tests/Devices/DeviceServiceTests.cs
@@ -11,18 +11,18 @@ using EpicGames.Horde.Jobs;
 using EpicGames.Horde.Jobs.Templates;
 using EpicGames.Horde.Projects;
 using EpicGames.Horde.Streams;
-using Horde.Server.Devices;
-using Horde.Server.Jobs;
-using Horde.Server.Jobs.Graphs;
-using Horde.Server.Jobs.Templates;
-using Horde.Server.Users;
-using Horde.Server.Utilities;
+using HordeServer.Devices;
+using HordeServer.Jobs;
+using HordeServer.Jobs.Graphs;
+using HordeServer.Jobs.Templates;
+using HordeServer.Users;
+using HordeServer.Utilities;
 using Microsoft.AspNetCore.Http;
 using Microsoft.AspNetCore.Mvc;
 using Microsoft.VisualStudio.TestTools.UnitTesting;
 using Moq;
 
-namespace Horde.Server.Tests.Devices
+namespace HordeServer.Tests.Devices
 {
 	/// <summary>
 	/// Tests for the device service
diff --git a/Engine/Source/Programs/Horde/HordeServer.Tests/Fixture.cs b/Engine/Source/Programs/Horde/HordeServer.Tests/Fixture.cs
--- a/Engine/Source/Programs/Horde/HordeServer.Tests/Fixture.cs
+++ b/Engine/Source/Programs/Horde/HordeServer.Tests/Fixture.cs
@@ -10,17 +10,17 @@ using EpicGames.Horde.Jobs;
 using EpicGames.Horde.Jobs.Templates;
 using EpicGames.Horde.Projects;
 using EpicGames.Horde.Streams;
-using Horde.Server.Agents;
-using Horde.Server.Configuration;
-using Horde.Server.Jobs;
-using Horde.Server.Jobs.Artifacts;
-using Horde.Server.Jobs.Graphs;
-using Horde.Server.Jobs.Templates;
-using Horde.Server.Projects;
-using Horde.Server.Server;
-using Horde.Server.Streams;
-
-namespace Horde.Server.Tests
+using HordeServer.Agents;
+using HordeServer.Configuration;
+using HordeServer.Jobs;
+using HordeServer.Jobs.Artifacts;
+using HordeServer.Jobs.Graphs;
+using HordeServer.Jobs.Templates;
+using HordeServer.Projects;
+using HordeServer.Server;
+using HordeServer.Streams;
+
+namespace HordeServer.Tests
 {
 	public class Fixture
 	{
diff --git a/Engine/Source/Programs/Horde/HordeServer.Tests/Fleet/AwsAutoScalingLifecycleTest.cs b/Engine/Source/Programs/Horde/HordeServer.Tests/Fleet/AwsAutoScalingLifecycleTest.cs
--- a/Engine/Source/Programs/Horde/HordeServer.Tests/Fleet/AwsAutoScalingLifecycleTest.cs
+++ b/Engine/Source/Programs/Horde/HordeServer.Tests/Fleet/AwsAutoScalingLifecycleTest.cs
@@ -12,15 +12,15 @@ using Amazon.AutoScaling.Model;
 using EpicGames.Horde.Agents;
 using EpicGames.Horde.Agents.Leases;
 using EpicGames.Horde.Agents.Pools;
-using Horde.Server.Agents;
-using Horde.Server.Agents.Fleet;
-using Horde.Server.Utilities;
+using HordeServer.Agents;
+using HordeServer.Agents.Fleet;
+using HordeServer.Utilities;
 using Microsoft.Extensions.DependencyInjection;
 using Microsoft.Extensions.Logging;
 using Microsoft.VisualStudio.TestTools.UnitTesting;
 using Moq;
 
-namespace Horde.Server.Tests.Fleet;
+namespace HordeServer.Tests.Fleet;
 
 [TestClass]
 public class AwsAutoScalingLifecycleServiceTest : TestSetup
diff --git a/Engine/Source/Programs/Horde/HordeServer.Tests/Fleet/AwsRecyclingFleetManagerTest.cs b/Engine/Source/Programs/Horde/HordeServer.Tests/Fleet/AwsRecyclingFleetManagerTest.cs
--- a/Engine/Source/Programs/Horde/HordeServer.Tests/Fleet/AwsRecyclingFleetManagerTest.cs
+++ b/Engine/Source/Programs/Horde/HordeServer.Tests/Fleet/AwsRecyclingFleetManagerTest.cs
@@ -7,16 +7,16 @@ using System.Threading;
 using System.Threading.Tasks;
 using Amazon.EC2;
 using Amazon.EC2.Model;
-using Horde.Server.Agents;
-using Horde.Server.Agents.Fleet;
-using Horde.Server.Agents.Fleet.Providers;
-using Horde.Server.Agents.Pools;
+using HordeServer.Agents;
+using HordeServer.Agents.Fleet;
+using HordeServer.Agents.Fleet.Providers;
+using HordeServer.Agents.Pools;
 using HordeCommon;
 using Microsoft.Extensions.Logging;
 using Microsoft.VisualStudio.TestTools.UnitTesting;
 using JsonSerializer = System.Text.Json.JsonSerializer;
 
-namespace Horde.Server.Tests.Fleet
+namespace HordeServer.Tests.Fleet
 {
 	[TestClass]
 	public class AwsRecyclingFleetManagerTest : TestSetup
diff --git a/Engine/Source/Programs/Horde/HordeServer.Tests/Fleet/AwsReuseFleetManagerTest.cs b/Engine/Source/Programs/Horde/HordeServer.Tests/Fleet/AwsReuseFleetManagerTest.cs
--- a/Engine/Source/Programs/Horde/HordeServer.Tests/Fleet/AwsReuseFleetManagerTest.cs
+++ b/Engine/Source/Programs/Horde/HordeServer.Tests/Fleet/AwsReuseFleetManagerTest.cs
@@ -5,15 +5,15 @@ using System.Threading;
 using System.Threading.Tasks;
 using Amazon.EC2;
 using Amazon.EC2.Model;
-using Horde.Server.Agents;
-using Horde.Server.Agents.Fleet;
-using Horde.Server.Agents.Fleet.Providers;
-using Horde.Server.Agents.Pools;
+using HordeServer.Agents;
+using HordeServer.Agents.Fleet;
+using HordeServer.Agents.Fleet.Providers;
+using HordeServer.Agents.Pools;
 using Microsoft.Extensions.DependencyInjection;
 using Microsoft.Extensions.Logging;
 using Microsoft.VisualStudio.TestTools.UnitTesting;
 
-namespace Horde.Server.Tests.Fleet
+namespace HordeServer.Tests.Fleet
 {
 	[TestClass]
 	public class AwsReuseFleetManagerTest : TestSetup
diff --git a/Engine/Source/Programs/Horde/HordeServer.Tests/Fleet/FakeAmazonCloudWatch.cs b/Engine/Source/Programs/Horde/HordeServer.Tests/Fleet/FakeAmazonCloudWatch.cs
--- a/Engine/Source/Programs/Horde/HordeServer.Tests/Fleet/FakeAmazonCloudWatch.cs
+++ b/Engine/Source/Programs/Horde/HordeServer.Tests/Fleet/FakeAmazonCloudWatch.cs
@@ -8,7 +8,7 @@ using Amazon.CloudWatch;
 using Amazon.CloudWatch.Model;
 using Moq;
 
-namespace Horde.Server.Tests.Fleet;
+namespace HordeServer.Tests.Fleet;
 
 /// <summary>
 /// Fake implementation of the IAmazonCloudWatch interface
diff --git a/Engine/Source/Programs/Horde/HordeServer.Tests/Fleet/FakeAmazonEc2.cs b/Engine/Source/Programs/Horde/HordeServer.Tests/Fleet/FakeAmazonEc2.cs
--- a/Engine/Source/Programs/Horde/HordeServer.Tests/Fleet/FakeAmazonEc2.cs
+++ b/Engine/Source/Programs/Horde/HordeServer.Tests/Fleet/FakeAmazonEc2.cs
@@ -10,7 +10,7 @@ using Amazon.EC2;
 using Amazon.EC2.Model;
 using Moq;
 
-namespace Horde.Server.Tests.Fleet;
+namespace HordeServer.Tests.Fleet;
 
 /// <summary>
 /// Fake implementation of the IAmazonEC2 interface
diff --git a/Engine/Source/Programs/Horde/HordeServer.Tests/Fleet/JobQueueStrategyTest.cs b/Engine/Source/Programs/Horde/HordeServer.Tests/Fleet/JobQueueStrategyTest.cs
--- a/Engine/Source/Programs/Horde/HordeServer.Tests/Fleet/JobQueueStrategyTest.cs
+++ b/Engine/Source/Programs/Horde/HordeServer.Tests/Fleet/JobQueueStrategyTest.cs
@@ -9,17 +9,17 @@ using EpicGames.Horde.Jobs;
 using EpicGames.Horde.Jobs.Templates;
 using EpicGames.Horde.Projects;
 using EpicGames.Horde.Streams;
-using Horde.Server.Agents;
-using Horde.Server.Agents.Fleet;
-using Horde.Server.Agents.Pools;
-using Horde.Server.Jobs;
-using Horde.Server.Jobs.Graphs;
-using Horde.Server.Projects;
-using Horde.Server.Server;
-using Horde.Server.Streams;
+using HordeServer.Agents;
+using HordeServer.Agents.Fleet;
+using HordeServer.Agents.Pools;
+using HordeServer.Jobs;
+using HordeServer.Jobs.Graphs;
+using HordeServer.Projects;
+using HordeServer.Server;
+using HordeServer.Streams;
 using Microsoft.VisualStudio.TestTools.UnitTesting;
 
-namespace Horde.Server.Tests.Fleet
+namespace HordeServer.Tests.Fleet
 {
 	[TestClass]
 	public class JobQueueStrategyTest : TestSetup
diff --git a/Engine/Source/Programs/Horde/HordeServer.Tests/Fleet/LeaseUtilizationAwsMetricStrategyTests.cs b/Engine/Source/Programs/Horde/HordeServer.Tests/Fleet/LeaseUtilizationAwsMetricStrategyTests.cs
--- a/Engine/Source/Programs/Horde/HordeServer.Tests/Fleet/LeaseUtilizationAwsMetricStrategyTests.cs
+++ b/Engine/Source/Programs/Horde/HordeServer.Tests/Fleet/LeaseUtilizationAwsMetricStrategyTests.cs
@@ -10,16 +10,16 @@ using EpicGames.Horde.Agents.Pools;
 using EpicGames.Horde.Streams;
 using Google.Protobuf;
 using Google.Protobuf.WellKnownTypes;
-using Horde.Server.Agents;
-using Horde.Server.Agents.Fleet;
-using Horde.Server.Agents.Leases;
-using Horde.Server.Agents.Pools;
-using Horde.Server.Utilities;
+using HordeServer.Agents;
+using HordeServer.Agents.Fleet;
+using HordeServer.Agents.Leases;
+using HordeServer.Agents.Pools;
+using HordeServer.Utilities;
 using HordeCommon.Rpc.Tasks;
 using Microsoft.Extensions.Logging.Abstractions;
 using Microsoft.VisualStudio.TestTools.UnitTesting;
 
-namespace Horde.Server.Tests.Fleet
+namespace HordeServer.Tests.Fleet
 {
 	[TestClass]
 	public class LeaseUtilizationAwsMetricStrategyTest : TestSetup
diff --git a/Engine/Source/Programs/Horde/HordeServer.Tests/Fleet/LeaseUtilizationStrategyTest.cs b/Engine/Source/Programs/Horde/HordeServer.Tests/Fleet/LeaseUtilizationStrategyTest.cs
--- a/Engine/Source/Programs/Horde/HordeServer.Tests/Fleet/LeaseUtilizationStrategyTest.cs
+++ b/Engine/Source/Programs/Horde/HordeServer.Tests/Fleet/LeaseUtilizationStrategyTest.cs
@@ -7,15 +7,15 @@ using EpicGames.Horde.Agents.Leases;
 using EpicGames.Horde.Streams;
 using Google.Protobuf;
 using Google.Protobuf.WellKnownTypes;
-using Horde.Server.Agents;
-using Horde.Server.Agents.Fleet;
-using Horde.Server.Agents.Leases;
-using Horde.Server.Agents.Pools;
-using Horde.Server.Utilities;
+using HordeServer.Agents;
+using HordeServer.Agents.Fleet;
+using HordeServer.Agents.Leases;
+using HordeServer.Agents.Pools;
+using HordeServer.Utilities;
 using HordeCommon.Rpc.Tasks;
 using Microsoft.VisualStudio.TestTools.UnitTesting;
 
-namespace Horde.Server.Tests.Fleet
+namespace HordeServer.Tests.Fleet
 {
 	[TestClass]
 	public class LeaseUtilizationStrategyTest : TestSetup
diff --git a/Engine/Source/Programs/Horde/HordeServer.Tests/GrpcTestHelpers.cs b/Engine/Source/Programs/Horde/HordeServer.Tests/GrpcTestHelpers.cs
--- a/Engine/Source/Programs/Horde/HordeServer.Tests/GrpcTestHelpers.cs
+++ b/Engine/Source/Programs/Horde/HordeServer.Tests/GrpcTestHelpers.cs
@@ -25,7 +25,7 @@ using System.Threading.Channels;
 using System.Threading.Tasks;
 using Grpc.Core;
 
-namespace Horde.Server.Tests
+namespace HordeServer.Tests
 {
 	public class TestServerStreamWriter<T> : IServerStreamWriter<T> where T : class
 	{
diff --git a/Engine/Source/Programs/Horde/HordeServer.Tests/Issues/IssueServiceTests.cs b/Engine/Source/Programs/Horde/HordeServer.Tests/Issues/IssueServiceTests.cs
--- a/Engine/Source/Programs/Horde/HordeServer.Tests/Issues/IssueServiceTests.cs
+++ b/Engine/Source/Programs/Horde/HordeServer.Tests/Issues/IssueServiceTests.cs
@@ -19,16 +19,16 @@ using EpicGames.Horde.Storage;
 using EpicGames.Horde.Streams;
 using EpicGames.Horde.Users;
 using EpicGames.Perforce;
-using Horde.Server.Issues;
-using Horde.Server.Jobs;
-using Horde.Server.Jobs.Graphs;
-using Horde.Server.Logs;
-using Horde.Server.Projects;
-using Horde.Server.Server;
-using Horde.Server.Storage;
-using Horde.Server.Streams;
-using Horde.Server.Tests.Stubs.Services;
-using Horde.Server.Users;
+using HordeServer.Issues;
+using HordeServer.Jobs;
+using HordeServer.Jobs.Graphs;
+using HordeServer.Logs;
+using HordeServer.Projects;
+using HordeServer.Server;
+using HordeServer.Storage;
+using HordeServer.Streams;
+using HordeServer.Tests.Stubs.Services;
+using HordeServer.Users;
 using JobDriver.JobDriver.Parser;
 using Microsoft.Extensions.Logging;
 using Microsoft.Extensions.Logging.Abstractions;
@@ -37,7 +37,7 @@ using MongoDB.Bson;
 using MongoDB.Driver;
 using Moq;
 
-namespace Horde.Server.Tests.Issues
+namespace HordeServer.Tests.Issues
 {
 	[TestClass]
 	public class IssueServiceTests : TestSetup
@@ -2887,7 +2887,7 @@ namespace Horde.Server.Tests.Issues
 				spans.Add(MockSpan(new StreamId("ue5-release-staging"), MockSuspect(101, 1), MockSuspect(102, 2)));
 				spans.Add(MockSpan(new StreamId("ue5-release"), MockSuspect(1, null), MockSuspect(2, null)));
 
-				List<IIssueSpan> results = Horde.Server.Issues.IssueService.FindMergeOriginSpans(spans);
+				List<IIssueSpan> results = HordeServer.Issues.IssueService.FindMergeOriginSpans(spans);
 				Assert.AreEqual(1, results.Count);
 				Assert.AreEqual(new StreamId("ue5-release"), results[0].StreamId);
 			}
@@ -2898,7 +2898,7 @@ namespace Horde.Server.Tests.Issues
 				spans.Add(MockSpan(new StreamId("ue5-release-staging"), MockSuspect(101, 1), MockSuspect(102, null)));
 				spans.Add(MockSpan(new StreamId("ue5-release"), MockSuspect(1, null), MockSuspect(2, null)));
 
-				List<IIssueSpan> results = Horde.Server.Issues.IssueService.FindMergeOriginSpans(spans);
+				List<IIssueSpan> results = HordeServer.Issues.IssueService.FindMergeOriginSpans(spans);
 				Assert.AreEqual(1, results.Count);
 				Assert.AreEqual(new StreamId("ue5-release"), results[0].StreamId);
 			}
@@ -2909,7 +2909,7 @@ namespace Horde.Server.Tests.Issues
 				spans.Add(MockSpan(new StreamId("ue5-release-staging"), MockSuspect(101, 1), MockSuspect(102, null)));
 				spans.Add(MockSpan(new StreamId("ue5-main"), MockSuspect(201, null), MockSuspect(202, 2)));
 
-				List<IIssueSpan> results = Horde.Server.Issues.IssueService.FindMergeOriginSpans(spans);
+				List<IIssueSpan> results = HordeServer.Issues.IssueService.FindMergeOriginSpans(spans);
 				Assert.AreEqual(1, results.Count);
 				Assert.AreEqual(new StreamId("ue5-release"), results[0].StreamId);
 			}
diff --git a/Engine/Source/Programs/Horde/HordeServer.Tests/Issues/IssueTagTests.cs b/Engine/Source/Programs/Horde/HordeServer.Tests/Issues/IssueTagTests.cs
--- a/Engine/Source/Programs/Horde/HordeServer.Tests/Issues/IssueTagTests.cs
+++ b/Engine/Source/Programs/Horde/HordeServer.Tests/Issues/IssueTagTests.cs
@@ -2,10 +2,10 @@
 
 using System;
 using System.Linq;
-using Horde.Server.Issues;
+using HordeServer.Issues;
 using Microsoft.VisualStudio.TestTools.UnitTesting;
 
-namespace Horde.Server.Tests.Issues
+namespace HordeServer.Tests.Issues
 {
 	[TestClass]
 	public class IssueTagTests
diff --git a/Engine/Source/Programs/Horde/HordeServer.Tests/Jobs/Bisect/BisectTests.cs b/Engine/Source/Programs/Horde/HordeServer.Tests/Jobs/Bisect/BisectTests.cs
--- a/Engine/Source/Programs/Horde/HordeServer.Tests/Jobs/Bisect/BisectTests.cs
+++ b/Engine/Source/Programs/Horde/HordeServer.Tests/Jobs/Bisect/BisectTests.cs
@@ -11,22 +11,22 @@ using EpicGames.Horde.Jobs;
 using EpicGames.Horde.Jobs.Templates;
 using EpicGames.Horde.Projects;
 using EpicGames.Horde.Streams;
-using Horde.Server.Jobs;
-using Horde.Server.Jobs.Bisect;
-using Horde.Server.Jobs.Graphs;
-using Horde.Server.Jobs.Templates;
-using Horde.Server.Logs;
-using Horde.Server.Projects;
-using Horde.Server.Streams;
-using Horde.Server.Users;
-using Horde.Server.Utilities;
+using HordeServer.Jobs;
+using HordeServer.Jobs.Bisect;
+using HordeServer.Jobs.Graphs;
+using HordeServer.Jobs.Templates;
+using HordeServer.Logs;
+using HordeServer.Projects;
+using HordeServer.Streams;
+using HordeServer.Users;
+using HordeServer.Utilities;
 using Microsoft.AspNetCore.Http;
 using Microsoft.AspNetCore.Mvc;
 using Microsoft.Extensions.DependencyInjection;
 using Microsoft.VisualStudio.TestTools.UnitTesting;
 using Moq;
 
-namespace Horde.Server.Tests.Jobs.Bisect
+namespace HordeServer.Tests.Jobs.Bisect
 {
 	[TestClass]
 	public class BisectTests : TestSetup
diff --git a/Engine/Source/Programs/Horde/HordeServer.Tests/Jobs/GetLabelStatesTests.cs b/Engine/Source/Programs/Horde/HordeServer.Tests/Jobs/GetLabelStatesTests.cs
--- a/Engine/Source/Programs/Horde/HordeServer.Tests/Jobs/GetLabelStatesTests.cs
+++ b/Engine/Source/Programs/Horde/HordeServer.Tests/Jobs/GetLabelStatesTests.cs
@@ -2,12 +2,12 @@
 
 using System.Collections.Generic;
 using EpicGames.Horde.Jobs;
-using Horde.Server.Jobs;
-using Horde.Server.Jobs.Graphs;
+using HordeServer.Jobs;
+using HordeServer.Jobs.Graphs;
 using Microsoft.VisualStudio.TestTools.UnitTesting;
 using Moq;
 
-namespace Horde.Server.Tests.Jobs
+namespace HordeServer.Tests.Jobs
 {
 	[TestClass]
 	public class GetLabelStatesTests
diff --git a/Engine/Source/Programs/Horde/HordeServer.Tests/Jobs/JobCollectionTests.cs b/Engine/Source/Programs/Horde/HordeServer.Tests/Jobs/JobCollectionTests.cs
--- a/Engine/Source/Programs/Horde/HordeServer.Tests/Jobs/JobCollectionTests.cs
+++ b/Engine/Source/Programs/Horde/HordeServer.Tests/Jobs/JobCollectionTests.cs
@@ -15,16 +15,16 @@ using EpicGames.Horde.Jobs.Templates;
 using EpicGames.Horde.Logs;
 using EpicGames.Horde.Streams;
 using EpicGames.Horde.Users;
-using Horde.Server.Agents.Sessions;
-using Horde.Server.Jobs;
-using Horde.Server.Jobs.Graphs;
-using Horde.Server.Jobs.Templates;
-using Horde.Server.Logs;
-using Horde.Server.Utilities;
+using HordeServer.Agents.Sessions;
+using HordeServer.Jobs;
+using HordeServer.Jobs.Graphs;
+using HordeServer.Jobs.Templates;
+using HordeServer.Logs;
+using HordeServer.Utilities;
 using Microsoft.VisualStudio.TestTools.UnitTesting;
 using Moq;
 
-namespace Horde.Server.Tests.Jobs
+namespace HordeServer.Tests.Jobs
 {
 	[TestClass]
 	public class JobCollectionTests : TestSetup
diff --git a/Engine/Source/Programs/Horde/HordeServer.Tests/Jobs/JobExpirationTests.cs b/Engine/Source/Programs/Horde/HordeServer.Tests/Jobs/JobExpirationTests.cs
--- a/Engine/Source/Programs/Horde/HordeServer.Tests/Jobs/JobExpirationTests.cs
+++ b/Engine/Source/Programs/Horde/HordeServer.Tests/Jobs/JobExpirationTests.cs
@@ -5,16 +5,16 @@ using System.Threading.Tasks;
 using EpicGames.Horde.Jobs.Templates;
 using EpicGames.Horde.Projects;
 using EpicGames.Horde.Streams;
-using Horde.Server.Jobs;
-using Horde.Server.Jobs.Graphs;
-using Horde.Server.Jobs.Templates;
-using Horde.Server.Projects;
-using Horde.Server.Server;
-using Horde.Server.Streams;
+using HordeServer.Jobs;
+using HordeServer.Jobs.Graphs;
+using HordeServer.Jobs.Templates;
+using HordeServer.Projects;
+using HordeServer.Server;
+using HordeServer.Streams;
 using Microsoft.Extensions.DependencyInjection;
 using Microsoft.VisualStudio.TestTools.UnitTesting;
 
-namespace Horde.Server.Tests.Jobs
+namespace HordeServer.Tests.Jobs
 {
 	[TestClass]
 	public class JobExpirationTests : TestSetup
diff --git a/Engine/Source/Programs/Horde/HordeServer.Tests/Jobs/JobServiceTests.cs b/Engine/Source/Programs/Horde/HordeServer.Tests/Jobs/JobServiceTests.cs
--- a/Engine/Source/Programs/Horde/HordeServer.Tests/Jobs/JobServiceTests.cs
+++ b/Engine/Source/Programs/Horde/HordeServer.Tests/Jobs/JobServiceTests.cs
@@ -12,20 +12,20 @@ using EpicGames.Horde.Commits;
 using EpicGames.Horde.Jobs;
 using EpicGames.Horde.Jobs.Templates;
 using EpicGames.Horde.Streams;
-using Horde.Server.Agents;
-using Horde.Server.Commits;
-using Horde.Server.Jobs;
-using Horde.Server.Jobs.Graphs;
-using Horde.Server.Jobs.Templates;
-using Horde.Server.Logs;
-using Horde.Server.Projects;
-using Horde.Server.Server;
-using Horde.Server.Streams;
-using Horde.Server.Users;
+using HordeServer.Agents;
+using HordeServer.Commits;
+using HordeServer.Jobs;
+using HordeServer.Jobs.Graphs;
+using HordeServer.Jobs.Templates;
+using HordeServer.Logs;
+using HordeServer.Projects;
+using HordeServer.Server;
+using HordeServer.Streams;
+using HordeServer.Users;
 using Microsoft.CodeAnalysis;
 using Microsoft.VisualStudio.TestTools.UnitTesting;
 
-namespace Horde.Server.Tests.Jobs
+namespace HordeServer.Tests.Jobs
 {
 	using ProjectId = EpicGames.Horde.Projects.ProjectId;
 
diff --git a/Engine/Source/Programs/Horde/HordeServer.Tests/Jobs/JobTaskSourceTests.cs b/Engine/Source/Programs/Horde/HordeServer.Tests/Jobs/JobTaskSourceTests.cs
--- a/Engine/Source/Programs/Horde/HordeServer.Tests/Jobs/JobTaskSourceTests.cs
+++ b/Engine/Source/Programs/Horde/HordeServer.Tests/Jobs/JobTaskSourceTests.cs
@@ -9,15 +9,15 @@ using EpicGames.Horde.Agents;
 using EpicGames.Horde.Agents.Pools;
 using EpicGames.Horde.Jobs;
 using EpicGames.Horde.Users;
-using Horde.Server.Agents;
-using Horde.Server.Agents.Pools;
-using Horde.Server.Jobs;
-using Horde.Server.Jobs.Graphs;
-using Horde.Server.Streams;
-using Horde.Server.Utilities;
+using HordeServer.Agents;
+using HordeServer.Agents.Pools;
+using HordeServer.Jobs;
+using HordeServer.Jobs.Graphs;
+using HordeServer.Streams;
+using HordeServer.Utilities;
 using Microsoft.VisualStudio.TestTools.UnitTesting;
 
-namespace Horde.Server.Tests.Jobs
+namespace HordeServer.Tests.Jobs
 {
 	[TestClass]
 	public class JobTaskSourceTests : TestSetup
diff --git a/Engine/Source/Programs/Horde/HordeServer.Tests/Jobs/JobsControllerDbTest.cs b/Engine/Source/Programs/Horde/HordeServer.Tests/Jobs/JobsControllerDbTest.cs
--- a/Engine/Source/Programs/Horde/HordeServer.Tests/Jobs/JobsControllerDbTest.cs
+++ b/Engine/Source/Programs/Horde/HordeServer.Tests/Jobs/JobsControllerDbTest.cs
@@ -4,11 +4,11 @@ using System.Collections.Generic;
 using System.Threading.Tasks;
 using EpicGames.Core;
 using EpicGames.Horde.Jobs;
-using Horde.Server.Jobs;
+using HordeServer.Jobs;
 using Microsoft.AspNetCore.Mvc;
 using Microsoft.VisualStudio.TestTools.UnitTesting;
 
-namespace Horde.Server.Tests.Jobs
+namespace HordeServer.Tests.Jobs
 {
 	/// <summary>
 	/// Database-only test for testing the Job controller. Different from the JobsController test that set up
diff --git a/Engine/Source/Programs/Horde/HordeServer.Tests/Jobs/JobsControllerTest.cs b/Engine/Source/Programs/Horde/HordeServer.Tests/Jobs/JobsControllerTest.cs
--- a/Engine/Source/Programs/Horde/HordeServer.Tests/Jobs/JobsControllerTest.cs
+++ b/Engine/Source/Programs/Horde/HordeServer.Tests/Jobs/JobsControllerTest.cs
@@ -2,12 +2,12 @@
 
 using System.Net;
 using System.Threading.Tasks;
-using Horde.Server.Jobs.Artifacts;
+using HordeServer.Jobs.Artifacts;
 using Microsoft.VisualStudio.TestTools.UnitTesting;
 
 #pragma warning disable CA2234 // Pass system uri objects instead of strings
 
-namespace Horde.Server.Tests.Jobs
+namespace HordeServer.Tests.Jobs
 {
 	[TestClass]
 	public class JobsControllerTest : ControllerIntegrationTest
diff --git a/Engine/Source/Programs/Horde/HordeServer.Tests/Jobs/Schedules/SchedulerTests.cs b/Engine/Source/Programs/Horde/HordeServer.Tests/Jobs/Schedules/SchedulerTests.cs
--- a/Engine/Source/Programs/Horde/HordeServer.Tests/Jobs/Schedules/SchedulerTests.cs
+++ b/Engine/Source/Programs/Horde/HordeServer.Tests/Jobs/Schedules/SchedulerTests.cs
@@ -9,19 +9,19 @@ using EpicGames.Horde.Jobs;
 using EpicGames.Horde.Jobs.Templates;
 using EpicGames.Horde.Projects;
 using EpicGames.Horde.Streams;
-using Horde.Server.Jobs;
-using Horde.Server.Jobs.Graphs;
-using Horde.Server.Jobs.Schedules;
-using Horde.Server.Jobs.Templates;
-using Horde.Server.Logs;
-using Horde.Server.Projects;
-using Horde.Server.Server;
-using Horde.Server.Streams;
-using Horde.Server.Users;
+using HordeServer.Jobs;
+using HordeServer.Jobs.Graphs;
+using HordeServer.Jobs.Schedules;
+using HordeServer.Jobs.Templates;
+using HordeServer.Logs;
+using HordeServer.Projects;
+using HordeServer.Server;
+using HordeServer.Streams;
+using HordeServer.Users;
 using Microsoft.Extensions.DependencyInjection;
 using Microsoft.VisualStudio.TestTools.UnitTesting;
 
-namespace Horde.Server.Tests.Jobs.Schedules
+namespace HordeServer.Tests.Jobs.Schedules
 {
 	[TestClass]
 	public class SchedulerTests : TestSetup
diff --git a/Engine/Source/Programs/Horde/HordeServer.Tests/Jobs/TestData/TestDataTests.cs b/Engine/Source/Programs/Horde/HordeServer.Tests/Jobs/TestData/TestDataTests.cs
--- a/Engine/Source/Programs/Horde/HordeServer.Tests/Jobs/TestData/TestDataTests.cs
+++ b/Engine/Source/Programs/Horde/HordeServer.Tests/Jobs/TestData/TestDataTests.cs
@@ -13,22 +13,22 @@ using EpicGames.Horde.Jobs.TestData;
 using EpicGames.Horde.Logs;
 using EpicGames.Horde.Projects;
 using EpicGames.Horde.Streams;
-using Horde.Server.Jobs;
-using Horde.Server.Jobs.Graphs;
-using Horde.Server.Logs;
-using Horde.Server.Projects;
-using Horde.Server.Server;
-using Horde.Server.Streams;
-using Horde.Server.Jobs.TestData;
-using Horde.Server.Users;
-using Horde.Server.Utilities;
+using HordeServer.Jobs;
+using HordeServer.Jobs.Graphs;
+using HordeServer.Logs;
+using HordeServer.Projects;
+using HordeServer.Server;
+using HordeServer.Streams;
+using HordeServer.Jobs.TestData;
+using HordeServer.Users;
+using HordeServer.Utilities;
 using Microsoft.AspNetCore.Http;
 using Microsoft.AspNetCore.Mvc;
 using Microsoft.VisualStudio.TestTools.UnitTesting;
 using MongoDB.Bson;
 using Moq;
 
-namespace Horde.Server.Tests.Jobs.TestData
+namespace HordeServer.Tests.Jobs.TestData
 {
 	/// <summary>
 	/// Tests for the device service
diff --git a/Engine/Source/Programs/Horde/HordeServer.Tests/Logs/LogChunkTests.cs b/Engine/Source/Programs/Horde/HordeServer.Tests/Logs/LogChunkTests.cs
--- a/Engine/Source/Programs/Horde/HordeServer.Tests/Logs/LogChunkTests.cs
+++ b/Engine/Source/Programs/Horde/HordeServer.Tests/Logs/LogChunkTests.cs
@@ -4,11 +4,11 @@ using System;
 using System.Collections.Generic;
 using System.Text;
 using EpicGames.Horde.Logs;
-using Horde.Server.Logs.Data;
+using HordeServer.Logs.Data;
 using Microsoft.Extensions.Logging.Abstractions;
 using Microsoft.VisualStudio.TestTools.UnitTesting;
 
-namespace Horde.Server.Tests.Logs
+namespace HordeServer.Tests.Logs
 {
 	[TestClass]
 	public class LogChunkTests
diff --git a/Engine/Source/Programs/Horde/HordeServer.Tests/Logs/LogIndexingTests.cs b/Engine/Source/Programs/Horde/HordeServer.Tests/Logs/LogIndexingTests.cs
--- a/Engine/Source/Programs/Horde/HordeServer.Tests/Logs/LogIndexingTests.cs
+++ b/Engine/Source/Programs/Horde/HordeServer.Tests/Logs/LogIndexingTests.cs
@@ -10,13 +10,13 @@ using System.Threading.Tasks;
 using EpicGames.Horde.Jobs;
 using EpicGames.Horde.Logs;
 using EpicGames.Horde.Storage;
-using Horde.Server.Jobs;
-using Horde.Server.Logs;
-using Horde.Server.Storage;
+using HordeServer.Jobs;
+using HordeServer.Logs;
+using HordeServer.Storage;
 using Microsoft.Extensions.Logging.Abstractions;
 using Microsoft.VisualStudio.TestTools.UnitTesting;
 
-namespace Horde.Server.Tests.Logs
+namespace HordeServer.Tests.Logs
 {
 	class TestLogWriter : IAsyncDisposable
 	{
@@ -60,7 +60,7 @@ namespace Horde.Server.Tests.Logs
 	[TestClass]
 	public class LogIndexingTests : TestSetup
 	{
-		private readonly byte[] _data = Horde.Server.Compute.Tests.Properties.Resources.TextFile;
+		private readonly byte[] _data = HordeServer.Compute.Tests.Properties.Resources.TextFile;
 
 		[TestMethod]
 		public async Task IndexTestsAsync()
diff --git a/Engine/Source/Programs/Horde/HordeServer.Tests/Logs/LogServiceTests.cs b/Engine/Source/Programs/Horde/HordeServer.Tests/Logs/LogServiceTests.cs
--- a/Engine/Source/Programs/Horde/HordeServer.Tests/Logs/LogServiceTests.cs
+++ b/Engine/Source/Programs/Horde/HordeServer.Tests/Logs/LogServiceTests.cs
@@ -12,14 +12,14 @@ using EpicGames.Core;
 using EpicGames.Horde.Agents.Sessions;
 using EpicGames.Horde.Jobs;
 using EpicGames.Horde.Logs;
-using Horde.Server.Agents.Sessions;
-using Horde.Server.Jobs;
-using Horde.Server.Logs;
-using Horde.Server.Logs.Data;
-using Horde.Server.Utilities;
+using HordeServer.Agents.Sessions;
+using HordeServer.Jobs;
+using HordeServer.Logs;
+using HordeServer.Logs.Data;
+using HordeServer.Utilities;
 using Microsoft.VisualStudio.TestTools.UnitTesting;
 
-namespace Horde.Server.Tests.Logs
+namespace HordeServer.Tests.Logs
 {
 	[TestClass]
 	public sealed class LogServiceTest : TestSetup
diff --git a/Engine/Source/Programs/Horde/HordeServer.Tests/Logs/LogTailServiceTests.cs b/Engine/Source/Programs/Horde/HordeServer.Tests/Logs/LogTailServiceTests.cs
--- a/Engine/Source/Programs/Horde/HordeServer.Tests/Logs/LogTailServiceTests.cs
+++ b/Engine/Source/Programs/Horde/HordeServer.Tests/Logs/LogTailServiceTests.cs
@@ -8,14 +8,14 @@ using System.Threading.Tasks;
 using EpicGames.Core;
 using EpicGames.Horde.Logs;
 using EpicGames.Redis;
-using Horde.Server.Logs;
-using Horde.Server.Server;
+using HordeServer.Logs;
+using HordeServer.Server;
 using Microsoft.Extensions.DependencyInjection;
 using Microsoft.Extensions.Logging;
 using Microsoft.Extensions.Options;
 using Microsoft.VisualStudio.TestTools.UnitTesting;
 
-namespace Horde.Server.Tests.Logs
+namespace HordeServer.Tests.Logs
 {
 	[TestClass]
 	public class LogTailServiceTests : TestSetup
diff --git a/Engine/Source/Programs/Horde/HordeServer.Tests/Notifications/NotificationServiceTests.cs b/Engine/Source/Programs/Horde/HordeServer.Tests/Notifications/NotificationServiceTests.cs
--- a/Engine/Source/Programs/Horde/HordeServer.Tests/Notifications/NotificationServiceTests.cs
+++ b/Engine/Source/Programs/Horde/HordeServer.Tests/Notifications/NotificationServiceTests.cs
@@ -6,23 +6,23 @@ using System.Threading;
 using System.Threading.Tasks;
 using EpicGames.Horde.Jobs;
 using EpicGames.Horde.Streams;
-using Horde.Server.Agents;
-using Horde.Server.Agents.Pools;
-using Horde.Server.Configuration;
-using Horde.Server.Devices;
-using Horde.Server.Issues;
-using Horde.Server.Jobs;
-using Horde.Server.Jobs.Graphs;
-using Horde.Server.Logs;
-using Horde.Server.Notifications;
-using Horde.Server.Streams;
-using Horde.Server.Users;
+using HordeServer.Agents;
+using HordeServer.Agents.Pools;
+using HordeServer.Configuration;
+using HordeServer.Devices;
+using HordeServer.Issues;
+using HordeServer.Jobs;
+using HordeServer.Jobs.Graphs;
+using HordeServer.Logs;
+using HordeServer.Notifications;
+using HordeServer.Streams;
+using HordeServer.Users;
 using Microsoft.Extensions.Caching.Memory;
 using Microsoft.Extensions.DependencyInjection;
 using Microsoft.VisualStudio.TestTools.UnitTesting;
 using Moq;
 
-namespace Horde.Server.Tests.Notifications
+namespace HordeServer.Tests.Notifications
 {
 	public class FakeNotificationSink : INotificationSink
 	{
diff --git a/Engine/Source/Programs/Horde/HordeServer.Tests/Perforce/PerforceLoadBalancerTests.cs b/Engine/Source/Programs/Horde/HordeServer.Tests/Perforce/PerforceLoadBalancerTests.cs
--- a/Engine/Source/Programs/Horde/HordeServer.Tests/Perforce/PerforceLoadBalancerTests.cs
+++ b/Engine/Source/Programs/Horde/HordeServer.Tests/Perforce/PerforceLoadBalancerTests.cs
@@ -7,13 +7,13 @@ using System.Net;
 using System.Net.Http;
 using System.Threading;
 using System.Threading.Tasks;
-using Horde.Server.Perforce;
-using Horde.Server.Server;
+using HordeServer.Perforce;
+using HordeServer.Server;
 using Microsoft.Extensions.Diagnostics.HealthChecks;
 using Microsoft.Extensions.Logging;
 using Microsoft.VisualStudio.TestTools.UnitTesting;
 
-namespace Horde.Server.Tests.Perforce;
+namespace HordeServer.Tests.Perforce;
 
 /// <summary>
 /// Fake implementation of IHealthMonitor for tests
diff --git a/Engine/Source/Programs/Horde/HordeServer.Tests/ProgramTests.cs b/Engine/Source/Programs/Horde/HordeServer.Tests/ProgramTests.cs
--- a/Engine/Source/Programs/Horde/HordeServer.Tests/ProgramTests.cs
+++ b/Engine/Source/Programs/Horde/HordeServer.Tests/ProgramTests.cs
@@ -5,10 +5,10 @@ using System.IO;
 using System.Runtime.InteropServices;
 using System.Security.Cryptography.X509Certificates;
 using EpicGames.Core;
-using Horde.Server.Commands;
+using HordeServer.Commands;
 using Microsoft.VisualStudio.TestTools.UnitTesting;
 
-namespace Horde.Server.Tests
+namespace HordeServer.Tests
 {
 	[TestClass]
 	public class ProgramTests
diff --git a/Engine/Source/Programs/Horde/HordeServer.Tests/Server/DebugControllerTests.cs b/Engine/Source/Programs/Horde/HordeServer.Tests/Server/DebugControllerTests.cs
--- a/Engine/Source/Programs/Horde/HordeServer.Tests/Server/DebugControllerTests.cs
+++ b/Engine/Source/Programs/Horde/HordeServer.Tests/Server/DebugControllerTests.cs
@@ -9,7 +9,7 @@ using Microsoft.VisualStudio.TestTools.UnitTesting;
 
 #pragma warning disable CA2234 // Pass system uri objects instead of strings
 
-namespace Horde.Server.Tests.Server;
+namespace HordeServer.Tests.Server;
 
 [TestClass]
 public class DebugControllerTests
diff --git a/Engine/Source/Programs/Horde/HordeServer.Tests/Server/MongoTests.cs b/Engine/Source/Programs/Horde/HordeServer.Tests/Server/MongoTests.cs
--- a/Engine/Source/Programs/Horde/HordeServer.Tests/Server/MongoTests.cs
+++ b/Engine/Source/Programs/Horde/HordeServer.Tests/Server/MongoTests.cs
@@ -6,7 +6,7 @@ using MongoDB.Bson;
 using MongoDB.Bson.Serialization;
 using MongoDB.Bson.Serialization.Attributes;
 
-namespace Horde.Server.Tests.Server
+namespace HordeServer.Tests.Server
 {
 	[TestClass]
 	public class MongoTests
diff --git a/Engine/Source/Programs/Horde/HordeServer.Tests/Server/RpcServiceTest.cs b/Engine/Source/Programs/Horde/HordeServer.Tests/Server/RpcServiceTest.cs
--- a/Engine/Source/Programs/Horde/HordeServer.Tests/Server/RpcServiceTest.cs
+++ b/Engine/Source/Programs/Horde/HordeServer.Tests/Server/RpcServiceTest.cs
@@ -12,9 +12,9 @@ using System.Threading.Tasks;
 using EpicGames.Horde.Agents;
 using EpicGames.Horde.Agents.Pools;
 using Grpc.Core;
-using Horde.Server.Agents;
-using Horde.Server.Server;
-using Horde.Server.Utilities;
+using HordeServer.Agents;
+using HordeServer.Server;
+using HordeServer.Utilities;
 using HordeCommon.Rpc;
 using HordeCommon.Rpc.Messages;
 using Microsoft.AspNetCore.Http;
@@ -23,7 +23,7 @@ using Microsoft.Extensions.FileProviders;
 using Microsoft.Extensions.Hosting;
 using Microsoft.VisualStudio.TestTools.UnitTesting;
 
-namespace Horde.Server.Tests.Server
+namespace HordeServer.Tests.Server
 {
 	using ISession = Microsoft.AspNetCore.Http.ISession;
 
diff --git a/Engine/Source/Programs/Horde/HordeServer.Tests/Server/SchemaControllerTest.cs b/Engine/Source/Programs/Horde/HordeServer.Tests/Server/SchemaControllerTest.cs
--- a/Engine/Source/Programs/Horde/HordeServer.Tests/Server/SchemaControllerTest.cs
+++ b/Engine/Source/Programs/Horde/HordeServer.Tests/Server/SchemaControllerTest.cs
@@ -9,7 +9,7 @@ using Microsoft.VisualStudio.TestTools.UnitTesting;
 
 #pragma warning disable CA2234 // Pass system uri objects instead of strings
 
-namespace Horde.Server.Tests.Server
+namespace HordeServer.Tests.Server
 {
 	[TestClass]
 	public class SchemaControllerTest : ControllerIntegrationTest
diff --git a/Engine/Source/Programs/Horde/HordeServer.Tests/Server/SchemaTests.cs b/Engine/Source/Programs/Horde/HordeServer.Tests/Server/SchemaTests.cs
--- a/Engine/Source/Programs/Horde/HordeServer.Tests/Server/SchemaTests.cs
+++ b/Engine/Source/Programs/Horde/HordeServer.Tests/Server/SchemaTests.cs
@@ -1,11 +1,11 @@
 // Copyright Epic Games, Inc. All Rights Reserved.
 
 using EpicGames.Core;
-using Horde.Server.Streams;
-using Horde.Server.Utilities;
+using HordeServer.Streams;
+using HordeServer.Utilities;
 using Microsoft.VisualStudio.TestTools.UnitTesting;
 
-namespace Horde.Server.Tests.Server
+namespace HordeServer.Tests.Server
 {
 	[TestClass]
 	public class SchemaTests
diff --git a/Engine/Source/Programs/Horde/HordeServer.Tests/Server/ServerStatusServiceTests.cs b/Engine/Source/Programs/Horde/HordeServer.Tests/Server/ServerStatusServiceTests.cs
--- a/Engine/Source/Programs/Horde/HordeServer.Tests/Server/ServerStatusServiceTests.cs
+++ b/Engine/Source/Programs/Horde/HordeServer.Tests/Server/ServerStatusServiceTests.cs
@@ -5,12 +5,12 @@ using System.Collections.Generic;
 using System.Linq;
 using System.Threading;
 using System.Threading.Tasks;
-using Horde.Server.Configuration;
-using Horde.Server.Server;
+using HordeServer.Configuration;
+using HordeServer.Server;
 using Microsoft.Extensions.Diagnostics.HealthChecks;
 using Microsoft.VisualStudio.TestTools.UnitTesting;
 
-namespace Horde.Server.Tests.Server
+namespace HordeServer.Tests.Server
 {
 	[TestClass]
 	public class ServerStatusServiceTest : TestSetup
diff --git a/Engine/Source/Programs/Horde/HordeServer.Tests/ServiceAccounts/ServiceAccountCollectionTests.cs b/Engine/Source/Programs/Horde/HordeServer.Tests/ServiceAccounts/ServiceAccountCollectionTests.cs
--- a/Engine/Source/Programs/Horde/HordeServer.Tests/ServiceAccounts/ServiceAccountCollectionTests.cs
+++ b/Engine/Source/Programs/Horde/HordeServer.Tests/ServiceAccounts/ServiceAccountCollectionTests.cs
@@ -2,13 +2,13 @@
 
 using System.Collections.Generic;
 using System.Threading.Tasks;
-using Horde.Server.Accounts;
-using Horde.Server.Server;
-using Horde.Server.ServiceAccounts;
-using Horde.Server.Users;
+using HordeServer.Accounts;
+using HordeServer.Server;
+using HordeServer.ServiceAccounts;
+using HordeServer.Users;
 using Microsoft.VisualStudio.TestTools.UnitTesting;
 
-namespace Horde.Server.Tests.Accounts
+namespace HordeServer.Tests.Accounts
 {
 	[TestClass]
 	public class HordeAccountCollectionTests : DatabaseIntegrationTest
diff --git a/Engine/Source/Programs/Horde/HordeServer.Tests/Storage/BlobStoreTests.cs b/Engine/Source/Programs/Horde/HordeServer.Tests/Storage/BlobStoreTests.cs
--- a/Engine/Source/Programs/Horde/HordeServer.Tests/Storage/BlobStoreTests.cs
+++ b/Engine/Source/Programs/Horde/HordeServer.Tests/Storage/BlobStoreTests.cs
@@ -6,11 +6,11 @@ using System.Linq;
 using System.Threading.Tasks;
 using EpicGames.Core;
 using EpicGames.Horde.Storage;
-using Horde.Server.Server;
-using Horde.Server.Storage;
+using HordeServer.Server;
+using HordeServer.Storage;
 using Microsoft.VisualStudio.TestTools.UnitTesting;
 
-namespace Horde.Server.Tests.Storage
+namespace HordeServer.Tests.Storage
 {
 	[TestClass]
 	public class BlobStoreTests : TestSetup
diff --git a/Engine/Source/Programs/Horde/HordeServer.Tests/Storage/GcServiceTests.cs b/Engine/Source/Programs/Horde/HordeServer.Tests/Storage/GcServiceTests.cs
--- a/Engine/Source/Programs/Horde/HordeServer.Tests/Storage/GcServiceTests.cs
+++ b/Engine/Source/Programs/Horde/HordeServer.Tests/Storage/GcServiceTests.cs
@@ -6,12 +6,12 @@ using System.Linq;
 using System.Threading;
 using System.Threading.Tasks;
 using EpicGames.Horde.Storage;
-using Horde.Server.Server;
-using Horde.Server.Storage;
+using HordeServer.Server;
+using HordeServer.Storage;
 using Microsoft.Extensions.DependencyInjection;
 using Microsoft.VisualStudio.TestTools.UnitTesting;
 
-namespace Horde.Server.Tests.Storage
+namespace HordeServer.Tests.Storage
 {
 	[TestClass]
 	public sealed class GcServiceTests : TestSetup
diff --git a/Engine/Source/Programs/Horde/HordeServer.Tests/Storage/StorageControllerTests.cs b/Engine/Source/Programs/Horde/HordeServer.Tests/Storage/StorageControllerTests.cs
--- a/Engine/Source/Programs/Horde/HordeServer.Tests/Storage/StorageControllerTests.cs
+++ b/Engine/Source/Programs/Horde/HordeServer.Tests/Storage/StorageControllerTests.cs
@@ -7,13 +7,13 @@ using EpicGames.Core;
 using EpicGames.Horde.Storage;
 using EpicGames.Horde.Storage.Backends;
 using EpicGames.Horde.Storage.Bundles;
-using Horde.Server.Storage;
+using HordeServer.Storage;
 using Microsoft.Extensions.DependencyInjection;
 using Microsoft.Extensions.Logging;
 using Microsoft.VisualStudio.TestTools.UnitTesting;
 using MongoDB.Driver;
 
-namespace Horde.Server.Tests.Storage
+namespace HordeServer.Tests.Storage
 {
 	[TestClass]
 	public class StorageControllerTests : ControllerIntegrationTest
diff --git a/Engine/Source/Programs/Horde/HordeServer.Tests/Storage/StorageServiceTests.cs b/Engine/Source/Programs/Horde/HordeServer.Tests/Storage/StorageServiceTests.cs
--- a/Engine/Source/Programs/Horde/HordeServer.Tests/Storage/StorageServiceTests.cs
+++ b/Engine/Source/Programs/Horde/HordeServer.Tests/Storage/StorageServiceTests.cs
@@ -4,11 +4,11 @@ using System;
 using System.Threading.Tasks;
 using EpicGames.Core;
 using EpicGames.Horde.Storage;
-using Horde.Server.Storage;
+using HordeServer.Storage;
 using Microsoft.Extensions.DependencyInjection;
 using Microsoft.VisualStudio.TestTools.UnitTesting;
 
-namespace Horde.Server.Tests.Storage
+namespace HordeServer.Tests.Storage
 {
 	[TestClass]
 	public class StorageServiceTests : TestSetup
diff --git a/Engine/Source/Programs/Horde/HordeServer.Tests/Streams/StreamCollectionTests.cs b/Engine/Source/Programs/Horde/HordeServer.Tests/Streams/StreamCollectionTests.cs
--- a/Engine/Source/Programs/Horde/HordeServer.Tests/Streams/StreamCollectionTests.cs
+++ b/Engine/Source/Programs/Horde/HordeServer.Tests/Streams/StreamCollectionTests.cs
@@ -3,10 +3,10 @@
 using System.Collections.Generic;
 using EpicGames.Horde.Jobs.Templates;
 using EpicGames.Horde.Streams;
-using Horde.Server.Streams;
+using HordeServer.Streams;
 using Microsoft.VisualStudio.TestTools.UnitTesting;
 
-namespace Horde.Server.Tests.Streams
+namespace HordeServer.Tests.Streams
 {
 	[TestClass]
 	public class StreamCollectionTests : TestSetup
@@ -26,10 +26,10 @@ namespace Horde.Server.Tests.Streams
 				Templates = new() { new TemplateRefConfig { Id = new TemplateId("foo") } }
 			};
 
-			Assert.ThrowsException<InvalidStreamException>(() => Horde.Server.Streams.StreamCollection.Validate(_streamId, config));
+			Assert.ThrowsException<InvalidStreamException>(() => HordeServer.Streams.StreamCollection.Validate(_streamId, config));
 
 			config.Templates.Add(new TemplateRefConfig { Id = new TemplateId("bar") });
-			Horde.Server.Streams.StreamCollection.Validate(_streamId, config);
+			HordeServer.Streams.StreamCollection.Validate(_streamId, config);
 		}
 	}
 }
diff --git a/Engine/Source/Programs/Horde/HordeServer.Tests/Streams/StreamServiceTest.cs b/Engine/Source/Programs/Horde/HordeServer.Tests/Streams/StreamServiceTest.cs
--- a/Engine/Source/Programs/Horde/HordeServer.Tests/Streams/StreamServiceTest.cs
+++ b/Engine/Source/Programs/Horde/HordeServer.Tests/Streams/StreamServiceTest.cs
@@ -3,10 +3,10 @@
 using System;
 using System.Globalization;
 using System.Threading.Tasks;
-using Horde.Server.Streams;
+using HordeServer.Streams;
 using Microsoft.VisualStudio.TestTools.UnitTesting;
 
-namespace Horde.Server.Tests.Streams
+namespace HordeServer.Tests.Streams
 {
 	[TestClass]
 	public class StreamServiceTests : TestSetup
diff --git a/Engine/Source/Programs/Horde/HordeServer.Tests/Stubs/Collections/SingletonDocumentStub.cs b/Engine/Source/Programs/Horde/HordeServer.Tests/Stubs/Collections/SingletonDocumentStub.cs
--- a/Engine/Source/Programs/Horde/HordeServer.Tests/Stubs/Collections/SingletonDocumentStub.cs
+++ b/Engine/Source/Programs/Horde/HordeServer.Tests/Stubs/Collections/SingletonDocumentStub.cs
@@ -2,9 +2,9 @@
 
 using System.Threading;
 using System.Threading.Tasks;
-using Horde.Server.Utilities;
+using HordeServer.Utilities;
 
-namespace Horde.Server.Tests.Stubs.Collections
+namespace HordeServer.Tests.Stubs.Collections
 {
 	class SingletonDocumentStub<T> : ISingletonDocument<T> where T : SingletonBase, new()
 	{
diff --git a/Engine/Source/Programs/Horde/HordeServer.Tests/Stubs/Services/PerforceServiceStub.cs b/Engine/Source/Programs/Horde/HordeServer.Tests/Stubs/Services/PerforceServiceStub.cs
--- a/Engine/Source/Programs/Horde/HordeServer.Tests/Stubs/Services/PerforceServiceStub.cs
+++ b/Engine/Source/Programs/Horde/HordeServer.Tests/Stubs/Services/PerforceServiceStub.cs
@@ -10,14 +10,14 @@ using EpicGames.Core;
 using EpicGames.Horde.Commits;
 using EpicGames.Horde.Streams;
 using EpicGames.Horde.Users;
-using Horde.Server.Commits;
-using Horde.Server.Perforce;
-using Horde.Server.Streams;
-using Horde.Server.Users;
-using Horde.Server.Utilities;
+using HordeServer.Commits;
+using HordeServer.Perforce;
+using HordeServer.Streams;
+using HordeServer.Users;
+using HordeServer.Utilities;
 using Microsoft.VisualStudio.TestTools.UnitTesting;
 
-namespace Horde.Server.Tests.Stubs.Services
+namespace HordeServer.Tests.Stubs.Services
 {
 	class PerforceServiceStub : IPerforceService
 	{
diff --git a/Engine/Source/Programs/Horde/HordeServer.Tests/Telemetry/TelemetryTests.cs b/Engine/Source/Programs/Horde/HordeServer.Tests/Telemetry/TelemetryTests.cs
--- a/Engine/Source/Programs/Horde/HordeServer.Tests/Telemetry/TelemetryTests.cs
+++ b/Engine/Source/Programs/Horde/HordeServer.Tests/Telemetry/TelemetryTests.cs
@@ -8,15 +8,15 @@ using System.Threading;
 using System.Threading.Tasks;
 using EpicGames.Horde.Telemetry;
 using EpicGames.Horde.Telemetry.Metrics;
-using Horde.Server.Server;
-using Horde.Server.Telemetry;
-using Horde.Server.Telemetry.Metrics;
-using Horde.Server.Telemetry.Sinks;
+using HordeServer.Server;
+using HordeServer.Telemetry;
+using HordeServer.Telemetry.Metrics;
+using HordeServer.Telemetry.Sinks;
 using Json.Path;
 using Microsoft.Extensions.DependencyInjection;
 using Microsoft.VisualStudio.TestTools.UnitTesting;
 
-namespace Horde.Server.Tests.Telemetry
+namespace HordeServer.Tests.Telemetry
 {
 	[TestClass]
 	public class TelemetryTests : TestSetup
diff --git a/Engine/Source/Programs/Horde/HordeServer.Tests/TestSetup.cs b/Engine/Source/Programs/Horde/HordeServer.Tests/TestSetup.cs
--- a/Engine/Source/Programs/Horde/HordeServer.Tests/TestSetup.cs
+++ b/Engine/Source/Programs/Horde/HordeServer.Tests/TestSetup.cs
@@ -17,49 +17,49 @@ using EpicGames.Horde.Agents.Pools;
 using EpicGames.Horde.Storage;
 using EpicGames.Horde.Storage.Bundles;
 using EpicGames.Horde.Storage.ObjectStores;
-using Horde.Server.Accounts;
-using Horde.Server.Acls;
-using Horde.Server.Agents;
-using Horde.Server.Agents.Enrollment;
-using Horde.Server.Agents.Fleet;
-using Horde.Server.Agents.Leases;
-using Horde.Server.Agents.Pools;
-using Horde.Server.Agents.Relay;
-using Horde.Server.Agents.Sessions;
-using Horde.Server.Agents.Telemetry;
-using Horde.Server.Agents.Utilization;
-using Horde.Server.Artifacts;
-using Horde.Server.Auditing;
-using Horde.Server.Commits;
-using Horde.Server.Compute;
-using Horde.Server.Configuration;
-using Horde.Server.Dashboard;
-using Horde.Server.Devices;
-using Horde.Server.Issues;
-using Horde.Server.Jobs;
-using Horde.Server.Jobs.Artifacts;
-using Horde.Server.Jobs.Bisect;
-using Horde.Server.Jobs.Graphs;
-using Horde.Server.Jobs.Schedules;
-using Horde.Server.Jobs.Templates;
-using Horde.Server.Jobs.TestData;
-using Horde.Server.Jobs.Timing;
-using Horde.Server.Logs;
-using Horde.Server.Logs.Storage;
-using Horde.Server.Notifications;
-using Horde.Server.Perforce;
-using Horde.Server.Server;
-using Horde.Server.Storage;
-using Horde.Server.Streams;
-using Horde.Server.Tasks;
-using Horde.Server.Telemetry;
-using Horde.Server.Telemetry.Sinks;
-using Horde.Server.Tests.Server;
-using Horde.Server.Tests.Stubs.Services;
-using Horde.Server.Tools;
-using Horde.Server.Ugs;
-using Horde.Server.Users;
-using Horde.Server.Utilities;
+using HordeServer.Accounts;
+using HordeServer.Acls;
+using HordeServer.Agents;
+using HordeServer.Agents.Enrollment;
+using HordeServer.Agents.Fleet;
+using HordeServer.Agents.Leases;
+using HordeServer.Agents.Pools;
+using HordeServer.Agents.Relay;
+using HordeServer.Agents.Sessions;
+using HordeServer.Agents.Telemetry;
+using HordeServer.Agents.Utilization;
+using HordeServer.Artifacts;
+using HordeServer.Auditing;
+using HordeServer.Commits;
+using HordeServer.Compute;
+using HordeServer.Configuration;
+using HordeServer.Dashboard;
+using HordeServer.Devices;
+using HordeServer.Issues;
+using HordeServer.Jobs;
+using HordeServer.Jobs.Artifacts;
+using HordeServer.Jobs.Bisect;
+using HordeServer.Jobs.Graphs;
+using HordeServer.Jobs.Schedules;
+using HordeServer.Jobs.Templates;
+using HordeServer.Jobs.TestData;
+using HordeServer.Jobs.Timing;
+using HordeServer.Logs;
+using HordeServer.Logs.Storage;
+using HordeServer.Notifications;
+using HordeServer.Perforce;
+using HordeServer.Server;
+using HordeServer.Storage;
+using HordeServer.Streams;
+using HordeServer.Tasks;
+using HordeServer.Telemetry;
+using HordeServer.Telemetry.Sinks;
+using HordeServer.Tests.Server;
+using HordeServer.Tests.Stubs.Services;
+using HordeServer.Tools;
+using HordeServer.Ugs;
+using HordeServer.Users;
+using HordeServer.Utilities;
 using HordeCommon;
 using Microsoft.AspNetCore.Http;
 using Microsoft.AspNetCore.Mvc;
@@ -73,7 +73,7 @@ using Microsoft.VisualStudio.TestTools.UnitTesting;
 using Moq;
 using OpenTelemetry.Trace;
 
-namespace Horde.Server.Tests
+namespace HordeServer.Tests
 {
 	/// <summary>
 	/// Handles set up of collections, services, fixtures etc during testing
