Horde: Rename namespaces to match project folders.

diff --git a/Engine/Source/Programs/Horde/HordeServer.Tests/Tools/ToolTests.cs b/Engine/Source/Programs/Horde/HordeServer.Tests/Tools/ToolTests.cs
--- a/Engine/Source/Programs/Horde/HordeServer.Tests/Tools/ToolTests.cs
+++ b/Engine/Source/Programs/Horde/HordeServer.Tests/Tools/ToolTests.cs
@@ -8,14 +8,14 @@ using System.Text;
 using System.Threading;
 using System.Threading.Tasks;
 using EpicGames.Horde.Tools;
-using Horde.Server.Server;
-using Horde.Server.Storage;
-using Horde.Server.Tools;
+using HordeServer.Server;
+using HordeServer.Storage;
+using HordeServer.Tools;
 using HordeCommon;
 using Microsoft.Extensions.DependencyInjection;
 using Microsoft.VisualStudio.TestTools.UnitTesting;
 
-namespace Horde.Server.Tests.Tools
+namespace HordeServer.Tests.Tools
 {
 	[TestClass]
 	public class ToolTests : TestSetup
diff --git a/Engine/Source/Programs/Horde/HordeServer.Tests/Tools/ToolsControllerTests.cs b/Engine/Source/Programs/Horde/HordeServer.Tests/Tools/ToolsControllerTests.cs
--- a/Engine/Source/Programs/Horde/HordeServer.Tests/Tools/ToolsControllerTests.cs
+++ b/Engine/Source/Programs/Horde/HordeServer.Tests/Tools/ToolsControllerTests.cs
@@ -8,18 +8,18 @@ using System.Net.Http.Headers;
 using System.Threading.Tasks;
 using EpicGames.Horde.Server;
 using EpicGames.Horde.Tools;
-using Horde.Server.Acls;
-using Horde.Server.Configuration;
-using Horde.Server.Server;
-using Horde.Server.ServiceAccounts;
-using Horde.Server.Storage;
-using Horde.Server.Tools;
-using Horde.Server.Users;
-using Horde.Server.Utilities;
+using HordeServer.Acls;
+using HordeServer.Configuration;
+using HordeServer.Server;
+using HordeServer.ServiceAccounts;
+using HordeServer.Storage;
+using HordeServer.Tools;
+using HordeServer.Users;
+using HordeServer.Utilities;
 using Microsoft.Extensions.DependencyInjection;
 using Microsoft.VisualStudio.TestTools.UnitTesting;
 
-namespace Horde.Server.Tests.Tools;
+namespace HordeServer.Tests.Tools;
 
 [TestClass]
 public class ToolsControllerTests
diff --git a/Engine/Source/Programs/Horde/HordeServer.Tests/Utilities/AsyncTaskQueueTests.cs b/Engine/Source/Programs/Horde/HordeServer.Tests/Utilities/AsyncTaskQueueTests.cs
--- a/Engine/Source/Programs/Horde/HordeServer.Tests/Utilities/AsyncTaskQueueTests.cs
+++ b/Engine/Source/Programs/Horde/HordeServer.Tests/Utilities/AsyncTaskQueueTests.cs
@@ -1,11 +1,11 @@
 // Copyright Epic Games, Inc. All Rights Reserved.
 
 using System.Threading.Tasks;
-using Horde.Server.Utilities;
+using HordeServer.Utilities;
 using Microsoft.Extensions.Logging.Abstractions;
 using Microsoft.VisualStudio.TestTools.UnitTesting;
 
-namespace Horde.Server.Tests.Utilities
+namespace HordeServer.Tests.Utilities
 {
 	[TestClass]
 	public class AsyncTaskQueueTests
diff --git a/Engine/Source/Programs/Horde/HordeServer.Tests/Utilities/IntervalJsonConverterTests.cs b/Engine/Source/Programs/Horde/HordeServer.Tests/Utilities/IntervalJsonConverterTests.cs
--- a/Engine/Source/Programs/Horde/HordeServer.Tests/Utilities/IntervalJsonConverterTests.cs
+++ b/Engine/Source/Programs/Horde/HordeServer.Tests/Utilities/IntervalJsonConverterTests.cs
@@ -3,10 +3,10 @@
 using System;
 using System.Text.Json;
 using System.Text.Json.Serialization;
-using Horde.Server.Utilities;
+using HordeServer.Utilities;
 using Microsoft.VisualStudio.TestTools.UnitTesting;
 
-namespace Horde.Server.Tests.Utilities
+namespace HordeServer.Tests.Utilities
 {
 	[TestClass]
 	public class IntervalJsonConverterTests
diff --git a/Engine/Source/Programs/Horde/HordeServer.Tests/Utilities/MongoCommandTracer.cs b/Engine/Source/Programs/Horde/HordeServer.Tests/Utilities/MongoCommandTracer.cs
--- a/Engine/Source/Programs/Horde/HordeServer.Tests/Utilities/MongoCommandTracer.cs
+++ b/Engine/Source/Programs/Horde/HordeServer.Tests/Utilities/MongoCommandTracer.cs
@@ -3,7 +3,7 @@
 using System;
 using System.Collections.Generic;
 using System.Net;
-using Horde.Server.Utilities;
+using HordeServer.Utilities;
 using Microsoft.Extensions.Logging;
 using Microsoft.VisualStudio.TestTools.UnitTesting;
 using MongoDB.Bson;
@@ -14,7 +14,7 @@ using MongoDB.Driver.Core.Events;
 using MongoDB.Driver.Core.Servers;
 using OpenTelemetry.Trace;
 
-namespace Horde.Server.Tests.Utilities;
+namespace HordeServer.Tests.Utilities;
 
 [TestClass]
 public class MongoCommandTracerTests
diff --git a/Engine/Source/Programs/Horde/HordeServer/Accounts/AccountAclAction.cs b/Engine/Source/Programs/Horde/HordeServer/Accounts/AccountAclAction.cs
--- a/Engine/Source/Programs/Horde/HordeServer/Accounts/AccountAclAction.cs
+++ b/Engine/Source/Programs/Horde/HordeServer/Accounts/AccountAclAction.cs
@@ -2,7 +2,7 @@
 
 using EpicGames.Horde.Acls;
 
-namespace Horde.Server.Accounts
+namespace HordeServer.Accounts
 {
 	/// <summary>
 	/// ACL actions valid for manipulating jobs
diff --git a/Engine/Source/Programs/Horde/HordeServer/Accounts/AccountCollection.cs b/Engine/Source/Programs/Horde/HordeServer/Accounts/AccountCollection.cs
--- a/Engine/Source/Programs/Horde/HordeServer/Accounts/AccountCollection.cs
+++ b/Engine/Source/Programs/Horde/HordeServer/Accounts/AccountCollection.cs
@@ -8,14 +8,14 @@ using System.Threading;
 using System.Threading.Tasks;
 using EpicGames.Core;
 using EpicGames.Horde.Accounts;
-using Horde.Server.Acls;
-using Horde.Server.Server;
-using Horde.Server.Users;
-using Horde.Server.Utilities;
+using HordeServer.Acls;
+using HordeServer.Server;
+using HordeServer.Users;
+using HordeServer.Utilities;
 using MongoDB.Bson.Serialization.Attributes;
 using MongoDB.Driver;
 
-namespace Horde.Server.Accounts
+namespace HordeServer.Accounts
 {
 	/// <summary>
 	/// Password hasher
diff --git a/Engine/Source/Programs/Horde/HordeServer/Accounts/AccountsController.cs b/Engine/Source/Programs/Horde/HordeServer/Accounts/AccountsController.cs
--- a/Engine/Source/Programs/Horde/HordeServer/Accounts/AccountsController.cs
+++ b/Engine/Source/Programs/Horde/HordeServer/Accounts/AccountsController.cs
@@ -5,14 +5,14 @@ using System.Threading;
 using System.Threading.Tasks;
 using EpicGames.Horde.Accounts;
 using EpicGames.Horde.Server;
-using Horde.Server.Server;
-using Horde.Server.Users;
-using Horde.Server.Utilities;
+using HordeServer.Server;
+using HordeServer.Users;
+using HordeServer.Utilities;
 using Microsoft.AspNetCore.Authorization;
 using Microsoft.AspNetCore.Mvc;
 using Microsoft.Extensions.Options;
 
-namespace Horde.Server.Accounts
+namespace HordeServer.Accounts
 {
 	/// <summary>
 	/// Controller for the /api/v1/accounts endpoint
diff --git a/Engine/Source/Programs/Horde/HordeServer/Accounts/IAccount.cs b/Engine/Source/Programs/Horde/HordeServer/Accounts/IAccount.cs
--- a/Engine/Source/Programs/Horde/HordeServer/Accounts/IAccount.cs
+++ b/Engine/Source/Programs/Horde/HordeServer/Accounts/IAccount.cs
@@ -5,10 +5,10 @@ using System.Collections.Generic;
 using System.Threading;
 using System.Threading.Tasks;
 using EpicGames.Horde.Accounts;
-using Horde.Server.Acls;
-using Horde.Server.Users;
+using HordeServer.Acls;
+using HordeServer.Users;
 
-namespace Horde.Server.Accounts
+namespace HordeServer.Accounts
 {
 	/// <summary>
 	/// An internal Horde account representing a user or service
diff --git a/Engine/Source/Programs/Horde/HordeServer/Accounts/IAccountCollection.cs b/Engine/Source/Programs/Horde/HordeServer/Accounts/IAccountCollection.cs
--- a/Engine/Source/Programs/Horde/HordeServer/Accounts/IAccountCollection.cs
+++ b/Engine/Source/Programs/Horde/HordeServer/Accounts/IAccountCollection.cs
@@ -4,9 +4,9 @@ using System.Collections.Generic;
 using System.Threading;
 using System.Threading.Tasks;
 using EpicGames.Horde.Accounts;
-using Horde.Server.Users;
+using HordeServer.Users;
 
-namespace Horde.Server.Accounts
+namespace HordeServer.Accounts
 {
 	/// <summary>
 	/// Interface for a collection of accounts
diff --git a/Engine/Source/Programs/Horde/HordeServer/Accounts/OAuthController.cs b/Engine/Source/Programs/Horde/HordeServer/Accounts/OAuthController.cs
--- a/Engine/Source/Programs/Horde/HordeServer/Accounts/OAuthController.cs
+++ b/Engine/Source/Programs/Horde/HordeServer/Accounts/OAuthController.cs
@@ -11,9 +11,9 @@ using System.Threading.Tasks;
 using System.Web;
 using EpicGames.Horde.Accounts;
 using EpicGames.Horde.Server;
-using Horde.Server.Server;
-using Horde.Server.Users;
-using Horde.Server.Utilities;
+using HordeServer.Server;
+using HordeServer.Users;
+using HordeServer.Utilities;
 using IdentityModel;
 using Microsoft.AspNetCore.Authorization;
 using Microsoft.AspNetCore.Mvc;
@@ -22,7 +22,7 @@ using Microsoft.AspNetCore.WebUtilities;
 using Microsoft.Extensions.Options;
 using Microsoft.IdentityModel.Tokens;
 
-namespace Horde.Server.Accounts
+namespace HordeServer.Accounts
 {
 	/// <summary>
 	/// Controller for /api/v1/oauth2 endpoints
diff --git a/Engine/Source/Programs/Horde/HordeServer/Accounts/OAuthMessages.cs b/Engine/Source/Programs/Horde/HordeServer/Accounts/OAuthMessages.cs
--- a/Engine/Source/Programs/Horde/HordeServer/Accounts/OAuthMessages.cs
+++ b/Engine/Source/Programs/Horde/HordeServer/Accounts/OAuthMessages.cs
@@ -4,7 +4,7 @@ using System;
 using System.Text.Json.Serialization;
 using Microsoft.AspNetCore.Mvc;
 
-namespace Horde.Server.Accounts
+namespace HordeServer.Accounts
 {
 	/// <summary>
 	/// Request to authorize a user using OAuth2
diff --git a/Engine/Source/Programs/Horde/HordeServer/Acls/AclConfig.cs b/Engine/Source/Programs/Horde/HordeServer/Acls/AclConfig.cs
--- a/Engine/Source/Programs/Horde/HordeServer/Acls/AclConfig.cs
+++ b/Engine/Source/Programs/Horde/HordeServer/Acls/AclConfig.cs
@@ -7,23 +7,23 @@ using System.Linq;
 using System.Security.Claims;
 using System.Text.Json.Serialization;
 using EpicGames.Horde.Acls;
-using Horde.Server.Agents;
-using Horde.Server.Agents.Pools;
-using Horde.Server.Agents.Sessions;
-using Horde.Server.Agents.Software;
-using Horde.Server.Artifacts;
-using Horde.Server.Devices;
-using Horde.Server.Jobs;
-using Horde.Server.Jobs.Bisect;
-using Horde.Server.Logs;
-using Horde.Server.Notifications;
-using Horde.Server.Projects;
-using Horde.Server.Server;
-using Horde.Server.Streams;
-using Horde.Server.Tools;
-using Horde.Server.Utilities;
-
-namespace Horde.Server.Acls
+using HordeServer.Agents;
+using HordeServer.Agents.Pools;
+using HordeServer.Agents.Sessions;
+using HordeServer.Agents.Software;
+using HordeServer.Artifacts;
+using HordeServer.Devices;
+using HordeServer.Jobs;
+using HordeServer.Jobs.Bisect;
+using HordeServer.Logs;
+using HordeServer.Notifications;
+using HordeServer.Projects;
+using HordeServer.Server;
+using HordeServer.Streams;
+using HordeServer.Tools;
+using HordeServer.Utilities;
+
+namespace HordeServer.Acls
 {
 	/// <summary>
 	/// Parameters to update an ACL
diff --git a/Engine/Source/Programs/Horde/HordeServer/Acls/AclProfileId.cs b/Engine/Source/Programs/Horde/HordeServer/Acls/AclProfileId.cs
--- a/Engine/Source/Programs/Horde/HordeServer/Acls/AclProfileId.cs
+++ b/Engine/Source/Programs/Horde/HordeServer/Acls/AclProfileId.cs
@@ -4,7 +4,7 @@ using System.ComponentModel;
 using EpicGames.Core;
 using EpicGames.Horde;
 
-namespace Horde.Server.Acls
+namespace HordeServer.Acls
 {
 	/// <summary>
 	/// Identifier for a profile; a group of actions which can be assigned to a user as a whole.
diff --git a/Engine/Source/Programs/Horde/HordeServer/Acls/AclScopeName.cs b/Engine/Source/Programs/Horde/HordeServer/Acls/AclScopeName.cs
--- a/Engine/Source/Programs/Horde/HordeServer/Acls/AclScopeName.cs
+++ b/Engine/Source/Programs/Horde/HordeServer/Acls/AclScopeName.cs
@@ -8,7 +8,7 @@ using System.Globalization;
 using System.Text.Json;
 using System.Text.Json.Serialization;
 
-namespace Horde.Server.Acls
+namespace HordeServer.Acls
 {
 	/// <summary>
 	/// Name of an ACL scope
diff --git a/Engine/Source/Programs/Horde/HordeServer/Acls/AclService.cs b/Engine/Source/Programs/Horde/HordeServer/Acls/AclService.cs
--- a/Engine/Source/Programs/Horde/HordeServer/Acls/AclService.cs
+++ b/Engine/Source/Programs/Horde/HordeServer/Acls/AclService.cs
@@ -10,12 +10,12 @@ using System.Threading.Tasks;
 using EpicGames.Horde.Agents;
 using EpicGames.Horde.Agents.Leases;
 using EpicGames.Horde.Agents.Sessions;
-using Horde.Server.Server;
-using Horde.Server.Utilities;
+using HordeServer.Server;
+using HordeServer.Utilities;
 using Microsoft.IdentityModel.Tokens;
 using MongoDB.Driver;
 
-namespace Horde.Server.Acls
+namespace HordeServer.Acls
 {
 	/// <summary>
 	/// Wraps functionality for manipulating permissions
diff --git a/Engine/Source/Programs/Horde/HordeServer/Agents/AgentAclAction.cs b/Engine/Source/Programs/Horde/HordeServer/Agents/AgentAclAction.cs
--- a/Engine/Source/Programs/Horde/HordeServer/Agents/AgentAclAction.cs
+++ b/Engine/Source/Programs/Horde/HordeServer/Agents/AgentAclAction.cs
@@ -2,7 +2,7 @@
 
 using EpicGames.Horde.Acls;
 
-namespace Horde.Server.Agents
+namespace HordeServer.Agents
 {
 	/// <summary>
 	/// ACL actions which apply to agents
diff --git a/Engine/Source/Programs/Horde/HordeServer/Agents/AgentCollection.cs b/Engine/Source/Programs/Horde/HordeServer/Agents/AgentCollection.cs
--- a/Engine/Source/Programs/Horde/HordeServer/Agents/AgentCollection.cs
+++ b/Engine/Source/Programs/Horde/HordeServer/Agents/AgentCollection.cs
@@ -15,15 +15,15 @@ using EpicGames.Horde.Agents.Pools;
 using EpicGames.Horde.Agents.Sessions;
 using EpicGames.Redis;
 using Google.Protobuf.WellKnownTypes;
-using Horde.Server.Auditing;
-using Horde.Server.Server;
+using HordeServer.Auditing;
+using HordeServer.Server;
 using HordeCommon.Rpc.Tasks;
 using MongoDB.Bson;
 using MongoDB.Bson.Serialization.Attributes;
 using MongoDB.Driver;
 using StackExchange.Redis;
 
-namespace Horde.Server.Agents
+namespace HordeServer.Agents
 {
 	/// <summary>
 	/// Collection of agent documents
diff --git a/Engine/Source/Programs/Horde/HordeServer/Agents/AgentReportService.cs b/Engine/Source/Programs/Horde/HordeServer/Agents/AgentReportService.cs
--- a/Engine/Source/Programs/Horde/HordeServer/Agents/AgentReportService.cs
+++ b/Engine/Source/Programs/Horde/HordeServer/Agents/AgentReportService.cs
@@ -6,15 +6,15 @@ using System.Threading;
 using System.Threading.Tasks;
 using EpicGames.Core;
 using EpicGames.Horde.Agents;
-using Horde.Server.Jobs;
-using Horde.Server.Notifications;
-using Horde.Server.Server;
-using Horde.Server.Utilities;
+using HordeServer.Jobs;
+using HordeServer.Notifications;
+using HordeServer.Server;
+using HordeServer.Utilities;
 using HordeCommon;
 using Microsoft.Extensions.Hosting;
 using Microsoft.Extensions.Logging;
 
-namespace Horde.Server.Agents
+namespace HordeServer.Agents
 {
 	/// <summary>
 	/// Report for the state of agents in the farm
diff --git a/Engine/Source/Programs/Horde/HordeServer/Agents/AgentService.cs b/Engine/Source/Programs/Horde/HordeServer/Agents/AgentService.cs
--- a/Engine/Source/Programs/Horde/HordeServer/Agents/AgentService.cs
+++ b/Engine/Source/Programs/Horde/HordeServer/Agents/AgentService.cs
@@ -16,14 +16,14 @@ using EpicGames.Horde.Agents.Sessions;
 using EpicGames.Redis;
 using EpicGames.Serialization;
 using Google.Protobuf.WellKnownTypes;
-using Horde.Server.Acls;
-using Horde.Server.Agents.Leases;
-using Horde.Server.Agents.Pools;
-using Horde.Server.Agents.Sessions;
-using Horde.Server.Auditing;
-using Horde.Server.Server;
-using Horde.Server.Tasks;
-using Horde.Server.Utilities;
+using HordeServer.Acls;
+using HordeServer.Agents.Leases;
+using HordeServer.Agents.Pools;
+using HordeServer.Agents.Sessions;
+using HordeServer.Auditing;
+using HordeServer.Server;
+using HordeServer.Tasks;
+using HordeServer.Utilities;
 using HordeCommon.Rpc;
 using HordeCommon.Rpc.Messages;
 using Microsoft.Extensions.Hosting;
@@ -31,7 +31,7 @@ using Microsoft.Extensions.Logging;
 using OpenTelemetry.Trace;
 using StackExchange.Redis;
 
-namespace Horde.Server.Agents
+namespace HordeServer.Agents
 {
 	/// <summary>
 	/// Singleton used to store agent costs
diff --git a/Engine/Source/Programs/Horde/HordeServer/Agents/AgentsController.cs b/Engine/Source/Programs/Horde/HordeServer/Agents/AgentsController.cs
--- a/Engine/Source/Programs/Horde/HordeServer/Agents/AgentsController.cs
+++ b/Engine/Source/Programs/Horde/HordeServer/Agents/AgentsController.cs
@@ -11,20 +11,20 @@ using EpicGames.Horde.Agents.Leases;
 using EpicGames.Horde.Agents.Pools;
 using EpicGames.Horde.Agents.Sessions;
 using EpicGames.Horde.Common;
-using Horde.Server.Agents.Leases;
-using Horde.Server.Agents.Sessions;
-using Horde.Server.Agents.Telemetry;
-using Horde.Server.Auditing;
-using Horde.Server.Server;
-using Horde.Server.Users;
-using Horde.Server.Utilities;
+using HordeServer.Agents.Leases;
+using HordeServer.Agents.Sessions;
+using HordeServer.Agents.Telemetry;
+using HordeServer.Auditing;
+using HordeServer.Server;
+using HordeServer.Users;
+using HordeServer.Utilities;
 using Microsoft.AspNetCore.Authorization;
 using Microsoft.AspNetCore.Mvc;
 using Microsoft.Extensions.Logging;
 using Microsoft.Extensions.Options;
 using OpenTelemetry.Trace;
 
-namespace Horde.Server.Agents
+namespace HordeServer.Agents
 {
 	/// <summary>
 	/// Controller for the /api/v1/agents endpoint
diff --git a/Engine/Source/Programs/Horde/HordeServer/Agents/ConformList.cs b/Engine/Source/Programs/Horde/HordeServer/Agents/ConformList.cs
--- a/Engine/Source/Programs/Horde/HordeServer/Agents/ConformList.cs
+++ b/Engine/Source/Programs/Horde/HordeServer/Agents/ConformList.cs
@@ -4,11 +4,11 @@ using System;
 using System.Collections.Generic;
 using EpicGames.Horde.Agents;
 using EpicGames.Horde.Agents.Leases;
-using Horde.Server.Utilities;
+using HordeServer.Utilities;
 using MongoDB.Bson;
 using MongoDB.Bson.Serialization.Attributes;
 
-namespace Horde.Server.Agents
+namespace HordeServer.Agents
 {
 	/// <summary>
 	/// Entry to the conform list
diff --git a/Engine/Source/Programs/Horde/HordeServer/Agents/Enrollment/EnrollmentController.cs b/Engine/Source/Programs/Horde/HordeServer/Agents/Enrollment/EnrollmentController.cs
--- a/Engine/Source/Programs/Horde/HordeServer/Agents/Enrollment/EnrollmentController.cs
+++ b/Engine/Source/Programs/Horde/HordeServer/Agents/Enrollment/EnrollmentController.cs
@@ -6,13 +6,13 @@ using System.Threading;
 using System.Threading.Tasks;
 using EpicGames.Core;
 using EpicGames.Horde.Agents.Enrollment;
-using Horde.Server.Server;
-using Horde.Server.Utilities;
+using HordeServer.Server;
+using HordeServer.Utilities;
 using Microsoft.AspNetCore.Authorization;
 using Microsoft.AspNetCore.Mvc;
 using Microsoft.Extensions.Options;
 
-namespace Horde.Server.Agents.Enrollment
+namespace HordeServer.Agents.Enrollment
 {
 	/// <summary>
 	/// Controller for the /api/v1/enrollment endpoint
diff --git a/Engine/Source/Programs/Horde/HordeServer/Agents/Enrollment/EnrollmentRpc.cs b/Engine/Source/Programs/Horde/HordeServer/Agents/Enrollment/EnrollmentRpc.cs
--- a/Engine/Source/Programs/Horde/HordeServer/Agents/Enrollment/EnrollmentRpc.cs
+++ b/Engine/Source/Programs/Horde/HordeServer/Agents/Enrollment/EnrollmentRpc.cs
@@ -8,18 +8,18 @@ using EpicGames.Core;
 using EpicGames.Horde.Agents;
 using Grpc.Core;
 using Horde.Common.Rpc;
-using Horde.Server.Acls;
-using Horde.Server.Utilities;
+using HordeServer.Acls;
+using HordeServer.Utilities;
 using Microsoft.AspNetCore.Authorization;
 using Microsoft.Extensions.Logging;
 
-namespace Horde.Server.Agents.Enrollment
+namespace HordeServer.Agents.Enrollment
 {
 	/// <summary>
 	/// Allow agents to register with the server
 	/// </summary>
 	[AllowAnonymous]
-	public class EnrollmentRpc : Common.Rpc.EnrollmentRpc.EnrollmentRpcBase
+	public class EnrollmentRpc : Horde.Common.Rpc.EnrollmentRpc.EnrollmentRpcBase
 	{
 		readonly EnrollmentService _registrationService;
 		readonly AgentService _agentService;
diff --git a/Engine/Source/Programs/Horde/HordeServer/Agents/Enrollment/EnrollmentService.cs b/Engine/Source/Programs/Horde/HordeServer/Agents/Enrollment/EnrollmentService.cs
--- a/Engine/Source/Programs/Horde/HordeServer/Agents/Enrollment/EnrollmentService.cs
+++ b/Engine/Source/Programs/Horde/HordeServer/Agents/Enrollment/EnrollmentService.cs
@@ -7,12 +7,12 @@ using System.Threading.Tasks;
 using EpicGames.Core;
 using EpicGames.Horde.Agents;
 using EpicGames.Redis;
-using Horde.Server.Server;
+using HordeServer.Server;
 using Microsoft.Extensions.Hosting;
 using Microsoft.Extensions.Options;
 using StackExchange.Redis;
 
-namespace Horde.Server.Agents.Enrollment
+namespace HordeServer.Agents.Enrollment
 {
 	/// <summary>
 	/// Information about a machine requesting to be registered
diff --git a/Engine/Source/Programs/Horde/HordeServer/Agents/Fleet/AwsAutoScalingLifecycle.cs b/Engine/Source/Programs/Horde/HordeServer/Agents/Fleet/AwsAutoScalingLifecycle.cs
--- a/Engine/Source/Programs/Horde/HordeServer/Agents/Fleet/AwsAutoScalingLifecycle.cs
+++ b/Engine/Source/Programs/Horde/HordeServer/Agents/Fleet/AwsAutoScalingLifecycle.cs
@@ -16,8 +16,8 @@ using Amazon.SQS;
 using Amazon.SQS.Model;
 using EpicGames.Core;
 using EpicGames.Horde.Agents;
-using Horde.Server.Server;
-using Horde.Server.Utilities;
+using HordeServer.Server;
+using HordeServer.Utilities;
 using Microsoft.AspNetCore.Authorization;
 using Microsoft.AspNetCore.Http;
 using Microsoft.AspNetCore.Mvc;
@@ -28,7 +28,7 @@ using Microsoft.Extensions.Options;
 using OpenTelemetry.Trace;
 using StackExchange.Redis;
 
-namespace Horde.Server.Agents.Fleet;
+namespace HordeServer.Agents.Fleet;
 
 /// <summary>
 /// Service handling callbacks and lifecycle events triggered by EC2 auto-scaling in AWS
diff --git a/Engine/Source/Programs/Horde/HordeServer/Agents/Fleet/ComputeQueueAwsMetricStrategy.cs b/Engine/Source/Programs/Horde/HordeServer/Agents/Fleet/ComputeQueueAwsMetricStrategy.cs
--- a/Engine/Source/Programs/Horde/HordeServer/Agents/Fleet/ComputeQueueAwsMetricStrategy.cs
+++ b/Engine/Source/Programs/Horde/HordeServer/Agents/Fleet/ComputeQueueAwsMetricStrategy.cs
@@ -11,13 +11,13 @@ using Amazon.CloudWatch;
 using Amazon.CloudWatch.Model;
 using EpicGames.Horde.Agents;
 using EpicGames.Horde.Compute;
-using Horde.Server.Agents.Pools;
-using Horde.Server.Compute;
-using Horde.Server.Utilities;
+using HordeServer.Agents.Pools;
+using HordeServer.Compute;
+using HordeServer.Utilities;
 using Microsoft.Extensions.Logging;
 using OpenTelemetry.Trace;
 
-namespace Horde.Server.Agents.Fleet
+namespace HordeServer.Agents.Fleet
 {
 	/// <summary>
 	/// Settings for <see cref="ComputeQueueAwsMetricStrategy" />
diff --git a/Engine/Source/Programs/Horde/HordeServer/Agents/Fleet/FleetManagerFactory.cs b/Engine/Source/Programs/Horde/HordeServer/Agents/Fleet/FleetManagerFactory.cs
--- a/Engine/Source/Programs/Horde/HordeServer/Agents/Fleet/FleetManagerFactory.cs
+++ b/Engine/Source/Programs/Horde/HordeServer/Agents/Fleet/FleetManagerFactory.cs
@@ -7,13 +7,13 @@ using System.Text.Json;
 using Amazon.AutoScaling;
 using Amazon.EC2;
 using EpicGames.Core;
-using Horde.Server.Agents.Fleet.Providers;
+using HordeServer.Agents.Fleet.Providers;
 using Microsoft.Extensions.DependencyInjection;
 using Microsoft.Extensions.Logging;
 using Microsoft.Extensions.Options;
 using OpenTelemetry.Trace;
 
-namespace Horde.Server.Agents.Fleet;
+namespace HordeServer.Agents.Fleet;
 
 /// <summary>
 /// Interface for fleet manager factory
diff --git a/Engine/Source/Programs/Horde/HordeServer/Agents/Fleet/FleetService.cs b/Engine/Source/Programs/Horde/HordeServer/Agents/Fleet/FleetService.cs
--- a/Engine/Source/Programs/Horde/HordeServer/Agents/Fleet/FleetService.cs
+++ b/Engine/Source/Programs/Horde/HordeServer/Agents/Fleet/FleetService.cs
@@ -10,15 +10,15 @@ using System.Threading.Tasks;
 using EpicGames.Core;
 using EpicGames.Horde.Agents;
 using EpicGames.Horde.Agents.Pools;
-using Horde.Server.Agents.Pools;
-using Horde.Server.Server;
-using Horde.Server.Utilities;
+using HordeServer.Agents.Pools;
+using HordeServer.Server;
+using HordeServer.Utilities;
 using Microsoft.Extensions.Hosting;
 using Microsoft.Extensions.Logging;
 using Microsoft.Extensions.Options;
 using OpenTelemetry.Trace;
 
-namespace Horde.Server.Agents.Fleet
+namespace HordeServer.Agents.Fleet
 {
 	/// <summary>
 	/// Parameters required for calculating pool size
diff --git a/Engine/Source/Programs/Horde/HordeServer/Agents/Fleet/IFleetManager.cs b/Engine/Source/Programs/Horde/HordeServer/Agents/Fleet/IFleetManager.cs
--- a/Engine/Source/Programs/Horde/HordeServer/Agents/Fleet/IFleetManager.cs
+++ b/Engine/Source/Programs/Horde/HordeServer/Agents/Fleet/IFleetManager.cs
@@ -4,9 +4,9 @@ using System;
 using System.Collections.Generic;
 using System.Threading;
 using System.Threading.Tasks;
-using Horde.Server.Agents.Pools;
+using HordeServer.Agents.Pools;
 
-namespace Horde.Server.Agents.Fleet
+namespace HordeServer.Agents.Fleet
 {
 	/// <summary>
 	/// Available fleet managers
diff --git a/Engine/Source/Programs/Horde/HordeServer/Agents/Fleet/IPoolSizeStrategy.cs b/Engine/Source/Programs/Horde/HordeServer/Agents/Fleet/IPoolSizeStrategy.cs
--- a/Engine/Source/Programs/Horde/HordeServer/Agents/Fleet/IPoolSizeStrategy.cs
+++ b/Engine/Source/Programs/Horde/HordeServer/Agents/Fleet/IPoolSizeStrategy.cs
@@ -5,9 +5,9 @@ using System.Collections.Generic;
 using System.Text.Json;
 using System.Threading;
 using System.Threading.Tasks;
-using Horde.Server.Agents.Pools;
+using HordeServer.Agents.Pools;
 
-namespace Horde.Server.Agents.Fleet
+namespace HordeServer.Agents.Fleet
 {
 	/// <summary>
 	/// Available pool sizing strategies
diff --git a/Engine/Source/Programs/Horde/HordeServer/Agents/Fleet/JobQueueStrategy.cs b/Engine/Source/Programs/Horde/HordeServer/Agents/Fleet/JobQueueStrategy.cs
--- a/Engine/Source/Programs/Horde/HordeServer/Agents/Fleet/JobQueueStrategy.cs
+++ b/Engine/Source/Programs/Horde/HordeServer/Agents/Fleet/JobQueueStrategy.cs
@@ -11,17 +11,17 @@ using EpicGames.Core;
 using EpicGames.Horde.Agents.Pools;
 using EpicGames.Horde.Jobs;
 using EpicGames.Horde.Streams;
-using Horde.Server.Agents.Pools;
-using Horde.Server.Jobs;
-using Horde.Server.Jobs.Graphs;
-using Horde.Server.Server;
-using Horde.Server.Streams;
-using Horde.Server.Utilities;
+using HordeServer.Agents.Pools;
+using HordeServer.Jobs;
+using HordeServer.Jobs.Graphs;
+using HordeServer.Server;
+using HordeServer.Streams;
+using HordeServer.Utilities;
 using Microsoft.Extensions.Caching.Memory;
 using Microsoft.Extensions.Options;
 using OpenTelemetry.Trace;
 
-namespace Horde.Server.Agents.Fleet
+namespace HordeServer.Agents.Fleet
 {
 	/// <summary>
 	/// Job queue sizing settings for a pool
diff --git a/Engine/Source/Programs/Horde/HordeServer/Agents/Fleet/LeaseUtilizationAwsMetricStrategy.cs b/Engine/Source/Programs/Horde/HordeServer/Agents/Fleet/LeaseUtilizationAwsMetricStrategy.cs
--- a/Engine/Source/Programs/Horde/HordeServer/Agents/Fleet/LeaseUtilizationAwsMetricStrategy.cs
+++ b/Engine/Source/Programs/Horde/HordeServer/Agents/Fleet/LeaseUtilizationAwsMetricStrategy.cs
@@ -10,13 +10,13 @@ using System.Threading.Tasks;
 using Amazon.CloudWatch;
 using Amazon.CloudWatch.Model;
 using EpicGames.Core;
-using Horde.Server.Agents.Leases;
-using Horde.Server.Agents.Pools;
-using Horde.Server.Utilities;
+using HordeServer.Agents.Leases;
+using HordeServer.Agents.Pools;
+using HordeServer.Utilities;
 using Microsoft.Extensions.Logging;
 using OpenTelemetry.Trace;
 
-namespace Horde.Server.Agents.Fleet;
+namespace HordeServer.Agents.Fleet;
 
 /// <summary>
 /// Settings for <see cref="LeaseUtilizationAwsMetricStrategy" />
diff --git a/Engine/Source/Programs/Horde/HordeServer/Agents/Fleet/LeaseUtilizationStrategy.cs b/Engine/Source/Programs/Horde/HordeServer/Agents/Fleet/LeaseUtilizationStrategy.cs
--- a/Engine/Source/Programs/Horde/HordeServer/Agents/Fleet/LeaseUtilizationStrategy.cs
+++ b/Engine/Source/Programs/Horde/HordeServer/Agents/Fleet/LeaseUtilizationStrategy.cs
@@ -10,14 +10,14 @@ using EpicGames.Core;
 using EpicGames.Horde.Agents;
 using EpicGames.Horde.Agents.Pools;
 using Google.Protobuf.WellKnownTypes;
-using Horde.Server.Agents.Leases;
-using Horde.Server.Agents.Pools;
-using Horde.Server.Utilities;
+using HordeServer.Agents.Leases;
+using HordeServer.Agents.Pools;
+using HordeServer.Utilities;
 using HordeCommon.Rpc.Tasks;
 using Microsoft.Extensions.Caching.Memory;
 using OpenTelemetry.Trace;
 
-namespace Horde.Server.Agents.Fleet
+namespace HordeServer.Agents.Fleet
 {
 	/// <summary>
 	/// Lease utilization sizing settings for a pool
diff --git a/Engine/Source/Programs/Horde/HordeServer/Agents/Fleet/Providers/AwsAsgFleetManager.cs b/Engine/Source/Programs/Horde/HordeServer/Agents/Fleet/Providers/AwsAsgFleetManager.cs
--- a/Engine/Source/Programs/Horde/HordeServer/Agents/Fleet/Providers/AwsAsgFleetManager.cs
+++ b/Engine/Source/Programs/Horde/HordeServer/Agents/Fleet/Providers/AwsAsgFleetManager.cs
@@ -9,11 +9,11 @@ using System.Threading.Tasks;
 using Amazon.AutoScaling;
 using Amazon.AutoScaling.Model;
 using EpicGames.Horde.Agents.Pools;
-using Horde.Server.Agents.Pools;
+using HordeServer.Agents.Pools;
 using Microsoft.Extensions.Logging;
 using OpenTelemetry.Trace;
 
-namespace Horde.Server.Agents.Fleet.Providers
+namespace HordeServer.Agents.Fleet.Providers
 {
 	/// <summary>
 	/// Settings for controlling an EC2 auto-scaling group (ASG)
diff --git a/Engine/Source/Programs/Horde/HordeServer/Agents/Fleet/Providers/AwsFleetManager.cs b/Engine/Source/Programs/Horde/HordeServer/Agents/Fleet/Providers/AwsFleetManager.cs
--- a/Engine/Source/Programs/Horde/HordeServer/Agents/Fleet/Providers/AwsFleetManager.cs
+++ b/Engine/Source/Programs/Horde/HordeServer/Agents/Fleet/Providers/AwsFleetManager.cs
@@ -8,13 +8,13 @@ using System.Threading.Tasks;
 using Amazon.EC2;
 using Amazon.EC2.Model;
 using EpicGames.Horde.Agents;
-using Horde.Server.Agents.Pools;
-using Horde.Server.Auditing;
-using Horde.Server.Utilities;
+using HordeServer.Agents.Pools;
+using HordeServer.Auditing;
+using HordeServer.Utilities;
 using Microsoft.Extensions.Logging;
 using OpenTelemetry.Trace;
 
-namespace Horde.Server.Agents.Fleet.Providers
+namespace HordeServer.Agents.Fleet.Providers
 {
 	/// <summary>
 	/// Settings for the AWS fleet manager
diff --git a/Engine/Source/Programs/Horde/HordeServer/Agents/Fleet/Providers/AwsRecyclingFleetManager.cs b/Engine/Source/Programs/Horde/HordeServer/Agents/Fleet/Providers/AwsRecyclingFleetManager.cs
--- a/Engine/Source/Programs/Horde/HordeServer/Agents/Fleet/Providers/AwsRecyclingFleetManager.cs
+++ b/Engine/Source/Programs/Horde/HordeServer/Agents/Fleet/Providers/AwsRecyclingFleetManager.cs
@@ -11,11 +11,11 @@ using Amazon.EC2;
 using Amazon.EC2.Model;
 using Amazon.Runtime;
 using EpicGames.Core;
-using Horde.Server.Agents.Pools;
+using HordeServer.Agents.Pools;
 using Microsoft.Extensions.Logging;
 using OpenTelemetry.Trace;
 
-namespace Horde.Server.Agents.Fleet.Providers;
+namespace HordeServer.Agents.Fleet.Providers;
 
 /// <summary>
 /// Settings for AWS fleet manager that reuses instances
diff --git a/Engine/Source/Programs/Horde/HordeServer/Agents/Fleet/Providers/AwsReuseFleetManager.cs b/Engine/Source/Programs/Horde/HordeServer/Agents/Fleet/Providers/AwsReuseFleetManager.cs
--- a/Engine/Source/Programs/Horde/HordeServer/Agents/Fleet/Providers/AwsReuseFleetManager.cs
+++ b/Engine/Source/Programs/Horde/HordeServer/Agents/Fleet/Providers/AwsReuseFleetManager.cs
@@ -7,11 +7,11 @@ using System.Threading;
 using System.Threading.Tasks;
 using Amazon.EC2;
 using Amazon.EC2.Model;
-using Horde.Server.Agents.Pools;
+using HordeServer.Agents.Pools;
 using Microsoft.Extensions.Logging;
 using OpenTelemetry.Trace;
 
-namespace Horde.Server.Agents.Fleet.Providers
+namespace HordeServer.Agents.Fleet.Providers
 {
 	/// <summary>
 	/// Settings for AWS fleet manager that reuses instances
diff --git a/Engine/Source/Programs/Horde/HordeServer/Agents/IAgent.cs b/Engine/Source/Programs/Horde/HordeServer/Agents/IAgent.cs
--- a/Engine/Source/Programs/Horde/HordeServer/Agents/IAgent.cs
+++ b/Engine/Source/Programs/Horde/HordeServer/Agents/IAgent.cs
@@ -18,14 +18,14 @@ using EpicGames.Horde.Logs;
 using EpicGames.Horde.Streams;
 using EpicGames.Horde.Tools;
 using Google.Protobuf.WellKnownTypes;
-using Horde.Server.Server;
+using HordeServer.Server;
 using HordeCommon.Rpc;
 using HordeCommon.Rpc.Messages;
 using HordeCommon.Rpc.Tasks;
 using MongoDB.Bson;
 using MongoDB.Bson.Serialization.Attributes;
 
-namespace Horde.Server.Agents
+namespace HordeServer.Agents
 {
 	/// <summary>
 	/// Mirrors an Agent document in the database
diff --git a/Engine/Source/Programs/Horde/HordeServer/Agents/IAgentCollection.cs b/Engine/Source/Programs/Horde/HordeServer/Agents/IAgentCollection.cs
--- a/Engine/Source/Programs/Horde/HordeServer/Agents/IAgentCollection.cs
+++ b/Engine/Source/Programs/Horde/HordeServer/Agents/IAgentCollection.cs
@@ -7,9 +7,9 @@ using System.Threading.Tasks;
 using EpicGames.Horde.Agents;
 using EpicGames.Horde.Agents.Leases;
 using EpicGames.Horde.Agents.Pools;
-using Horde.Server.Auditing;
+using HordeServer.Auditing;
 
-namespace Horde.Server.Agents
+namespace HordeServer.Agents
 {
 	/// <summary>
 	/// Interface for a collection of agent documents
diff --git a/Engine/Source/Programs/Horde/HordeServer/Agents/Leases/ILease.cs b/Engine/Source/Programs/Horde/HordeServer/Agents/Leases/ILease.cs
--- a/Engine/Source/Programs/Horde/HordeServer/Agents/Leases/ILease.cs
+++ b/Engine/Source/Programs/Horde/HordeServer/Agents/Leases/ILease.cs
@@ -10,7 +10,7 @@ using EpicGames.Horde.Streams;
 using Google.Protobuf;
 using Google.Protobuf.WellKnownTypes;
 
-namespace Horde.Server.Agents.Leases
+namespace HordeServer.Agents.Leases
 {
 	/// <summary>
 	/// Document describing a lease. This exists to permanently record a lease; the agent object tracks internal state of any active leases through AgentLease objects.
diff --git a/Engine/Source/Programs/Horde/HordeServer/Agents/Leases/ILeaseCollection.cs b/Engine/Source/Programs/Horde/HordeServer/Agents/Leases/ILeaseCollection.cs
--- a/Engine/Source/Programs/Horde/HordeServer/Agents/Leases/ILeaseCollection.cs
+++ b/Engine/Source/Programs/Horde/HordeServer/Agents/Leases/ILeaseCollection.cs
@@ -11,7 +11,7 @@ using EpicGames.Horde.Agents.Sessions;
 using EpicGames.Horde.Logs;
 using EpicGames.Horde.Streams;
 
-namespace Horde.Server.Agents.Leases
+namespace HordeServer.Agents.Leases
 {
 	/// <summary>
 	/// Interface for a collection of lease documents
diff --git a/Engine/Source/Programs/Horde/HordeServer/Agents/Leases/LeaseCollection.cs b/Engine/Source/Programs/Horde/HordeServer/Agents/Leases/LeaseCollection.cs
--- a/Engine/Source/Programs/Horde/HordeServer/Agents/Leases/LeaseCollection.cs
+++ b/Engine/Source/Programs/Horde/HordeServer/Agents/Leases/LeaseCollection.cs
@@ -11,13 +11,13 @@ using EpicGames.Horde.Agents.Pools;
 using EpicGames.Horde.Agents.Sessions;
 using EpicGames.Horde.Logs;
 using EpicGames.Horde.Streams;
-using Horde.Server.Server;
-using Horde.Server.Utilities;
+using HordeServer.Server;
+using HordeServer.Utilities;
 using MongoDB.Bson.Serialization.Attributes;
 using MongoDB.Driver;
 using OpenTelemetry.Trace;
 
-namespace Horde.Server.Agents.Leases
+namespace HordeServer.Agents.Leases
 {
 	/// <summary>
 	/// Collection of lease documents
diff --git a/Engine/Source/Programs/Horde/HordeServer/Agents/Leases/LeasesController.cs b/Engine/Source/Programs/Horde/HordeServer/Agents/Leases/LeasesController.cs
--- a/Engine/Source/Programs/Horde/HordeServer/Agents/Leases/LeasesController.cs
+++ b/Engine/Source/Programs/Horde/HordeServer/Agents/Leases/LeasesController.cs
@@ -9,15 +9,15 @@ using EpicGames.Horde.Agents;
 using EpicGames.Horde.Agents.Leases;
 using EpicGames.Horde.Agents.Sessions;
 using Google.Protobuf.WellKnownTypes;
-using Horde.Server.Server;
-using Horde.Server.Tasks;
-using Horde.Server.Utilities;
+using HordeServer.Server;
+using HordeServer.Tasks;
+using HordeServer.Utilities;
 using Microsoft.AspNetCore.Authorization;
 using Microsoft.AspNetCore.Mvc;
 using Microsoft.Extensions.Options;
 using OpenTelemetry.Trace;
 
-namespace Horde.Server.Agents.Leases
+namespace HordeServer.Agents.Leases
 {
 	/// <summary>
 	/// Controller for the /api/v1/leases endpoint
diff --git a/Engine/Source/Programs/Horde/HordeServer/Agents/Pools/IPool.cs b/Engine/Source/Programs/Horde/HordeServer/Agents/Pools/IPool.cs
--- a/Engine/Source/Programs/Horde/HordeServer/Agents/Pools/IPool.cs
+++ b/Engine/Source/Programs/Horde/HordeServer/Agents/Pools/IPool.cs
@@ -6,9 +6,9 @@ using System.Text.Json.Serialization;
 using System.Threading;
 using System.Threading.Tasks;
 using EpicGames.Horde.Common;
-using Horde.Server.Agents.Fleet;
+using HordeServer.Agents.Fleet;
 
-namespace Horde.Server.Agents.Pools
+namespace HordeServer.Agents.Pools
 {
 	/// <summary>
 	/// Metadata for configuring and picking a pool sizing strategy
diff --git a/Engine/Source/Programs/Horde/HordeServer/Agents/Pools/IPoolCollection.cs b/Engine/Source/Programs/Horde/HordeServer/Agents/Pools/IPoolCollection.cs
--- a/Engine/Source/Programs/Horde/HordeServer/Agents/Pools/IPoolCollection.cs
+++ b/Engine/Source/Programs/Horde/HordeServer/Agents/Pools/IPoolCollection.cs
@@ -6,9 +6,9 @@ using System.Threading;
 using System.Threading.Tasks;
 using EpicGames.Horde.Agents.Pools;
 using EpicGames.Horde.Common;
-using Horde.Server.Agents.Fleet;
+using HordeServer.Agents.Fleet;
 
-namespace Horde.Server.Agents.Pools
+namespace HordeServer.Agents.Pools
 {
 	/// <summary>
 	/// Collection of pool documents
diff --git a/Engine/Source/Programs/Horde/HordeServer/Agents/Pools/PoolAclAction.cs b/Engine/Source/Programs/Horde/HordeServer/Agents/Pools/PoolAclAction.cs
--- a/Engine/Source/Programs/Horde/HordeServer/Agents/Pools/PoolAclAction.cs
+++ b/Engine/Source/Programs/Horde/HordeServer/Agents/Pools/PoolAclAction.cs
@@ -2,7 +2,7 @@
 
 using EpicGames.Horde.Acls;
 
-namespace Horde.Server.Agents.Pools
+namespace HordeServer.Agents.Pools
 {
 	/// <summary>
 	/// ACL actions for pools
diff --git a/Engine/Source/Programs/Horde/HordeServer/Agents/Pools/PoolCollection.cs b/Engine/Source/Programs/Horde/HordeServer/Agents/Pools/PoolCollection.cs
--- a/Engine/Source/Programs/Horde/HordeServer/Agents/Pools/PoolCollection.cs
+++ b/Engine/Source/Programs/Horde/HordeServer/Agents/Pools/PoolCollection.cs
@@ -7,15 +7,15 @@ using System.Threading;
 using System.Threading.Tasks;
 using EpicGames.Horde.Agents.Pools;
 using EpicGames.Horde.Common;
-using Horde.Server.Agents.Fleet;
-using Horde.Server.Server;
-using Horde.Server.Utilities;
+using HordeServer.Agents.Fleet;
+using HordeServer.Server;
+using HordeServer.Utilities;
 using Microsoft.Extensions.Options;
 using MongoDB.Bson;
 using MongoDB.Bson.Serialization.Attributes;
 using MongoDB.Driver;
 
-namespace Horde.Server.Agents.Pools
+namespace HordeServer.Agents.Pools
 {
 	/// <summary>
 	/// Collection of pool documents
diff --git a/Engine/Source/Programs/Horde/HordeServer/Agents/Pools/PoolCondition.cs b/Engine/Source/Programs/Horde/HordeServer/Agents/Pools/PoolCondition.cs
--- a/Engine/Source/Programs/Horde/HordeServer/Agents/Pools/PoolCondition.cs
+++ b/Engine/Source/Programs/Horde/HordeServer/Agents/Pools/PoolCondition.cs
@@ -5,7 +5,7 @@ using System.Collections.Generic;
 using System.Globalization;
 using EpicGames.Horde.Common;
 
-namespace Horde.Server.Agents.Pools
+namespace HordeServer.Agents.Pools
 {
 	/// <summary>
 	/// Defined pool property names, for use in dashboard filter expressions
diff --git a/Engine/Source/Programs/Horde/HordeServer/Agents/Pools/PoolUpdateService.cs b/Engine/Source/Programs/Horde/HordeServer/Agents/Pools/PoolUpdateService.cs
--- a/Engine/Source/Programs/Horde/HordeServer/Agents/Pools/PoolUpdateService.cs
+++ b/Engine/Source/Programs/Horde/HordeServer/Agents/Pools/PoolUpdateService.cs
@@ -8,15 +8,15 @@ using System.Threading.Tasks;
 using EpicGames.Core;
 using EpicGames.Horde.Agents;
 using EpicGames.Horde.Agents.Pools;
-using Horde.Server.Perforce;
-using Horde.Server.Server;
-using Horde.Server.Streams;
+using HordeServer.Perforce;
+using HordeServer.Server;
+using HordeServer.Streams;
 using Microsoft.Extensions.Hosting;
 using Microsoft.Extensions.Logging;
 using Microsoft.Extensions.Options;
 using OpenTelemetry.Trace;
 
-namespace Horde.Server.Agents.Pools
+namespace HordeServer.Agents.Pools
 {
 	/// <summary>
 	/// Periodically updates pool documents to contain the correct workspaces
diff --git a/Engine/Source/Programs/Horde/HordeServer/Agents/Pools/PoolsController.cs b/Engine/Source/Programs/Horde/HordeServer/Agents/Pools/PoolsController.cs
--- a/Engine/Source/Programs/Horde/HordeServer/Agents/Pools/PoolsController.cs
+++ b/Engine/Source/Programs/Horde/HordeServer/Agents/Pools/PoolsController.cs
@@ -9,15 +9,15 @@ using EpicGames.Core;
 using EpicGames.Horde;
 using EpicGames.Horde.Agents.Pools;
 using EpicGames.Horde.Common;
-using Horde.Server.Agents.Fleet;
-using Horde.Server.Agents.Utilization;
-using Horde.Server.Server;
-using Horde.Server.Utilities;
+using HordeServer.Agents.Fleet;
+using HordeServer.Agents.Utilization;
+using HordeServer.Server;
+using HordeServer.Utilities;
 using Microsoft.AspNetCore.Authorization;
 using Microsoft.AspNetCore.Mvc;
 using Microsoft.Extensions.Options;
 
-namespace Horde.Server.Agents.Pools
+namespace HordeServer.Agents.Pools
 {
 	/// <summary>
 	/// Controller for the /api/v1/pools endpoint
diff --git a/Engine/Source/Programs/Horde/HordeServer/Agents/Relay/AgentRelayClient.cs b/Engine/Source/Programs/Horde/HordeServer/Agents/Relay/AgentRelayClient.cs
--- a/Engine/Source/Programs/Horde/HordeServer/Agents/Relay/AgentRelayClient.cs
+++ b/Engine/Source/Programs/Horde/HordeServer/Agents/Relay/AgentRelayClient.cs
@@ -9,7 +9,7 @@ using Grpc.Core;
 using Horde.Common.Rpc;
 using Microsoft.Extensions.Logging;
 
-namespace Horde.Server.Agents.Relay;
+namespace HordeServer.Agents.Relay;
 
 /// <summary>
 /// Relay service mirroring port mappings sent from server
diff --git a/Engine/Source/Programs/Horde/HordeServer/Agents/Relay/AgentRelayService.cs b/Engine/Source/Programs/Horde/HordeServer/Agents/Relay/AgentRelayService.cs
--- a/Engine/Source/Programs/Horde/HordeServer/Agents/Relay/AgentRelayService.cs
+++ b/Engine/Source/Programs/Horde/HordeServer/Agents/Relay/AgentRelayService.cs
@@ -16,13 +16,13 @@ using Google.Protobuf;
 using Google.Protobuf.WellKnownTypes;
 using Grpc.Core;
 using Horde.Common.Rpc;
-using Horde.Server.Server;
-using Horde.Server.Utilities;
+using HordeServer.Server;
+using HordeServer.Utilities;
 using Microsoft.Extensions.Hosting;
 using Microsoft.Extensions.Logging;
 using StackExchange.Redis;
 
-namespace Horde.Server.Agents.Relay;
+namespace HordeServer.Agents.Relay;
 
 /// <summary>
 /// Exception thrown by agent relay related code
diff --git a/Engine/Source/Programs/Horde/HordeServer/Agents/Relay/Nftables.cs b/Engine/Source/Programs/Horde/HordeServer/Agents/Relay/Nftables.cs
--- a/Engine/Source/Programs/Horde/HordeServer/Agents/Relay/Nftables.cs
+++ b/Engine/Source/Programs/Horde/HordeServer/Agents/Relay/Nftables.cs
@@ -12,7 +12,7 @@ using Horde.Common.Rpc;
 using Microsoft.Extensions.Logging;
 using Microsoft.Extensions.Logging.Abstractions;
 
-namespace Horde.Server.Agents.Relay;
+namespace HordeServer.Agents.Relay;
 
 /// <summary>
 /// nftables specific exception
diff --git a/Engine/Source/Programs/Horde/HordeServer/Agents/Sessions/ISession.cs b/Engine/Source/Programs/Horde/HordeServer/Agents/Sessions/ISession.cs
--- a/Engine/Source/Programs/Horde/HordeServer/Agents/Sessions/ISession.cs
+++ b/Engine/Source/Programs/Horde/HordeServer/Agents/Sessions/ISession.cs
@@ -5,7 +5,7 @@ using System.Collections.Generic;
 using EpicGames.Horde.Agents;
 using EpicGames.Horde.Agents.Sessions;
 
-namespace Horde.Server.Agents.Sessions
+namespace HordeServer.Agents.Sessions
 {
 	/// <summary>
 	/// Information about an agent session.
diff --git a/Engine/Source/Programs/Horde/HordeServer/Agents/Sessions/ISessionCollection.cs b/Engine/Source/Programs/Horde/HordeServer/Agents/Sessions/ISessionCollection.cs
--- a/Engine/Source/Programs/Horde/HordeServer/Agents/Sessions/ISessionCollection.cs
+++ b/Engine/Source/Programs/Horde/HordeServer/Agents/Sessions/ISessionCollection.cs
@@ -7,7 +7,7 @@ using System.Threading.Tasks;
 using EpicGames.Horde.Agents;
 using EpicGames.Horde.Agents.Sessions;
 
-namespace Horde.Server.Agents.Sessions
+namespace HordeServer.Agents.Sessions
 {
 	/// <summary>
 	/// Interface for a collection of session documents
diff --git a/Engine/Source/Programs/Horde/HordeServer/Agents/Sessions/SessionCollection.cs b/Engine/Source/Programs/Horde/HordeServer/Agents/Sessions/SessionCollection.cs
--- a/Engine/Source/Programs/Horde/HordeServer/Agents/Sessions/SessionCollection.cs
+++ b/Engine/Source/Programs/Horde/HordeServer/Agents/Sessions/SessionCollection.cs
@@ -7,12 +7,12 @@ using System.Threading;
 using System.Threading.Tasks;
 using EpicGames.Horde.Agents;
 using EpicGames.Horde.Agents.Sessions;
-using Horde.Server.Server;
-using Horde.Server.Utilities;
+using HordeServer.Server;
+using HordeServer.Utilities;
 using MongoDB.Bson.Serialization.Attributes;
 using MongoDB.Driver;
 
-namespace Horde.Server.Agents.Sessions
+namespace HordeServer.Agents.Sessions
 {
 	/// <summary>
 	/// Collection of session documents
diff --git a/Engine/Source/Programs/Horde/HordeServer/Agents/Sessions/SessionIdUtils.cs b/Engine/Source/Programs/Horde/HordeServer/Agents/Sessions/SessionIdUtils.cs
--- a/Engine/Source/Programs/Horde/HordeServer/Agents/Sessions/SessionIdUtils.cs
+++ b/Engine/Source/Programs/Horde/HordeServer/Agents/Sessions/SessionIdUtils.cs
@@ -1,9 +1,9 @@
 // Copyright Epic Games, Inc. All Rights Reserved.
 
 using EpicGames.Horde.Agents.Sessions;
-using Horde.Server.Utilities;
+using HordeServer.Utilities;
 
-namespace Horde.Server.Agents.Sessions
+namespace HordeServer.Agents.Sessions
 {
 	/// <summary>
 	/// Server helper methods for <see cref="SessionId"/>
diff --git a/Engine/Source/Programs/Horde/HordeServer/Agents/Software/AgentSoftwareAclAction.cs b/Engine/Source/Programs/Horde/HordeServer/Agents/Software/AgentSoftwareAclAction.cs
--- a/Engine/Source/Programs/Horde/HordeServer/Agents/Software/AgentSoftwareAclAction.cs
+++ b/Engine/Source/Programs/Horde/HordeServer/Agents/Software/AgentSoftwareAclAction.cs
@@ -2,7 +2,7 @@
 
 using EpicGames.Horde.Acls;
 
-namespace Horde.Server.Agents.Software
+namespace HordeServer.Agents.Software
 {
 	/// <summary>
 	/// ACL actions for agent software
diff --git a/Engine/Source/Programs/Horde/HordeServer/Agents/Software/AgentSoftwareController.cs b/Engine/Source/Programs/Horde/HordeServer/Agents/Software/AgentSoftwareController.cs
--- a/Engine/Source/Programs/Horde/HordeServer/Agents/Software/AgentSoftwareController.cs
+++ b/Engine/Source/Programs/Horde/HordeServer/Agents/Software/AgentSoftwareController.cs
@@ -5,14 +5,14 @@ using System.IO;
 using System.Threading;
 using System.Threading.Tasks;
 using EpicGames.Core;
-using Horde.Server.Server;
-using Horde.Server.Tools;
+using HordeServer.Server;
+using HordeServer.Tools;
 using Microsoft.AspNetCore.Authorization;
 using Microsoft.AspNetCore.Mvc;
 using Microsoft.Extensions.Options;
 using Microsoft.Net.Http.Headers;
 
-namespace Horde.Server.Agents.Software
+namespace HordeServer.Agents.Software
 {
 	/// <summary>
 	/// Information about an agent software channel
diff --git a/Engine/Source/Programs/Horde/HordeServer/Agents/Telemetry/AgentTelemetryCollection.cs b/Engine/Source/Programs/Horde/HordeServer/Agents/Telemetry/AgentTelemetryCollection.cs
--- a/Engine/Source/Programs/Horde/HordeServer/Agents/Telemetry/AgentTelemetryCollection.cs
+++ b/Engine/Source/Programs/Horde/HordeServer/Agents/Telemetry/AgentTelemetryCollection.cs
@@ -7,15 +7,15 @@ using System.Threading;
 using System.Threading.Tasks;
 using EpicGames.Core;
 using EpicGames.Horde.Agents;
-using Horde.Server.Server;
-using Horde.Server.Utilities;
+using HordeServer.Server;
+using HordeServer.Utilities;
 using Microsoft.Extensions.Hosting;
 using Microsoft.Extensions.Logging;
 using MongoDB.Bson;
 using MongoDB.Bson.Serialization.Attributes;
 using MongoDB.Driver;
 
-namespace Horde.Server.Agents.Telemetry
+namespace HordeServer.Agents.Telemetry
 {
 	/// <summary>
 	/// Collection which writes agent data to a MongoDB collection
diff --git a/Engine/Source/Programs/Horde/HordeServer/Agents/Telemetry/IAgentTelemetry.cs b/Engine/Source/Programs/Horde/HordeServer/Agents/Telemetry/IAgentTelemetry.cs
--- a/Engine/Source/Programs/Horde/HordeServer/Agents/Telemetry/IAgentTelemetry.cs
+++ b/Engine/Source/Programs/Horde/HordeServer/Agents/Telemetry/IAgentTelemetry.cs
@@ -2,7 +2,7 @@
 
 using System;
 
-namespace Horde.Server.Agents.Telemetry
+namespace HordeServer.Agents.Telemetry
 {
 	/// <summary>
 	/// Telemetry for an agent
diff --git a/Engine/Source/Programs/Horde/HordeServer/Agents/Telemetry/IAgentTelemetryCollection.cs b/Engine/Source/Programs/Horde/HordeServer/Agents/Telemetry/IAgentTelemetryCollection.cs
--- a/Engine/Source/Programs/Horde/HordeServer/Agents/Telemetry/IAgentTelemetryCollection.cs
+++ b/Engine/Source/Programs/Horde/HordeServer/Agents/Telemetry/IAgentTelemetryCollection.cs
@@ -6,7 +6,7 @@ using System.Threading;
 using System.Threading.Tasks;
 using EpicGames.Horde.Agents;
 
-namespace Horde.Server.Agents.Telemetry
+namespace HordeServer.Agents.Telemetry
 {
 	/// <summary>
 	/// Interface for a collection of agent telemetry documents
diff --git a/Engine/Source/Programs/Horde/HordeServer/Agents/Utilization/IUtilizationData.cs b/Engine/Source/Programs/Horde/HordeServer/Agents/Utilization/IUtilizationData.cs
--- a/Engine/Source/Programs/Horde/HordeServer/Agents/Utilization/IUtilizationData.cs
+++ b/Engine/Source/Programs/Horde/HordeServer/Agents/Utilization/IUtilizationData.cs
@@ -5,7 +5,7 @@ using System.Collections.Generic;
 using EpicGames.Horde.Agents.Pools;
 using EpicGames.Horde.Streams;
 
-namespace Horde.Server.Agents.Utilization
+namespace HordeServer.Agents.Utilization
 {
 	/// <summary>
 	/// Information about the utilization of a pool
diff --git a/Engine/Source/Programs/Horde/HordeServer/Agents/Utilization/IUtilizationDataCollection.cs b/Engine/Source/Programs/Horde/HordeServer/Agents/Utilization/IUtilizationDataCollection.cs
--- a/Engine/Source/Programs/Horde/HordeServer/Agents/Utilization/IUtilizationDataCollection.cs
+++ b/Engine/Source/Programs/Horde/HordeServer/Agents/Utilization/IUtilizationDataCollection.cs
@@ -5,7 +5,7 @@ using System.Collections.Generic;
 using System.Threading;
 using System.Threading.Tasks;
 
-namespace Horde.Server.Agents.Utilization
+namespace HordeServer.Agents.Utilization
 {
 	/// <summary>
 	/// Collection of utilization collection
diff --git a/Engine/Source/Programs/Horde/HordeServer/Agents/Utilization/UtilizationDataCollection.cs b/Engine/Source/Programs/Horde/HordeServer/Agents/Utilization/UtilizationDataCollection.cs
--- a/Engine/Source/Programs/Horde/HordeServer/Agents/Utilization/UtilizationDataCollection.cs
+++ b/Engine/Source/Programs/Horde/HordeServer/Agents/Utilization/UtilizationDataCollection.cs
@@ -8,11 +8,11 @@ using System.Threading.Tasks;
 using EpicGames.Core;
 using EpicGames.Horde.Agents.Pools;
 using EpicGames.Horde.Streams;
-using Horde.Server.Server;
+using HordeServer.Server;
 using MongoDB.Bson.Serialization.Attributes;
 using MongoDB.Driver;
 
-namespace Horde.Server.Agents.Utilization
+namespace HordeServer.Agents.Utilization
 {
 	/// <summary>
 	/// Collection of utilization data
diff --git a/Engine/Source/Programs/Horde/HordeServer/Agents/Utilization/UtilizationDataService.cs b/Engine/Source/Programs/Horde/HordeServer/Agents/Utilization/UtilizationDataService.cs
--- a/Engine/Source/Programs/Horde/HordeServer/Agents/Utilization/UtilizationDataService.cs
+++ b/Engine/Source/Programs/Horde/HordeServer/Agents/Utilization/UtilizationDataService.cs
@@ -8,13 +8,13 @@ using System.Threading.Tasks;
 using EpicGames.Core;
 using EpicGames.Horde.Agents;
 using EpicGames.Horde.Agents.Pools;
-using Horde.Server.Agents.Fleet;
-using Horde.Server.Agents.Leases;
-using Horde.Server.Agents.Pools;
+using HordeServer.Agents.Fleet;
+using HordeServer.Agents.Leases;
+using HordeServer.Agents.Pools;
 using Microsoft.Extensions.Hosting;
 using Microsoft.Extensions.Logging;
 
-namespace Horde.Server.Agents.Utilization
+namespace HordeServer.Agents.Utilization
 {
 	/// <summary>
 	/// Service which updates telemetry periodically
diff --git a/Engine/Source/Programs/Horde/HordeServer/Artifacts/ArtifactAclAction.cs b/Engine/Source/Programs/Horde/HordeServer/Artifacts/ArtifactAclAction.cs
--- a/Engine/Source/Programs/Horde/HordeServer/Artifacts/ArtifactAclAction.cs
+++ b/Engine/Source/Programs/Horde/HordeServer/Artifacts/ArtifactAclAction.cs
@@ -2,7 +2,7 @@
 
 using EpicGames.Horde.Acls;
 
-namespace Horde.Server.Artifacts
+namespace HordeServer.Artifacts
 {
 	/// <summary>
 	/// ACL actions for artifacts
diff --git a/Engine/Source/Programs/Horde/HordeServer/Artifacts/ArtifactCollection.cs b/Engine/Source/Programs/Horde/HordeServer/Artifacts/ArtifactCollection.cs
--- a/Engine/Source/Programs/Horde/HordeServer/Artifacts/ArtifactCollection.cs
+++ b/Engine/Source/Programs/Horde/HordeServer/Artifacts/ArtifactCollection.cs
@@ -10,15 +10,15 @@ using EpicGames.Core;
 using EpicGames.Horde.Artifacts;
 using EpicGames.Horde.Storage;
 using EpicGames.Horde.Streams;
-using Horde.Server.Acls;
-using Horde.Server.Server;
-using Horde.Server.Storage;
-using Horde.Server.Utilities;
+using HordeServer.Acls;
+using HordeServer.Server;
+using HordeServer.Storage;
+using HordeServer.Utilities;
 using MongoDB.Bson;
 using MongoDB.Bson.Serialization.Attributes;
 using MongoDB.Driver;
 
-namespace Horde.Server.Artifacts
+namespace HordeServer.Artifacts
 {
 	/// <summary>
 	/// Implementation of <see cref="IArtifactCollection"/>
diff --git a/Engine/Source/Programs/Horde/HordeServer/Artifacts/ArtifactExpirationService.cs b/Engine/Source/Programs/Horde/HordeServer/Artifacts/ArtifactExpirationService.cs
--- a/Engine/Source/Programs/Horde/HordeServer/Artifacts/ArtifactExpirationService.cs
+++ b/Engine/Source/Programs/Horde/HordeServer/Artifacts/ArtifactExpirationService.cs
@@ -9,13 +9,13 @@ using System.Threading.Tasks;
 using EpicGames.Core;
 using EpicGames.Horde.Storage;
 using EpicGames.Horde.Streams;
-using Horde.Server.Server;
-using Horde.Server.Storage;
+using HordeServer.Server;
+using HordeServer.Storage;
 using Microsoft.Extensions.Hosting;
 using Microsoft.Extensions.Logging;
 using Microsoft.Extensions.Options;
 
-namespace Horde.Server.Artifacts
+namespace HordeServer.Artifacts
 {
 	/// <summary>
 	/// Expires artifacts according to their 
diff --git a/Engine/Source/Programs/Horde/HordeServer/Artifacts/ArtifactsController.cs b/Engine/Source/Programs/Horde/HordeServer/Artifacts/ArtifactsController.cs
--- a/Engine/Source/Programs/Horde/HordeServer/Artifacts/ArtifactsController.cs
+++ b/Engine/Source/Programs/Horde/HordeServer/Artifacts/ArtifactsController.cs
@@ -21,13 +21,13 @@ using EpicGames.Horde.Storage.Bundles;
 using EpicGames.Horde.Storage.Nodes;
 using EpicGames.Horde.Streams;
 using Google.Protobuf.WellKnownTypes;
-using Horde.Server.Acls;
-using Horde.Server.Agents.Leases;
-using Horde.Server.Jobs;
-using Horde.Server.Server;
-using Horde.Server.Storage;
-using Horde.Server.Streams;
-using Horde.Server.Utilities;
+using HordeServer.Acls;
+using HordeServer.Agents.Leases;
+using HordeServer.Jobs;
+using HordeServer.Server;
+using HordeServer.Storage;
+using HordeServer.Streams;
+using HordeServer.Utilities;
 using HordeCommon.Rpc.Tasks;
 using Microsoft.AspNetCore.Authorization;
 using Microsoft.AspNetCore.Http;
@@ -37,7 +37,7 @@ using Microsoft.AspNetCore.StaticFiles;
 using Microsoft.Extensions.Logging;
 using Microsoft.Extensions.Options;
 
-namespace Horde.Server.Artifacts
+namespace HordeServer.Artifacts
 {
 	/// <summary>
 	/// Public interface for artifacts
diff --git a/Engine/Source/Programs/Horde/HordeServer/Artifacts/BlockCache.cs b/Engine/Source/Programs/Horde/HordeServer/Artifacts/BlockCache.cs
--- a/Engine/Source/Programs/Horde/HordeServer/Artifacts/BlockCache.cs
+++ b/Engine/Source/Programs/Horde/HordeServer/Artifacts/BlockCache.cs
@@ -16,7 +16,7 @@ using EpicGames.Core;
 using Microsoft.Extensions.Logging;
 using Microsoft.Extensions.Logging.Abstractions;
 
-namespace Horde.Server.Artifacts
+namespace HordeServer.Artifacts
 {
 	/// <summary>
 	/// Stats for the cache
diff --git a/Engine/Source/Programs/Horde/HordeServer/Artifacts/IArtifact.cs b/Engine/Source/Programs/Horde/HordeServer/Artifacts/IArtifact.cs
--- a/Engine/Source/Programs/Horde/HordeServer/Artifacts/IArtifact.cs
+++ b/Engine/Source/Programs/Horde/HordeServer/Artifacts/IArtifact.cs
@@ -5,9 +5,9 @@ using System.Collections.Generic;
 using EpicGames.Horde.Artifacts;
 using EpicGames.Horde.Storage;
 using EpicGames.Horde.Streams;
-using Horde.Server.Acls;
+using HordeServer.Acls;
 
-namespace Horde.Server.Artifacts
+namespace HordeServer.Artifacts
 {
 	/// <summary>
 	/// Information about an artifact
diff --git a/Engine/Source/Programs/Horde/HordeServer/Artifacts/IArtifactCollection.cs b/Engine/Source/Programs/Horde/HordeServer/Artifacts/IArtifactCollection.cs
--- a/Engine/Source/Programs/Horde/HordeServer/Artifacts/IArtifactCollection.cs
+++ b/Engine/Source/Programs/Horde/HordeServer/Artifacts/IArtifactCollection.cs
@@ -6,9 +6,9 @@ using System.Threading;
 using System.Threading.Tasks;
 using EpicGames.Horde.Artifacts;
 using EpicGames.Horde.Streams;
-using Horde.Server.Acls;
+using HordeServer.Acls;
 
-namespace Horde.Server.Artifacts
+namespace HordeServer.Artifacts
 {
 	/// <summary>
 	/// Interface for a collection of artifacts
diff --git a/Engine/Source/Programs/Horde/HordeServer/Artifacts/UnsyncCache.cs b/Engine/Source/Programs/Horde/HordeServer/Artifacts/UnsyncCache.cs
--- a/Engine/Source/Programs/Horde/HordeServer/Artifacts/UnsyncCache.cs
+++ b/Engine/Source/Programs/Horde/HordeServer/Artifacts/UnsyncCache.cs
@@ -11,11 +11,11 @@ using System.Threading.Tasks;
 using EpicGames.Core;
 using EpicGames.Horde.Storage;
 using EpicGames.Horde.Storage.Nodes;
-using Horde.Server.Storage;
+using HordeServer.Storage;
 using Microsoft.Extensions.Caching.Memory;
 using Microsoft.Extensions.Logging;
 
-namespace Horde.Server.Artifacts
+namespace HordeServer.Artifacts
 {
 	/// <summary>
 	/// Manifest for Unsync
diff --git a/Engine/Source/Programs/Horde/HordeServer/Auditing/AuditLog.cs b/Engine/Source/Programs/Horde/HordeServer/Auditing/AuditLog.cs
--- a/Engine/Source/Programs/Horde/HordeServer/Auditing/AuditLog.cs
+++ b/Engine/Source/Programs/Horde/HordeServer/Auditing/AuditLog.cs
@@ -8,14 +8,14 @@ using System.Threading;
 using System.Threading.Channels;
 using System.Threading.Tasks;
 using EpicGames.Core;
-using Horde.Server.Server;
-using Horde.Server.Utilities;
+using HordeServer.Server;
+using HordeServer.Utilities;
 using Microsoft.Extensions.Logging;
 using MongoDB.Bson;
 using MongoDB.Bson.Serialization.Attributes;
 using MongoDB.Driver;
 
-namespace Horde.Server.Auditing
+namespace HordeServer.Auditing
 {
 	class AuditLog<TSubject> : IAuditLog<TSubject>, IAsyncDisposable
 	{
diff --git a/Engine/Source/Programs/Horde/HordeServer/Auditing/IAuditLog.cs b/Engine/Source/Programs/Horde/HordeServer/Auditing/IAuditLog.cs
--- a/Engine/Source/Programs/Horde/HordeServer/Auditing/IAuditLog.cs
+++ b/Engine/Source/Programs/Horde/HordeServer/Auditing/IAuditLog.cs
@@ -8,7 +8,7 @@ using System.Threading;
 using System.Threading.Tasks;
 using Microsoft.Extensions.Logging;
 
-namespace Horde.Server.Auditing
+namespace HordeServer.Auditing
 {
 	/// <summary>
 	/// Message from an audit log
diff --git a/Engine/Source/Programs/Horde/HordeServer/Authentication/AnonymousAuthHandler.cs b/Engine/Source/Programs/Horde/HordeServer/Authentication/AnonymousAuthHandler.cs
--- a/Engine/Source/Programs/Horde/HordeServer/Authentication/AnonymousAuthHandler.cs
+++ b/Engine/Source/Programs/Horde/HordeServer/Authentication/AnonymousAuthHandler.cs
@@ -6,12 +6,12 @@ using System.Security.Claims;
 using System.Text.Encodings.Web;
 using System.Threading.Tasks;
 using EpicGames.Horde.Users;
-using Horde.Server.Utilities;
+using HordeServer.Utilities;
 using Microsoft.AspNetCore.Authentication;
 using Microsoft.Extensions.Logging;
 using Microsoft.Extensions.Options;
 
-namespace Horde.Server.Authentication
+namespace HordeServer.Authentication
 {
 	class AnonymousAuthenticationOptions : AuthenticationSchemeOptions
 	{
diff --git a/Engine/Source/Programs/Horde/HordeServer/Authentication/ExternalJwtAuthHandler.cs b/Engine/Source/Programs/Horde/HordeServer/Authentication/ExternalJwtAuthHandler.cs
--- a/Engine/Source/Programs/Horde/HordeServer/Authentication/ExternalJwtAuthHandler.cs
+++ b/Engine/Source/Programs/Horde/HordeServer/Authentication/ExternalJwtAuthHandler.cs
@@ -6,15 +6,15 @@ using System.IdentityModel.Tokens.Jwt;
 using System.Linq;
 using System.Security.Claims;
 using System.Threading.Tasks;
-using Horde.Server.Users;
-using Horde.Server.Utilities;
+using HordeServer.Users;
+using HordeServer.Utilities;
 using Microsoft.AspNetCore.Authentication;
 using Microsoft.AspNetCore.Authentication.JwtBearer;
 using Microsoft.Extensions.DependencyInjection;
 using Microsoft.Extensions.Logging;
 using Microsoft.IdentityModel.Tokens;
 
-namespace Horde.Server.Authentication
+namespace HordeServer.Authentication
 {
 	/// <summary>
 	/// JWT bearer handler capable of processing tokens generated by external OIDC providers (ie. access tokens issued by the OIDC server rather than the Horde server)
diff --git a/Engine/Source/Programs/Horde/HordeServer/Authentication/JwtAuthHandler.cs b/Engine/Source/Programs/Horde/HordeServer/Authentication/JwtAuthHandler.cs
--- a/Engine/Source/Programs/Horde/HordeServer/Authentication/JwtAuthHandler.cs
+++ b/Engine/Source/Programs/Horde/HordeServer/Authentication/JwtAuthHandler.cs
@@ -4,14 +4,14 @@ using System;
 using System.IdentityModel.Tokens.Jwt;
 using System.Text.Encodings.Web;
 using System.Threading.Tasks;
-using Horde.Server.Server;
-using Horde.Server.Utilities;
+using HordeServer.Server;
+using HordeServer.Utilities;
 using Microsoft.AspNetCore.Authentication;
 using Microsoft.AspNetCore.Authentication.JwtBearer;
 using Microsoft.Extensions.Logging;
 using Microsoft.Extensions.Options;
 
-namespace Horde.Server.Authentication
+namespace HordeServer.Authentication
 {
 	/// <summary>
 	/// JWT handler for server-issued bearer tokens. These tokens are signed using a randomly generated key per DB instance.
diff --git a/Engine/Source/Programs/Horde/HordeServer/Authentication/OidcAuthHandler.cs b/Engine/Source/Programs/Horde/HordeServer/Authentication/OidcAuthHandler.cs
--- a/Engine/Source/Programs/Horde/HordeServer/Authentication/OidcAuthHandler.cs
+++ b/Engine/Source/Programs/Horde/HordeServer/Authentication/OidcAuthHandler.cs
@@ -6,8 +6,8 @@ using System.Security.Claims;
 using System.Text.Encodings.Web;
 using System.Text.Json;
 using System.Threading.Tasks;
-using Horde.Server.Users;
-using Horde.Server.Utilities;
+using HordeServer.Users;
+using HordeServer.Utilities;
 using Microsoft.AspNetCore.Authentication;
 using Microsoft.AspNetCore.Authentication.OAuth.Claims;
 using Microsoft.AspNetCore.Authentication.OpenIdConnect;
@@ -17,7 +17,7 @@ using Microsoft.Extensions.Logging;
 using Microsoft.Extensions.Options;
 using Microsoft.IdentityModel.Protocols.OpenIdConnect;
 
-namespace Horde.Server.Authentication
+namespace HordeServer.Authentication
 {
 	class OidcAuthHandler : OpenIdConnectHandler
 	{
diff --git a/Engine/Source/Programs/Horde/HordeServer/Authentication/ServiceAccountAuthHandler.cs b/Engine/Source/Programs/Horde/HordeServer/Authentication/ServiceAccountAuthHandler.cs
--- a/Engine/Source/Programs/Horde/HordeServer/Authentication/ServiceAccountAuthHandler.cs
+++ b/Engine/Source/Programs/Horde/HordeServer/Authentication/ServiceAccountAuthHandler.cs
@@ -6,13 +6,13 @@ using System.Linq;
 using System.Security.Claims;
 using System.Text.Encodings.Web;
 using System.Threading.Tasks;
-using Horde.Server.ServiceAccounts;
+using HordeServer.ServiceAccounts;
 using Microsoft.AspNetCore.Authentication;
 using Microsoft.Extensions.Logging;
 using Microsoft.Extensions.Options;
 using Microsoft.Net.Http.Headers;
 
-namespace Horde.Server.Authentication
+namespace HordeServer.Authentication
 {
 	class ServiceAccountAuthOptions : AuthenticationSchemeOptions
 	{
diff --git a/Engine/Source/Programs/Horde/HordeServer/Aws/AwsCloudWatchMetricExporter.cs b/Engine/Source/Programs/Horde/HordeServer/Aws/AwsCloudWatchMetricExporter.cs
--- a/Engine/Source/Programs/Horde/HordeServer/Aws/AwsCloudWatchMetricExporter.cs
+++ b/Engine/Source/Programs/Horde/HordeServer/Aws/AwsCloudWatchMetricExporter.cs
@@ -8,12 +8,12 @@ using System.Threading;
 using System.Threading.Tasks;
 using Amazon.CloudWatch;
 using Amazon.CloudWatch.Model;
-using Horde.Server.Compute;
-using Horde.Server.Utilities;
+using HordeServer.Compute;
+using HordeServer.Utilities;
 using Microsoft.Extensions.Hosting;
 using Microsoft.Extensions.Logging;
 
-namespace Horde.Server.Aws;
+namespace HordeServer.Aws;
 
 /// <summary>
 /// Exports metric inside Horde server to AWS CloudWatch
diff --git a/Engine/Source/Programs/Horde/HordeServer/Commands/Generate/ConfigDocsCommand.cs b/Engine/Source/Programs/Horde/HordeServer/Commands/Generate/ConfigDocsCommand.cs
--- a/Engine/Source/Programs/Horde/HordeServer/Commands/Generate/ConfigDocsCommand.cs
+++ b/Engine/Source/Programs/Horde/HordeServer/Commands/Generate/ConfigDocsCommand.cs
@@ -12,14 +12,14 @@ using System.Threading.Tasks;
 using System.Xml;
 using EpicGames.Core;
 using EpicGames.Horde.Acls;
-using Horde.Server.Dashboard;
-using Horde.Server.Projects;
-using Horde.Server.Server;
-using Horde.Server.Streams;
-using Horde.Server.Telemetry.Metrics;
+using HordeServer.Dashboard;
+using HordeServer.Projects;
+using HordeServer.Server;
+using HordeServer.Streams;
+using HordeServer.Telemetry.Metrics;
 using Microsoft.Extensions.Logging;
 
-namespace Horde.Server.Commands.Generate
+namespace HordeServer.Commands.Generate
 {
 	[Command("generate", "configdocs", "Writes Markdown docs for server settings")]
 	class DocsCommand : Command
@@ -66,7 +66,7 @@ namespace Horde.Server.Commands.Generate
 			else
 			{
 				Assembly agentAssembly = Assembly.LoadFile(Agent.FullName);
-				Type agentSettingsType = agentAssembly.GetType("Horde.Agent.AgentSettings")!;
+				Type agentSettingsType = agentAssembly.GetType("HordeAgent.AgentSettings")!;
 
 				JsonSchema agentSchema = JsonSchema.FromType(agentSettingsType, xmlDocReader);
 				await WriteDocAsync(agentSchema.RootType, "Agent.json (Agent)", "Deployment/AgentSettings.md", "[Horde](../../README.md) > [Deployment](../Deployment.md) > [Agent](Agent.md)", AppSettingsIntro, new Dictionary<string, string>(), logger);
diff --git a/Engine/Source/Programs/Horde/HordeServer/Commands/Generate/ConfigSchemasCommand.cs b/Engine/Source/Programs/Horde/HordeServer/Commands/Generate/ConfigSchemasCommand.cs
--- a/Engine/Source/Programs/Horde/HordeServer/Commands/Generate/ConfigSchemasCommand.cs
+++ b/Engine/Source/Programs/Horde/HordeServer/Commands/Generate/ConfigSchemasCommand.cs
@@ -4,11 +4,11 @@ using System;
 using System.ComponentModel;
 using System.Threading.Tasks;
 using EpicGames.Core;
-using Horde.Server.Server;
-using Horde.Server.Utilities;
+using HordeServer.Server;
+using HordeServer.Utilities;
 using Microsoft.Extensions.Logging;
 
-namespace Horde.Server.Commands.Generate
+namespace HordeServer.Commands.Generate
 {
 	[Command("generate", "configschemas", "Writes JSON schemas for server settings")]
 	class SchemasCommand : Command
diff --git a/Engine/Source/Programs/Horde/HordeServer/Commands/Generate/CppApiCommand.cs b/Engine/Source/Programs/Horde/HordeServer/Commands/Generate/CppApiCommand.cs
--- a/Engine/Source/Programs/Horde/HordeServer/Commands/Generate/CppApiCommand.cs
+++ b/Engine/Source/Programs/Horde/HordeServer/Commands/Generate/CppApiCommand.cs
@@ -10,12 +10,12 @@ using System.Threading;
 using System.Threading.Tasks;
 using EpicGames.Core;
 using EpicGames.Horde;
-using Horde.Server.Projects;
-using Horde.Server.Server;
+using HordeServer.Projects;
+using HordeServer.Server;
 using Microsoft.AspNetCore.Mvc;
 using Microsoft.Extensions.Logging;
 
-namespace Horde.Server.Commands.Generate
+namespace HordeServer.Commands.Generate
 {
 	[Command("generate", "cppapi", "Outputs C++ source files describing the Horde API")]
 	class CppApiCommand : Command
diff --git a/Engine/Source/Programs/Horde/HordeServer/Commands/Generate/ToolDataCommand.cs b/Engine/Source/Programs/Horde/HordeServer/Commands/Generate/ToolDataCommand.cs
--- a/Engine/Source/Programs/Horde/HordeServer/Commands/Generate/ToolDataCommand.cs
+++ b/Engine/Source/Programs/Horde/HordeServer/Commands/Generate/ToolDataCommand.cs
@@ -10,10 +10,10 @@ using EpicGames.Horde.Storage;
 using EpicGames.Horde.Storage.Backends;
 using EpicGames.Horde.Storage.Bundles;
 using EpicGames.Horde.Storage.Nodes;
-using Horde.Server.Utilities;
+using HordeServer.Utilities;
 using Microsoft.Extensions.Logging;
 
-namespace Horde.Server.Commands.Generate
+namespace HordeServer.Commands.Generate
 {
 	using JsonObject = System.Text.Json.Nodes.JsonObject;
 
diff --git a/Engine/Source/Programs/Horde/HordeServer/Commands/OpenCommand.cs b/Engine/Source/Programs/Horde/HordeServer/Commands/OpenCommand.cs
--- a/Engine/Source/Programs/Horde/HordeServer/Commands/OpenCommand.cs
+++ b/Engine/Source/Programs/Horde/HordeServer/Commands/OpenCommand.cs
@@ -6,7 +6,7 @@ using EpicGames.Core;
 using Microsoft.Extensions.Logging;
 using Microsoft.Extensions.Options;
 
-namespace Horde.Server.Commands
+namespace HordeServer.Commands
 {
 	/// <summary>
 	/// Opens a web browser to the server homepage
diff --git a/Engine/Source/Programs/Horde/HordeServer/Commands/RelayCommand.cs b/Engine/Source/Programs/Horde/HordeServer/Commands/RelayCommand.cs
--- a/Engine/Source/Programs/Horde/HordeServer/Commands/RelayCommand.cs
+++ b/Engine/Source/Programs/Horde/HordeServer/Commands/RelayCommand.cs
@@ -8,11 +8,11 @@ using System.Threading.Tasks;
 using EpicGames.Core;
 using Grpc.Net.Client;
 using Horde.Common.Rpc;
-using Horde.Server.Agents.Relay;
+using HordeServer.Agents.Relay;
 using Microsoft.Extensions.Logging;
 using Microsoft.Extensions.Logging.Abstractions;
 
-namespace Horde.Server.Commands;
+namespace HordeServer.Commands;
 
 /// <summary>
 /// Run server in relay mode
diff --git a/Engine/Source/Programs/Horde/HordeServer/Commands/ServerCommand.cs b/Engine/Source/Programs/Horde/HordeServer/Commands/ServerCommand.cs
--- a/Engine/Source/Programs/Horde/HordeServer/Commands/ServerCommand.cs
+++ b/Engine/Source/Programs/Horde/HordeServer/Commands/ServerCommand.cs
@@ -18,7 +18,7 @@ using Microsoft.Extensions.Logging;
 using Microsoft.Extensions.Options;
 using Serilog;
 
-namespace Horde.Server.Commands
+namespace HordeServer.Commands
 {
 	using ILogger = Microsoft.Extensions.Logging.ILogger;
 
diff --git a/Engine/Source/Programs/Horde/HordeServer/Commands/Test/ReplicationCommand.cs b/Engine/Source/Programs/Horde/HordeServer/Commands/Test/ReplicationCommand.cs
--- a/Engine/Source/Programs/Horde/HordeServer/Commands/Test/ReplicationCommand.cs
+++ b/Engine/Source/Programs/Horde/HordeServer/Commands/Test/ReplicationCommand.cs
@@ -5,15 +5,15 @@ using System.Threading.Tasks;
 using EpicGames.Core;
 using EpicGames.Horde.Replicators;
 using EpicGames.Horde.Streams;
-using Horde.Server.Configuration;
-using Horde.Server.Replicators;
-using Horde.Server.Server;
-using Horde.Server.Streams;
+using HordeServer.Configuration;
+using HordeServer.Replicators;
+using HordeServer.Server;
+using HordeServer.Streams;
 using Microsoft.Extensions.Configuration;
 using Microsoft.Extensions.DependencyInjection;
 using Microsoft.Extensions.Logging;
 
-namespace Horde.Server.Commands.Test
+namespace HordeServer.Commands.Test
 {
 	[Command("test", "replication", "Replicates commits for a particular change of changes from Perforce")]
 	class TestReplicationCommand : Command
diff --git a/Engine/Source/Programs/Horde/HordeServer/Commits/CommitService.cs b/Engine/Source/Programs/Horde/HordeServer/Commits/CommitService.cs
--- a/Engine/Source/Programs/Horde/HordeServer/Commits/CommitService.cs
+++ b/Engine/Source/Programs/Horde/HordeServer/Commits/CommitService.cs
@@ -1,9 +1,9 @@
 // Copyright Epic Games, Inc. All Rights Reserved.
 
-using Horde.Server.Perforce;
-using Horde.Server.Streams;
+using HordeServer.Perforce;
+using HordeServer.Streams;
 
-namespace Horde.Server.Commits
+namespace HordeServer.Commits
 {
 	/// <summary>
 	/// Provides commit information for streams
diff --git a/Engine/Source/Programs/Horde/HordeServer/Commits/ICommit.cs b/Engine/Source/Programs/Horde/HordeServer/Commits/ICommit.cs
--- a/Engine/Source/Programs/Horde/HordeServer/Commits/ICommit.cs
+++ b/Engine/Source/Programs/Horde/HordeServer/Commits/ICommit.cs
@@ -9,7 +9,7 @@ using EpicGames.Horde.Commits;
 using EpicGames.Horde.Streams;
 using EpicGames.Horde.Users;
 
-namespace Horde.Server.Commits
+namespace HordeServer.Commits
 {
 	/// <summary>
 	/// Stores metadata about a commit
diff --git a/Engine/Source/Programs/Horde/HordeServer/Commits/ICommitCollection.cs b/Engine/Source/Programs/Horde/HordeServer/Commits/ICommitCollection.cs
--- a/Engine/Source/Programs/Horde/HordeServer/Commits/ICommitCollection.cs
+++ b/Engine/Source/Programs/Horde/HordeServer/Commits/ICommitCollection.cs
@@ -7,9 +7,9 @@ using System.Threading;
 using System.Threading.Tasks;
 using EpicGames.Horde.Commits;
 using EpicGames.Perforce;
-using Horde.Server.Jobs.Templates;
+using HordeServer.Jobs.Templates;
 
-namespace Horde.Server.Commits
+namespace HordeServer.Commits
 {
 	/// <summary>
 	/// VCS abstraction. Provides information about commits to a particular stream.
diff --git a/Engine/Source/Programs/Horde/HordeServer/Commits/ICommitService.cs b/Engine/Source/Programs/Horde/HordeServer/Commits/ICommitService.cs
--- a/Engine/Source/Programs/Horde/HordeServer/Commits/ICommitService.cs
+++ b/Engine/Source/Programs/Horde/HordeServer/Commits/ICommitService.cs
@@ -1,8 +1,8 @@
 // Copyright Epic Games, Inc. All Rights Reserved.
 
-using Horde.Server.Streams;
+using HordeServer.Streams;
 
-namespace Horde.Server.Commits
+namespace HordeServer.Commits
 {
 	/// <summary>
 	/// Provides information about commits to a stream
diff --git a/Engine/Source/Programs/Horde/HordeServer/Compute/ComputeAclAction.cs b/Engine/Source/Programs/Horde/HordeServer/Compute/ComputeAclAction.cs
--- a/Engine/Source/Programs/Horde/HordeServer/Compute/ComputeAclAction.cs
+++ b/Engine/Source/Programs/Horde/HordeServer/Compute/ComputeAclAction.cs
@@ -2,7 +2,7 @@
 
 using EpicGames.Horde.Acls;
 
-namespace Horde.Server.Compute
+namespace HordeServer.Compute
 {
 	/// <summary>
 	/// ACL actions for compute service
diff --git a/Engine/Source/Programs/Horde/HordeServer/Compute/ComputeController.cs b/Engine/Source/Programs/Horde/HordeServer/Compute/ComputeController.cs
--- a/Engine/Source/Programs/Horde/HordeServer/Compute/ComputeController.cs
+++ b/Engine/Source/Programs/Horde/HordeServer/Compute/ComputeController.cs
@@ -9,14 +9,14 @@ using EpicGames.Core;
 using EpicGames.Horde.Agents;
 using EpicGames.Horde.Agents.Pools;
 using EpicGames.Horde.Compute;
-using Horde.Server.Acls;
-using Horde.Server.Server;
-using Horde.Server.Utilities;
+using HordeServer.Acls;
+using HordeServer.Server;
+using HordeServer.Utilities;
 using Microsoft.AspNetCore.Authorization;
 using Microsoft.AspNetCore.Mvc;
 using Microsoft.Extensions.Options;
 
-namespace Horde.Server.Compute
+namespace HordeServer.Compute
 {
 	/// <summary>
 	/// Controller for the /api/v2/compute endpoint
diff --git a/Engine/Source/Programs/Horde/HordeServer/Compute/ComputeRequest.cs b/Engine/Source/Programs/Horde/HordeServer/Compute/ComputeRequest.cs
--- a/Engine/Source/Programs/Horde/HordeServer/Compute/ComputeRequest.cs
+++ b/Engine/Source/Programs/Horde/HordeServer/Compute/ComputeRequest.cs
@@ -3,7 +3,7 @@
 using EpicGames.Horde.Compute;
 using HordeCommon.Rpc.Tasks;
 
-namespace Horde.Server.Compute
+namespace HordeServer.Compute
 {
 	/// <summary>
 	/// Request for a compute resource
diff --git a/Engine/Source/Programs/Horde/HordeServer/Compute/ComputeService.cs b/Engine/Source/Programs/Horde/HordeServer/Compute/ComputeService.cs
--- a/Engine/Source/Programs/Horde/HordeServer/Compute/ComputeService.cs
+++ b/Engine/Source/Programs/Horde/HordeServer/Compute/ComputeService.cs
@@ -23,12 +23,12 @@ using EpicGames.Horde.Users;
 using Google.Protobuf;
 using Google.Protobuf.WellKnownTypes;
 using Horde.Common.Rpc;
-using Horde.Server.Agents;
-using Horde.Server.Agents.Leases;
-using Horde.Server.Agents.Relay;
-using Horde.Server.Logs;
-using Horde.Server.Server;
-using Horde.Server.Utilities;
+using HordeServer.Agents;
+using HordeServer.Agents.Leases;
+using HordeServer.Agents.Relay;
+using HordeServer.Logs;
+using HordeServer.Server;
+using HordeServer.Utilities;
 using HordeCommon.Rpc.Tasks;
 using Microsoft.Extensions.Hosting;
 using Microsoft.Extensions.Logging;
@@ -37,7 +37,7 @@ using OpenTelemetry.Trace;
 using StackExchange.Redis;
 using Condition = EpicGames.Horde.Common.Condition;
 
-namespace Horde.Server.Compute
+namespace HordeServer.Compute
 {
 	/// <summary>
 	/// Exceptions related to compute service
diff --git a/Engine/Source/Programs/Horde/HordeServer/Compute/ComputeTaskSource.cs b/Engine/Source/Programs/Horde/HordeServer/Compute/ComputeTaskSource.cs
--- a/Engine/Source/Programs/Horde/HordeServer/Compute/ComputeTaskSource.cs
+++ b/Engine/Source/Programs/Horde/HordeServer/Compute/ComputeTaskSource.cs
@@ -10,16 +10,16 @@ using EpicGames.Horde.Agents;
 using EpicGames.Horde.Agents.Leases;
 using EpicGames.Horde.Compute;
 using Google.Protobuf.WellKnownTypes;
-using Horde.Server.Agents;
-using Horde.Server.Agents.Leases;
-using Horde.Server.Agents.Relay;
-using Horde.Server.Server;
-using Horde.Server.Tasks;
+using HordeServer.Agents;
+using HordeServer.Agents.Leases;
+using HordeServer.Agents.Relay;
+using HordeServer.Server;
+using HordeServer.Tasks;
 using HordeCommon.Rpc.Tasks;
 using Microsoft.Extensions.Logging;
 using Microsoft.Extensions.Options;
 
-namespace Horde.Server.Compute
+namespace HordeServer.Compute
 {
 	/// <summary>
 	/// A TCP/IP port used by a compute resource, and how it is mapped externally -> internally
@@ -48,7 +48,7 @@ namespace Horde.Server.Compute
 			AgentPort = agentPort;
 		}
 
-		/// <inheritdoc cref="Equals(Horde.Server.Compute.ComputeResourcePort)" />
+		/// <inheritdoc cref="Equals(HordeServer.Compute.ComputeResourcePort)" />
 		protected bool Equals(ComputeResourcePort other)
 		{
 			return Port == other.Port && AgentPort == other.AgentPort;
diff --git a/Engine/Source/Programs/Horde/HordeServer/Configuration/ConfigContext.cs b/Engine/Source/Programs/Horde/HordeServer/Configuration/ConfigContext.cs
--- a/Engine/Source/Programs/Horde/HordeServer/Configuration/ConfigContext.cs
+++ b/Engine/Source/Programs/Horde/HordeServer/Configuration/ConfigContext.cs
@@ -10,7 +10,7 @@ using System.Threading.Tasks;
 using EpicGames.Core;
 using Microsoft.Extensions.Logging;
 
-namespace Horde.Server.Configuration
+namespace HordeServer.Configuration
 {
 	using JsonObject = System.Text.Json.Nodes.JsonObject;
 
diff --git a/Engine/Source/Programs/Horde/HordeServer/Configuration/ConfigException.cs b/Engine/Source/Programs/Horde/HordeServer/Configuration/ConfigException.cs
--- a/Engine/Source/Programs/Horde/HordeServer/Configuration/ConfigException.cs
+++ b/Engine/Source/Programs/Horde/HordeServer/Configuration/ConfigException.cs
@@ -3,7 +3,7 @@
 using System;
 using System.Collections.Generic;
 
-namespace Horde.Server.Configuration
+namespace HordeServer.Configuration
 {
 	/// <summary>
 	/// Exception thrown when reading config files
diff --git a/Engine/Source/Programs/Horde/HordeServer/Configuration/ConfigService.cs b/Engine/Source/Programs/Horde/HordeServer/Configuration/ConfigService.cs
--- a/Engine/Source/Programs/Horde/HordeServer/Configuration/ConfigService.cs
+++ b/Engine/Source/Programs/Horde/HordeServer/Configuration/ConfigService.cs
@@ -12,8 +12,8 @@ using System.Threading.Tasks;
 using EpicGames.Core;
 using EpicGames.Horde.Users;
 using EpicGames.Redis;
-using Horde.Server.Server;
-using Horde.Server.Users;
+using HordeServer.Server;
+using HordeServer.Users;
 using Microsoft.Extensions.Diagnostics.HealthChecks;
 using Microsoft.Extensions.Hosting;
 using Microsoft.Extensions.Logging;
@@ -23,7 +23,7 @@ using Microsoft.Extensions.Primitives;
 using ProtoBuf;
 using StackExchange.Redis;
 
-namespace Horde.Server.Configuration
+namespace HordeServer.Configuration
 {
 	/// <summary>
 	/// Information about the updated config
diff --git a/Engine/Source/Programs/Horde/HordeServer/Configuration/IConfigFile.cs b/Engine/Source/Programs/Horde/HordeServer/Configuration/IConfigFile.cs
--- a/Engine/Source/Programs/Horde/HordeServer/Configuration/IConfigFile.cs
+++ b/Engine/Source/Programs/Horde/HordeServer/Configuration/IConfigFile.cs
@@ -3,9 +3,9 @@
 using System;
 using System.Threading;
 using System.Threading.Tasks;
-using Horde.Server.Users;
+using HordeServer.Users;
 
-namespace Horde.Server.Configuration
+namespace HordeServer.Configuration
 {
 	/// <summary>
 	/// Accessor for a specific revision of a config file. Provides metadata about the current revision, and allows reading its data.
diff --git a/Engine/Source/Programs/Horde/HordeServer/Configuration/IConfigSource.cs b/Engine/Source/Programs/Horde/HordeServer/Configuration/IConfigSource.cs
--- a/Engine/Source/Programs/Horde/HordeServer/Configuration/IConfigSource.cs
+++ b/Engine/Source/Programs/Horde/HordeServer/Configuration/IConfigSource.cs
@@ -10,12 +10,12 @@ using System.Threading.Tasks;
 using EpicGames.Core;
 using EpicGames.Horde.Users;
 using EpicGames.Perforce;
-using Horde.Server.Users;
+using HordeServer.Users;
 using Microsoft.Extensions.Caching.Memory;
 using Microsoft.Extensions.Logging;
 using Microsoft.Extensions.Options;
 
-namespace Horde.Server.Configuration
+namespace HordeServer.Configuration
 {
 	/// <summary>
 	/// Source for reading config files
diff --git a/Engine/Source/Programs/Horde/HordeServer/Dashboard/DashboardConfig.cs b/Engine/Source/Programs/Horde/HordeServer/Dashboard/DashboardConfig.cs
--- a/Engine/Source/Programs/Horde/HordeServer/Dashboard/DashboardConfig.cs
+++ b/Engine/Source/Programs/Horde/HordeServer/Dashboard/DashboardConfig.cs
@@ -3,10 +3,10 @@
 using System.Collections.Generic;
 using EpicGames.Core;
 using EpicGames.Horde.Common;
-using Horde.Server.Configuration;
-using Horde.Server.Server;
+using HordeServer.Configuration;
+using HordeServer.Server;
 
-namespace Horde.Server.Dashboard
+namespace HordeServer.Dashboard
 {
 
 	/// <summary>
diff --git a/Engine/Source/Programs/Horde/HordeServer/Dashboard/DashboardController.cs b/Engine/Source/Programs/Horde/HordeServer/Dashboard/DashboardController.cs
--- a/Engine/Source/Programs/Horde/HordeServer/Dashboard/DashboardController.cs
+++ b/Engine/Source/Programs/Horde/HordeServer/Dashboard/DashboardController.cs
@@ -9,17 +9,17 @@ using System.Threading.Tasks;
 using EpicGames.Horde.Accounts;
 using EpicGames.Horde.Dashboard;
 using EpicGames.Horde.Server;
-using Horde.Server.Accounts;
-using Horde.Server.Server;
-using Horde.Server.Telemetry;
-using Horde.Server.Utilities;
+using HordeServer.Accounts;
+using HordeServer.Server;
+using HordeServer.Telemetry;
+using HordeServer.Utilities;
 using Microsoft.AspNetCore.Authentication;
 using Microsoft.AspNetCore.Authentication.Cookies;
 using Microsoft.AspNetCore.Authorization;
 using Microsoft.AspNetCore.Mvc;
 using Microsoft.Extensions.Options;
 
-namespace Horde.Server.Dashboard
+namespace HordeServer.Dashboard
 {
 	/// <summary>	
 	/// Dashboard authorization challenge controller	
diff --git a/Engine/Source/Programs/Horde/HordeServer/Dashboard/DashboardPreviewCollection.cs b/Engine/Source/Programs/Horde/HordeServer/Dashboard/DashboardPreviewCollection.cs
--- a/Engine/Source/Programs/Horde/HordeServer/Dashboard/DashboardPreviewCollection.cs
+++ b/Engine/Source/Programs/Horde/HordeServer/Dashboard/DashboardPreviewCollection.cs
@@ -5,12 +5,12 @@ using System.Collections.Generic;
 using System.Linq;
 using System.Threading;
 using System.Threading.Tasks;
-using Horde.Server.Server;
-using Horde.Server.Utilities;
+using HordeServer.Server;
+using HordeServer.Utilities;
 using MongoDB.Bson.Serialization.Attributes;
 using MongoDB.Driver;
 
-namespace Horde.Server.Dashboard
+namespace HordeServer.Dashboard
 {
 
 	/// <summary>
diff --git a/Engine/Source/Programs/Horde/HordeServer/Dashboard/DashboardTelemetryConfig.cs b/Engine/Source/Programs/Horde/HordeServer/Dashboard/DashboardTelemetryConfig.cs
--- a/Engine/Source/Programs/Horde/HordeServer/Dashboard/DashboardTelemetryConfig.cs
+++ b/Engine/Source/Programs/Horde/HordeServer/Dashboard/DashboardTelemetryConfig.cs
@@ -9,7 +9,7 @@ using EpicGames.Horde.Telemetry;
 using EpicGames.Horde.Telemetry.Metrics;
 using EpicGames.Serialization;
 
-namespace Horde.Server.Dashboard
+namespace HordeServer.Dashboard
 {
 	/// <summary>
 	/// The units used to present the telemetry
diff --git a/Engine/Source/Programs/Horde/HordeServer/Ddc/BlobContents.cs b/Engine/Source/Programs/Horde/HordeServer/Ddc/BlobContents.cs
--- a/Engine/Source/Programs/Horde/HordeServer/Ddc/BlobContents.cs
+++ b/Engine/Source/Programs/Horde/HordeServer/Ddc/BlobContents.cs
@@ -6,7 +6,7 @@ using System.Threading.Tasks;
 
 #pragma warning disable CS1591
 
-namespace Horde.Server.Ddc
+namespace HordeServer.Ddc
 {
 	public sealed class BlobContents : IDisposable, IAsyncDisposable
 	{
diff --git a/Engine/Source/Programs/Horde/HordeServer/Ddc/BlobId.cs b/Engine/Source/Programs/Horde/HordeServer/Ddc/BlobId.cs
--- a/Engine/Source/Programs/Horde/HordeServer/Ddc/BlobId.cs
+++ b/Engine/Source/Programs/Horde/HordeServer/Ddc/BlobId.cs
@@ -15,7 +15,7 @@ using EpicGames.Serialization;
 
 #pragma warning disable CS1591
 
-namespace Horde.Server.Ddc
+namespace HordeServer.Ddc
 {
 	[TypeConverter(typeof(BlobIdentifierTypeConverter))]
 	[JsonConverter(typeof(BlobIdentifierJsonConverter))]
diff --git a/Engine/Source/Programs/Horde/HordeServer/Ddc/IBlobService.cs b/Engine/Source/Programs/Horde/HordeServer/Ddc/IBlobService.cs
--- a/Engine/Source/Programs/Horde/HordeServer/Ddc/IBlobService.cs
+++ b/Engine/Source/Programs/Horde/HordeServer/Ddc/IBlobService.cs
@@ -13,7 +13,7 @@ using OpenTelemetry.Trace;
 
 #pragma warning disable CS1591
 
-namespace Horde.Server.Ddc;
+namespace HordeServer.Ddc;
 
 public interface IBlobService
 {
diff --git a/Engine/Source/Programs/Horde/HordeServer/Ddc/RefService.cs b/Engine/Source/Programs/Horde/HordeServer/Ddc/RefService.cs
--- a/Engine/Source/Programs/Horde/HordeServer/Ddc/RefService.cs
+++ b/Engine/Source/Programs/Horde/HordeServer/Ddc/RefService.cs
@@ -6,9 +6,10 @@ using System.Linq;
 using System.Threading;
 using System.Threading.Tasks;
 using EpicGames.Horde.Storage;
+using EpicGames.Horde.Storage.Nodes;
 using EpicGames.Serialization;
 
-namespace Horde.Server.Ddc
+namespace HordeServer.Ddc
 {
 	class RefService : IRefService
 	{
diff --git a/Engine/Source/Programs/Horde/HordeServer/Ddc/RequestHelper.cs b/Engine/Source/Programs/Horde/HordeServer/Ddc/RequestHelper.cs
--- a/Engine/Source/Programs/Horde/HordeServer/Ddc/RequestHelper.cs
+++ b/Engine/Source/Programs/Horde/HordeServer/Ddc/RequestHelper.cs
@@ -5,13 +5,13 @@ using System.Security.Claims;
 using System.Threading.Tasks;
 using EpicGames.Horde.Acls;
 using EpicGames.Horde.Storage;
-using Horde.Server.Server;
-using Horde.Server.Storage;
+using HordeServer.Server;
+using HordeServer.Storage;
 using Microsoft.AspNetCore.Http;
 using Microsoft.AspNetCore.Mvc;
 using Microsoft.Extensions.Options;
 
-namespace Horde.Server.Ddc
+namespace HordeServer.Ddc
 {
 	class RequestHelper : IRequestHelper
 	{
diff --git a/Engine/Source/Programs/Horde/HordeServer/Devices/DeviceAclAction.cs b/Engine/Source/Programs/Horde/HordeServer/Devices/DeviceAclAction.cs
--- a/Engine/Source/Programs/Horde/HordeServer/Devices/DeviceAclAction.cs
+++ b/Engine/Source/Programs/Horde/HordeServer/Devices/DeviceAclAction.cs
@@ -2,7 +2,7 @@
 
 using EpicGames.Horde.Acls;
 
-namespace Horde.Server.Devices
+namespace HordeServer.Devices
 {
 	/// <summary>
 	/// ACL actions which apply to devices
diff --git a/Engine/Source/Programs/Horde/HordeServer/Devices/DeviceCollection.cs b/Engine/Source/Programs/Horde/HordeServer/Devices/DeviceCollection.cs
--- a/Engine/Source/Programs/Horde/HordeServer/Devices/DeviceCollection.cs
+++ b/Engine/Source/Programs/Horde/HordeServer/Devices/DeviceCollection.cs
@@ -10,16 +10,16 @@ using EpicGames.Horde.Devices;
 using EpicGames.Horde.Jobs;
 using EpicGames.Horde.Streams;
 using EpicGames.Horde.Users;
-using Horde.Server.Jobs;
-using Horde.Server.Server;
-using Horde.Server.Utilities;
+using HordeServer.Jobs;
+using HordeServer.Server;
+using HordeServer.Utilities;
 using Microsoft.Extensions.Logging;
 using MongoDB.Bson;
 using MongoDB.Bson.Serialization.Attributes;
 using MongoDB.Bson.Serialization.Options;
 using MongoDB.Driver;
 
-namespace Horde.Server.Devices
+namespace HordeServer.Devices
 {
 	/// <summary>
 	/// Collection of device documents
diff --git a/Engine/Source/Programs/Horde/HordeServer/Devices/DeviceConfig.cs b/Engine/Source/Programs/Horde/HordeServer/Devices/DeviceConfig.cs
--- a/Engine/Source/Programs/Horde/HordeServer/Devices/DeviceConfig.cs
+++ b/Engine/Source/Programs/Horde/HordeServer/Devices/DeviceConfig.cs
@@ -8,7 +8,7 @@ using System.Text.Json.Serialization;
 using EpicGames.Horde.Devices;
 using EpicGames.Horde.Projects;
 
-namespace Horde.Server.Devices
+namespace HordeServer.Devices
 {
 	/// <summary>
 	/// Configuration for devices
diff --git a/Engine/Source/Programs/Horde/HordeServer/Devices/DeviceReportService.cs b/Engine/Source/Programs/Horde/HordeServer/Devices/DeviceReportService.cs
--- a/Engine/Source/Programs/Horde/HordeServer/Devices/DeviceReportService.cs
+++ b/Engine/Source/Programs/Horde/HordeServer/Devices/DeviceReportService.cs
@@ -7,15 +7,15 @@ using System.Threading;
 using System.Threading.Tasks;
 using EpicGames.Core;
 using EpicGames.Horde.Devices;
-using Horde.Server.Jobs;
-using Horde.Server.Notifications;
-using Horde.Server.Server;
-using Horde.Server.Utilities;
+using HordeServer.Jobs;
+using HordeServer.Notifications;
+using HordeServer.Server;
+using HordeServer.Utilities;
 using Microsoft.Extensions.Hosting;
 using Microsoft.Extensions.Logging;
 using Microsoft.Extensions.Options;
 
-namespace Horde.Server.Devices
+namespace HordeServer.Devices
 {
 #pragma warning disable CS1591 // Missing XML comment for publicly visible type or member
 	public class DeviceReport
diff --git a/Engine/Source/Programs/Horde/HordeServer/Devices/DeviceService.cs b/Engine/Source/Programs/Horde/HordeServer/Devices/DeviceService.cs
--- a/Engine/Source/Programs/Horde/HordeServer/Devices/DeviceService.cs
+++ b/Engine/Source/Programs/Horde/HordeServer/Devices/DeviceService.cs
@@ -12,20 +12,20 @@ using EpicGames.Horde.Devices;
 using EpicGames.Horde.Jobs;
 using EpicGames.Horde.Projects;
 using EpicGames.Horde.Users;
-using Horde.Server.Jobs;
-using Horde.Server.Jobs.Graphs;
-using Horde.Server.Notifications;
-using Horde.Server.Projects;
-using Horde.Server.Server;
-using Horde.Server.Streams;
-using Horde.Server.Users;
+using HordeServer.Jobs;
+using HordeServer.Jobs.Graphs;
+using HordeServer.Notifications;
+using HordeServer.Projects;
+using HordeServer.Server;
+using HordeServer.Streams;
+using HordeServer.Users;
 using Microsoft.Extensions.Hosting;
 using Microsoft.Extensions.Logging;
 using Microsoft.Extensions.Options;
 using MongoDB.Bson;
 using OpenTelemetry.Trace;
 
-namespace Horde.Server.Devices
+namespace HordeServer.Devices
 {
 	/// <summary>
 	///  Device Pool Authorization (convenience class for pool ACL's)
diff --git a/Engine/Source/Programs/Horde/HordeServer/Devices/DevicesController.cs b/Engine/Source/Programs/Horde/HordeServer/Devices/DevicesController.cs
--- a/Engine/Source/Programs/Horde/HordeServer/Devices/DevicesController.cs
+++ b/Engine/Source/Programs/Horde/HordeServer/Devices/DevicesController.cs
@@ -8,8 +8,8 @@ using System.Threading;
 using System.Threading.Tasks;
 using EpicGames.Horde.Devices;
 using EpicGames.Horde.Jobs;
-using Horde.Server.Server;
-using Horde.Server.Users;
+using HordeServer.Server;
+using HordeServer.Users;
 using Microsoft.AspNetCore.Authorization;
 using Microsoft.AspNetCore.Mvc;
 using Microsoft.Extensions.Logging;
@@ -17,7 +17,7 @@ using Microsoft.Extensions.Options;
 using MongoDB.Bson;
 using MongoDB.Driver;
 
-namespace Horde.Server.Devices
+namespace HordeServer.Devices
 {
 	/// <summary>
 	/// Controller for device service
diff --git a/Engine/Source/Programs/Horde/HordeServer/Devices/IDeviceCollection.cs b/Engine/Source/Programs/Horde/HordeServer/Devices/IDeviceCollection.cs
--- a/Engine/Source/Programs/Horde/HordeServer/Devices/IDeviceCollection.cs
+++ b/Engine/Source/Programs/Horde/HordeServer/Devices/IDeviceCollection.cs
@@ -6,10 +6,10 @@ using System.Threading.Tasks;
 using EpicGames.Horde.Devices;
 using EpicGames.Horde.Jobs;
 using EpicGames.Horde.Users;
-using Horde.Server.Jobs;
+using HordeServer.Jobs;
 using MongoDB.Bson;
 
-namespace Horde.Server.Devices
+namespace HordeServer.Devices
 {
 	/// <summary>
 	/// Device reservation request data
diff --git a/Engine/Source/Programs/Horde/HordeServer/HordeServer.csproj b/Engine/Source/Programs/Horde/HordeServer/HordeServer.csproj
--- a/Engine/Source/Programs/Horde/HordeServer/HordeServer.csproj
+++ b/Engine/Source/Programs/Horde/HordeServer/HordeServer.csproj
@@ -1,4 +1,4 @@
-<Project Sdk="Microsoft.NET.Sdk.Web">
+﻿<Project Sdk="Microsoft.NET.Sdk.Web">
 
   <PropertyGroup>
     <TargetFramework>net8.0</TargetFramework>
@@ -91,7 +91,7 @@
   </ItemGroup>
 
   <ItemGroup>
-    <None Remove="Horde.Server.xml" />
+    <None Remove="HordeServer.xml" />
     <None Remove="HordeServer.xml" />
   </ItemGroup>
 
diff --git a/Engine/Source/Programs/Horde/HordeServer/Issues/External/IExternalIssue.cs b/Engine/Source/Programs/Horde/HordeServer/Issues/External/IExternalIssue.cs
--- a/Engine/Source/Programs/Horde/HordeServer/Issues/External/IExternalIssue.cs
+++ b/Engine/Source/Programs/Horde/HordeServer/Issues/External/IExternalIssue.cs
@@ -1,6 +1,6 @@
 // Copyright Epic Games, Inc. All Rights Reserved.	
 
-namespace Horde.Server.Issues.External
+namespace HordeServer.Issues.External
 {
 	/// <summary>
 	/// 
diff --git a/Engine/Source/Programs/Horde/HordeServer/Issues/External/IExternalIssueService.cs b/Engine/Source/Programs/Horde/HordeServer/Issues/External/IExternalIssueService.cs
--- a/Engine/Source/Programs/Horde/HordeServer/Issues/External/IExternalIssueService.cs
+++ b/Engine/Source/Programs/Horde/HordeServer/Issues/External/IExternalIssueService.cs
@@ -4,11 +4,11 @@ using System;
 using System.Collections.Generic;
 using System.Threading;
 using System.Threading.Tasks;
-using Horde.Server.Streams;
-using Horde.Server.Users;
+using HordeServer.Streams;
+using HordeServer.Users;
 using Microsoft.Extensions.Hosting;
 
-namespace Horde.Server.Issues.External
+namespace HordeServer.Issues.External
 {
 	/// <summary>
 	/// External issue tracking project
diff --git a/Engine/Source/Programs/Horde/HordeServer/Issues/External/JiraService.cs b/Engine/Source/Programs/Horde/HordeServer/Issues/External/JiraService.cs
--- a/Engine/Source/Programs/Horde/HordeServer/Issues/External/JiraService.cs
+++ b/Engine/Source/Programs/Horde/HordeServer/Issues/External/JiraService.cs
@@ -11,16 +11,16 @@ using System.Text.Json;
 using System.Threading;
 using System.Threading.Tasks;
 using EpicGames.Core;
-using Horde.Server.Server;
-using Horde.Server.Streams;
-using Horde.Server.Users;
+using HordeServer.Server;
+using HordeServer.Streams;
+using HordeServer.Users;
 using Microsoft.Extensions.Hosting;
 using Microsoft.Extensions.Logging;
 using Microsoft.Extensions.Options;
 using Polly;
 using Polly.Extensions.Http;
 
-namespace Horde.Server.Issues.External
+namespace HordeServer.Issues.External
 {
 	/// <summary>
 	/// Jira issue information
diff --git a/Engine/Source/Programs/Horde/HordeServer/Issues/IIssue.cs b/Engine/Source/Programs/Horde/HordeServer/Issues/IIssue.cs
--- a/Engine/Source/Programs/Horde/HordeServer/Issues/IIssue.cs
+++ b/Engine/Source/Programs/Horde/HordeServer/Issues/IIssue.cs
@@ -8,7 +8,7 @@ using EpicGames.Horde.Streams;
 using EpicGames.Horde.Users;
 using MongoDB.Bson;
 
-namespace Horde.Server.Issues
+namespace HordeServer.Issues
 {
 	/// <summary>
 	/// Stores information about a build health issue
diff --git a/Engine/Source/Programs/Horde/HordeServer/Issues/IIssueCollection.cs b/Engine/Source/Programs/Horde/HordeServer/Issues/IIssueCollection.cs
--- a/Engine/Source/Programs/Horde/HordeServer/Issues/IIssueCollection.cs
+++ b/Engine/Source/Programs/Horde/HordeServer/Issues/IIssueCollection.cs
@@ -11,13 +11,13 @@ using EpicGames.Horde.Jobs.Templates;
 using EpicGames.Horde.Logs;
 using EpicGames.Horde.Streams;
 using EpicGames.Horde.Users;
-using Horde.Server.Auditing;
-using Horde.Server.Jobs;
-using Horde.Server.Jobs.Graphs;
+using HordeServer.Auditing;
+using HordeServer.Jobs;
+using HordeServer.Jobs.Graphs;
 using MongoDB.Bson;
 using MongoDB.Driver;
 
-namespace Horde.Server.Issues
+namespace HordeServer.Issues
 {
 	/// <summary>
 	/// Fingerprint for an issue
diff --git a/Engine/Source/Programs/Horde/HordeServer/Issues/IIssueDetails.cs b/Engine/Source/Programs/Horde/HordeServer/Issues/IIssueDetails.cs
--- a/Engine/Source/Programs/Horde/HordeServer/Issues/IIssueDetails.cs
+++ b/Engine/Source/Programs/Horde/HordeServer/Issues/IIssueDetails.cs
@@ -5,9 +5,9 @@ using System.Collections.Generic;
 using System.Threading;
 using System.Threading.Tasks;
 using EpicGames.Horde.Users;
-using Horde.Server.Users;
+using HordeServer.Users;
 
-namespace Horde.Server.Issues
+namespace HordeServer.Issues
 {
 	/// <summary>
 	/// Cached issue information
diff --git a/Engine/Source/Programs/Horde/HordeServer/Issues/IIssueFingerprint.cs b/Engine/Source/Programs/Horde/HordeServer/Issues/IIssueFingerprint.cs
--- a/Engine/Source/Programs/Horde/HordeServer/Issues/IIssueFingerprint.cs
+++ b/Engine/Source/Programs/Horde/HordeServer/Issues/IIssueFingerprint.cs
@@ -5,7 +5,7 @@ using System.Collections.Generic;
 using System.Linq;
 using EpicGames.Horde.Issues;
 
-namespace Horde.Server.Issues
+namespace HordeServer.Issues
 {
 	/// <summary>
 	/// Fingerprint for an issue
diff --git a/Engine/Source/Programs/Horde/HordeServer/Issues/IIssueStep.cs b/Engine/Source/Programs/Horde/HordeServer/Issues/IIssueStep.cs
--- a/Engine/Source/Programs/Horde/HordeServer/Issues/IIssueStep.cs
+++ b/Engine/Source/Programs/Horde/HordeServer/Issues/IIssueStep.cs
@@ -4,10 +4,10 @@ using System;
 using EpicGames.Horde.Issues;
 using EpicGames.Horde.Jobs;
 using EpicGames.Horde.Logs;
-using Horde.Server.Jobs.Graphs;
+using HordeServer.Jobs.Graphs;
 using MongoDB.Bson;
 
-namespace Horde.Server.Issues
+namespace HordeServer.Issues
 {
 	/// <summary>
 	/// Identifies a particular changelist and job that contributes to a span
