Horde: Rename namespaces to match project folders.

diff --git a/Engine/Source/Programs/Horde/HordeServer/Issues/IssueCollection.cs b/Engine/Source/Programs/Horde/HordeServer/Issues/IssueCollection.cs
--- a/Engine/Source/Programs/Horde/HordeServer/Issues/IssueCollection.cs
+++ b/Engine/Source/Programs/Horde/HordeServer/Issues/IssueCollection.cs
@@ -17,14 +17,14 @@ using EpicGames.Horde.Streams;
 using EpicGames.Horde.Telemetry;
 using EpicGames.Horde.Users;
 using EpicGames.Redis.Utility;
-using Horde.Server.Auditing;
-using Horde.Server.Jobs;
-using Horde.Server.Jobs.Graphs;
-using Horde.Server.Server;
-using Horde.Server.Streams;
-using Horde.Server.Telemetry;
-using Horde.Server.Users;
-using Horde.Server.Utilities;
+using HordeServer.Auditing;
+using HordeServer.Jobs;
+using HordeServer.Jobs.Graphs;
+using HordeServer.Server;
+using HordeServer.Streams;
+using HordeServer.Telemetry;
+using HordeServer.Users;
+using HordeServer.Utilities;
 using Microsoft.Extensions.Caching.Memory;
 using Microsoft.Extensions.Logging;
 using Microsoft.Extensions.Options;
@@ -34,7 +34,7 @@ using MongoDB.Bson.Serialization.Attributes;
 using MongoDB.Driver;
 using OpenTelemetry.Trace;
 
-namespace Horde.Server.Issues
+namespace HordeServer.Issues
 {
 	class IssueCollection : IIssueCollection
 	{
diff --git a/Engine/Source/Programs/Horde/HordeServer/Issues/IssueReportService.cs b/Engine/Source/Programs/Horde/HordeServer/Issues/IssueReportService.cs
--- a/Engine/Source/Programs/Horde/HordeServer/Issues/IssueReportService.cs
+++ b/Engine/Source/Programs/Horde/HordeServer/Issues/IssueReportService.cs
@@ -10,12 +10,12 @@ using EpicGames.Horde.Issues;
 using EpicGames.Horde.Jobs;
 using EpicGames.Horde.Jobs.Templates;
 using EpicGames.Horde.Streams;
-using Horde.Server.Jobs;
-using Horde.Server.Jobs.Graphs;
-using Horde.Server.Notifications;
-using Horde.Server.Server;
-using Horde.Server.Streams;
-using Horde.Server.Utilities;
+using HordeServer.Jobs;
+using HordeServer.Jobs.Graphs;
+using HordeServer.Notifications;
+using HordeServer.Server;
+using HordeServer.Streams;
+using HordeServer.Utilities;
 using HordeCommon;
 using Microsoft.Extensions.Hosting;
 using Microsoft.Extensions.Logging;
@@ -23,7 +23,7 @@ using Microsoft.Extensions.Options;
 using MongoDB.Bson.Serialization.Attributes;
 using MongoDB.Bson.Serialization.Options;
 
-namespace Horde.Server.Issues
+namespace HordeServer.Issues
 {
 #pragma warning disable CS1591 // Missing XML comment for publicly visible type or member
 	public class WorkflowStats
diff --git a/Engine/Source/Programs/Horde/HordeServer/Issues/IssueService.cs b/Engine/Source/Programs/Horde/HordeServer/Issues/IssueService.cs
--- a/Engine/Source/Programs/Horde/HordeServer/Issues/IssueService.cs
+++ b/Engine/Source/Programs/Horde/HordeServer/Issues/IssueService.cs
@@ -16,13 +16,13 @@ using EpicGames.Horde.Jobs.Templates;
 using EpicGames.Horde.Logs;
 using EpicGames.Horde.Streams;
 using EpicGames.Horde.Users;
-using Horde.Server.Commits;
-using Horde.Server.Jobs;
-using Horde.Server.Jobs.Graphs;
-using Horde.Server.Logs;
-using Horde.Server.Server;
-using Horde.Server.Streams;
-using Horde.Server.Users;
+using HordeServer.Commits;
+using HordeServer.Jobs;
+using HordeServer.Jobs.Graphs;
+using HordeServer.Logs;
+using HordeServer.Server;
+using HordeServer.Streams;
+using HordeServer.Users;
 using Microsoft.Extensions.DependencyInjection;
 using Microsoft.Extensions.Hosting;
 using Microsoft.Extensions.Logging;
@@ -31,7 +31,7 @@ using MongoDB.Bson;
 using MongoDB.Bson.Serialization;
 using OpenTelemetry.Trace;
 
-namespace Horde.Server.Issues
+namespace HordeServer.Issues
 {
 	/// <summary>
 	/// Detailed issue information
diff --git a/Engine/Source/Programs/Horde/HordeServer/Issues/IssueTagService.cs b/Engine/Source/Programs/Horde/HordeServer/Issues/IssueTagService.cs
--- a/Engine/Source/Programs/Horde/HordeServer/Issues/IssueTagService.cs
+++ b/Engine/Source/Programs/Horde/HordeServer/Issues/IssueTagService.cs
@@ -7,17 +7,17 @@ using System.Threading;
 using System.Threading.Tasks;
 using EpicGames.Core;
 using EpicGames.Horde.Streams;
-using Horde.Server.Commits;
-using Horde.Server.Server;
-using Horde.Server.Streams;
-using Horde.Server.Utilities;
+using HordeServer.Commits;
+using HordeServer.Server;
+using HordeServer.Streams;
+using HordeServer.Utilities;
 using Microsoft.Extensions.Hosting;
 using Microsoft.Extensions.Logging;
 using Microsoft.Extensions.Options;
 using MongoDB.Bson.Serialization.Attributes;
 using MongoDB.Bson.Serialization.Options;
 
-namespace Horde.Server.Issues
+namespace HordeServer.Issues
 {
 	/// <summary>
 	/// Polls revision control for fix changelist numbers in commits using the syntax '#horde 1234'
diff --git a/Engine/Source/Programs/Horde/HordeServer/Issues/IssuesController.cs b/Engine/Source/Programs/Horde/HordeServer/Issues/IssuesController.cs
--- a/Engine/Source/Programs/Horde/HordeServer/Issues/IssuesController.cs
+++ b/Engine/Source/Programs/Horde/HordeServer/Issues/IssuesController.cs
@@ -12,22 +12,22 @@ using EpicGames.Horde.Jobs.Templates;
 using EpicGames.Horde.Logs;
 using EpicGames.Horde.Streams;
 using EpicGames.Horde.Users;
-using Horde.Server.Auditing;
-using Horde.Server.Issues.External;
-using Horde.Server.Jobs;
-using Horde.Server.Jobs.Graphs;
-using Horde.Server.Logs;
-using Horde.Server.Server;
-using Horde.Server.Streams;
-using Horde.Server.Users;
-using Horde.Server.Utilities;
+using HordeServer.Auditing;
+using HordeServer.Issues.External;
+using HordeServer.Jobs;
+using HordeServer.Jobs.Graphs;
+using HordeServer.Logs;
+using HordeServer.Server;
+using HordeServer.Streams;
+using HordeServer.Users;
+using HordeServer.Utilities;
 using Microsoft.AspNetCore.Authorization;
 using Microsoft.AspNetCore.Mvc;
 using Microsoft.Extensions.Logging;
 using Microsoft.Extensions.Options;
 using MongoDB.Bson;
 
-namespace Horde.Server.Issues
+namespace HordeServer.Issues
 {
 	/// <summary>
 	/// Controller for the /api/v1/issues endpoint
diff --git a/Engine/Source/Programs/Horde/HordeServer/Issues/WorkflowConfig.cs b/Engine/Source/Programs/Horde/HordeServer/Issues/WorkflowConfig.cs
--- a/Engine/Source/Programs/Horde/HordeServer/Issues/WorkflowConfig.cs
+++ b/Engine/Source/Programs/Horde/HordeServer/Issues/WorkflowConfig.cs
@@ -4,9 +4,9 @@ using System;
 using System.Collections.Generic;
 using System.ComponentModel.DataAnnotations;
 using EpicGames.Horde.Issues;
-using Horde.Server.Jobs.Graphs;
+using HordeServer.Jobs.Graphs;
 
-namespace Horde.Server.Issues
+namespace HordeServer.Issues
 {
 	/// <summary>
 	/// Configuration for an issue workflow
diff --git a/Engine/Source/Programs/Horde/HordeServer/Jobs/Artifacts/ArtifactCollectionV1.cs b/Engine/Source/Programs/Horde/HordeServer/Jobs/Artifacts/ArtifactCollectionV1.cs
--- a/Engine/Source/Programs/Horde/HordeServer/Jobs/Artifacts/ArtifactCollectionV1.cs
+++ b/Engine/Source/Programs/Horde/HordeServer/Jobs/Artifacts/ArtifactCollectionV1.cs
@@ -8,12 +8,12 @@ using System.Threading.Tasks;
 using EpicGames.Core;
 using EpicGames.Horde.Jobs;
 using EpicGames.Horde.Storage;
-using Horde.Server.Server;
+using HordeServer.Server;
 using MongoDB.Bson;
 using MongoDB.Bson.Serialization.Attributes;
 using MongoDB.Driver;
 
-namespace Horde.Server.Jobs.Artifacts
+namespace HordeServer.Jobs.Artifacts
 {
 	/// <summary>
 	/// Wraps functionality for manipulating artifacts
diff --git a/Engine/Source/Programs/Horde/HordeServer/Jobs/Artifacts/ArtifactMessagesV1.cs b/Engine/Source/Programs/Horde/HordeServer/Jobs/Artifacts/ArtifactMessagesV1.cs
--- a/Engine/Source/Programs/Horde/HordeServer/Jobs/Artifacts/ArtifactMessagesV1.cs
+++ b/Engine/Source/Programs/Horde/HordeServer/Jobs/Artifacts/ArtifactMessagesV1.cs
@@ -3,7 +3,7 @@
 using System.Collections.Generic;
 using EpicGames.Horde.Jobs;
 
-namespace Horde.Server.Jobs.Artifacts
+namespace HordeServer.Jobs.Artifacts
 {
 	/// <summary>
 	/// Response from creating an artifact
diff --git a/Engine/Source/Programs/Horde/HordeServer/Jobs/Artifacts/ArtifactsControllerV1.cs b/Engine/Source/Programs/Horde/HordeServer/Jobs/Artifacts/ArtifactsControllerV1.cs
--- a/Engine/Source/Programs/Horde/HordeServer/Jobs/Artifacts/ArtifactsControllerV1.cs
+++ b/Engine/Source/Programs/Horde/HordeServer/Jobs/Artifacts/ArtifactsControllerV1.cs
@@ -11,12 +11,12 @@ using System.Threading;
 using System.Threading.Tasks;
 using EpicGames.Core;
 using EpicGames.Horde.Jobs;
-using Horde.Server.Acls;
-using Horde.Server.Artifacts;
-using Horde.Server.Jobs.Graphs;
-using Horde.Server.Server;
-using Horde.Server.Streams;
-using Horde.Server.Utilities;
+using HordeServer.Acls;
+using HordeServer.Artifacts;
+using HordeServer.Jobs.Graphs;
+using HordeServer.Server;
+using HordeServer.Streams;
+using HordeServer.Utilities;
 using Microsoft.AspNetCore.Authorization;
 using Microsoft.AspNetCore.Http;
 using Microsoft.AspNetCore.Mvc;
@@ -24,7 +24,7 @@ using Microsoft.Extensions.Options;
 using Microsoft.IdentityModel.Tokens;
 using MongoDB.Bson;
 
-namespace Horde.Server.Jobs.Artifacts
+namespace HordeServer.Jobs.Artifacts
 {
 	/// <summary>
 	/// Controller for the /api/artifacts endpoint
diff --git a/Engine/Source/Programs/Horde/HordeServer/Jobs/Artifacts/IArtifactCollectionV1.cs b/Engine/Source/Programs/Horde/HordeServer/Jobs/Artifacts/IArtifactCollectionV1.cs
--- a/Engine/Source/Programs/Horde/HordeServer/Jobs/Artifacts/IArtifactCollectionV1.cs
+++ b/Engine/Source/Programs/Horde/HordeServer/Jobs/Artifacts/IArtifactCollectionV1.cs
@@ -6,7 +6,7 @@ using System.Threading.Tasks;
 using EpicGames.Horde.Jobs;
 using MongoDB.Bson;
 
-namespace Horde.Server.Jobs.Artifacts
+namespace HordeServer.Jobs.Artifacts
 {
 	/// <summary>
 	/// Interface for a collection of artifacts
diff --git a/Engine/Source/Programs/Horde/HordeServer/Jobs/Bisect/BisectController.cs b/Engine/Source/Programs/Horde/HordeServer/Jobs/Bisect/BisectController.cs
--- a/Engine/Source/Programs/Horde/HordeServer/Jobs/Bisect/BisectController.cs
+++ b/Engine/Source/Programs/Horde/HordeServer/Jobs/Bisect/BisectController.cs
@@ -11,18 +11,18 @@ using EpicGames.Horde.Jobs.Bisect;
 using EpicGames.Horde.Jobs.Templates;
 using EpicGames.Horde.Streams;
 using EpicGames.Horde.Users;
-using Horde.Server.Jobs.Graphs;
-using Horde.Server.Server;
-using Horde.Server.Streams;
-using Horde.Server.Users;
-using Horde.Server.Utilities;
+using HordeServer.Jobs.Graphs;
+using HordeServer.Server;
+using HordeServer.Streams;
+using HordeServer.Users;
+using HordeServer.Utilities;
 using Microsoft.AspNetCore.Authorization;
 using Microsoft.AspNetCore.Mvc;
 using Microsoft.Extensions.Logging;
 using Microsoft.Extensions.Options;
 using OpenTelemetry.Trace;
 
-namespace Horde.Server.Jobs.Bisect
+namespace HordeServer.Jobs.Bisect
 {
 	/// <summary>
 	/// Request to create a new bisect task
diff --git a/Engine/Source/Programs/Horde/HordeServer/Jobs/Bisect/BisectIdUtils.cs b/Engine/Source/Programs/Horde/HordeServer/Jobs/Bisect/BisectIdUtils.cs
--- a/Engine/Source/Programs/Horde/HordeServer/Jobs/Bisect/BisectIdUtils.cs
+++ b/Engine/Source/Programs/Horde/HordeServer/Jobs/Bisect/BisectIdUtils.cs
@@ -1,9 +1,9 @@
 // Copyright Epic Games, Inc. All Rights Reserved.
 
 using EpicGames.Horde.Jobs.Bisect;
-using Horde.Server.Utilities;
+using HordeServer.Utilities;
 
-namespace Horde.Server.Jobs.Bisect
+namespace HordeServer.Jobs.Bisect
 {
 	/// <summary>
 	/// Server helper methods for <see cref="BisectTaskId"/>
diff --git a/Engine/Source/Programs/Horde/HordeServer/Jobs/Bisect/BisectService.cs b/Engine/Source/Programs/Horde/HordeServer/Jobs/Bisect/BisectService.cs
--- a/Engine/Source/Programs/Horde/HordeServer/Jobs/Bisect/BisectService.cs
+++ b/Engine/Source/Programs/Horde/HordeServer/Jobs/Bisect/BisectService.cs
@@ -7,17 +7,17 @@ using System.Threading;
 using System.Threading.Tasks;
 using EpicGames.Core;
 using EpicGames.Horde.Jobs;
-using Horde.Server.Commits;
-using Horde.Server.Jobs.Graphs;
-using Horde.Server.Jobs.Templates;
-using Horde.Server.Server;
-using Horde.Server.Streams;
+using HordeServer.Commits;
+using HordeServer.Jobs.Graphs;
+using HordeServer.Jobs.Templates;
+using HordeServer.Server;
+using HordeServer.Streams;
 using Microsoft.Extensions.Hosting;
 using Microsoft.Extensions.Logging;
 using Microsoft.Extensions.Options;
 using MongoDB.Driver;
 
-namespace Horde.Server.Jobs.Bisect
+namespace HordeServer.Jobs.Bisect
 {
 	class BisectService : IHostedService
 	{
diff --git a/Engine/Source/Programs/Horde/HordeServer/Jobs/Bisect/BisectTaskAclAction.cs b/Engine/Source/Programs/Horde/HordeServer/Jobs/Bisect/BisectTaskAclAction.cs
--- a/Engine/Source/Programs/Horde/HordeServer/Jobs/Bisect/BisectTaskAclAction.cs
+++ b/Engine/Source/Programs/Horde/HordeServer/Jobs/Bisect/BisectTaskAclAction.cs
@@ -2,7 +2,7 @@
 
 using EpicGames.Horde.Acls;
 
-namespace Horde.Server.Jobs.Bisect
+namespace HordeServer.Jobs.Bisect
 {
 	/// <summary>
 	/// ACL actions valid for manipulating jobs
diff --git a/Engine/Source/Programs/Horde/HordeServer/Jobs/Bisect/BisectTaskCollection.cs b/Engine/Source/Programs/Horde/HordeServer/Jobs/Bisect/BisectTaskCollection.cs
--- a/Engine/Source/Programs/Horde/HordeServer/Jobs/Bisect/BisectTaskCollection.cs
+++ b/Engine/Source/Programs/Horde/HordeServer/Jobs/Bisect/BisectTaskCollection.cs
@@ -11,14 +11,14 @@ using EpicGames.Horde.Jobs.Bisect;
 using EpicGames.Horde.Jobs.Templates;
 using EpicGames.Horde.Streams;
 using EpicGames.Horde.Users;
-using Horde.Server.Server;
-using Horde.Server.Utilities;
+using HordeServer.Server;
+using HordeServer.Utilities;
 using MongoDB.Bson;
 using MongoDB.Bson.Serialization.Attributes;
 using MongoDB.Driver;
 using OpenTelemetry.Trace;
 
-namespace Horde.Server.Jobs.Bisect
+namespace HordeServer.Jobs.Bisect
 {
 	class BisectTaskCollection : IBisectTaskCollection
 	{
diff --git a/Engine/Source/Programs/Horde/HordeServer/Jobs/Bisect/IBisectTask.cs b/Engine/Source/Programs/Horde/HordeServer/Jobs/Bisect/IBisectTask.cs
--- a/Engine/Source/Programs/Horde/HordeServer/Jobs/Bisect/IBisectTask.cs
+++ b/Engine/Source/Programs/Horde/HordeServer/Jobs/Bisect/IBisectTask.cs
@@ -8,7 +8,7 @@ using EpicGames.Horde.Jobs.Templates;
 using EpicGames.Horde.Streams;
 using EpicGames.Horde.Users;
 
-namespace Horde.Server.Jobs.Bisect
+namespace HordeServer.Jobs.Bisect
 {
 	/// <summary>
 	/// State of a bisect task
diff --git a/Engine/Source/Programs/Horde/HordeServer/Jobs/Bisect/IBisectTaskCollection.cs b/Engine/Source/Programs/Horde/HordeServer/Jobs/Bisect/IBisectTaskCollection.cs
--- a/Engine/Source/Programs/Horde/HordeServer/Jobs/Bisect/IBisectTaskCollection.cs
+++ b/Engine/Source/Programs/Horde/HordeServer/Jobs/Bisect/IBisectTaskCollection.cs
@@ -8,10 +8,10 @@ using EpicGames.Horde.Commits;
 using EpicGames.Horde.Jobs;
 using EpicGames.Horde.Jobs.Bisect;
 using EpicGames.Horde.Users;
-using Horde.Server.Jobs.Graphs;
+using HordeServer.Jobs.Graphs;
 using Microsoft.Extensions.Logging;
 
-namespace Horde.Server.Jobs.Bisect
+namespace HordeServer.Jobs.Bisect
 {
 	/// <summary>
 	/// Manages a collection of bisect tasks
diff --git a/Engine/Source/Programs/Horde/HordeServer/Jobs/DebugJobsController.cs b/Engine/Source/Programs/Horde/HordeServer/Jobs/DebugJobsController.cs
--- a/Engine/Source/Programs/Horde/HordeServer/Jobs/DebugJobsController.cs
+++ b/Engine/Source/Programs/Horde/HordeServer/Jobs/DebugJobsController.cs
@@ -9,18 +9,18 @@ using System.Text;
 using System.Threading;
 using System.Threading.Tasks;
 using EpicGames.Horde.Jobs;
-using Horde.Server.Agents.Leases;
-using Horde.Server.Jobs.Graphs;
-using Horde.Server.Server;
-using Horde.Server.Streams;
-using Horde.Server.Utilities;
+using HordeServer.Agents.Leases;
+using HordeServer.Jobs.Graphs;
+using HordeServer.Server;
+using HordeServer.Streams;
+using HordeServer.Utilities;
 using Microsoft.AspNetCore.Authorization;
 using Microsoft.AspNetCore.Http;
 using Microsoft.AspNetCore.Mvc;
 using Microsoft.Extensions.Logging;
 using Microsoft.Extensions.Options;
 
-namespace Horde.Server.Jobs
+namespace HordeServer.Jobs
 {
 	/// <summary>
 	/// Debug functionality for jobs
diff --git a/Engine/Source/Programs/Horde/HordeServer/Jobs/Graphs/GraphCollection.cs b/Engine/Source/Programs/Horde/HordeServer/Jobs/Graphs/GraphCollection.cs
--- a/Engine/Source/Programs/Horde/HordeServer/Jobs/Graphs/GraphCollection.cs
+++ b/Engine/Source/Programs/Horde/HordeServer/Jobs/Graphs/GraphCollection.cs
@@ -9,8 +9,8 @@ using System.Threading.Tasks;
 using EpicGames.Core;
 using EpicGames.Horde.Artifacts;
 using EpicGames.Horde.Jobs;
-using Horde.Server.Jobs.Templates;
-using Horde.Server.Server;
+using HordeServer.Jobs.Templates;
+using HordeServer.Server;
 using HordeCommon;
 using Microsoft.Extensions.Caching.Memory;
 using MongoDB.Bson;
@@ -18,7 +18,7 @@ using MongoDB.Bson.Serialization.Attributes;
 using MongoDB.Bson.Serialization.Options;
 using MongoDB.Driver;
 
-namespace Horde.Server.Jobs.Graphs
+namespace HordeServer.Jobs.Graphs
 {
 	/// <summary>
 	/// Collection of graph documents
diff --git a/Engine/Source/Programs/Horde/HordeServer/Jobs/Graphs/IGraph.cs b/Engine/Source/Programs/Horde/HordeServer/Jobs/Graphs/IGraph.cs
--- a/Engine/Source/Programs/Horde/HordeServer/Jobs/Graphs/IGraph.cs
+++ b/Engine/Source/Programs/Horde/HordeServer/Jobs/Graphs/IGraph.cs
@@ -11,7 +11,7 @@ using EpicGames.Horde.Jobs;
 using HordeCommon;
 using MongoDB.Bson.Serialization.Attributes;
 
-namespace Horde.Server.Jobs.Graphs
+namespace HordeServer.Jobs.Graphs
 {
 	/// <summary>
 	/// A unique dependency graph instance
diff --git a/Engine/Source/Programs/Horde/HordeServer/Jobs/Graphs/IGraphCollection.cs b/Engine/Source/Programs/Horde/HordeServer/Jobs/Graphs/IGraphCollection.cs
--- a/Engine/Source/Programs/Horde/HordeServer/Jobs/Graphs/IGraphCollection.cs
+++ b/Engine/Source/Programs/Horde/HordeServer/Jobs/Graphs/IGraphCollection.cs
@@ -4,9 +4,9 @@ using System.Collections.Generic;
 using System.Threading;
 using System.Threading.Tasks;
 using EpicGames.Core;
-using Horde.Server.Jobs.Templates;
+using HordeServer.Jobs.Templates;
 
-namespace Horde.Server.Jobs.Graphs
+namespace HordeServer.Jobs.Graphs
 {
 	/// <summary>
 	/// Interface for a collection of graph documents
diff --git a/Engine/Source/Programs/Horde/HordeServer/Jobs/Graphs/NodeAnnotations.cs b/Engine/Source/Programs/Horde/HordeServer/Jobs/Graphs/NodeAnnotations.cs
--- a/Engine/Source/Programs/Horde/HordeServer/Jobs/Graphs/NodeAnnotations.cs
+++ b/Engine/Source/Programs/Horde/HordeServer/Jobs/Graphs/NodeAnnotations.cs
@@ -3,9 +3,9 @@
 using System;
 using System.Collections.Generic;
 using EpicGames.Horde.Issues;
-using Horde.Server.Utilities;
+using HordeServer.Utilities;
 
-namespace Horde.Server.Jobs.Graphs
+namespace HordeServer.Jobs.Graphs
 {
 	/// <summary>
 	/// Interface which wraps a generic key/value dictionary to provide specific node annotations
diff --git a/Engine/Source/Programs/Horde/HordeServer/Jobs/IJob.cs b/Engine/Source/Programs/Horde/HordeServer/Jobs/IJob.cs
--- a/Engine/Source/Programs/Horde/HordeServer/Jobs/IJob.cs
+++ b/Engine/Source/Programs/Horde/HordeServer/Jobs/IJob.cs
@@ -20,16 +20,16 @@ using EpicGames.Horde.Logs;
 using EpicGames.Horde.Streams;
 using EpicGames.Horde.Users;
 using Horde.Common.Rpc;
-using Horde.Server.Acls;
-using Horde.Server.Jobs.Graphs;
-using Horde.Server.Jobs.Timing;
-using Horde.Server.Server;
-using Horde.Server.Streams;
-using Horde.Server.Ugs;
+using HordeServer.Acls;
+using HordeServer.Jobs.Graphs;
+using HordeServer.Jobs.Timing;
+using HordeServer.Server;
+using HordeServer.Streams;
+using HordeServer.Ugs;
 using MongoDB.Bson;
 using MongoDB.Driver;
 
-namespace Horde.Server.Jobs
+namespace HordeServer.Jobs
 {
 	/// <summary>
 	/// Document describing a job
diff --git a/Engine/Source/Programs/Horde/HordeServer/Jobs/IJobCollection.cs b/Engine/Source/Programs/Horde/HordeServer/Jobs/IJobCollection.cs
--- a/Engine/Source/Programs/Horde/HordeServer/Jobs/IJobCollection.cs
+++ b/Engine/Source/Programs/Horde/HordeServer/Jobs/IJobCollection.cs
@@ -10,11 +10,11 @@ using EpicGames.Horde.Jobs.Bisect;
 using EpicGames.Horde.Jobs.Templates;
 using EpicGames.Horde.Streams;
 using EpicGames.Horde.Users;
-using Horde.Server.Acls;
-using Horde.Server.Jobs.Graphs;
-using Horde.Server.Streams;
+using HordeServer.Acls;
+using HordeServer.Jobs.Graphs;
+using HordeServer.Streams;
 
-namespace Horde.Server.Jobs
+namespace HordeServer.Jobs
 {
 	/// <summary>
 	/// Options for creating a new job
diff --git a/Engine/Source/Programs/Horde/HordeServer/Jobs/IJobStepRef.cs b/Engine/Source/Programs/Horde/HordeServer/Jobs/IJobStepRef.cs
--- a/Engine/Source/Programs/Horde/HordeServer/Jobs/IJobStepRef.cs
+++ b/Engine/Source/Programs/Horde/HordeServer/Jobs/IJobStepRef.cs
@@ -12,7 +12,7 @@ using EpicGames.Horde.Streams;
 using MongoDB.Bson.Serialization;
 using MongoDB.Bson.Serialization.Attributes;
 
-namespace Horde.Server.Jobs
+namespace HordeServer.Jobs
 {
 	/// <summary>
 	/// Unique id struct for JobStepRef objects. Includes a job id, batch id, and step id to uniquely identify the step.
diff --git a/Engine/Source/Programs/Horde/HordeServer/Jobs/IJobStepRefCollection.cs b/Engine/Source/Programs/Horde/HordeServer/Jobs/IJobStepRefCollection.cs
--- a/Engine/Source/Programs/Horde/HordeServer/Jobs/IJobStepRefCollection.cs
+++ b/Engine/Source/Programs/Horde/HordeServer/Jobs/IJobStepRefCollection.cs
@@ -10,10 +10,10 @@ using EpicGames.Horde.Jobs;
 using EpicGames.Horde.Jobs.Templates;
 using EpicGames.Horde.Logs;
 using EpicGames.Horde.Streams;
-using Horde.Server.Jobs.Graphs;
+using HordeServer.Jobs.Graphs;
 using Microsoft.Extensions.Logging;
 
-namespace Horde.Server.Jobs
+namespace HordeServer.Jobs
 {
 	/// <summary>
 	/// Interface for a collection of JobStepRef documents
diff --git a/Engine/Source/Programs/Horde/HordeServer/Jobs/JobAclAction.cs b/Engine/Source/Programs/Horde/HordeServer/Jobs/JobAclAction.cs
--- a/Engine/Source/Programs/Horde/HordeServer/Jobs/JobAclAction.cs
+++ b/Engine/Source/Programs/Horde/HordeServer/Jobs/JobAclAction.cs
@@ -2,7 +2,7 @@
 
 using EpicGames.Horde.Acls;
 
-namespace Horde.Server.Jobs
+namespace HordeServer.Jobs
 {
 	/// <summary>
 	/// ACL actions valid for manipulating jobs
diff --git a/Engine/Source/Programs/Horde/HordeServer/Jobs/JobCollection.cs b/Engine/Source/Programs/Horde/HordeServer/Jobs/JobCollection.cs
--- a/Engine/Source/Programs/Horde/HordeServer/Jobs/JobCollection.cs
+++ b/Engine/Source/Programs/Horde/HordeServer/Jobs/JobCollection.cs
@@ -21,13 +21,13 @@ using EpicGames.Horde.Jobs.Templates;
 using EpicGames.Horde.Logs;
 using EpicGames.Horde.Streams;
 using EpicGames.Horde.Users;
-using Horde.Server.Acls;
-using Horde.Server.Jobs.Graphs;
-using Horde.Server.Logs;
-using Horde.Server.Server;
-using Horde.Server.Streams;
-using Horde.Server.Telemetry;
-using Horde.Server.Utilities;
+using HordeServer.Acls;
+using HordeServer.Jobs.Graphs;
+using HordeServer.Logs;
+using HordeServer.Server;
+using HordeServer.Streams;
+using HordeServer.Telemetry;
+using HordeServer.Utilities;
 using Microsoft.Extensions.Logging;
 using Microsoft.Extensions.Options;
 using MongoDB.Bson;
@@ -38,7 +38,7 @@ using MongoDB.Bson.Serialization.Options;
 using MongoDB.Driver;
 using OpenTelemetry.Trace;
 
-namespace Horde.Server.Jobs
+namespace HordeServer.Jobs
 {
 	/// <summary>
 	/// Wrapper around the jobs collection in a mongo DB
diff --git a/Engine/Source/Programs/Horde/HordeServer/Jobs/JobExpirationService.cs b/Engine/Source/Programs/Horde/HordeServer/Jobs/JobExpirationService.cs
--- a/Engine/Source/Programs/Horde/HordeServer/Jobs/JobExpirationService.cs
+++ b/Engine/Source/Programs/Horde/HordeServer/Jobs/JobExpirationService.cs
@@ -5,13 +5,13 @@ using System.Collections.Generic;
 using System.Threading;
 using System.Threading.Tasks;
 using EpicGames.Core;
-using Horde.Server.Server;
-using Horde.Server.Streams;
+using HordeServer.Server;
+using HordeServer.Streams;
 using Microsoft.Extensions.Hosting;
 using Microsoft.Extensions.Logging;
 using Microsoft.Extensions.Options;
 
-namespace Horde.Server.Jobs
+namespace HordeServer.Jobs
 {
 	/// <summary>
 	/// Expires jobs a certain amount of time after they have run
diff --git a/Engine/Source/Programs/Horde/HordeServer/Jobs/JobIdUtils.cs b/Engine/Source/Programs/Horde/HordeServer/Jobs/JobIdUtils.cs
--- a/Engine/Source/Programs/Horde/HordeServer/Jobs/JobIdUtils.cs
+++ b/Engine/Source/Programs/Horde/HordeServer/Jobs/JobIdUtils.cs
@@ -1,9 +1,9 @@
 // Copyright Epic Games, Inc. All Rights Reserved.
 
 using EpicGames.Horde.Jobs;
-using Horde.Server.Utilities;
+using HordeServer.Utilities;
 
-namespace Horde.Server.Jobs
+namespace HordeServer.Jobs
 {
 	/// <summary>
 	/// Utility methods for <see cref="JobId"/>
diff --git a/Engine/Source/Programs/Horde/HordeServer/Jobs/JobRpcService.cs b/Engine/Source/Programs/Horde/HordeServer/Jobs/JobRpcService.cs
--- a/Engine/Source/Programs/Horde/HordeServer/Jobs/JobRpcService.cs
+++ b/Engine/Source/Programs/Horde/HordeServer/Jobs/JobRpcService.cs
@@ -21,26 +21,26 @@ using EpicGames.Horde.Logs;
 using Google.Protobuf.WellKnownTypes;
 using Grpc.Core;
 using Horde.Common.Rpc;
-using Horde.Server.Acls;
-using Horde.Server.Agents;
-using Horde.Server.Agents.Pools;
-using Horde.Server.Artifacts;
-using Horde.Server.Jobs.Graphs;
-using Horde.Server.Jobs.Templates;
-using Horde.Server.Jobs.TestData;
-using Horde.Server.Logs;
-using Horde.Server.Perforce;
-using Horde.Server.Server;
-using Horde.Server.Streams;
-using Horde.Server.Tasks;
-using Horde.Server.Utilities;
+using HordeServer.Acls;
+using HordeServer.Agents;
+using HordeServer.Agents.Pools;
+using HordeServer.Artifacts;
+using HordeServer.Jobs.Graphs;
+using HordeServer.Jobs.Templates;
+using HordeServer.Jobs.TestData;
+using HordeServer.Logs;
+using HordeServer.Perforce;
+using HordeServer.Server;
+using HordeServer.Streams;
+using HordeServer.Tasks;
+using HordeServer.Utilities;
 using Microsoft.AspNetCore.Authorization;
 using Microsoft.Extensions.Logging;
 using Microsoft.Extensions.Options;
 using MongoDB.Bson;
 using MongoDB.Bson.Serialization;
 
-namespace Horde.Server.Jobs
+namespace HordeServer.Jobs
 {
 	/// <summary>
 	/// Implements the Horde gRPC service for bots updating their status and dequeing work
diff --git a/Engine/Source/Programs/Horde/HordeServer/Jobs/JobService.cs b/Engine/Source/Programs/Horde/HordeServer/Jobs/JobService.cs
--- a/Engine/Source/Programs/Horde/HordeServer/Jobs/JobService.cs
+++ b/Engine/Source/Programs/Horde/HordeServer/Jobs/JobService.cs
@@ -19,24 +19,24 @@ using EpicGames.Horde.Jobs.Templates;
 using EpicGames.Horde.Logs;
 using EpicGames.Horde.Streams;
 using EpicGames.Horde.Users;
-using Horde.Server.Agents;
-using Horde.Server.Commits;
-using Horde.Server.Issues;
-using Horde.Server.Jobs.Bisect;
-using Horde.Server.Jobs.Graphs;
-using Horde.Server.Jobs.Templates;
-using Horde.Server.Jobs.Timing;
-using Horde.Server.Notifications;
-using Horde.Server.Perforce;
-using Horde.Server.Server;
-using Horde.Server.Streams;
-using Horde.Server.Users;
-using Horde.Server.Utilities;
+using HordeServer.Agents;
+using HordeServer.Commits;
+using HordeServer.Issues;
+using HordeServer.Jobs.Bisect;
+using HordeServer.Jobs.Graphs;
+using HordeServer.Jobs.Templates;
+using HordeServer.Jobs.Timing;
+using HordeServer.Notifications;
+using HordeServer.Perforce;
+using HordeServer.Server;
+using HordeServer.Streams;
+using HordeServer.Users;
+using HordeServer.Utilities;
 using Microsoft.Extensions.Logging;
 using MongoDB.Bson;
 using OpenTelemetry.Trace;
 
-namespace Horde.Server.Jobs
+namespace HordeServer.Jobs
 {
 	/// <summary>
 	/// Exception thrown when attempting to retry executing a node that does not allow retries
diff --git a/Engine/Source/Programs/Horde/HordeServer/Jobs/JobStepRefCollection.cs b/Engine/Source/Programs/Horde/HordeServer/Jobs/JobStepRefCollection.cs
--- a/Engine/Source/Programs/Horde/HordeServer/Jobs/JobStepRefCollection.cs
+++ b/Engine/Source/Programs/Horde/HordeServer/Jobs/JobStepRefCollection.cs
@@ -11,15 +11,15 @@ using EpicGames.Horde.Jobs;
 using EpicGames.Horde.Jobs.Templates;
 using EpicGames.Horde.Logs;
 using EpicGames.Horde.Streams;
-using Horde.Server.Server;
-using Horde.Server.Streams;
-using Horde.Server.Telemetry;
-using Horde.Server.Utilities;
+using HordeServer.Server;
+using HordeServer.Streams;
+using HordeServer.Telemetry;
+using HordeServer.Utilities;
 using Microsoft.Extensions.Options;
 using MongoDB.Bson.Serialization.Attributes;
 using MongoDB.Driver;
 
-namespace Horde.Server.Jobs
+namespace HordeServer.Jobs
 {
 	/// <summary>
 	/// Collection of JobStepRef documents
diff --git a/Engine/Source/Programs/Horde/HordeServer/Jobs/JobTaskSource.cs b/Engine/Source/Programs/Horde/HordeServer/Jobs/JobTaskSource.cs
--- a/Engine/Source/Programs/Horde/HordeServer/Jobs/JobTaskSource.cs
+++ b/Engine/Source/Programs/Horde/HordeServer/Jobs/JobTaskSource.cs
@@ -20,19 +20,19 @@ using EpicGames.Horde.Storage;
 using EpicGames.Horde.Streams;
 using Google.Protobuf;
 using Google.Protobuf.WellKnownTypes;
-using Horde.Server.Acls;
-using Horde.Server.Agents;
-using Horde.Server.Agents.Pools;
-using Horde.Server.Jobs.Bisect;
-using Horde.Server.Jobs.Graphs;
-using Horde.Server.Logs;
-using Horde.Server.Perforce;
-using Horde.Server.Server;
-using Horde.Server.Storage;
-using Horde.Server.Streams;
-using Horde.Server.Tasks;
-using Horde.Server.Ugs;
-using Horde.Server.Utilities;
+using HordeServer.Acls;
+using HordeServer.Agents;
+using HordeServer.Agents.Pools;
+using HordeServer.Jobs.Bisect;
+using HordeServer.Jobs.Graphs;
+using HordeServer.Logs;
+using HordeServer.Perforce;
+using HordeServer.Server;
+using HordeServer.Storage;
+using HordeServer.Streams;
+using HordeServer.Tasks;
+using HordeServer.Ugs;
+using HordeServer.Utilities;
 using HordeCommon;
 using HordeCommon.Rpc.Messages;
 using HordeCommon.Rpc.Tasks;
@@ -41,7 +41,7 @@ using Microsoft.Extensions.Logging;
 using Microsoft.Extensions.Options;
 using OpenTelemetry.Trace;
 
-namespace Horde.Server.Jobs
+namespace HordeServer.Jobs
 {
 	/// <summary>
 	/// Background service to dispatch pending work to agents in priority order.
diff --git a/Engine/Source/Programs/Horde/HordeServer/Jobs/JobsController.cs b/Engine/Source/Programs/Horde/HordeServer/Jobs/JobsController.cs
--- a/Engine/Source/Programs/Horde/HordeServer/Jobs/JobsController.cs
+++ b/Engine/Source/Programs/Horde/HordeServer/Jobs/JobsController.cs
@@ -13,26 +13,26 @@ using EpicGames.Horde.Jobs.Templates;
 using EpicGames.Horde.Logs;
 using EpicGames.Horde.Streams;
 using EpicGames.Horde.Users;
-using Horde.Server.Acls;
-using Horde.Server.Agents;
-using Horde.Server.Artifacts;
-using Horde.Server.Commits;
-using Horde.Server.Jobs.Artifacts;
-using Horde.Server.Jobs.Graphs;
-using Horde.Server.Jobs.Templates;
-using Horde.Server.Jobs.Timing;
-using Horde.Server.Notifications;
-using Horde.Server.Perforce;
-using Horde.Server.Server;
-using Horde.Server.Streams;
-using Horde.Server.Users;
-using Horde.Server.Utilities;
+using HordeServer.Acls;
+using HordeServer.Agents;
+using HordeServer.Artifacts;
+using HordeServer.Commits;
+using HordeServer.Jobs.Artifacts;
+using HordeServer.Jobs.Graphs;
+using HordeServer.Jobs.Templates;
+using HordeServer.Jobs.Timing;
+using HordeServer.Notifications;
+using HordeServer.Perforce;
+using HordeServer.Server;
+using HordeServer.Streams;
+using HordeServer.Users;
+using HordeServer.Utilities;
 using Microsoft.AspNetCore.Authorization;
 using Microsoft.AspNetCore.Mvc;
 using Microsoft.Extensions.Options;
 using MongoDB.Bson;
 
-namespace Horde.Server.Jobs
+namespace HordeServer.Jobs
 {
 	/// <summary>
 	/// Controller for the /api/v1/jobs endpoing
diff --git a/Engine/Source/Programs/Horde/HordeServer/Jobs/Schedules/ScheduleService.cs b/Engine/Source/Programs/Horde/HordeServer/Jobs/Schedules/ScheduleService.cs
--- a/Engine/Source/Programs/Horde/HordeServer/Jobs/Schedules/ScheduleService.cs
+++ b/Engine/Source/Programs/Horde/HordeServer/Jobs/Schedules/ScheduleService.cs
@@ -15,12 +15,12 @@ using EpicGames.Horde.Streams;
 using EpicGames.Redis;
 using EpicGames.Redis.Utility;
 using EpicGames.Serialization;
-using Horde.Server.Commits;
-using Horde.Server.Jobs.Graphs;
-using Horde.Server.Jobs.Templates;
-using Horde.Server.Server;
-using Horde.Server.Streams;
-using Horde.Server.Utilities;
+using HordeServer.Commits;
+using HordeServer.Jobs.Graphs;
+using HordeServer.Jobs.Templates;
+using HordeServer.Server;
+using HordeServer.Streams;
+using HordeServer.Utilities;
 using HordeCommon;
 using Microsoft.Extensions.Hosting;
 using Microsoft.Extensions.Logging;
@@ -28,7 +28,7 @@ using Microsoft.Extensions.Options;
 using OpenTelemetry.Trace;
 using StackExchange.Redis;
 
-namespace Horde.Server.Jobs.Schedules
+namespace HordeServer.Jobs.Schedules
 {
 	/// <summary>
 	/// Manipulates schedule instances
diff --git a/Engine/Source/Programs/Horde/HordeServer/Jobs/Schedules/Schedules.cs b/Engine/Source/Programs/Horde/HordeServer/Jobs/Schedules/Schedules.cs
--- a/Engine/Source/Programs/Horde/HordeServer/Jobs/Schedules/Schedules.cs
+++ b/Engine/Source/Programs/Horde/HordeServer/Jobs/Schedules/Schedules.cs
@@ -4,9 +4,9 @@ using System;
 using System.Collections.Generic;
 using EpicGames.Horde.Commits;
 using EpicGames.Horde.Jobs;
-using Horde.Server.Streams;
+using HordeServer.Streams;
 
-namespace Horde.Server.Jobs.Schedules
+namespace HordeServer.Jobs.Schedules
 {
 	/// <summary>
 	/// Response describing a schedule
diff --git a/Engine/Source/Programs/Horde/HordeServer/Jobs/Templates/ITemplate.cs b/Engine/Source/Programs/Horde/HordeServer/Jobs/Templates/ITemplate.cs
--- a/Engine/Source/Programs/Horde/HordeServer/Jobs/Templates/ITemplate.cs
+++ b/Engine/Source/Programs/Horde/HordeServer/Jobs/Templates/ITemplate.cs
@@ -7,7 +7,7 @@ using EpicGames.Horde.Jobs;
 using EpicGames.Horde.Jobs.Templates;
 using MongoDB.Bson.Serialization.Attributes;
 
-namespace Horde.Server.Jobs.Templates
+namespace HordeServer.Jobs.Templates
 {
 	/// <summary>
 	/// Base class for parameters used to configure templates via the new build dialog
diff --git a/Engine/Source/Programs/Horde/HordeServer/Jobs/Templates/ITemplateCollection.cs b/Engine/Source/Programs/Horde/HordeServer/Jobs/Templates/ITemplateCollection.cs
--- a/Engine/Source/Programs/Horde/HordeServer/Jobs/Templates/ITemplateCollection.cs
+++ b/Engine/Source/Programs/Horde/HordeServer/Jobs/Templates/ITemplateCollection.cs
@@ -3,7 +3,7 @@
 using System.Threading.Tasks;
 using EpicGames.Core;
 
-namespace Horde.Server.Jobs.Templates
+namespace HordeServer.Jobs.Templates
 {
 	/// <summary>
 	/// Interface for a collection of template documents
diff --git a/Engine/Source/Programs/Horde/HordeServer/Jobs/Templates/TemplateCollection.cs b/Engine/Source/Programs/Horde/HordeServer/Jobs/Templates/TemplateCollection.cs
--- a/Engine/Source/Programs/Horde/HordeServer/Jobs/Templates/TemplateCollection.cs
+++ b/Engine/Source/Programs/Horde/HordeServer/Jobs/Templates/TemplateCollection.cs
@@ -6,13 +6,13 @@ using System.Linq;
 using System.Threading.Tasks;
 using EpicGames.Core;
 using EpicGames.Horde.Jobs;
-using Horde.Server.Server;
+using HordeServer.Server;
 using Microsoft.Extensions.Caching.Memory;
 using MongoDB.Bson;
 using MongoDB.Bson.Serialization.Attributes;
 using MongoDB.Driver;
 
-namespace Horde.Server.Jobs.Templates
+namespace HordeServer.Jobs.Templates
 {
 	/// <summary>
 	/// Collection of template documents
diff --git a/Engine/Source/Programs/Horde/HordeServer/Jobs/Templates/TemplateConfig.cs b/Engine/Source/Programs/Horde/HordeServer/Jobs/Templates/TemplateConfig.cs
--- a/Engine/Source/Programs/Horde/HordeServer/Jobs/Templates/TemplateConfig.cs
+++ b/Engine/Source/Programs/Horde/HordeServer/Jobs/Templates/TemplateConfig.cs
@@ -9,9 +9,9 @@ using EpicGames.Core;
 using EpicGames.Horde;
 using EpicGames.Horde.Jobs;
 using EpicGames.Horde.Jobs.Templates;
-using Horde.Server.Configuration;
+using HordeServer.Configuration;
 
-namespace Horde.Server.Jobs.Templates
+namespace HordeServer.Jobs.Templates
 {
 	/// <summary>
 	/// Parameters to create a new template
diff --git a/Engine/Source/Programs/Horde/HordeServer/Jobs/TestData/ITestData.cs b/Engine/Source/Programs/Horde/HordeServer/Jobs/TestData/ITestData.cs
--- a/Engine/Source/Programs/Horde/HordeServer/Jobs/TestData/ITestData.cs
+++ b/Engine/Source/Programs/Horde/HordeServer/Jobs/TestData/ITestData.cs
@@ -8,7 +8,7 @@ using EpicGames.Horde.Jobs.TestData;
 using EpicGames.Horde.Streams;
 using MongoDB.Bson;
 
-namespace Horde.Server.Jobs.TestData
+namespace HordeServer.Jobs.TestData
 {
 	/// <summary>
 	/// Defines a testing environment based on platforms, configurations, targets, etc.
diff --git a/Engine/Source/Programs/Horde/HordeServer/Jobs/TestData/ITestDataCollection.cs b/Engine/Source/Programs/Horde/HordeServer/Jobs/TestData/ITestDataCollection.cs
--- a/Engine/Source/Programs/Horde/HordeServer/Jobs/TestData/ITestDataCollection.cs
+++ b/Engine/Source/Programs/Horde/HordeServer/Jobs/TestData/ITestDataCollection.cs
@@ -8,7 +8,7 @@ using EpicGames.Horde.Jobs;
 using EpicGames.Horde.Streams;
 using MongoDB.Bson;
 
-namespace Horde.Server.Jobs.TestData
+namespace HordeServer.Jobs.TestData
 {
 	/// <summary>
 	/// Collection of test data documents
diff --git a/Engine/Source/Programs/Horde/HordeServer/Jobs/TestData/TestDataCollection.cs b/Engine/Source/Programs/Horde/HordeServer/Jobs/TestData/TestDataCollection.cs
--- a/Engine/Source/Programs/Horde/HordeServer/Jobs/TestData/TestDataCollection.cs
+++ b/Engine/Source/Programs/Horde/HordeServer/Jobs/TestData/TestDataCollection.cs
@@ -9,8 +9,8 @@ using EpicGames.Horde.Jobs;
 using EpicGames.Horde.Jobs.Templates;
 using EpicGames.Horde.Jobs.TestData;
 using EpicGames.Horde.Streams;
-using Horde.Server.Server;
-using Horde.Server.Utilities;
+using HordeServer.Server;
+using HordeServer.Utilities;
 using Microsoft.Extensions.Logging;
 using MongoDB.Bson;
 using MongoDB.Bson.Serialization;
@@ -18,7 +18,7 @@ using MongoDB.Bson.Serialization.Attributes;
 using MongoDB.Driver;
 using OpenTelemetry.Trace;
 
-namespace Horde.Server.Jobs.TestData
+namespace HordeServer.Jobs.TestData
 {
 	/// <summary>
 	/// Collection of test data documents
diff --git a/Engine/Source/Programs/Horde/HordeServer/Jobs/TestData/TestDataController.cs b/Engine/Source/Programs/Horde/HordeServer/Jobs/TestData/TestDataController.cs
--- a/Engine/Source/Programs/Horde/HordeServer/Jobs/TestData/TestDataController.cs
+++ b/Engine/Source/Programs/Horde/HordeServer/Jobs/TestData/TestDataController.cs
@@ -9,16 +9,16 @@ using EpicGames.Core;
 using EpicGames.Horde.Jobs;
 using EpicGames.Horde.Jobs.TestData;
 using EpicGames.Horde.Streams;
-using Horde.Server.Server;
-using Horde.Server.Streams;
-using Horde.Server.Utilities;
+using HordeServer.Server;
+using HordeServer.Streams;
+using HordeServer.Utilities;
 using Microsoft.AspNetCore.Authorization;
 using Microsoft.AspNetCore.Mvc;
 using Microsoft.Extensions.Options;
 using MongoDB.Bson;
 using MongoDB.Bson.Serialization;
 
-namespace Horde.Server.Jobs.TestData
+namespace HordeServer.Jobs.TestData
 {
 	/// <summary>
 	/// Controller for the /api/v1/testdata endpoint
diff --git a/Engine/Source/Programs/Horde/HordeServer/Jobs/TestData/TestId.cs b/Engine/Source/Programs/Horde/HordeServer/Jobs/TestData/TestId.cs
--- a/Engine/Source/Programs/Horde/HordeServer/Jobs/TestData/TestId.cs
+++ b/Engine/Source/Programs/Horde/HordeServer/Jobs/TestData/TestId.cs
@@ -2,10 +2,10 @@
 
 using System.ComponentModel;
 using EpicGames.Core;
-using Horde.Server.Utilities;
+using HordeServer.Utilities;
 using MongoDB.Bson;
 
-namespace Horde.Server.Jobs.TestData
+namespace HordeServer.Jobs.TestData
 {
 	/// <summary>
 	/// Identifier for a session
diff --git a/Engine/Source/Programs/Horde/HordeServer/Jobs/TestData/TestMetaId.cs b/Engine/Source/Programs/Horde/HordeServer/Jobs/TestData/TestMetaId.cs
--- a/Engine/Source/Programs/Horde/HordeServer/Jobs/TestData/TestMetaId.cs
+++ b/Engine/Source/Programs/Horde/HordeServer/Jobs/TestData/TestMetaId.cs
@@ -2,10 +2,10 @@
 
 using System.ComponentModel;
 using EpicGames.Core;
-using Horde.Server.Utilities;
+using HordeServer.Utilities;
 using MongoDB.Bson;
 
-namespace Horde.Server.Jobs.TestData
+namespace HordeServer.Jobs.TestData
 {
 	/// <summary>
 	/// Identifier for a session
diff --git a/Engine/Source/Programs/Horde/HordeServer/Jobs/Timing/IJobTiming.cs b/Engine/Source/Programs/Horde/HordeServer/Jobs/Timing/IJobTiming.cs
--- a/Engine/Source/Programs/Horde/HordeServer/Jobs/Timing/IJobTiming.cs
+++ b/Engine/Source/Programs/Horde/HordeServer/Jobs/Timing/IJobTiming.cs
@@ -2,7 +2,7 @@
 
 using System.Diagnostics.CodeAnalysis;
 
-namespace Horde.Server.Jobs.Timing
+namespace HordeServer.Jobs.Timing
 {
 	/// <summary>
 	/// Information about the timing for an individual step
diff --git a/Engine/Source/Programs/Horde/HordeServer/Jobs/Timing/IJobTimingCollection.cs b/Engine/Source/Programs/Horde/HordeServer/Jobs/Timing/IJobTimingCollection.cs
--- a/Engine/Source/Programs/Horde/HordeServer/Jobs/Timing/IJobTimingCollection.cs
+++ b/Engine/Source/Programs/Horde/HordeServer/Jobs/Timing/IJobTimingCollection.cs
@@ -5,7 +5,7 @@ using System.Threading.Tasks;
 using EpicGames.Horde.Jobs;
 using MongoDB.Bson.Serialization.Attributes;
 
-namespace Horde.Server.Jobs.Timing
+namespace HordeServer.Jobs.Timing
 {
 	/// <summary>
 	/// Average timing information for a node
diff --git a/Engine/Source/Programs/Horde/HordeServer/Jobs/Timing/JobTimingCollection.cs b/Engine/Source/Programs/Horde/HordeServer/Jobs/Timing/JobTimingCollection.cs
--- a/Engine/Source/Programs/Horde/HordeServer/Jobs/Timing/JobTimingCollection.cs
+++ b/Engine/Source/Programs/Horde/HordeServer/Jobs/Timing/JobTimingCollection.cs
@@ -5,12 +5,12 @@ using System.Diagnostics.CodeAnalysis;
 using System.Linq;
 using System.Threading.Tasks;
 using EpicGames.Horde.Jobs;
-using Horde.Server.Server;
-using Horde.Server.Utilities;
+using HordeServer.Server;
+using HordeServer.Utilities;
 using MongoDB.Bson.Serialization.Attributes;
 using MongoDB.Driver;
 
-namespace Horde.Server.Jobs.Timing
+namespace HordeServer.Jobs.Timing
 {
 	/// <summary>
 	/// Concrete implementation of IJobTimingCollection
diff --git a/Engine/Source/Programs/Horde/HordeServer/Logs/Data/LogChunkData.cs b/Engine/Source/Programs/Horde/HordeServer/Logs/Data/LogChunkData.cs
--- a/Engine/Source/Programs/Horde/HordeServer/Logs/Data/LogChunkData.cs
+++ b/Engine/Source/Programs/Horde/HordeServer/Logs/Data/LogChunkData.cs
@@ -7,7 +7,7 @@ using EpicGames.Core;
 using EpicGames.Horde.Logs;
 using Microsoft.Extensions.Logging;
 
-namespace Horde.Server.Logs.Data
+namespace HordeServer.Logs.Data
 {
 	/// <summary>
 	/// Data for a log chunk
diff --git a/Engine/Source/Programs/Horde/HordeServer/Logs/Data/LogIndexData.cs b/Engine/Source/Programs/Horde/HordeServer/Logs/Data/LogIndexData.cs
--- a/Engine/Source/Programs/Horde/HordeServer/Logs/Data/LogIndexData.cs
+++ b/Engine/Source/Programs/Horde/HordeServer/Logs/Data/LogIndexData.cs
@@ -6,10 +6,10 @@ using System.IO;
 using System.Linq;
 using EpicGames.Core;
 using EpicGames.Horde.Logs;
-using Horde.Server.Utilities;
+using HordeServer.Utilities;
 using OpenTelemetry.Trace;
 
-namespace Horde.Server.Logs.Data
+namespace HordeServer.Logs.Data
 {
 	/// <summary>
 	/// Contains a source block of text to be indexed
diff --git a/Engine/Source/Programs/Horde/HordeServer/Logs/Data/LogSubChunkData.cs b/Engine/Source/Programs/Horde/HordeServer/Logs/Data/LogSubChunkData.cs
--- a/Engine/Source/Programs/Horde/HordeServer/Logs/Data/LogSubChunkData.cs
+++ b/Engine/Source/Programs/Horde/HordeServer/Logs/Data/LogSubChunkData.cs
@@ -3,10 +3,10 @@
 using System;
 using EpicGames.Core;
 using EpicGames.Horde.Logs;
-using Horde.Server.Utilities;
+using HordeServer.Utilities;
 using Microsoft.Extensions.Logging;
 
-namespace Horde.Server.Logs.Data
+namespace HordeServer.Logs.Data
 {
 	/// <summary>
 	/// Pending data for a sub-chunk
diff --git a/Engine/Source/Programs/Horde/HordeServer/Logs/Data/LogText.cs b/Engine/Source/Programs/Horde/HordeServer/Logs/Data/LogText.cs
--- a/Engine/Source/Programs/Horde/HordeServer/Logs/Data/LogText.cs
+++ b/Engine/Source/Programs/Horde/HordeServer/Logs/Data/LogText.cs
@@ -9,7 +9,7 @@ using System.Text.Json;
 using EpicGames.Core;
 using Microsoft.Extensions.Logging;
 
-namespace Horde.Server.Logs.Data
+namespace HordeServer.Logs.Data
 {
 	/// <summary>
 	/// Exception thrown trying to parse log output
diff --git a/Engine/Source/Programs/Horde/HordeServer/Logs/ILog.cs b/Engine/Source/Programs/Horde/HordeServer/Logs/ILog.cs
--- a/Engine/Source/Programs/Horde/HordeServer/Logs/ILog.cs
+++ b/Engine/Source/Programs/Horde/HordeServer/Logs/ILog.cs
@@ -13,11 +13,11 @@ using EpicGames.Horde.Agents.Sessions;
 using EpicGames.Horde.Jobs;
 using EpicGames.Horde.Logs;
 using EpicGames.Horde.Storage;
-using Horde.Server.Acls;
+using HordeServer.Acls;
 using Microsoft.Extensions.Logging;
 using MongoDB.Bson;
 
-namespace Horde.Server.Logs
+namespace HordeServer.Logs
 {
 	/// <summary>
 	/// Information about a log file
diff --git a/Engine/Source/Programs/Horde/HordeServer/Logs/ILogCollection.cs b/Engine/Source/Programs/Horde/HordeServer/Logs/ILogCollection.cs
--- a/Engine/Source/Programs/Horde/HordeServer/Logs/ILogCollection.cs
+++ b/Engine/Source/Programs/Horde/HordeServer/Logs/ILogCollection.cs
@@ -9,7 +9,7 @@ using EpicGames.Horde.Jobs;
 using EpicGames.Horde.Logs;
 using MongoDB.Bson;
 
-namespace Horde.Server.Logs
+namespace HordeServer.Logs
 {
 	/// <summary>
 	/// Wrapper around the jobs collection in a mongo DB
diff --git a/Engine/Source/Programs/Horde/HordeServer/Logs/ILogStorage.cs b/Engine/Source/Programs/Horde/HordeServer/Logs/ILogStorage.cs
--- a/Engine/Source/Programs/Horde/HordeServer/Logs/ILogStorage.cs
+++ b/Engine/Source/Programs/Horde/HordeServer/Logs/ILogStorage.cs
@@ -3,9 +3,9 @@
 using System;
 using System.Threading.Tasks;
 using EpicGames.Horde.Logs;
-using Horde.Server.Logs.Data;
+using HordeServer.Logs.Data;
 
-namespace Horde.Server.Logs
+namespace HordeServer.Logs
 {
 	/// <summary>
 	/// Legacy interface for accessing log data. This is not used for new log files; the agent writes log chunks directly to the storage client.
diff --git a/Engine/Source/Programs/Horde/HordeServer/Logs/LogAclAction.cs b/Engine/Source/Programs/Horde/HordeServer/Logs/LogAclAction.cs
--- a/Engine/Source/Programs/Horde/HordeServer/Logs/LogAclAction.cs
+++ b/Engine/Source/Programs/Horde/HordeServer/Logs/LogAclAction.cs
@@ -2,7 +2,7 @@
 
 using EpicGames.Horde.Acls;
 
-namespace Horde.Server.Logs
+namespace HordeServer.Logs
 {
 	/// <summary>
 	/// ACL actions for logs
diff --git a/Engine/Source/Programs/Horde/HordeServer/Logs/LogCollection.cs b/Engine/Source/Programs/Horde/HordeServer/Logs/LogCollection.cs
--- a/Engine/Source/Programs/Horde/HordeServer/Logs/LogCollection.cs
+++ b/Engine/Source/Programs/Horde/HordeServer/Logs/LogCollection.cs
@@ -16,10 +16,10 @@ using EpicGames.Horde.Agents.Sessions;
 using EpicGames.Horde.Jobs;
 using EpicGames.Horde.Logs;
 using EpicGames.Horde.Storage;
-using Horde.Server.Logs.Data;
-using Horde.Server.Server;
-using Horde.Server.Storage;
-using Horde.Server.Utilities;
+using HordeServer.Logs.Data;
+using HordeServer.Server;
+using HordeServer.Storage;
+using HordeServer.Utilities;
 using Microsoft.Extensions.Caching.Memory;
 using Microsoft.Extensions.Logging;
 using MongoDB.Bson;
@@ -27,7 +27,7 @@ using MongoDB.Bson.Serialization.Attributes;
 using MongoDB.Driver;
 using OpenTelemetry.Trace;
 
-namespace Horde.Server.Logs
+namespace HordeServer.Logs
 {
 	/// <summary>
 	/// Wrapper around the jobs collection in a mongo DB
diff --git a/Engine/Source/Programs/Horde/HordeServer/Logs/LogIdUtils.cs b/Engine/Source/Programs/Horde/HordeServer/Logs/LogIdUtils.cs
--- a/Engine/Source/Programs/Horde/HordeServer/Logs/LogIdUtils.cs
+++ b/Engine/Source/Programs/Horde/HordeServer/Logs/LogIdUtils.cs
@@ -1,9 +1,9 @@
 // Copyright Epic Games, Inc. All Rights Reserved.
 
 using EpicGames.Horde.Logs;
-using Horde.Server.Utilities;
+using HordeServer.Utilities;
 
-namespace Horde.Server.Logs
+namespace HordeServer.Logs
 {
 	/// <summary>
 	/// Server helper methods for <see cref="LogId"/>
diff --git a/Engine/Source/Programs/Horde/HordeServer/Logs/LogRpcService.cs b/Engine/Source/Programs/Horde/HordeServer/Logs/LogRpcService.cs
--- a/Engine/Source/Programs/Horde/HordeServer/Logs/LogRpcService.cs
+++ b/Engine/Source/Programs/Horde/HordeServer/Logs/LogRpcService.cs
@@ -11,14 +11,14 @@ using EpicGames.Horde.Storage;
 using Google.Protobuf.WellKnownTypes;
 using Grpc.Core;
 using Horde.Common.Rpc;
-using Horde.Server.Server;
-using Horde.Server.Storage;
-using Horde.Server.Utilities;
+using HordeServer.Server;
+using HordeServer.Storage;
+using HordeServer.Utilities;
 using Microsoft.AspNetCore.Authorization;
 using Microsoft.Extensions.Logging;
 using Microsoft.Extensions.Options;
 
-namespace Horde.Server.Logs
+namespace HordeServer.Logs
 {
 	/// <summary>
 	/// Implements the Horde gRPC service for bots updating their status and dequeing work
diff --git a/Engine/Source/Programs/Horde/HordeServer/Logs/LogTailService.cs b/Engine/Source/Programs/Horde/HordeServer/Logs/LogTailService.cs
--- a/Engine/Source/Programs/Horde/HordeServer/Logs/LogTailService.cs
+++ b/Engine/Source/Programs/Horde/HordeServer/Logs/LogTailService.cs
@@ -9,13 +9,13 @@ using System.Threading.Tasks;
 using EpicGames.Core;
 using EpicGames.Horde.Logs;
 using EpicGames.Redis;
-using Horde.Server.Server;
+using HordeServer.Server;
 using Microsoft.Extensions.Hosting;
 using Microsoft.Extensions.Logging;
 using Microsoft.Extensions.Options;
 using StackExchange.Redis;
 
-namespace Horde.Server.Logs
+namespace HordeServer.Logs
 {
 	/// <summary>
 	/// Stores data for logs which has not yet been flushed to persistent storage.
diff --git a/Engine/Source/Programs/Horde/HordeServer/Logs/LogsController.cs b/Engine/Source/Programs/Horde/HordeServer/Logs/LogsController.cs
--- a/Engine/Source/Programs/Horde/HordeServer/Logs/LogsController.cs
+++ b/Engine/Source/Programs/Horde/HordeServer/Logs/LogsController.cs
@@ -14,19 +14,19 @@ using EpicGames.Horde.Acls;
 using EpicGames.Horde.Jobs;
 using EpicGames.Horde.Logs;
 using EpicGames.Horde.Storage;
-using Horde.Server.Acls;
-using Horde.Server.Agents.Sessions;
-using Horde.Server.Issues;
-using Horde.Server.Jobs;
-using Horde.Server.Server;
-using Horde.Server.Storage;
-using Horde.Server.Utilities;
+using HordeServer.Acls;
+using HordeServer.Agents.Sessions;
+using HordeServer.Issues;
+using HordeServer.Jobs;
+using HordeServer.Server;
+using HordeServer.Storage;
+using HordeServer.Utilities;
 using Microsoft.AspNetCore.Authorization;
 using Microsoft.AspNetCore.Mvc;
 using Microsoft.Extensions.Options;
 using MongoDB.Bson;
 
-namespace Horde.Server.Logs
+namespace HordeServer.Logs
 {
 	/// <summary>
 	/// Format for the returned data
diff --git a/Engine/Source/Programs/Horde/HordeServer/Logs/Storage/LocalLogStorage.cs b/Engine/Source/Programs/Horde/HordeServer/Logs/Storage/LocalLogStorage.cs
--- a/Engine/Source/Programs/Horde/HordeServer/Logs/Storage/LocalLogStorage.cs
+++ b/Engine/Source/Programs/Horde/HordeServer/Logs/Storage/LocalLogStorage.cs
@@ -3,10 +3,10 @@
 using System;
 using System.Threading.Tasks;
 using EpicGames.Horde.Logs;
-using Horde.Server.Logs.Data;
+using HordeServer.Logs.Data;
 using Microsoft.Extensions.Caching.Memory;
 
-namespace Horde.Server.Logs.Storage
+namespace HordeServer.Logs.Storage
 {
 	/// <summary>
 	/// In-memory cache for chunk and index data
diff --git a/Engine/Source/Programs/Horde/HordeServer/Logs/Storage/NullLogStorage.cs b/Engine/Source/Programs/Horde/HordeServer/Logs/Storage/NullLogStorage.cs
--- a/Engine/Source/Programs/Horde/HordeServer/Logs/Storage/NullLogStorage.cs
+++ b/Engine/Source/Programs/Horde/HordeServer/Logs/Storage/NullLogStorage.cs
@@ -2,9 +2,9 @@
 
 using System.Threading.Tasks;
 using EpicGames.Horde.Logs;
-using Horde.Server.Logs.Data;
+using HordeServer.Logs.Data;
 
-namespace Horde.Server.Logs.Storage
+namespace HordeServer.Logs.Storage
 {
 	/// <summary>
 	/// Empty implementation of log storage
diff --git a/Engine/Source/Programs/Horde/HordeServer/Notifications/IEvent.cs b/Engine/Source/Programs/Horde/HordeServer/Notifications/IEvent.cs
--- a/Engine/Source/Programs/Horde/HordeServer/Notifications/IEvent.cs
+++ b/Engine/Source/Programs/Horde/HordeServer/Notifications/IEvent.cs
@@ -3,7 +3,7 @@
 using EpicGames.Horde.Jobs.Templates;
 using EpicGames.Horde.Streams;
 
-namespace Horde.Server.Notifications
+namespace HordeServer.Notifications
 {
 	/// <summary>
 	/// Base interface for events that users can subscribe for notifications to
diff --git a/Engine/Source/Programs/Horde/HordeServer/Notifications/INotificationService.cs b/Engine/Source/Programs/Horde/HordeServer/Notifications/INotificationService.cs
--- a/Engine/Source/Programs/Horde/HordeServer/Notifications/INotificationService.cs
+++ b/Engine/Source/Programs/Horde/HordeServer/Notifications/INotificationService.cs
@@ -7,17 +7,17 @@ using System.Security.Claims;
 using System.Threading;
 using System.Threading.Tasks;
 using EpicGames.Horde.Jobs;
-using Horde.Server.Agents;
-using Horde.Server.Configuration;
-using Horde.Server.Devices;
-using Horde.Server.Issues;
-using Horde.Server.Jobs;
-using Horde.Server.Jobs.Graphs;
-using Horde.Server.Streams;
-using Horde.Server.Users;
+using HordeServer.Agents;
+using HordeServer.Configuration;
+using HordeServer.Devices;
+using HordeServer.Issues;
+using HordeServer.Jobs;
+using HordeServer.Jobs.Graphs;
+using HordeServer.Streams;
+using HordeServer.Users;
 using MongoDB.Bson;
 
-namespace Horde.Server.Notifications
+namespace HordeServer.Notifications
 {
 	/// <summary>
 	/// Marker interface for (serializable) notifications
diff --git a/Engine/Source/Programs/Horde/HordeServer/Notifications/INotificationSink.cs b/Engine/Source/Programs/Horde/HordeServer/Notifications/INotificationSink.cs
--- a/Engine/Source/Programs/Horde/HordeServer/Notifications/INotificationSink.cs
+++ b/Engine/Source/Programs/Horde/HordeServer/Notifications/INotificationSink.cs
@@ -5,17 +5,17 @@ using System.Collections.Generic;
 using System.Threading;
 using System.Threading.Tasks;
 using EpicGames.Horde.Jobs;
-using Horde.Server.Agents;
-using Horde.Server.Configuration;
-using Horde.Server.Devices;
-using Horde.Server.Issues;
-using Horde.Server.Jobs;
-using Horde.Server.Jobs.Graphs;
-using Horde.Server.Logs;
-using Horde.Server.Streams;
-using Horde.Server.Users;
-
-namespace Horde.Server.Notifications
+using HordeServer.Agents;
+using HordeServer.Configuration;
+using HordeServer.Devices;
+using HordeServer.Issues;
+using HordeServer.Jobs;
+using HordeServer.Jobs.Graphs;
+using HordeServer.Logs;
+using HordeServer.Streams;
+using HordeServer.Users;
+
+namespace HordeServer.Notifications
 {
 	/// <summary>
 	/// Implements a notification method
diff --git a/Engine/Source/Programs/Horde/HordeServer/Notifications/INotificationTrigger.cs b/Engine/Source/Programs/Horde/HordeServer/Notifications/INotificationTrigger.cs
--- a/Engine/Source/Programs/Horde/HordeServer/Notifications/INotificationTrigger.cs
+++ b/Engine/Source/Programs/Horde/HordeServer/Notifications/INotificationTrigger.cs
@@ -4,7 +4,7 @@ using System.Collections.Generic;
 using EpicGames.Horde.Users;
 using MongoDB.Bson;
 
-namespace Horde.Server.Notifications
+namespace HordeServer.Notifications
 {
 	/// <summary>
 	/// An individual subscription
diff --git a/Engine/Source/Programs/Horde/HordeServer/Notifications/NotificationService.cs b/Engine/Source/Programs/Horde/HordeServer/Notifications/NotificationService.cs
--- a/Engine/Source/Programs/Horde/HordeServer/Notifications/NotificationService.cs
+++ b/Engine/Source/Programs/Horde/HordeServer/Notifications/NotificationService.cs
@@ -12,18 +12,18 @@ using EpicGames.Core;
 using EpicGames.Horde.Jobs;
 using EpicGames.Horde.Users;
 using EpicGames.Redis;
-using Horde.Server.Agents;
-using Horde.Server.Agents.Pools;
-using Horde.Server.Configuration;
-using Horde.Server.Devices;
-using Horde.Server.Issues;
-using Horde.Server.Jobs;
-using Horde.Server.Jobs.Graphs;
-using Horde.Server.Logs;
-using Horde.Server.Server;
-using Horde.Server.Streams;
-using Horde.Server.Users;
-using Horde.Server.Utilities;
+using HordeServer.Agents;
+using HordeServer.Agents.Pools;
+using HordeServer.Configuration;
+using HordeServer.Devices;
+using HordeServer.Issues;
+using HordeServer.Jobs;
+using HordeServer.Jobs.Graphs;
+using HordeServer.Logs;
+using HordeServer.Server;
+using HordeServer.Streams;
+using HordeServer.Users;
+using HordeServer.Utilities;
 using Microsoft.Extensions.Caching.Memory;
 using Microsoft.Extensions.Hosting;
 using Microsoft.Extensions.Logging;
@@ -32,7 +32,7 @@ using MongoDB.Bson;
 using StackExchange.Redis;
 using JsonSerializer = System.Text.Json.JsonSerializer;
 
-namespace Horde.Server.Notifications
+namespace HordeServer.Notifications
 {
 	/// <summary>
 	/// Wraps functionality for delivering notifications.
diff --git a/Engine/Source/Programs/Horde/HordeServer/Notifications/NotificationTriggerCollection.cs b/Engine/Source/Programs/Horde/HordeServer/Notifications/NotificationTriggerCollection.cs
--- a/Engine/Source/Programs/Horde/HordeServer/Notifications/NotificationTriggerCollection.cs
+++ b/Engine/Source/Programs/Horde/HordeServer/Notifications/NotificationTriggerCollection.cs
@@ -5,14 +5,14 @@ using System.Linq;
 using System.Threading;
 using System.Threading.Tasks;
 using EpicGames.Horde.Users;
-using Horde.Server.Server;
-using Horde.Server.Users;
-using Horde.Server.Utilities;
+using HordeServer.Server;
+using HordeServer.Users;
+using HordeServer.Utilities;
 using MongoDB.Bson;
 using MongoDB.Bson.Serialization.Attributes;
 using MongoDB.Driver;
 
-namespace Horde.Server.Notifications
+namespace HordeServer.Notifications
 {
 	/// <summary>
 	/// Collection of notification triggers
diff --git a/Engine/Source/Programs/Horde/HordeServer/Notifications/Sinks/SlackNotificationSink.cs b/Engine/Source/Programs/Horde/HordeServer/Notifications/Sinks/SlackNotificationSink.cs
--- a/Engine/Source/Programs/Horde/HordeServer/Notifications/Sinks/SlackNotificationSink.cs
+++ b/Engine/Source/Programs/Horde/HordeServer/Notifications/Sinks/SlackNotificationSink.cs
@@ -29,17 +29,17 @@ using EpicGames.Redis.Utility;
 using EpicGames.Slack;
 using EpicGames.Slack.Blocks;
 using EpicGames.Slack.Elements;
-using Horde.Server.Agents;
-using Horde.Server.Configuration;
-using Horde.Server.Devices;
-using Horde.Server.Issues;
-using Horde.Server.Issues.External;
-using Horde.Server.Jobs;
-using Horde.Server.Jobs.Graphs;
-using Horde.Server.Logs;
-using Horde.Server.Server;
-using Horde.Server.Streams;
-using Horde.Server.Users;
+using HordeServer.Agents;
+using HordeServer.Configuration;
+using HordeServer.Devices;
+using HordeServer.Issues;
+using HordeServer.Issues.External;
+using HordeServer.Jobs;
+using HordeServer.Jobs.Graphs;
+using HordeServer.Logs;
+using HordeServer.Server;
+using HordeServer.Streams;
+using HordeServer.Users;
 using Microsoft.AspNetCore.Hosting;
 using Microsoft.CodeAnalysis;
 using Microsoft.Extensions.Caching.Memory;
@@ -51,7 +51,7 @@ using MongoDB.Bson.Serialization.Attributes;
 using MongoDB.Driver;
 using StackExchange.Redis;
 
-namespace Horde.Server.Notifications.Sinks
+namespace HordeServer.Notifications.Sinks
 {
 	/// <summary>
 	/// Maintains a connection to Slack, in order to receive socket-mode notifications of user interactions
diff --git a/Engine/Source/Programs/Horde/HordeServer/Notifications/SubscriptionCollection.cs b/Engine/Source/Programs/Horde/HordeServer/Notifications/SubscriptionCollection.cs
--- a/Engine/Source/Programs/Horde/HordeServer/Notifications/SubscriptionCollection.cs
+++ b/Engine/Source/Programs/Horde/HordeServer/Notifications/SubscriptionCollection.cs
@@ -11,12 +11,12 @@ using EpicGames.Horde.Jobs;
 using EpicGames.Horde.Jobs.Templates;
 using EpicGames.Horde.Streams;
 using EpicGames.Horde.Users;
-using Horde.Server.Server;
-using Horde.Server.Utilities;
+using HordeServer.Server;
+using HordeServer.Utilities;
 using MongoDB.Bson.Serialization.Attributes;
 using MongoDB.Driver;
 
-namespace Horde.Server.Notifications
+namespace HordeServer.Notifications
 {
 	/// <summary>
 	/// Collection of subscription documents
diff --git a/Engine/Source/Programs/Horde/HordeServer/Notifications/SubscriptionsController.cs b/Engine/Source/Programs/Horde/HordeServer/Notifications/SubscriptionsController.cs
--- a/Engine/Source/Programs/Horde/HordeServer/Notifications/SubscriptionsController.cs
+++ b/Engine/Source/Programs/Horde/HordeServer/Notifications/SubscriptionsController.cs
@@ -3,13 +3,13 @@
 using System.Collections.Generic;
 using System.Threading.Tasks;
 using EpicGames.Horde.Users;
-using Horde.Server.Server;
-using Horde.Server.Utilities;
+using HordeServer.Server;
+using HordeServer.Utilities;
 using Microsoft.AspNetCore.Authorization;
 using Microsoft.AspNetCore.Mvc;
 using Microsoft.Extensions.Options;
 
-namespace Horde.Server.Notifications
+namespace HordeServer.Notifications
 {
 	/// <summary>
 	/// Controller for the /api/v1/agents endpoint
diff --git a/Engine/Source/Programs/Horde/HordeServer/Parameters/ParametersController.cs b/Engine/Source/Programs/Horde/HordeServer/Parameters/ParametersController.cs
--- a/Engine/Source/Programs/Horde/HordeServer/Parameters/ParametersController.cs
+++ b/Engine/Source/Programs/Horde/HordeServer/Parameters/ParametersController.cs
@@ -2,13 +2,13 @@
 
 using System;
 using System.Text.Json.Nodes;
-using Horde.Server.Server;
-using Horde.Server.Utilities;
+using HordeServer.Server;
+using HordeServer.Utilities;
 using Microsoft.AspNetCore.Authorization;
 using Microsoft.AspNetCore.Mvc;
 using Microsoft.Extensions.Options;
 
-namespace Horde.Server.Parameters
+namespace HordeServer.Parameters
 {
 	/// <summary>
 	/// Controller for the /api/v1/parameters endpoint. Provides configuration data to other tools.
diff --git a/Engine/Source/Programs/Horde/HordeServer/Perforce/IPerforceService.cs b/Engine/Source/Programs/Horde/HordeServer/Perforce/IPerforceService.cs
--- a/Engine/Source/Programs/Horde/HordeServer/Perforce/IPerforceService.cs
+++ b/Engine/Source/Programs/Horde/HordeServer/Perforce/IPerforceService.cs
@@ -7,11 +7,11 @@ using System.Threading;
 using System.Threading.Tasks;
 using EpicGames.Horde.Commits;
 using EpicGames.Perforce;
-using Horde.Server.Commits;
-using Horde.Server.Streams;
-using Horde.Server.Users;
+using HordeServer.Commits;
+using HordeServer.Streams;
+using HordeServer.Users;
 
-namespace Horde.Server.Perforce
+namespace HordeServer.Perforce
 {
 	/// <summary>
 	/// Result from checking a shelved change status
diff --git a/Engine/Source/Programs/Horde/HordeServer/Perforce/PerforceController.cs b/Engine/Source/Programs/Horde/HordeServer/Perforce/PerforceController.cs
--- a/Engine/Source/Programs/Horde/HordeServer/Perforce/PerforceController.cs
+++ b/Engine/Source/Programs/Horde/HordeServer/Perforce/PerforceController.cs
@@ -6,7 +6,7 @@ using System.Net;
 using System.Text.Json.Serialization;
 using System.Threading;
 using System.Threading.Tasks;
-using Horde.Server.Server;
+using HordeServer.Server;
 using Microsoft.AspNetCore.Authorization;
 using Microsoft.AspNetCore.Http;
 using Microsoft.AspNetCore.Mvc;
@@ -14,7 +14,7 @@ using Microsoft.Extensions.Logging;
 using Microsoft.Extensions.Options;
 using OpenTelemetry.Trace;
 
-namespace Horde.Server.Perforce
+namespace HordeServer.Perforce
 {
 	/// <summary>
 	/// Controller for Perforce triggers and callbacks
diff --git a/Engine/Source/Programs/Horde/HordeServer/Perforce/PerforceLoadBalancer.cs b/Engine/Source/Programs/Horde/HordeServer/Perforce/PerforceLoadBalancer.cs
--- a/Engine/Source/Programs/Horde/HordeServer/Perforce/PerforceLoadBalancer.cs
+++ b/Engine/Source/Programs/Horde/HordeServer/Perforce/PerforceLoadBalancer.cs
@@ -13,10 +13,10 @@ using System.Threading.Tasks;
 using EpicGames.Core;
 using EpicGames.Perforce;
 using Google.Protobuf.WellKnownTypes;
-using Horde.Server.Agents;
-using Horde.Server.Agents.Leases;
-using Horde.Server.Server;
-using Horde.Server.Utilities;
+using HordeServer.Agents;
+using HordeServer.Agents.Leases;
+using HordeServer.Server;
+using HordeServer.Utilities;
 using HordeCommon;
 using HordeCommon.Rpc.Messages;
 using HordeCommon.Rpc.Tasks;
@@ -28,7 +28,7 @@ using MongoDB.Bson.Serialization.Attributes;
 using OpenTelemetry.Trace;
 using StackExchange.Redis;
 
-namespace Horde.Server.Perforce
+namespace HordeServer.Perforce
 {
 	using Condition = EpicGames.Horde.Common.Condition;
 
diff --git a/Engine/Source/Programs/Horde/HordeServer/Perforce/PerforceService.cs b/Engine/Source/Programs/Horde/HordeServer/Perforce/PerforceService.cs
--- a/Engine/Source/Programs/Horde/HordeServer/Perforce/PerforceService.cs
+++ b/Engine/Source/Programs/Horde/HordeServer/Perforce/PerforceService.cs
@@ -18,17 +18,17 @@ using EpicGames.Horde.Commits;
 using EpicGames.Horde.Streams;
 using EpicGames.Horde.Users;
 using EpicGames.Perforce;
-using Horde.Server.Commits;
-using Horde.Server.Server;
-using Horde.Server.Streams;
-using Horde.Server.Users;
-using Horde.Server.Utilities;
+using HordeServer.Commits;
+using HordeServer.Server;
+using HordeServer.Streams;
+using HordeServer.Users;
+using HordeServer.Utilities;
 using Microsoft.Extensions.Caching.Memory;
 using Microsoft.Extensions.Logging;
 using Microsoft.Extensions.Options;
 using OpenTelemetry.Trace;
 
-namespace Horde.Server.Perforce
+namespace HordeServer.Perforce
 {
 	/// <summary>
 	/// P4API implementation of the Perforce service
diff --git a/Engine/Source/Programs/Horde/HordeServer/Perforce/PerforceServiceCache.cs b/Engine/Source/Programs/Horde/HordeServer/Perforce/PerforceServiceCache.cs
--- a/Engine/Source/Programs/Horde/HordeServer/Perforce/PerforceServiceCache.cs
+++ b/Engine/Source/Programs/Horde/HordeServer/Perforce/PerforceServiceCache.cs
@@ -14,11 +14,11 @@ using EpicGames.Horde.Streams;
 using EpicGames.Horde.Users;
 using EpicGames.Perforce;
 using EpicGames.Redis;
-using Horde.Server.Commits;
-using Horde.Server.Server;
-using Horde.Server.Streams;
-using Horde.Server.Users;
-using Horde.Server.Utilities;
+using HordeServer.Commits;
+using HordeServer.Server;
+using HordeServer.Streams;
+using HordeServer.Users;
+using HordeServer.Utilities;
 using Microsoft.Extensions.Hosting;
 using Microsoft.Extensions.Logging;
 using Microsoft.Extensions.Options;
@@ -29,7 +29,7 @@ using MongoDB.Driver;
 using OpenTelemetry.Trace;
 using StackExchange.Redis;
 
-namespace Horde.Server.Perforce
+namespace HordeServer.Perforce
 {
 	/// <summary>
 	/// Service which mirrors changes from Perforce
diff --git a/Engine/Source/Programs/Horde/HordeServer/Perforce/PerforceWorkspace.cs b/Engine/Source/Programs/Horde/HordeServer/Perforce/PerforceWorkspace.cs
--- a/Engine/Source/Programs/Horde/HordeServer/Perforce/PerforceWorkspace.cs
+++ b/Engine/Source/Programs/Horde/HordeServer/Perforce/PerforceWorkspace.cs
@@ -7,13 +7,13 @@ using System.Linq;
 using System.Threading;
 using System.Threading.Tasks;
 using EpicGames.Horde.Agents.Pools;
-using Horde.Server.Agents;
-using Horde.Server.Agents.Pools;
-using Horde.Server.Server;
-using Horde.Server.Streams;
+using HordeServer.Agents;
+using HordeServer.Agents.Pools;
+using HordeServer.Server;
+using HordeServer.Streams;
 using HordeCommon.Rpc.Messages;
 
-namespace Horde.Server.Perforce
+namespace HordeServer.Perforce
 {
 	/// <summary>
 	/// Helper methods for dealing with Perforce workspaces
diff --git a/Engine/Source/Programs/Horde/HordeServer/Perforce/ServerPerforceController.cs b/Engine/Source/Programs/Horde/HordeServer/Perforce/ServerPerforceController.cs
--- a/Engine/Source/Programs/Horde/HordeServer/Perforce/ServerPerforceController.cs
+++ b/Engine/Source/Programs/Horde/HordeServer/Perforce/ServerPerforceController.cs
@@ -8,12 +8,12 @@ using System.Threading.Tasks;
 using EpicGames.Core;
 using EpicGames.Horde.Server;
 using EpicGames.Perforce;
-using Horde.Server.Configuration;
-using Horde.Server.Utilities;
+using HordeServer.Configuration;
+using HordeServer.Utilities;
 using Microsoft.AspNetCore.Authorization;
 using Microsoft.AspNetCore.Mvc;
 
-namespace Horde.Server.Perforce
+namespace HordeServer.Perforce
 {
 	/// <summary>
 	/// Implements preflight of config changes with Perforce
diff --git a/Engine/Source/Programs/Horde/HordeServer/Projects/ProjectAclAction.cs b/Engine/Source/Programs/Horde/HordeServer/Projects/ProjectAclAction.cs
--- a/Engine/Source/Programs/Horde/HordeServer/Projects/ProjectAclAction.cs
+++ b/Engine/Source/Programs/Horde/HordeServer/Projects/ProjectAclAction.cs
@@ -2,7 +2,7 @@
 
 using EpicGames.Horde.Acls;
 
-namespace Horde.Server.Projects
+namespace HordeServer.Projects
 {
 	/// <summary>
 	/// ACL actions that can be performed on projects
diff --git a/Engine/Source/Programs/Horde/HordeServer/Projects/ProjectConfig.cs b/Engine/Source/Programs/Horde/HordeServer/Projects/ProjectConfig.cs
--- a/Engine/Source/Programs/Horde/HordeServer/Projects/ProjectConfig.cs
+++ b/Engine/Source/Programs/Horde/HordeServer/Projects/ProjectConfig.cs
@@ -11,13 +11,13 @@ using EpicGames.Horde.Acls;
 using EpicGames.Horde.Jobs;
 using EpicGames.Horde.Projects;
 using EpicGames.Horde.Telemetry;
-using Horde.Server.Acls;
-using Horde.Server.Agents.Pools;
-using Horde.Server.Configuration;
-using Horde.Server.Server;
-using Horde.Server.Streams;
+using HordeServer.Acls;
+using HordeServer.Agents.Pools;
+using HordeServer.Configuration;
+using HordeServer.Server;
+using HordeServer.Streams;
 
-namespace Horde.Server.Projects
+namespace HordeServer.Projects
 {
 	/// <summary>
 	/// Stores configuration for a project
diff --git a/Engine/Source/Programs/Horde/HordeServer/Projects/ProjectsController.cs b/Engine/Source/Programs/Horde/HordeServer/Projects/ProjectsController.cs
--- a/Engine/Source/Programs/Horde/HordeServer/Projects/ProjectsController.cs
+++ b/Engine/Source/Programs/Horde/HordeServer/Projects/ProjectsController.cs
@@ -4,16 +4,16 @@ using System.Collections.Generic;
 using System.Linq;
 using System.Text.RegularExpressions;
 using EpicGames.Horde.Projects;
-using Horde.Server.Configuration;
-using Horde.Server.Server;
-using Horde.Server.Streams;
-using Horde.Server.Utilities;
+using HordeServer.Configuration;
+using HordeServer.Server;
+using HordeServer.Streams;
+using HordeServer.Utilities;
 using Microsoft.AspNetCore.Authorization;
 using Microsoft.AspNetCore.Mvc;
 using Microsoft.AspNetCore.StaticFiles;
 using Microsoft.Extensions.Options;
 
-namespace Horde.Server.Projects
+namespace HordeServer.Projects
 {
 	/// <summary>
 	/// Controller for the /api/v1/projects endpoint
diff --git a/Engine/Source/Programs/Horde/HordeServer/Replicators/IReplicator.cs b/Engine/Source/Programs/Horde/HordeServer/Replicators/IReplicator.cs
--- a/Engine/Source/Programs/Horde/HordeServer/Replicators/IReplicator.cs
+++ b/Engine/Source/Programs/Horde/HordeServer/Replicators/IReplicator.cs
@@ -5,7 +5,7 @@ using System.Threading;
 using System.Threading.Tasks;
 using EpicGames.Horde.Replicators;
 
-namespace Horde.Server.Replicators
+namespace HordeServer.Replicators
 {
 	/// <summary>
 	/// Snapshot of the state for a replicator; a process which transfers data from an external source control provider into Horde.
diff --git a/Engine/Source/Programs/Horde/HordeServer/Replicators/IReplicatorCollection.cs b/Engine/Source/Programs/Horde/HordeServer/Replicators/IReplicatorCollection.cs
--- a/Engine/Source/Programs/Horde/HordeServer/Replicators/IReplicatorCollection.cs
+++ b/Engine/Source/Programs/Horde/HordeServer/Replicators/IReplicatorCollection.cs
@@ -5,7 +5,7 @@ using System.Threading;
 using System.Threading.Tasks;
 using EpicGames.Horde.Replicators;
 
-namespace Horde.Server.Replicators
+namespace HordeServer.Replicators
 {
 	/// <summary>
 	/// Interface for a collection of replicators
diff --git a/Engine/Source/Programs/Horde/HordeServer/Replicators/PerforceReplicator.cs b/Engine/Source/Programs/Horde/HordeServer/Replicators/PerforceReplicator.cs
--- a/Engine/Source/Programs/Horde/HordeServer/Replicators/PerforceReplicator.cs
+++ b/Engine/Source/Programs/Horde/HordeServer/Replicators/PerforceReplicator.cs
@@ -16,14 +16,14 @@ using EpicGames.Horde.Storage.Bundles;
 using EpicGames.Horde.Storage.Nodes;
 using EpicGames.Horde.Streams;
 using EpicGames.Perforce;
-using Horde.Server.Commits;
-using Horde.Server.Perforce;
-using Horde.Server.Server;
-using Horde.Server.Storage;
-using Horde.Server.Streams;
+using HordeServer.Commits;
+using HordeServer.Perforce;
+using HordeServer.Server;
+using HordeServer.Storage;
+using HordeServer.Streams;
 using Microsoft.Extensions.Logging;
 
-namespace Horde.Server.Replicators
+namespace HordeServer.Replicators
 {
 	/// <summary>
 	/// Options for replicating commits
diff --git a/Engine/Source/Programs/Horde/HordeServer/Replicators/ReplicationService.cs b/Engine/Source/Programs/Horde/HordeServer/Replicators/ReplicationService.cs
--- a/Engine/Source/Programs/Horde/HordeServer/Replicators/ReplicationService.cs
+++ b/Engine/Source/Programs/Horde/HordeServer/Replicators/ReplicationService.cs
@@ -7,15 +7,15 @@ using System.Threading;
 using System.Threading.Tasks;
 using EpicGames.Core;
 using EpicGames.Horde.Replicators;
-using Horde.Server.Server;
-using Horde.Server.Streams;
+using HordeServer.Server;
+using HordeServer.Streams;
 using Microsoft.Extensions.Hosting;
 using Microsoft.Extensions.Logging;
 using Microsoft.Extensions.Options;
 using Polly;
 using Polly.Retry;
 
-namespace Horde.Server.Replicators
+namespace HordeServer.Replicators
 {
 	/// <summary>
 	/// Exception triggered during content replication
diff --git a/Engine/Source/Programs/Horde/HordeServer/Replicators/ReplicatorCollection.cs b/Engine/Source/Programs/Horde/HordeServer/Replicators/ReplicatorCollection.cs
--- a/Engine/Source/Programs/Horde/HordeServer/Replicators/ReplicatorCollection.cs
+++ b/Engine/Source/Programs/Horde/HordeServer/Replicators/ReplicatorCollection.cs
@@ -6,13 +6,13 @@ using System.Threading;
 using System.Threading.Tasks;
 using EpicGames.Core;
 using EpicGames.Horde.Replicators;
-using Horde.Server.Server;
-using Horde.Server.Utilities;
+using HordeServer.Server;
+using HordeServer.Utilities;
 using MongoDB.Bson.Serialization;
 using MongoDB.Bson.Serialization.Attributes;
 using MongoDB.Driver;
 
-namespace Horde.Server.Replicators
+namespace HordeServer.Replicators
 {
 	/// <summary>
 	/// Default implementation of <see cref="IReplicatorCollection"/>.
diff --git a/Engine/Source/Programs/Horde/HordeServer/Replicators/ReplicatorsController.cs b/Engine/Source/Programs/Horde/HordeServer/Replicators/ReplicatorsController.cs
--- a/Engine/Source/Programs/Horde/HordeServer/Replicators/ReplicatorsController.cs
+++ b/Engine/Source/Programs/Horde/HordeServer/Replicators/ReplicatorsController.cs
@@ -6,15 +6,15 @@ using System.Threading;
 using System.Threading.Tasks;
 using EpicGames.Horde.Replicators;
 using EpicGames.Horde.Streams;
-using Horde.Server.Projects;
-using Horde.Server.Server;
-using Horde.Server.Streams;
-using Horde.Server.Utilities;
+using HordeServer.Projects;
+using HordeServer.Server;
+using HordeServer.Streams;
+using HordeServer.Utilities;
 using Microsoft.AspNetCore.Authorization;
 using Microsoft.AspNetCore.Mvc;
 using Microsoft.Extensions.Options;
 
-namespace Horde.Server.Replicators
+namespace HordeServer.Replicators
 {
 	/// <summary>
 	/// Controller for the /api/v1/projects endpoint
diff --git a/Engine/Source/Programs/Horde/HordeServer/Secrets/ISecret.cs b/Engine/Source/Programs/Horde/HordeServer/Secrets/ISecret.cs
--- a/Engine/Source/Programs/Horde/HordeServer/Secrets/ISecret.cs
+++ b/Engine/Source/Programs/Horde/HordeServer/Secrets/ISecret.cs
@@ -3,7 +3,7 @@
 using System.Collections.Generic;
 using EpicGames.Horde.Secrets;
 
-namespace Horde.Server.Secrets
+namespace HordeServer.Secrets
 {
 	/// <summary>
 	/// Information about a secret
diff --git a/Engine/Source/Programs/Horde/HordeServer/Secrets/ISecretCollection.cs b/Engine/Source/Programs/Horde/HordeServer/Secrets/ISecretCollection.cs
--- a/Engine/Source/Programs/Horde/HordeServer/Secrets/ISecretCollection.cs
+++ b/Engine/Source/Programs/Horde/HordeServer/Secrets/ISecretCollection.cs
@@ -4,7 +4,7 @@ using System.Threading;
 using System.Threading.Tasks;
 using EpicGames.Horde.Secrets;
 
-namespace Horde.Server.Secrets
+namespace HordeServer.Secrets
 {
 	/// <summary>
 	/// Collection of secrets
diff --git a/Engine/Source/Programs/Horde/HordeServer/Secrets/Providers/AwsParameterStoreSecretProvider.cs b/Engine/Source/Programs/Horde/HordeServer/Secrets/Providers/AwsParameterStoreSecretProvider.cs
--- a/Engine/Source/Programs/Horde/HordeServer/Secrets/Providers/AwsParameterStoreSecretProvider.cs
+++ b/Engine/Source/Programs/Horde/HordeServer/Secrets/Providers/AwsParameterStoreSecretProvider.cs
@@ -7,7 +7,7 @@ using Amazon.Extensions.NETCore.Setup;
 using Amazon.SimpleSystemsManagement;
 using Amazon.SimpleSystemsManagement.Model;
 
-namespace Horde.Server.Secrets.Providers
+namespace HordeServer.Secrets.Providers
 {
 	/// <summary>
 	/// Fetches secrets from the AWS parameter store
diff --git a/Engine/Source/Programs/Horde/HordeServer/Secrets/SecretCollection.cs b/Engine/Source/Programs/Horde/HordeServer/Secrets/SecretCollection.cs
--- a/Engine/Source/Programs/Horde/HordeServer/Secrets/SecretCollection.cs
+++ b/Engine/Source/Programs/Horde/HordeServer/Secrets/SecretCollection.cs
@@ -7,10 +7,10 @@ using System.Text.Json;
 using System.Threading;
 using System.Threading.Tasks;
 using EpicGames.Horde.Secrets;
-using Horde.Server.Server;
+using HordeServer.Server;
 using Microsoft.Extensions.Options;
 
-namespace Horde.Server.Secrets
+namespace HordeServer.Secrets
 {
 	/// <summary>
 	/// Implementation of <see cref="ISecretCollection"/>
diff --git a/Engine/Source/Programs/Horde/HordeServer/Secrets/SecretConfig.cs b/Engine/Source/Programs/Horde/HordeServer/Secrets/SecretConfig.cs
--- a/Engine/Source/Programs/Horde/HordeServer/Secrets/SecretConfig.cs
+++ b/Engine/Source/Programs/Horde/HordeServer/Secrets/SecretConfig.cs
@@ -5,10 +5,10 @@ using System.Collections.Generic;
 using System.Security.Claims;
 using EpicGames.Horde.Acls;
 using EpicGames.Horde.Secrets;
-using Horde.Server.Acls;
-using Horde.Server.Server;
+using HordeServer.Acls;
+using HordeServer.Server;
 
-namespace Horde.Server.Secrets
+namespace HordeServer.Secrets
 {
 	/// <summary>
 	/// Configuration for a secret value
diff --git a/Engine/Source/Programs/Horde/HordeServer/Secrets/SecretsController.cs b/Engine/Source/Programs/Horde/HordeServer/Secrets/SecretsController.cs
--- a/Engine/Source/Programs/Horde/HordeServer/Secrets/SecretsController.cs
+++ b/Engine/Source/Programs/Horde/HordeServer/Secrets/SecretsController.cs
@@ -5,13 +5,13 @@ using System.Linq;
 using System.Threading;
 using System.Threading.Tasks;
 using EpicGames.Horde.Secrets;
-using Horde.Server.Server;
-using Horde.Server.Utilities;
+using HordeServer.Server;
+using HordeServer.Utilities;
 using Microsoft.AspNetCore.Authorization;
 using Microsoft.AspNetCore.Mvc;
 using Microsoft.Extensions.Options;
 
-namespace Horde.Server.Secrets
+namespace HordeServer.Secrets
 {
 	/// <summary>
 	/// Controller for the /api/v1/secrets endpoint
diff --git a/Engine/Source/Programs/Horde/HordeServer/Server/AccountController.cs b/Engine/Source/Programs/Horde/HordeServer/Server/AccountController.cs
--- a/Engine/Source/Programs/Horde/HordeServer/Server/AccountController.cs
+++ b/Engine/Source/Programs/Horde/HordeServer/Server/AccountController.cs
@@ -11,11 +11,11 @@ using System.Threading.Tasks;
 using EpicGames.Horde.Accounts;
 using EpicGames.Horde.Acls;
 using EpicGames.Horde.Server;
-using Horde.Server.Accounts;
-using Horde.Server.Acls;
-using Horde.Server.Authentication;
-using Horde.Server.Users;
-using Horde.Server.Utilities;
+using HordeServer.Accounts;
+using HordeServer.Acls;
+using HordeServer.Authentication;
+using HordeServer.Users;
+using HordeServer.Utilities;
 using Microsoft.AspNetCore.Authentication;
 using Microsoft.AspNetCore.Authentication.Cookies;
 using Microsoft.AspNetCore.Authentication.OpenIdConnect;
@@ -24,7 +24,7 @@ using Microsoft.Extensions.Options;
 
 #pragma warning disable CA1054 // URI-like parameters should not be strings
 
-namespace Horde.Server.Server
+namespace HordeServer.Server
 {
 	/// <summary>
 	/// Model for Horde account login view
diff --git a/Engine/Source/Programs/Horde/HordeServer/Server/AdminAclAction.cs b/Engine/Source/Programs/Horde/HordeServer/Server/AdminAclAction.cs
--- a/Engine/Source/Programs/Horde/HordeServer/Server/AdminAclAction.cs
+++ b/Engine/Source/Programs/Horde/HordeServer/Server/AdminAclAction.cs
@@ -2,7 +2,7 @@
 
 using EpicGames.Horde.Acls;
 
-namespace Horde.Server.Server
+namespace HordeServer.Server
 {
 	/// <summary>
 	/// ACL actions for admin operations
diff --git a/Engine/Source/Programs/Horde/HordeServer/Server/AdminController.cs b/Engine/Source/Programs/Horde/HordeServer/Server/AdminController.cs
--- a/Engine/Source/Programs/Horde/HordeServer/Server/AdminController.cs
+++ b/Engine/Source/Programs/Horde/HordeServer/Server/AdminController.cs
@@ -6,13 +6,13 @@ using System.Linq;
 using System.Security.Claims;
 using System.Threading;
 using System.Threading.Tasks;
-using Horde.Server.Acls;
-using Horde.Server.Utilities;
+using HordeServer.Acls;
+using HordeServer.Utilities;
 using Microsoft.AspNetCore.Authorization;
 using Microsoft.AspNetCore.Mvc;
 using Microsoft.Extensions.Options;
 
-namespace Horde.Server.Server
+namespace HordeServer.Server
 {
 	/// <summary>
 	/// Object containing settings for the server
diff --git a/Engine/Source/Programs/Horde/HordeServer/Server/ConsistencyService.cs b/Engine/Source/Programs/Horde/HordeServer/Server/ConsistencyService.cs
--- a/Engine/Source/Programs/Horde/HordeServer/Server/ConsistencyService.cs
+++ b/Engine/Source/Programs/Horde/HordeServer/Server/ConsistencyService.cs
@@ -9,13 +9,13 @@ using EpicGames.Core;
 using EpicGames.Horde.Agents;
 using EpicGames.Horde.Agents.Leases;
 using EpicGames.Horde.Agents.Sessions;
-using Horde.Server.Agents;
-using Horde.Server.Agents.Leases;
-using Horde.Server.Agents.Sessions;
+using HordeServer.Agents;
+using HordeServer.Agents.Leases;
+using HordeServer.Agents.Sessions;
 using Microsoft.Extensions.Hosting;
 using Microsoft.Extensions.Logging;
 
-namespace Horde.Server.Server
+namespace HordeServer.Server
 {
 	/// <summary>
 	/// Service which checks the database for consistency and fixes up any errors
diff --git a/Engine/Source/Programs/Horde/HordeServer/Server/DebugController.cs b/Engine/Source/Programs/Horde/HordeServer/Server/DebugController.cs
--- a/Engine/Source/Programs/Horde/HordeServer/Server/DebugController.cs
+++ b/Engine/Source/Programs/Horde/HordeServer/Server/DebugController.cs
@@ -16,11 +16,11 @@ using EpicGames.Horde.Compute;
 using EpicGames.Horde.Logs;
 using Google.Protobuf;
 using Horde.Common.Rpc;
-using Horde.Server.Agents.Relay;
-using Horde.Server.Configuration;
-using Horde.Server.Logs;
-using Horde.Server.Projects;
-using Horde.Server.Utilities;
+using HordeServer.Agents.Relay;
+using HordeServer.Configuration;
+using HordeServer.Logs;
+using HordeServer.Projects;
+using HordeServer.Utilities;
 using JetBrains.Profiler.SelfApi;
 using Microsoft.AspNetCore.Authorization;
 using Microsoft.AspNetCore.Http;
@@ -31,7 +31,7 @@ using Microsoft.Extensions.Options;
 using MongoDB.Bson;
 using MongoDB.Driver;
 
-namespace Horde.Server.Server
+namespace HordeServer.Server
 {
 	/// <summary>
 	/// Controller managing account status
diff --git a/Engine/Source/Programs/Horde/HordeServer/Server/DowntimeService.cs b/Engine/Source/Programs/Horde/HordeServer/Server/DowntimeService.cs
--- a/Engine/Source/Programs/Horde/HordeServer/Server/DowntimeService.cs
+++ b/Engine/Source/Programs/Horde/HordeServer/Server/DowntimeService.cs
@@ -9,7 +9,7 @@ using Microsoft.Extensions.Hosting;
 using Microsoft.Extensions.Logging;
 using Microsoft.Extensions.Options;
 
-namespace Horde.Server.Server
+namespace HordeServer.Server
 {
 	/// <summary>
 	/// Interface for a service which keeps track of whether we're during downtime
diff --git a/Engine/Source/Programs/Horde/HordeServer/Server/ExceptionController.cs b/Engine/Source/Programs/Horde/HordeServer/Server/ExceptionController.cs
--- a/Engine/Source/Programs/Horde/HordeServer/Server/ExceptionController.cs
+++ b/Engine/Source/Programs/Horde/HordeServer/Server/ExceptionController.cs
@@ -2,12 +2,12 @@
 
 using System.Net;
 using EpicGames.Core;
-using Horde.Server.Utilities;
+using HordeServer.Utilities;
 using Microsoft.AspNetCore.Diagnostics;
 using Microsoft.AspNetCore.Mvc;
 using Microsoft.Extensions.Logging;
 
-namespace Horde.Server.Server
+namespace HordeServer.Server
 {
 	/// <summary>
 	/// Controller managing account status
diff --git a/Engine/Source/Programs/Horde/HordeServer/Server/GlobalConfig.cs b/Engine/Source/Programs/Horde/HordeServer/Server/GlobalConfig.cs
--- a/Engine/Source/Programs/Horde/HordeServer/Server/GlobalConfig.cs
+++ b/Engine/Source/Programs/Horde/HordeServer/Server/GlobalConfig.cs
@@ -24,22 +24,22 @@ using EpicGames.Horde.Tools;
 using EpicGames.Horde.Users;
 using EpicGames.Perforce;
 using EpicGames.Serialization;
-using Horde.Server.Acls;
-using Horde.Server.Agents;
-using Horde.Server.Agents.Pools;
-using Horde.Server.Configuration;
-using Horde.Server.Dashboard;
-using Horde.Server.Devices;
-using Horde.Server.Perforce;
-using Horde.Server.Projects;
-using Horde.Server.Secrets;
-using Horde.Server.Storage;
-using Horde.Server.Streams;
-using Horde.Server.Telemetry.Metrics;
-using Horde.Server.Tools;
-using Horde.Server.Utilities;
-
-namespace Horde.Server.Server
+using HordeServer.Acls;
+using HordeServer.Agents;
+using HordeServer.Agents.Pools;
+using HordeServer.Configuration;
+using HordeServer.Dashboard;
+using HordeServer.Devices;
+using HordeServer.Perforce;
+using HordeServer.Projects;
+using HordeServer.Secrets;
+using HordeServer.Storage;
+using HordeServer.Streams;
+using HordeServer.Telemetry.Metrics;
+using HordeServer.Tools;
+using HordeServer.Utilities;
+
+namespace HordeServer.Server
 {
 	using JsonObject = System.Text.Json.Nodes.JsonObject;
 
diff --git a/Engine/Source/Programs/Horde/HordeServer/Server/GlobalsService.cs b/Engine/Source/Programs/Horde/HordeServer/Server/GlobalsService.cs
--- a/Engine/Source/Programs/Horde/HordeServer/Server/GlobalsService.cs
+++ b/Engine/Source/Programs/Horde/HordeServer/Server/GlobalsService.cs
@@ -5,13 +5,13 @@ using System.Net;
 using System.Security.Cryptography;
 using System.Threading;
 using System.Threading.Tasks;
-using Horde.Server.Utilities;
+using HordeServer.Utilities;
 using Microsoft.Extensions.Options;
 using Microsoft.IdentityModel.Tokens;
 using MongoDB.Bson;
 using MongoDB.Bson.Serialization.Attributes;
 
-namespace Horde.Server.Server
+namespace HordeServer.Server
 {
 	/// <summary>
 	/// Accessor for the global <see cref="Globals"/> singleton
diff --git a/Engine/Source/Programs/Horde/HordeServer/Server/Notices/INotice.cs b/Engine/Source/Programs/Horde/HordeServer/Server/Notices/INotice.cs
--- a/Engine/Source/Programs/Horde/HordeServer/Server/Notices/INotice.cs
+++ b/Engine/Source/Programs/Horde/HordeServer/Server/Notices/INotice.cs
@@ -5,7 +5,7 @@ using EpicGames.Horde.Users;
 using MongoDB.Bson;
 using MongoDB.Bson.Serialization.Attributes;
 
-namespace Horde.Server.Server.Notices
+namespace HordeServer.Server.Notices
 {
 	/// <summary>
 	/// User notice
diff --git a/Engine/Source/Programs/Horde/HordeServer/Server/Notices/INoticeCollection.cs b/Engine/Source/Programs/Horde/HordeServer/Server/Notices/INoticeCollection.cs
--- a/Engine/Source/Programs/Horde/HordeServer/Server/Notices/INoticeCollection.cs
+++ b/Engine/Source/Programs/Horde/HordeServer/Server/Notices/INoticeCollection.cs
@@ -7,7 +7,7 @@ using System.Threading.Tasks;
 using EpicGames.Horde.Users;
 using MongoDB.Bson;
 
-namespace Horde.Server.Server.Notices
+namespace HordeServer.Server.Notices
 {
 	/// <summary>
 	/// A collection of user specified notices
diff --git a/Engine/Source/Programs/Horde/HordeServer/Server/Notices/NoticesController.cs b/Engine/Source/Programs/Horde/HordeServer/Server/Notices/NoticesController.cs
--- a/Engine/Source/Programs/Horde/HordeServer/Server/Notices/NoticesController.cs
+++ b/Engine/Source/Programs/Horde/HordeServer/Server/Notices/NoticesController.cs
@@ -8,16 +8,16 @@ using System.Threading.Tasks;
 using EpicGames.Core;
 using EpicGames.Horde.Server.Notices;
 using EpicGames.Horde.Users;
-using Horde.Server.Acls;
-using Horde.Server.Users;
-using Horde.Server.Utilities;
+using HordeServer.Acls;
+using HordeServer.Users;
+using HordeServer.Utilities;
 using Microsoft.AspNetCore.Authorization;
 using Microsoft.AspNetCore.Mvc;
 using Microsoft.Extensions.Diagnostics.HealthChecks;
 using Microsoft.Extensions.Options;
 using MongoDB.Bson;
 
-namespace Horde.Server.Server.Notices
+namespace HordeServer.Server.Notices
 {
 	/// <summary>
 	/// Controller for the /api/v1/notices endpoint
diff --git a/Engine/Source/Programs/Horde/HordeServer/Server/RpcService.cs b/Engine/Source/Programs/Horde/HordeServer/Server/RpcService.cs
--- a/Engine/Source/Programs/Horde/HordeServer/Server/RpcService.cs
+++ b/Engine/Source/Programs/Horde/HordeServer/Server/RpcService.cs
@@ -17,13 +17,13 @@ using Google.Protobuf;
 using Google.Protobuf.Reflection;
 using Google.Protobuf.WellKnownTypes;
 using Grpc.Core;
-using Horde.Server.Acls;
-using Horde.Server.Agents;
-using Horde.Server.Agents.Sessions;
-using Horde.Server.Agents.Telemetry;
-using Horde.Server.Telemetry;
-using Horde.Server.Tools;
-using Horde.Server.Utilities;
+using HordeServer.Acls;
+using HordeServer.Agents;
+using HordeServer.Agents.Sessions;
+using HordeServer.Agents.Telemetry;
+using HordeServer.Telemetry;
+using HordeServer.Tools;
+using HordeServer.Utilities;
 using HordeCommon.Rpc;
 using HordeCommon.Rpc.Messages;
 using HordeCommon.Rpc.Messages.Telemetry;
@@ -31,7 +31,7 @@ using Microsoft.AspNetCore.Authorization;
 using Microsoft.Extensions.Logging;
 using Microsoft.Extensions.Options;
 
-namespace Horde.Server.Server
+namespace HordeServer.Server
 {
 	/// <summary>
 	/// Implements the Horde gRPC service for bots updating their status and dequeing work
diff --git a/Engine/Source/Programs/Horde/HordeServer/Server/ServerController.cs b/Engine/Source/Programs/Horde/HordeServer/Server/ServerController.cs
--- a/Engine/Source/Programs/Horde/HordeServer/Server/ServerController.cs
+++ b/Engine/Source/Programs/Horde/HordeServer/Server/ServerController.cs
@@ -11,16 +11,16 @@ using EpicGames.Core;
 using EpicGames.Horde;
 using EpicGames.Horde.Agents.Pools;
 using EpicGames.Horde.Server;
-using Horde.Server.Agents;
-using Horde.Server.Agents.Pools;
-using Horde.Server.Tools;
-using Horde.Server.Utilities;
+using HordeServer.Agents;
+using HordeServer.Agents.Pools;
+using HordeServer.Tools;
+using HordeServer.Utilities;
 using Microsoft.AspNetCore.Authorization;
 using Microsoft.AspNetCore.Mvc;
 using Microsoft.Extensions.DependencyInjection;
 using Microsoft.Extensions.Options;
 
-namespace Horde.Server.Server
+namespace HordeServer.Server
 {
 	/// <summary>
 	/// Controller managing account status
diff --git a/Engine/Source/Programs/Horde/HordeServer/Server/ServerStatusController.cs b/Engine/Source/Programs/Horde/HordeServer/Server/ServerStatusController.cs
--- a/Engine/Source/Programs/Horde/HordeServer/Server/ServerStatusController.cs
+++ b/Engine/Source/Programs/Horde/HordeServer/Server/ServerStatusController.cs
@@ -10,7 +10,7 @@ using Microsoft.AspNetCore.Http;
 using Microsoft.AspNetCore.Mvc;
 using Microsoft.Extensions.Diagnostics.HealthChecks;
 
-namespace Horde.Server.Server;
+namespace HordeServer.Server;
 
 /// <summary>
 /// ASP.NET view model for server status updates
diff --git a/Engine/Source/Programs/Horde/HordeServer/Server/ServerStatusService.cs b/Engine/Source/Programs/Horde/HordeServer/Server/ServerStatusService.cs
--- a/Engine/Source/Programs/Horde/HordeServer/Server/ServerStatusService.cs
+++ b/Engine/Source/Programs/Horde/HordeServer/Server/ServerStatusService.cs
@@ -12,7 +12,7 @@ using Microsoft.Extensions.Hosting;
 using Microsoft.Extensions.Logging;
 using StackExchange.Redis;
 
-namespace Horde.Server.Server;
+namespace HordeServer.Server;
 
 /// <summary>
 /// Allows reporting the health of a subsystem
diff --git a/Engine/Source/Programs/Horde/HordeServer/ServerApp.cs b/Engine/Source/Programs/Horde/HordeServer/ServerApp.cs
--- a/Engine/Source/Programs/Horde/HordeServer/ServerApp.cs
+++ b/Engine/Source/Programs/Horde/HordeServer/ServerApp.cs
@@ -9,8 +9,8 @@ using System.Runtime.InteropServices;
 using System.Threading;
 using System.Threading.Tasks;
 using EpicGames.Core;
-using Horde.Server.Commands;
-using Horde.Server.Utilities;
+using HordeServer.Commands;
+using HordeServer.Utilities;
 using Microsoft.Extensions.Configuration;
 using Microsoft.Extensions.DependencyInjection;
 using Microsoft.Extensions.Hosting;
@@ -23,7 +23,7 @@ using Serilog.Exceptions.Grpc.Destructurers;
 using Serilog.Formatting.Json;
 using Serilog.Sinks.SystemConsole.Themes;
 
-namespace Horde.Server
+namespace HordeServer
 {
 	static class LoggerExtensions
 	{
diff --git a/Engine/Source/Programs/Horde/HordeServer/ServerSettings.cs b/Engine/Source/Programs/Horde/HordeServer/ServerSettings.cs
--- a/Engine/Source/Programs/Horde/HordeServer/ServerSettings.cs
+++ b/Engine/Source/Programs/Horde/HordeServer/ServerSettings.cs
@@ -14,14 +14,14 @@ using EpicGames.Horde.Storage.Bundles;
 using EpicGames.Horde.Storage.Nodes;
 using EpicGames.Horde.Tools;
 using EpicGames.Perforce;
-using Horde.Server.Agents.Fleet;
-using Horde.Server.Server;
-using Horde.Server.Storage;
-using Horde.Server.Storage.ObjectStores;
-using Horde.Server.Telemetry.Sinks;
-using Horde.Server.Tools;
-
-namespace Horde.Server
+using HordeServer.Agents.Fleet;
+using HordeServer.Server;
+using HordeServer.Storage;
+using HordeServer.Storage.ObjectStores;
+using HordeServer.Telemetry.Sinks;
+using HordeServer.Tools;
+
+namespace HordeServer
 {
 	/// <summary>
 	/// Types of storage backend to use
diff --git a/Engine/Source/Programs/Horde/HordeServer/ServiceAccounts/IServiceAccount.cs b/Engine/Source/Programs/Horde/HordeServer/ServiceAccounts/IServiceAccount.cs
--- a/Engine/Source/Programs/Horde/HordeServer/ServiceAccounts/IServiceAccount.cs
+++ b/Engine/Source/Programs/Horde/HordeServer/ServiceAccounts/IServiceAccount.cs
@@ -5,10 +5,10 @@ using System.Collections.Generic;
 using System.Threading;
 using System.Threading.Tasks;
 using EpicGames.Horde.ServiceAccounts;
-using Horde.Server.Acls;
-using Horde.Server.Users;
+using HordeServer.Acls;
+using HordeServer.Users;
 
-namespace Horde.Server.ServiceAccounts
+namespace HordeServer.ServiceAccounts
 {
 	/// <summary>
 	/// An internal Horde account representing a service
diff --git a/Engine/Source/Programs/Horde/HordeServer/ServiceAccounts/IServiceAccountCollection.cs b/Engine/Source/Programs/Horde/HordeServer/ServiceAccounts/IServiceAccountCollection.cs
--- a/Engine/Source/Programs/Horde/HordeServer/ServiceAccounts/IServiceAccountCollection.cs
+++ b/Engine/Source/Programs/Horde/HordeServer/ServiceAccounts/IServiceAccountCollection.cs
@@ -4,9 +4,9 @@ using System.Collections.Generic;
 using System.Threading;
 using System.Threading.Tasks;
 using EpicGames.Horde.ServiceAccounts;
-using Horde.Server.Users;
+using HordeServer.Users;
 
-namespace Horde.Server.ServiceAccounts
+namespace HordeServer.ServiceAccounts
 {
 	/// <summary>
 	/// Interface for a collection of accounts
diff --git a/Engine/Source/Programs/Horde/HordeServer/ServiceAccounts/ServiceAccountAclAction.cs b/Engine/Source/Programs/Horde/HordeServer/ServiceAccounts/ServiceAccountAclAction.cs
--- a/Engine/Source/Programs/Horde/HordeServer/ServiceAccounts/ServiceAccountAclAction.cs
+++ b/Engine/Source/Programs/Horde/HordeServer/ServiceAccounts/ServiceAccountAclAction.cs
@@ -2,7 +2,7 @@
 
 using EpicGames.Horde.Acls;
 
-namespace Horde.Server.ServiceAccounts
+namespace HordeServer.ServiceAccounts
 {
 	/// <summary>
 	/// ACL actions valid for manipulating jobs
diff --git a/Engine/Source/Programs/Horde/HordeServer/ServiceAccounts/ServiceAccountCollection.cs b/Engine/Source/Programs/Horde/HordeServer/ServiceAccounts/ServiceAccountCollection.cs
--- a/Engine/Source/Programs/Horde/HordeServer/ServiceAccounts/ServiceAccountCollection.cs
+++ b/Engine/Source/Programs/Horde/HordeServer/ServiceAccounts/ServiceAccountCollection.cs
@@ -8,14 +8,14 @@ using System.Threading;
 using System.Threading.Tasks;
 using EpicGames.Core;
 using EpicGames.Horde.ServiceAccounts;
-using Horde.Server.Acls;
-using Horde.Server.Server;
-using Horde.Server.Users;
-using Horde.Server.Utilities;
+using HordeServer.Acls;
+using HordeServer.Server;
+using HordeServer.Users;
+using HordeServer.Utilities;
 using MongoDB.Bson.Serialization.Attributes;
 using MongoDB.Driver;
 
-namespace Horde.Server.ServiceAccounts
+namespace HordeServer.ServiceAccounts
 {
 	/// <summary>
 	/// Collection of service account documents
diff --git a/Engine/Source/Programs/Horde/HordeServer/ServiceAccounts/ServiceAccountsController.cs b/Engine/Source/Programs/Horde/HordeServer/ServiceAccounts/ServiceAccountsController.cs
--- a/Engine/Source/Programs/Horde/HordeServer/ServiceAccounts/ServiceAccountsController.cs
+++ b/Engine/Source/Programs/Horde/HordeServer/ServiceAccounts/ServiceAccountsController.cs
@@ -5,14 +5,14 @@ using System.Threading;
 using System.Threading.Tasks;
 using EpicGames.Horde.Accounts;
 using EpicGames.Horde.ServiceAccounts;
-using Horde.Server.Server;
-using Horde.Server.Users;
-using Horde.Server.Utilities;
+using HordeServer.Server;
+using HordeServer.Users;
+using HordeServer.Utilities;
 using Microsoft.AspNetCore.Authorization;
 using Microsoft.AspNetCore.Mvc;
 using Microsoft.Extensions.Options;
 
-namespace Horde.Server.ServiceAccounts
+namespace HordeServer.ServiceAccounts
 {
 	/// <summary>
 	/// Controller for the /api/v1/serviceaccounts endpoint
