Horde: Rename namespaces to match project folders.

diff --git a/Engine/Source/Programs/Horde/HordeServer/Startup.cs b/Engine/Source/Programs/Horde/HordeServer/Startup.cs
--- a/Engine/Source/Programs/Horde/HordeServer/Startup.cs
+++ b/Engine/Source/Programs/Horde/HordeServer/Startup.cs
@@ -38,59 +38,59 @@ using EpicGames.Horde.Users;
 using EpicGames.Redis;
 using Grpc.Core;
 using Grpc.Core.Interceptors;
-using Horde.Server.Accounts;
-using Horde.Server.Acls;
-using Horde.Server.Agents;
-using Horde.Server.Agents.Enrollment;
-using Horde.Server.Agents.Fleet;
-using Horde.Server.Agents.Leases;
-using Horde.Server.Agents.Pools;
-using Horde.Server.Agents.Relay;
-using Horde.Server.Agents.Sessions;
-using Horde.Server.Agents.Telemetry;
-using Horde.Server.Agents.Utilization;
-using Horde.Server.Artifacts;
-using Horde.Server.Auditing;
-using Horde.Server.Authentication;
-using Horde.Server.Aws;
-using Horde.Server.Commits;
-using Horde.Server.Compute;
-using Horde.Server.Configuration;
-using Horde.Server.Dashboard;
-using Horde.Server.Ddc;
-using Horde.Server.Devices;
-using Horde.Server.Issues;
-using Horde.Server.Issues.External;
-using Horde.Server.Jobs;
-using Horde.Server.Jobs.Artifacts;
-using Horde.Server.Jobs.Bisect;
-using Horde.Server.Jobs.Graphs;
-using Horde.Server.Jobs.Schedules;
-using Horde.Server.Jobs.Templates;
-using Horde.Server.Jobs.TestData;
-using Horde.Server.Jobs.Timing;
-using Horde.Server.Logs;
-using Horde.Server.Logs.Storage;
-using Horde.Server.Notifications;
-using Horde.Server.Notifications.Sinks;
-using Horde.Server.Perforce;
-using Horde.Server.Replicators;
-using Horde.Server.Secrets;
-using Horde.Server.Secrets.Providers;
-using Horde.Server.Server;
-using Horde.Server.Server.Notices;
-using Horde.Server.ServiceAccounts;
-using Horde.Server.Storage;
-using Horde.Server.Storage.ObjectStores;
-using Horde.Server.Streams;
-using Horde.Server.Tasks;
-using Horde.Server.Telemetry;
-using Horde.Server.Telemetry.Metrics;
-using Horde.Server.Telemetry.Sinks;
-using Horde.Server.Tools;
-using Horde.Server.Ugs;
-using Horde.Server.Users;
-using Horde.Server.Utilities;
+using HordeServer.Accounts;
+using HordeServer.Acls;
+using HordeServer.Agents;
+using HordeServer.Agents.Enrollment;
+using HordeServer.Agents.Fleet;
+using HordeServer.Agents.Leases;
+using HordeServer.Agents.Pools;
+using HordeServer.Agents.Relay;
+using HordeServer.Agents.Sessions;
+using HordeServer.Agents.Telemetry;
+using HordeServer.Agents.Utilization;
+using HordeServer.Artifacts;
+using HordeServer.Auditing;
+using HordeServer.Authentication;
+using HordeServer.Aws;
+using HordeServer.Commits;
+using HordeServer.Compute;
+using HordeServer.Configuration;
+using HordeServer.Dashboard;
+using HordeServer.Ddc;
+using HordeServer.Devices;
+using HordeServer.Issues;
+using HordeServer.Issues.External;
+using HordeServer.Jobs;
+using HordeServer.Jobs.Artifacts;
+using HordeServer.Jobs.Bisect;
+using HordeServer.Jobs.Graphs;
+using HordeServer.Jobs.Schedules;
+using HordeServer.Jobs.Templates;
+using HordeServer.Jobs.TestData;
+using HordeServer.Jobs.Timing;
+using HordeServer.Logs;
+using HordeServer.Logs.Storage;
+using HordeServer.Notifications;
+using HordeServer.Notifications.Sinks;
+using HordeServer.Perforce;
+using HordeServer.Replicators;
+using HordeServer.Secrets;
+using HordeServer.Secrets.Providers;
+using HordeServer.Server;
+using HordeServer.Server.Notices;
+using HordeServer.ServiceAccounts;
+using HordeServer.Storage;
+using HordeServer.Storage.ObjectStores;
+using HordeServer.Streams;
+using HordeServer.Tasks;
+using HordeServer.Telemetry;
+using HordeServer.Telemetry.Metrics;
+using HordeServer.Telemetry.Sinks;
+using HordeServer.Tools;
+using HordeServer.Ugs;
+using HordeServer.Users;
+using HordeServer.Utilities;
 using HordeCommon;
 using Microsoft.AspNetCore.Authentication;
 using Microsoft.AspNetCore.Authentication.Cookies;
@@ -131,11 +131,11 @@ using StatusCode = Grpc.Core.StatusCode;
 
 #pragma warning disable CA1505 // 'ConfigureServices' has a maintainability index of '9'. Rewrite or refactor the code to increase its maintainability index (MI) above '9'.
 
-namespace Horde.Server
+namespace HordeServer
 {
 	using ContentHash = EpicGames.Core.ContentHash;
 	using ILogger = Microsoft.Extensions.Logging.ILogger;
-	using ReferenceResolver = Horde.Server.Ddc.ReferenceResolver;
+	using ReferenceResolver = HordeServer.Ddc.ReferenceResolver;
 
 	class Startup
 	{
diff --git a/Engine/Source/Programs/Horde/HordeServer/Storage/BackendId.cs b/Engine/Source/Programs/Horde/HordeServer/Storage/BackendId.cs
--- a/Engine/Source/Programs/Horde/HordeServer/Storage/BackendId.cs
+++ b/Engine/Source/Programs/Horde/HordeServer/Storage/BackendId.cs
@@ -5,7 +5,7 @@ using EpicGames.Core;
 using EpicGames.Horde;
 using EpicGames.Serialization;
 
-namespace Horde.Server.Storage
+namespace HordeServer.Storage
 {
 	/// <summary>
 	/// Identifier for a pool
diff --git a/Engine/Source/Programs/Horde/HordeServer/Storage/IObjectStoreFactory.cs b/Engine/Source/Programs/Horde/HordeServer/Storage/IObjectStoreFactory.cs
--- a/Engine/Source/Programs/Horde/HordeServer/Storage/IObjectStoreFactory.cs
+++ b/Engine/Source/Programs/Horde/HordeServer/Storage/IObjectStoreFactory.cs
@@ -2,7 +2,7 @@
 
 using EpicGames.Horde.Storage;
 
-namespace Horde.Server.Storage
+namespace HordeServer.Storage
 {
 	/// <summary>
 	/// Service which resolves storage backend instances
diff --git a/Engine/Source/Programs/Horde/HordeServer/Storage/ObjectStoreFactory.cs b/Engine/Source/Programs/Horde/HordeServer/Storage/ObjectStoreFactory.cs
--- a/Engine/Source/Programs/Horde/HordeServer/Storage/ObjectStoreFactory.cs
+++ b/Engine/Source/Programs/Horde/HordeServer/Storage/ObjectStoreFactory.cs
@@ -5,11 +5,11 @@ using System.Collections.Generic;
 using EpicGames.Core;
 using EpicGames.Horde.Storage;
 using EpicGames.Horde.Storage.ObjectStores;
-using Horde.Server.Storage.ObjectStores;
+using HordeServer.Storage.ObjectStores;
 using Microsoft.Extensions.DependencyInjection;
 using Microsoft.Extensions.Logging;
 
-namespace Horde.Server.Storage
+namespace HordeServer.Storage
 {
 	/// <summary>
 	/// Default implementation of <see cref="IObjectStoreFactory"/>
diff --git a/Engine/Source/Programs/Horde/HordeServer/Storage/ObjectStores/AwsObjectStore.cs b/Engine/Source/Programs/Horde/HordeServer/Storage/ObjectStores/AwsObjectStore.cs
--- a/Engine/Source/Programs/Horde/HordeServer/Storage/ObjectStores/AwsObjectStore.cs
+++ b/Engine/Source/Programs/Horde/HordeServer/Storage/ObjectStores/AwsObjectStore.cs
@@ -15,12 +15,12 @@ using Amazon.S3;
 using Amazon.S3.Model;
 using EpicGames.Core;
 using EpicGames.Horde.Storage;
-using Horde.Server.Utilities;
+using HordeServer.Utilities;
 using Microsoft.Extensions.Configuration;
 using Microsoft.Extensions.Logging;
 using OpenTelemetry.Trace;
 
-namespace Horde.Server.Storage.ObjectStores
+namespace HordeServer.Storage.ObjectStores
 {
 	/// <summary>
 	/// Exception wrapper for S3 requests
diff --git a/Engine/Source/Programs/Horde/HordeServer/Storage/ObjectStores/AzureObjectStore.cs b/Engine/Source/Programs/Horde/HordeServer/Storage/ObjectStores/AzureObjectStore.cs
--- a/Engine/Source/Programs/Horde/HordeServer/Storage/ObjectStores/AzureObjectStore.cs
+++ b/Engine/Source/Programs/Horde/HordeServer/Storage/ObjectStores/AzureObjectStore.cs
@@ -15,7 +15,7 @@ using EpicGames.Horde.Storage;
 using Microsoft.Extensions.Logging;
 using OpenTelemetry.Trace;
 
-namespace Horde.Server.Storage.ObjectStores
+namespace HordeServer.Storage.ObjectStores
 {
 	/// <summary>
 	/// Exception wrapper for S3 requests
diff --git a/Engine/Source/Programs/Horde/HordeServer/Storage/StorageConfig.cs b/Engine/Source/Programs/Horde/HordeServer/Storage/StorageConfig.cs
--- a/Engine/Source/Programs/Horde/HordeServer/Storage/StorageConfig.cs
+++ b/Engine/Source/Programs/Horde/HordeServer/Storage/StorageConfig.cs
@@ -13,11 +13,11 @@ using System.Text.Json.Serialization;
 using EpicGames.Core;
 using EpicGames.Horde.Acls;
 using EpicGames.Horde.Storage;
-using Horde.Server.Acls;
-using Horde.Server.Server;
-using Horde.Server.Storage.ObjectStores;
+using HordeServer.Acls;
+using HordeServer.Server;
+using HordeServer.Storage.ObjectStores;
 
-namespace Horde.Server.Storage
+namespace HordeServer.Storage
 {
 	/// <summary>
 	/// Well-known namespace identifiers
diff --git a/Engine/Source/Programs/Horde/HordeServer/Storage/StorageController.cs b/Engine/Source/Programs/Horde/HordeServer/Storage/StorageController.cs
--- a/Engine/Source/Programs/Horde/HordeServer/Storage/StorageController.cs
+++ b/Engine/Source/Programs/Horde/HordeServer/Storage/StorageController.cs
@@ -16,18 +16,16 @@ using EpicGames.Horde.Acls;
 using EpicGames.Horde.Storage;
 using EpicGames.Horde.Storage.Nodes;
 using EpicGames.Serialization;
-using Horde.Server.Server;
-using Horde.Server.Utilities;
+using HordeServer.Server;
+using HordeServer.Utilities;
 using Microsoft.AspNetCore.Authorization;
 using Microsoft.AspNetCore.Http;
 using Microsoft.AspNetCore.Mvc;
 using Microsoft.Extensions.Logging;
 using Microsoft.Extensions.Options;
 
-namespace Horde.Server.Storage
+namespace HordeServer.Storage
 {
-	using DdcRefNode = Horde.Server.Ddc.DdcRefNode;
-
 	/// <summary>
 	/// Controller for the /api/v1/storage endpoint
 	/// </summary>
diff --git a/Engine/Source/Programs/Horde/HordeServer/Storage/StorageService.cs b/Engine/Source/Programs/Horde/HordeServer/Storage/StorageService.cs
--- a/Engine/Source/Programs/Horde/HordeServer/Storage/StorageService.cs
+++ b/Engine/Source/Programs/Horde/HordeServer/Storage/StorageService.cs
@@ -15,8 +15,8 @@ using EpicGames.Horde.Storage.Bundles;
 using EpicGames.Horde.Storage.ObjectStores;
 using EpicGames.Redis;
 using EpicGames.Redis.Utility;
-using Horde.Server.Server;
-using Horde.Server.Utilities;
+using HordeServer.Server;
+using HordeServer.Utilities;
 using Microsoft.Extensions.Caching.Memory;
 using Microsoft.Extensions.Configuration;
 using Microsoft.Extensions.DependencyInjection;
@@ -29,7 +29,7 @@ using MongoDB.Driver;
 using OpenTelemetry.Trace;
 using StackExchange.Redis;
 
-namespace Horde.Server.Storage
+namespace HordeServer.Storage
 {
 	/// <summary>
 	/// Interface for the storage service
diff --git a/Engine/Source/Programs/Horde/HordeServer/Streams/IStream.cs b/Engine/Source/Programs/Horde/HordeServer/Streams/IStream.cs
--- a/Engine/Source/Programs/Horde/HordeServer/Streams/IStream.cs
+++ b/Engine/Source/Programs/Horde/HordeServer/Streams/IStream.cs
@@ -9,7 +9,7 @@ using EpicGames.Horde.Streams;
 using EpicGames.Horde.Users;
 using Horde.Common.Rpc;
 
-namespace Horde.Server.Streams
+namespace HordeServer.Streams
 {
 	/// <summary>
 	/// Exception thrown when stream validation fails
diff --git a/Engine/Source/Programs/Horde/HordeServer/Streams/IStreamCollection.cs b/Engine/Source/Programs/Horde/HordeServer/Streams/IStreamCollection.cs
--- a/Engine/Source/Programs/Horde/HordeServer/Streams/IStreamCollection.cs
+++ b/Engine/Source/Programs/Horde/HordeServer/Streams/IStreamCollection.cs
@@ -8,7 +8,7 @@ using System.Threading.Tasks;
 using EpicGames.Horde.Jobs;
 using EpicGames.Horde.Jobs.Templates;
 
-namespace Horde.Server.Streams
+namespace HordeServer.Streams
 {
 	/// <summary>
 	/// Collection of stream documents
diff --git a/Engine/Source/Programs/Horde/HordeServer/Streams/StreamCollection.cs b/Engine/Source/Programs/Horde/HordeServer/Streams/StreamCollection.cs
--- a/Engine/Source/Programs/Horde/HordeServer/Streams/StreamCollection.cs
+++ b/Engine/Source/Programs/Horde/HordeServer/Streams/StreamCollection.cs
@@ -11,12 +11,12 @@ using EpicGames.Horde.Jobs;
 using EpicGames.Horde.Jobs.Templates;
 using EpicGames.Horde.Streams;
 using EpicGames.Horde.Users;
-using Horde.Server.Server;
-using Horde.Server.Utilities;
+using HordeServer.Server;
+using HordeServer.Utilities;
 using MongoDB.Bson.Serialization.Attributes;
 using MongoDB.Driver;
 
-namespace Horde.Server.Streams
+namespace HordeServer.Streams
 {
 	/// <summary>
 	/// Collection of stream documents
diff --git a/Engine/Source/Programs/Horde/HordeServer/Streams/StreamConfig.cs b/Engine/Source/Programs/Horde/HordeServer/Streams/StreamConfig.cs
--- a/Engine/Source/Programs/Horde/HordeServer/Streams/StreamConfig.cs
+++ b/Engine/Source/Programs/Horde/HordeServer/Streams/StreamConfig.cs
@@ -23,17 +23,17 @@ using EpicGames.Horde.Streams;
 using EpicGames.Horde.Telemetry;
 using EpicGames.Perforce;
 using Horde.Common.Rpc;
-using Horde.Server.Acls;
-using Horde.Server.Configuration;
-using Horde.Server.Issues;
-using Horde.Server.Jobs.Graphs;
-using Horde.Server.Jobs.Templates;
-using Horde.Server.Projects;
-using Horde.Server.Replicators;
-using Horde.Server.Server;
-using Horde.Server.Utilities;
-
-namespace Horde.Server.Streams
+using HordeServer.Acls;
+using HordeServer.Configuration;
+using HordeServer.Issues;
+using HordeServer.Jobs.Graphs;
+using HordeServer.Jobs.Templates;
+using HordeServer.Projects;
+using HordeServer.Replicators;
+using HordeServer.Server;
+using HordeServer.Utilities;
+
+namespace HordeServer.Streams
 {
 	/// <summary>
 	/// Flags identifying content of a changelist
diff --git a/Engine/Source/Programs/Horde/HordeServer/Streams/StreamMessages.cs b/Engine/Source/Programs/Horde/HordeServer/Streams/StreamMessages.cs
--- a/Engine/Source/Programs/Horde/HordeServer/Streams/StreamMessages.cs
+++ b/Engine/Source/Programs/Horde/HordeServer/Streams/StreamMessages.cs
@@ -9,12 +9,12 @@ using EpicGames.Horde.Jobs.Templates;
 using EpicGames.Horde.Projects;
 using EpicGames.Horde.Streams;
 using EpicGames.Horde.Users;
-using Horde.Server.Issues;
-using Horde.Server.Jobs.Schedules;
-using Horde.Server.Jobs.Templates;
-using Horde.Server.Projects;
+using HordeServer.Issues;
+using HordeServer.Jobs.Schedules;
+using HordeServer.Jobs.Templates;
+using HordeServer.Projects;
 
-namespace Horde.Server.Streams
+namespace HordeServer.Streams
 {
 	/// <summary>
 	/// Step state update request
diff --git a/Engine/Source/Programs/Horde/HordeServer/Streams/StreamsController.cs b/Engine/Source/Programs/Horde/HordeServer/Streams/StreamsController.cs
--- a/Engine/Source/Programs/Horde/HordeServer/Streams/StreamsController.cs
+++ b/Engine/Source/Programs/Horde/HordeServer/Streams/StreamsController.cs
@@ -11,20 +11,20 @@ using EpicGames.Horde.Jobs.Templates;
 using EpicGames.Horde.Projects;
 using EpicGames.Horde.Streams;
 using EpicGames.Horde.Users;
-using Horde.Server.Commits;
-using Horde.Server.Jobs;
-using Horde.Server.Jobs.Templates;
-using Horde.Server.Projects;
-using Horde.Server.Server;
-using Horde.Server.Users;
-using Horde.Server.Utilities;
+using HordeServer.Commits;
+using HordeServer.Jobs;
+using HordeServer.Jobs.Templates;
+using HordeServer.Projects;
+using HordeServer.Server;
+using HordeServer.Users;
+using HordeServer.Utilities;
 using Microsoft.AspNetCore.Authorization;
 using Microsoft.AspNetCore.Mvc;
 using Microsoft.Extensions.Options;
 using OpenTelemetry.Trace;
 using TimeZoneConverter;
 
-namespace Horde.Server.Streams
+namespace HordeServer.Streams
 {
 	/// <summary>
 	/// Controller for the /api/v1/streams endpoint
diff --git a/Engine/Source/Programs/Horde/HordeServer/Tasks/ConformTaskSource.cs b/Engine/Source/Programs/Horde/HordeServer/Tasks/ConformTaskSource.cs
--- a/Engine/Source/Programs/Horde/HordeServer/Tasks/ConformTaskSource.cs
+++ b/Engine/Source/Programs/Horde/HordeServer/Tasks/ConformTaskSource.cs
@@ -12,20 +12,20 @@ using EpicGames.Horde.Jobs;
 using EpicGames.Horde.Logs;
 using Google.Protobuf;
 using Google.Protobuf.WellKnownTypes;
-using Horde.Server.Agents;
-using Horde.Server.Agents.Pools;
-using Horde.Server.Jobs;
-using Horde.Server.Logs;
-using Horde.Server.Perforce;
-using Horde.Server.Server;
-using Horde.Server.Utilities;
+using HordeServer.Agents;
+using HordeServer.Agents.Pools;
+using HordeServer.Jobs;
+using HordeServer.Logs;
+using HordeServer.Perforce;
+using HordeServer.Server;
+using HordeServer.Utilities;
 using HordeCommon.Rpc.Messages;
 using HordeCommon.Rpc.Tasks;
 using Microsoft.Extensions.Hosting;
 using Microsoft.Extensions.Logging;
 using Microsoft.Extensions.Options;
 
-namespace Horde.Server.Tasks
+namespace HordeServer.Tasks
 {
 	/// <summary>
 	/// Generates tasks telling agents to sync their workspaces
diff --git a/Engine/Source/Programs/Horde/HordeServer/Tasks/ITaskSource.cs b/Engine/Source/Programs/Horde/HordeServer/Tasks/ITaskSource.cs
--- a/Engine/Source/Programs/Horde/HordeServer/Tasks/ITaskSource.cs
+++ b/Engine/Source/Programs/Horde/HordeServer/Tasks/ITaskSource.cs
@@ -13,10 +13,10 @@ using EpicGames.Horde.Agents.Leases;
 using Google.Protobuf;
 using Google.Protobuf.Reflection;
 using Google.Protobuf.WellKnownTypes;
-using Horde.Server.Agents;
+using HordeServer.Agents;
 using Microsoft.Extensions.Logging;
 
-namespace Horde.Server.Tasks
+namespace HordeServer.Tasks
 {
 	/// <summary>
 	/// Flags indicating when a task source is valid
diff --git a/Engine/Source/Programs/Horde/HordeServer/Tasks/RestartTaskSource.cs b/Engine/Source/Programs/Horde/HordeServer/Tasks/RestartTaskSource.cs
--- a/Engine/Source/Programs/Horde/HordeServer/Tasks/RestartTaskSource.cs
+++ b/Engine/Source/Programs/Horde/HordeServer/Tasks/RestartTaskSource.cs
@@ -7,12 +7,12 @@ using EpicGames.Horde.Jobs;
 using EpicGames.Horde.Logs;
 using Google.Protobuf;
 using Google.Protobuf.WellKnownTypes;
-using Horde.Server.Agents;
-using Horde.Server.Logs;
-using Horde.Server.Utilities;
+using HordeServer.Agents;
+using HordeServer.Logs;
+using HordeServer.Utilities;
 using HordeCommon.Rpc.Tasks;
 
-namespace Horde.Server.Tasks
+namespace HordeServer.Tasks
 {
 	class RestartTaskSource : TaskSourceBase<RestartTask>
 	{
diff --git a/Engine/Source/Programs/Horde/HordeServer/Tasks/ShutdownTaskSource.cs b/Engine/Source/Programs/Horde/HordeServer/Tasks/ShutdownTaskSource.cs
--- a/Engine/Source/Programs/Horde/HordeServer/Tasks/ShutdownTaskSource.cs
+++ b/Engine/Source/Programs/Horde/HordeServer/Tasks/ShutdownTaskSource.cs
@@ -7,12 +7,12 @@ using EpicGames.Horde.Jobs;
 using EpicGames.Horde.Logs;
 using Google.Protobuf;
 using Google.Protobuf.WellKnownTypes;
-using Horde.Server.Agents;
-using Horde.Server.Logs;
-using Horde.Server.Utilities;
+using HordeServer.Agents;
+using HordeServer.Logs;
+using HordeServer.Utilities;
 using HordeCommon.Rpc.Tasks;
 
-namespace Horde.Server.Tasks
+namespace HordeServer.Tasks
 {
 	class ShutdownTaskSource : TaskSourceBase<ShutdownTask>
 	{
diff --git a/Engine/Source/Programs/Horde/HordeServer/Tasks/UpgradeTaskSource.cs b/Engine/Source/Programs/Horde/HordeServer/Tasks/UpgradeTaskSource.cs
--- a/Engine/Source/Programs/Horde/HordeServer/Tasks/UpgradeTaskSource.cs
+++ b/Engine/Source/Programs/Horde/HordeServer/Tasks/UpgradeTaskSource.cs
@@ -11,15 +11,15 @@ using EpicGames.Horde.Logs;
 using EpicGames.Horde.Tools;
 using Google.Protobuf;
 using Google.Protobuf.WellKnownTypes;
-using Horde.Server.Agents;
-using Horde.Server.Logs;
-using Horde.Server.Server;
-using Horde.Server.Tools;
-using Horde.Server.Utilities;
+using HordeServer.Agents;
+using HordeServer.Logs;
+using HordeServer.Server;
+using HordeServer.Tools;
+using HordeServer.Utilities;
 using HordeCommon.Rpc.Tasks;
 using Microsoft.Extensions.Options;
 
-namespace Horde.Server.Tasks
+namespace HordeServer.Tasks
 {
 	class UpgradeTaskSource : TaskSourceBase<UpgradeTask>
 	{
diff --git a/Engine/Source/Programs/Horde/HordeServer/Telemetry/ITelemetrySink.cs b/Engine/Source/Programs/Horde/HordeServer/Telemetry/ITelemetrySink.cs
--- a/Engine/Source/Programs/Horde/HordeServer/Telemetry/ITelemetrySink.cs
+++ b/Engine/Source/Programs/Horde/HordeServer/Telemetry/ITelemetrySink.cs
@@ -5,7 +5,7 @@ using System.Threading;
 using System.Threading.Tasks;
 using EpicGames.Horde.Telemetry;
 
-namespace Horde.Server.Telemetry
+namespace HordeServer.Telemetry
 {
 	/// <summary>
 	/// Interface for a telemetry sink
diff --git a/Engine/Source/Programs/Horde/HordeServer/Telemetry/Metrics/IMetric.cs b/Engine/Source/Programs/Horde/HordeServer/Telemetry/Metrics/IMetric.cs
--- a/Engine/Source/Programs/Horde/HordeServer/Telemetry/Metrics/IMetric.cs
+++ b/Engine/Source/Programs/Horde/HordeServer/Telemetry/Metrics/IMetric.cs
@@ -3,7 +3,7 @@
 using System;
 using EpicGames.Horde.Telemetry.Metrics;
 
-namespace Horde.Server.Telemetry.Metrics
+namespace HordeServer.Telemetry.Metrics
 {
 	/// <summary>
 	/// Interface for a metric event
diff --git a/Engine/Source/Programs/Horde/HordeServer/Telemetry/Metrics/IMetricCollection.cs b/Engine/Source/Programs/Horde/HordeServer/Telemetry/Metrics/IMetricCollection.cs
--- a/Engine/Source/Programs/Horde/HordeServer/Telemetry/Metrics/IMetricCollection.cs
+++ b/Engine/Source/Programs/Horde/HordeServer/Telemetry/Metrics/IMetricCollection.cs
@@ -8,7 +8,7 @@ using System.Threading.Tasks;
 using EpicGames.Horde.Telemetry;
 using EpicGames.Horde.Telemetry.Metrics;
 
-namespace Horde.Server.Telemetry.Metrics
+namespace HordeServer.Telemetry.Metrics
 {
 	/// <summary>
 	/// Collection of aggregated metrics
diff --git a/Engine/Source/Programs/Horde/HordeServer/Telemetry/Metrics/MetricCollection.cs b/Engine/Source/Programs/Horde/HordeServer/Telemetry/Metrics/MetricCollection.cs
--- a/Engine/Source/Programs/Horde/HordeServer/Telemetry/Metrics/MetricCollection.cs
+++ b/Engine/Source/Programs/Horde/HordeServer/Telemetry/Metrics/MetricCollection.cs
@@ -9,8 +9,8 @@ using System.Threading.Tasks;
 using EpicGames.Core;
 using EpicGames.Horde.Telemetry;
 using EpicGames.Horde.Telemetry.Metrics;
-using Horde.Server.Server;
-using Horde.Server.Utilities;
+using HordeServer.Server;
+using HordeServer.Utilities;
 using Json.Path;
 using Microsoft.Extensions.Hosting;
 using Microsoft.Extensions.Logging;
@@ -20,7 +20,7 @@ using MongoDB.Bson.Serialization.Attributes;
 using MongoDB.Driver;
 using TDigestNet;
 
-namespace Horde.Server.Telemetry.Metrics
+namespace HordeServer.Telemetry.Metrics
 {
 	class MetricCollection : IMetricCollection, IHostedService, IAsyncDisposable
 	{
diff --git a/Engine/Source/Programs/Horde/HordeServer/Telemetry/Metrics/MetricConfig.cs b/Engine/Source/Programs/Horde/HordeServer/Telemetry/Metrics/MetricConfig.cs
--- a/Engine/Source/Programs/Horde/HordeServer/Telemetry/Metrics/MetricConfig.cs
+++ b/Engine/Source/Programs/Horde/HordeServer/Telemetry/Metrics/MetricConfig.cs
@@ -8,13 +8,13 @@ using System.Text.Json.Serialization;
 using EpicGames.Core;
 using EpicGames.Horde.Telemetry;
 using EpicGames.Horde.Telemetry.Metrics;
-using Horde.Server.Acls;
-using Horde.Server.Configuration;
-using Horde.Server.Server;
-using Horde.Server.Utilities;
+using HordeServer.Acls;
+using HordeServer.Configuration;
+using HordeServer.Server;
+using HordeServer.Utilities;
 using Json.Path;
 
-namespace Horde.Server.Telemetry.Metrics
+namespace HordeServer.Telemetry.Metrics
 {
 	/// <summary>
 	/// Method for aggregating samples into a metric
diff --git a/Engine/Source/Programs/Horde/HordeServer/Telemetry/Sinks/ClickHouseTelemetrySink.cs b/Engine/Source/Programs/Horde/HordeServer/Telemetry/Sinks/ClickHouseTelemetrySink.cs
--- a/Engine/Source/Programs/Horde/HordeServer/Telemetry/Sinks/ClickHouseTelemetrySink.cs
+++ b/Engine/Source/Programs/Horde/HordeServer/Telemetry/Sinks/ClickHouseTelemetrySink.cs
@@ -12,7 +12,7 @@ using EpicGames.Horde.Telemetry;
 using HordeCommon.Rpc.Messages.Telemetry;
 using Microsoft.Extensions.Logging;
 
-namespace Horde.Server.Telemetry.Sinks
+namespace HordeServer.Telemetry.Sinks
 {
 	/// <summary>
 	/// ClickHouse telemetry sink
diff --git a/Engine/Source/Programs/Horde/HordeServer/Telemetry/Sinks/EpicTelemetrySink.cs b/Engine/Source/Programs/Horde/HordeServer/Telemetry/Sinks/EpicTelemetrySink.cs
--- a/Engine/Source/Programs/Horde/HordeServer/Telemetry/Sinks/EpicTelemetrySink.cs
+++ b/Engine/Source/Programs/Horde/HordeServer/Telemetry/Sinks/EpicTelemetrySink.cs
@@ -17,7 +17,7 @@ using Microsoft.Extensions.Logging;
 using Polly;
 using Polly.Extensions.Http;
 
-namespace Horde.Server.Telemetry.Sinks
+namespace HordeServer.Telemetry.Sinks
 {
 	/// <summary>
 	/// Options for the Epic telemetry sink
diff --git a/Engine/Source/Programs/Horde/HordeServer/Telemetry/Sinks/MetricTelemetrySink.cs b/Engine/Source/Programs/Horde/HordeServer/Telemetry/Sinks/MetricTelemetrySink.cs
--- a/Engine/Source/Programs/Horde/HordeServer/Telemetry/Sinks/MetricTelemetrySink.cs
+++ b/Engine/Source/Programs/Horde/HordeServer/Telemetry/Sinks/MetricTelemetrySink.cs
@@ -5,9 +5,9 @@ using System.Text.Json.Nodes;
 using System.Threading;
 using System.Threading.Tasks;
 using EpicGames.Horde.Telemetry;
-using Horde.Server.Telemetry.Metrics;
+using HordeServer.Telemetry.Metrics;
 
-namespace Horde.Server.Telemetry.Sinks
+namespace HordeServer.Telemetry.Sinks
 {
 	/// <summary>
 	/// Consumes telemetry events and generates metrics
diff --git a/Engine/Source/Programs/Horde/HordeServer/Telemetry/Sinks/MongoTelemetrySink.cs b/Engine/Source/Programs/Horde/HordeServer/Telemetry/Sinks/MongoTelemetrySink.cs
--- a/Engine/Source/Programs/Horde/HordeServer/Telemetry/Sinks/MongoTelemetrySink.cs
+++ b/Engine/Source/Programs/Horde/HordeServer/Telemetry/Sinks/MongoTelemetrySink.cs
@@ -7,8 +7,8 @@ using System.Threading;
 using System.Threading.Tasks;
 using EpicGames.Core;
 using EpicGames.Horde.Telemetry;
-using Horde.Server.Server;
-using Horde.Server.Utilities;
+using HordeServer.Server;
+using HordeServer.Utilities;
 using Microsoft.Extensions.Hosting;
 using Microsoft.Extensions.Logging;
 using Microsoft.Extensions.Options;
@@ -16,7 +16,7 @@ using MongoDB.Bson;
 using MongoDB.Bson.Serialization.Attributes;
 using MongoDB.Driver;
 
-namespace Horde.Server.Telemetry.Sinks
+namespace HordeServer.Telemetry.Sinks
 {
 	/// <summary>
 	/// Telemetry sink which writes data to a MongoDB collection
diff --git a/Engine/Source/Programs/Horde/HordeServer/Telemetry/TelemetryAclAction.cs b/Engine/Source/Programs/Horde/HordeServer/Telemetry/TelemetryAclAction.cs
--- a/Engine/Source/Programs/Horde/HordeServer/Telemetry/TelemetryAclAction.cs
+++ b/Engine/Source/Programs/Horde/HordeServer/Telemetry/TelemetryAclAction.cs
@@ -2,7 +2,7 @@
 
 using EpicGames.Horde.Acls;
 
-namespace Horde.Server.Telemetry
+namespace HordeServer.Telemetry
 {
 	static class TelemetryAclAction
 	{
diff --git a/Engine/Source/Programs/Horde/HordeServer/Telemetry/TelemetryController.cs b/Engine/Source/Programs/Horde/HordeServer/Telemetry/TelemetryController.cs
--- a/Engine/Source/Programs/Horde/HordeServer/Telemetry/TelemetryController.cs
+++ b/Engine/Source/Programs/Horde/HordeServer/Telemetry/TelemetryController.cs
@@ -9,14 +9,14 @@ using System.Threading;
 using System.Threading.Tasks;
 using EpicGames.Horde.Telemetry;
 using EpicGames.Horde.Telemetry.Metrics;
-using Horde.Server.Server;
-using Horde.Server.Telemetry.Metrics;
-using Horde.Server.Utilities;
+using HordeServer.Server;
+using HordeServer.Telemetry.Metrics;
+using HordeServer.Utilities;
 using Microsoft.AspNetCore.Authorization;
 using Microsoft.AspNetCore.Mvc;
 using Microsoft.Extensions.Options;
 
-namespace Horde.Server.Telemetry
+namespace HordeServer.Telemetry
 {
 	/// <summary>
 	/// Controller for the /api/v1/telemetry endpoint
diff --git a/Engine/Source/Programs/Horde/HordeServer/Telemetry/TelemetryManager.cs b/Engine/Source/Programs/Horde/HordeServer/Telemetry/TelemetryManager.cs
--- a/Engine/Source/Programs/Horde/HordeServer/Telemetry/TelemetryManager.cs
+++ b/Engine/Source/Programs/Horde/HordeServer/Telemetry/TelemetryManager.cs
@@ -8,14 +8,14 @@ using System.Threading;
 using System.Threading.Tasks;
 using EpicGames.Core;
 using EpicGames.Horde.Telemetry;
-using Horde.Server.Telemetry.Sinks;
+using HordeServer.Telemetry.Sinks;
 using Microsoft.Extensions.DependencyInjection;
 using Microsoft.Extensions.Hosting;
 using Microsoft.Extensions.Logging;
 using Microsoft.Extensions.Options;
 using OpenTelemetry.Trace;
 
-namespace Horde.Server.Telemetry
+namespace HordeServer.Telemetry
 {
 	/// <summary>
 	/// Telemetry sink dispatching incoming events to all registered sinks
diff --git a/Engine/Source/Programs/Horde/HordeServer/Tools/ITool.cs b/Engine/Source/Programs/Horde/HordeServer/Tools/ITool.cs
--- a/Engine/Source/Programs/Horde/HordeServer/Tools/ITool.cs
+++ b/Engine/Source/Programs/Horde/HordeServer/Tools/ITool.cs
@@ -10,7 +10,7 @@ using EpicGames.Horde.Acls;
 using EpicGames.Horde.Storage;
 using EpicGames.Horde.Tools;
 
-namespace Horde.Server.Tools
+namespace HordeServer.Tools
 {
 	/// <summary>
 	/// Describes a standalone, external tool hosted and deployed by Horde. Provides basic functionality for performing
diff --git a/Engine/Source/Programs/Horde/HordeServer/Tools/IToolCollection.cs b/Engine/Source/Programs/Horde/HordeServer/Tools/IToolCollection.cs
--- a/Engine/Source/Programs/Horde/HordeServer/Tools/IToolCollection.cs
+++ b/Engine/Source/Programs/Horde/HordeServer/Tools/IToolCollection.cs
@@ -5,7 +5,7 @@ using System.Threading;
 using System.Threading.Tasks;
 using EpicGames.Horde.Tools;
 
-namespace Horde.Server.Tools
+namespace HordeServer.Tools
 {
 	/// <summary>
 	/// Collection of tools
diff --git a/Engine/Source/Programs/Horde/HordeServer/Tools/ToolCollection.cs b/Engine/Source/Programs/Horde/HordeServer/Tools/ToolCollection.cs
--- a/Engine/Source/Programs/Horde/HordeServer/Tools/ToolCollection.cs
+++ b/Engine/Source/Programs/Horde/HordeServer/Tools/ToolCollection.cs
@@ -18,15 +18,15 @@ using EpicGames.Horde.Storage.Bundles;
 using EpicGames.Horde.Storage.Nodes;
 using EpicGames.Horde.Storage.ObjectStores;
 using EpicGames.Horde.Tools;
-using Horde.Server.Server;
-using Horde.Server.Storage;
-using Horde.Server.Utilities;
+using HordeServer.Server;
+using HordeServer.Storage;
+using HordeServer.Utilities;
 using Microsoft.Extensions.Logging;
 using Microsoft.Extensions.Options;
 using MongoDB.Bson.Serialization.Attributes;
 using MongoDB.Driver;
 
-namespace Horde.Server.Tools
+namespace HordeServer.Tools
 {
 	/// <summary>
 	/// Collection of tool documents
diff --git a/Engine/Source/Programs/Horde/HordeServer/Tools/ToolConfig.cs b/Engine/Source/Programs/Horde/HordeServer/Tools/ToolConfig.cs
--- a/Engine/Source/Programs/Horde/HordeServer/Tools/ToolConfig.cs
+++ b/Engine/Source/Programs/Horde/HordeServer/Tools/ToolConfig.cs
@@ -9,11 +9,11 @@ using System.Text.Json.Serialization;
 using EpicGames.Horde.Acls;
 using EpicGames.Horde.Storage;
 using EpicGames.Horde.Tools;
-using Horde.Server.Acls;
-using Horde.Server.Server;
-using Horde.Server.Storage;
+using HordeServer.Acls;
+using HordeServer.Server;
+using HordeServer.Storage;
 
-namespace Horde.Server.Tools
+namespace HordeServer.Tools
 {
 	/// <summary>
 	/// Options for configuring a tool
diff --git a/Engine/Source/Programs/Horde/HordeServer/Tools/ToolsController.cs b/Engine/Source/Programs/Horde/HordeServer/Tools/ToolsController.cs
--- a/Engine/Source/Programs/Horde/HordeServer/Tools/ToolsController.cs
+++ b/Engine/Source/Programs/Horde/HordeServer/Tools/ToolsController.cs
@@ -10,15 +10,15 @@ using EpicGames.Core;
 using EpicGames.Horde.Storage;
 using EpicGames.Horde.Storage.Nodes;
 using EpicGames.Horde.Tools;
-using Horde.Server.Storage;
-using Horde.Server.Utilities;
+using HordeServer.Storage;
+using HordeServer.Utilities;
 using Microsoft.AspNetCore.Authorization;
 using Microsoft.AspNetCore.Http;
 using Microsoft.AspNetCore.Mvc;
 using Microsoft.AspNetCore.StaticFiles;
 using Microsoft.Extensions.Logging;
 
-namespace Horde.Server.Tools
+namespace HordeServer.Tools
 {
 	/// <summary>
 	/// Controller for the /api/v1/tools endpoint
diff --git a/Engine/Source/Programs/Horde/HordeServer/Ugs/IUgsMetadataCollection.cs b/Engine/Source/Programs/Horde/HordeServer/Ugs/IUgsMetadataCollection.cs
--- a/Engine/Source/Programs/Horde/HordeServer/Ugs/IUgsMetadataCollection.cs
+++ b/Engine/Source/Programs/Horde/HordeServer/Ugs/IUgsMetadataCollection.cs
@@ -5,7 +5,7 @@ using System.Collections.Generic;
 using System.Threading;
 using System.Threading.Tasks;
 
-namespace Horde.Server.Ugs
+namespace HordeServer.Ugs
 {
 	/// <summary>
 	/// Collection of stream documents
diff --git a/Engine/Source/Programs/Horde/HordeServer/Ugs/UgsController.cs b/Engine/Source/Programs/Horde/HordeServer/Ugs/UgsController.cs
--- a/Engine/Source/Programs/Horde/HordeServer/Ugs/UgsController.cs
+++ b/Engine/Source/Programs/Horde/HordeServer/Ugs/UgsController.cs
@@ -7,16 +7,16 @@ using System.Threading;
 using System.Threading.Tasks;
 using EpicGames.Horde.Issues;
 using EpicGames.Horde.Users;
-using Horde.Server.Issues;
-using Horde.Server.Logs;
-using Horde.Server.Users;
-using Horde.Server.Utilities;
+using HordeServer.Issues;
+using HordeServer.Logs;
+using HordeServer.Users;
+using HordeServer.Utilities;
 using Microsoft.AspNetCore.Mvc;
 using Microsoft.Extensions.Logging;
 using Microsoft.Extensions.Options;
 using MongoDB.Bson;
 
-namespace Horde.Server.Ugs
+namespace HordeServer.Ugs
 {
 	/// <summary>
 	/// Controller for the /api/v1/issues endpoint
diff --git a/Engine/Source/Programs/Horde/HordeServer/Ugs/UgsMessages.cs b/Engine/Source/Programs/Horde/HordeServer/Ugs/UgsMessages.cs
--- a/Engine/Source/Programs/Horde/HordeServer/Ugs/UgsMessages.cs
+++ b/Engine/Source/Programs/Horde/HordeServer/Ugs/UgsMessages.cs
@@ -6,10 +6,10 @@ using System.ComponentModel.DataAnnotations;
 using System.Linq;
 using EpicGames.Core;
 using EpicGames.Horde.Issues;
-using Horde.Server.Issues;
-using Horde.Server.Users;
+using HordeServer.Issues;
+using HordeServer.Users;
 
-namespace Horde.Server.Ugs
+namespace HordeServer.Ugs
 {
 	/// <summary>
 	/// Review by a user of a particular change
diff --git a/Engine/Source/Programs/Horde/HordeServer/Ugs/UgsMetadata.cs b/Engine/Source/Programs/Horde/HordeServer/Ugs/UgsMetadata.cs
--- a/Engine/Source/Programs/Horde/HordeServer/Ugs/UgsMetadata.cs
+++ b/Engine/Source/Programs/Horde/HordeServer/Ugs/UgsMetadata.cs
@@ -3,7 +3,7 @@
 using System;
 using System.Collections.Generic;
 
-namespace Horde.Server.Ugs
+namespace HordeServer.Ugs
 {
 	/// <summary>
 	/// Metadata for a particular user
diff --git a/Engine/Source/Programs/Horde/HordeServer/Ugs/UgsMetadataCollection.cs b/Engine/Source/Programs/Horde/HordeServer/Ugs/UgsMetadataCollection.cs
--- a/Engine/Source/Programs/Horde/HordeServer/Ugs/UgsMetadataCollection.cs
+++ b/Engine/Source/Programs/Horde/HordeServer/Ugs/UgsMetadataCollection.cs
@@ -6,14 +6,14 @@ using System.Diagnostics.CodeAnalysis;
 using System.Linq;
 using System.Threading;
 using System.Threading.Tasks;
-using Horde.Server.Server;
-using Horde.Server.Utilities;
+using HordeServer.Server;
+using HordeServer.Utilities;
 using Microsoft.Extensions.Logging;
 using MongoDB.Bson;
 using MongoDB.Bson.Serialization.Attributes;
 using MongoDB.Driver;
 
-namespace Horde.Server.Ugs
+namespace HordeServer.Ugs
 {
 	/// <summary>
 	/// Concrete implementation of IUgsMetadataCollection
diff --git a/Engine/Source/Programs/Horde/HordeServer/Users/IAvatar.cs b/Engine/Source/Programs/Horde/HordeServer/Users/IAvatar.cs
--- a/Engine/Source/Programs/Horde/HordeServer/Users/IAvatar.cs
+++ b/Engine/Source/Programs/Horde/HordeServer/Users/IAvatar.cs
@@ -1,6 +1,6 @@
 // Copyright Epic Games, Inc. All Rights Reserved.
 
-namespace Horde.Server.Users
+namespace HordeServer.Users
 {
 	/// <summary>
 	/// Avatar for the user
diff --git a/Engine/Source/Programs/Horde/HordeServer/Users/IAvatarService.cs b/Engine/Source/Programs/Horde/HordeServer/Users/IAvatarService.cs
--- a/Engine/Source/Programs/Horde/HordeServer/Users/IAvatarService.cs
+++ b/Engine/Source/Programs/Horde/HordeServer/Users/IAvatarService.cs
@@ -3,7 +3,7 @@
 using System.Threading;
 using System.Threading.Tasks;
 
-namespace Horde.Server.Users
+namespace HordeServer.Users
 {
 	/// <summary>
 	/// Provides avatar images
diff --git a/Engine/Source/Programs/Horde/HordeServer/Users/IUserCollection.cs b/Engine/Source/Programs/Horde/HordeServer/Users/IUserCollection.cs
--- a/Engine/Source/Programs/Horde/HordeServer/Users/IUserCollection.cs
+++ b/Engine/Source/Programs/Horde/HordeServer/Users/IUserCollection.cs
@@ -7,10 +7,10 @@ using System.Threading.Tasks;
 using EpicGames.Horde.Jobs;
 using EpicGames.Horde.Jobs.Bisect;
 using EpicGames.Horde.Users;
-using Horde.Server.Utilities;
+using HordeServer.Utilities;
 using MongoDB.Bson;
 
-namespace Horde.Server.Users
+namespace HordeServer.Users
 {
 	/// <summary>
 	/// Manages user documents
diff --git a/Engine/Source/Programs/Horde/HordeServer/Users/UserCollection.cs b/Engine/Source/Programs/Horde/HordeServer/Users/UserCollection.cs
--- a/Engine/Source/Programs/Horde/HordeServer/Users/UserCollection.cs
+++ b/Engine/Source/Programs/Horde/HordeServer/Users/UserCollection.cs
@@ -10,14 +10,14 @@ using System.Threading.Tasks;
 using EpicGames.Horde.Jobs;
 using EpicGames.Horde.Jobs.Bisect;
 using EpicGames.Horde.Users;
-using Horde.Server.Server;
-using Horde.Server.Utilities;
+using HordeServer.Server;
+using HordeServer.Utilities;
 using MongoDB.Bson;
 using MongoDB.Bson.Serialization;
 using MongoDB.Bson.Serialization.Attributes;
 using MongoDB.Driver;
 
-namespace Horde.Server.Users
+namespace HordeServer.Users
 {
 	/// <summary>
 	/// Manages user documents
diff --git a/Engine/Source/Programs/Horde/HordeServer/Users/UserCollectionV2.cs b/Engine/Source/Programs/Horde/HordeServer/Users/UserCollectionV2.cs
--- a/Engine/Source/Programs/Horde/HordeServer/Users/UserCollectionV2.cs
+++ b/Engine/Source/Programs/Horde/HordeServer/Users/UserCollectionV2.cs
@@ -10,8 +10,8 @@ using EpicGames.Horde.Jobs.Bisect;
 using EpicGames.Horde.Jobs.Templates;
 using EpicGames.Horde.Streams;
 using EpicGames.Horde.Users;
-using Horde.Server.Server;
-using Horde.Server.Utilities;
+using HordeServer.Server;
+using HordeServer.Utilities;
 using Microsoft.Extensions.Caching.Memory;
 using Microsoft.Extensions.Logging;
 using MongoDB.Bson;
@@ -20,7 +20,7 @@ using MongoDB.Bson.Serialization.Attributes;
 using MongoDB.Driver;
 using MongoDB.Driver.Linq;
 
-namespace Horde.Server.Users
+namespace HordeServer.Users
 {
 	/// <summary>
 	/// Manages user documents
diff --git a/Engine/Source/Programs/Horde/HordeServer/Users/UserController.cs b/Engine/Source/Programs/Horde/HordeServer/Users/UserController.cs
--- a/Engine/Source/Programs/Horde/HordeServer/Users/UserController.cs
+++ b/Engine/Source/Programs/Horde/HordeServer/Users/UserController.cs
@@ -7,17 +7,17 @@ using System.Threading.Tasks;
 using EpicGames.Horde.Jobs;
 using EpicGames.Horde.Jobs.Bisect;
 using EpicGames.Horde.Users;
-using Horde.Server.Accounts;
-using Horde.Server.Agents;
-using Horde.Server.Agents.Pools;
-using Horde.Server.Server;
-using Horde.Server.Server.Notices;
-using Horde.Server.Utilities;
+using HordeServer.Accounts;
+using HordeServer.Agents;
+using HordeServer.Agents.Pools;
+using HordeServer.Server;
+using HordeServer.Server.Notices;
+using HordeServer.Utilities;
 using Microsoft.AspNetCore.Authorization;
 using Microsoft.AspNetCore.Mvc;
 using Microsoft.Extensions.Options;
 
-namespace Horde.Server.Users
+namespace HordeServer.Users
 {
 	/// <summary>
 	/// Controller for the /api/v1/user endpoint
diff --git a/Engine/Source/Programs/Horde/HordeServer/Users/UsersController.cs b/Engine/Source/Programs/Horde/HordeServer/Users/UsersController.cs
--- a/Engine/Source/Programs/Horde/HordeServer/Users/UsersController.cs
+++ b/Engine/Source/Programs/Horde/HordeServer/Users/UsersController.cs
@@ -6,11 +6,11 @@ using System.Linq;
 using System.Threading;
 using System.Threading.Tasks;
 using EpicGames.Horde.Users;
-using Horde.Server.Utilities;
+using HordeServer.Utilities;
 using Microsoft.AspNetCore.Authorization;
 using Microsoft.AspNetCore.Mvc;
 
-namespace Horde.Server.Users
+namespace HordeServer.Users
 {
 	/// <summary>
 	/// Controller for the /api/v1/users endpoint
diff --git a/Engine/Source/Programs/Horde/HordeServer/Utilities/AgentUtilities.cs b/Engine/Source/Programs/Horde/HordeServer/Utilities/AgentUtilities.cs
--- a/Engine/Source/Programs/Horde/HordeServer/Utilities/AgentUtilities.cs
+++ b/Engine/Source/Programs/Horde/HordeServer/Utilities/AgentUtilities.cs
@@ -9,7 +9,7 @@ using System.Text.Json;
 
 #nullable disable
 
-namespace Horde.Server.Utilities
+namespace HordeServer.Utilities
 {
 	static class AgentUtilities
 	{
diff --git a/Engine/Source/Programs/Horde/HordeServer/Utilities/ArtifactChunkStream.cs b/Engine/Source/Programs/Horde/HordeServer/Utilities/ArtifactChunkStream.cs
--- a/Engine/Source/Programs/Horde/HordeServer/Utilities/ArtifactChunkStream.cs
+++ b/Engine/Source/Programs/Horde/HordeServer/Utilities/ArtifactChunkStream.cs
@@ -7,7 +7,7 @@ using System.Threading.Tasks;
 using Grpc.Core;
 using Horde.Common.Rpc;
 
-namespace Horde.Server.Utilities
+namespace HordeServer.Utilities
 {
 	/// <summary>
 	/// Stream that takes in rpc chunks from an upload and copies them?
diff --git a/Engine/Source/Programs/Horde/HordeServer/Utilities/Clock.cs b/Engine/Source/Programs/Horde/HordeServer/Utilities/Clock.cs
--- a/Engine/Source/Programs/Horde/HordeServer/Utilities/Clock.cs
+++ b/Engine/Source/Programs/Horde/HordeServer/Utilities/Clock.cs
@@ -6,8 +6,8 @@ using System.Threading;
 using System.Threading.Tasks;
 using EpicGames.Core;
 using EpicGames.Redis.Utility;
-using Horde.Server;
-using Horde.Server.Server;
+using HordeServer;
+using HordeServer.Server;
 using Microsoft.Extensions.Logging;
 using Microsoft.Extensions.Options;
 using StackExchange.Redis;
diff --git a/Engine/Source/Programs/Horde/HordeServer/Utilities/HordeClaimTypes.cs b/Engine/Source/Programs/Horde/HordeServer/Utilities/HordeClaimTypes.cs
--- a/Engine/Source/Programs/Horde/HordeServer/Utilities/HordeClaimTypes.cs
+++ b/Engine/Source/Programs/Horde/HordeServer/Utilities/HordeClaimTypes.cs
@@ -3,9 +3,9 @@
 using System.Security.Claims;
 using EpicGames.Horde.Accounts;
 using EpicGames.Horde.Users;
-using Horde.Server.Users;
+using HordeServer.Users;
 
-namespace Horde.Server.Utilities
+namespace HordeServer.Utilities
 {
 	/// <summary>
 	/// Claim types that are specific to horde
diff --git a/Engine/Source/Programs/Horde/HordeServer/Utilities/HordeClaims.cs b/Engine/Source/Programs/Horde/HordeServer/Utilities/HordeClaims.cs
--- a/Engine/Source/Programs/Horde/HordeServer/Utilities/HordeClaims.cs
+++ b/Engine/Source/Programs/Horde/HordeServer/Utilities/HordeClaims.cs
@@ -2,9 +2,9 @@
 
 using EpicGames.Horde.Agents;
 using EpicGames.Horde.Agents.Sessions;
-using Horde.Server.Acls;
+using HordeServer.Acls;
 
-namespace Horde.Server.Utilities
+namespace HordeServer.Utilities
 {
 	/// <summary>
 	/// Predetermined claim values
diff --git a/Engine/Source/Programs/Horde/HordeServer/Utilities/HordeControllerBase.cs b/Engine/Source/Programs/Horde/HordeServer/Utilities/HordeControllerBase.cs
--- a/Engine/Source/Programs/Horde/HordeServer/Utilities/HordeControllerBase.cs
+++ b/Engine/Source/Programs/Horde/HordeServer/Utilities/HordeControllerBase.cs
@@ -16,12 +16,12 @@ using EpicGames.Horde.Storage;
 using EpicGames.Horde.Streams;
 using EpicGames.Horde.Telemetry;
 using EpicGames.Horde.Tools;
-using Horde.Server.Acls;
+using HordeServer.Acls;
 using Microsoft.AspNetCore.Http;
 using Microsoft.AspNetCore.Mvc;
 using Microsoft.Extensions.Logging;
 
-namespace Horde.Server.Utilities
+namespace HordeServer.Utilities
 {
 	/// <summary>
 	/// Base class for Horde controllers
diff --git a/Engine/Source/Programs/Horde/HordeServer/Utilities/MongoCommandTracer.cs b/Engine/Source/Programs/Horde/HordeServer/Utilities/MongoCommandTracer.cs
--- a/Engine/Source/Programs/Horde/HordeServer/Utilities/MongoCommandTracer.cs
+++ b/Engine/Source/Programs/Horde/HordeServer/Utilities/MongoCommandTracer.cs
@@ -10,7 +10,7 @@ using MongoDB.Driver.Core.Configuration;
 using MongoDB.Driver.Core.Events;
 using OpenTelemetry.Trace;
 
-namespace Horde.Server.Utilities;
+namespace HordeServer.Utilities;
 
 /// <summary>
 /// OpenTelemetry-based tracer listening for MongoDB command events
diff --git a/Engine/Source/Programs/Horde/HordeServer/Utilities/OpenTelemetry.cs b/Engine/Source/Programs/Horde/HordeServer/Utilities/OpenTelemetry.cs
--- a/Engine/Source/Programs/Horde/HordeServer/Utilities/OpenTelemetry.cs
+++ b/Engine/Source/Programs/Horde/HordeServer/Utilities/OpenTelemetry.cs
@@ -23,7 +23,7 @@ using OpenTelemetry.Resources;
 using OpenTelemetry.Trace;
 using Serilog.Core;
 
-namespace Horde.Server.Utilities;
+namespace HordeServer.Utilities;
 
 /// <summary>
 /// Serilog event enricher attaching trace and span ID for Datadog using current System.Diagnostics.Activity
diff --git a/Engine/Source/Programs/Horde/HordeServer/Utilities/SingletonDocument.cs b/Engine/Source/Programs/Horde/HordeServer/Utilities/SingletonDocument.cs
--- a/Engine/Source/Programs/Horde/HordeServer/Utilities/SingletonDocument.cs
+++ b/Engine/Source/Programs/Horde/HordeServer/Utilities/SingletonDocument.cs
@@ -4,10 +4,10 @@ using System;
 using System.Threading;
 using System.Threading.Tasks;
 using EpicGames.Horde;
-using Horde.Server.Server;
+using HordeServer.Server;
 using MongoDB.Bson.Serialization.Attributes;
 
-namespace Horde.Server.Utilities
+namespace HordeServer.Utilities
 {
 	/// <summary>
 	/// Identifier for a pool
diff --git a/Engine/Source/Programs/Horde/HordeServer/appsettings.Development.json b/Engine/Source/Programs/Horde/HordeServer/appsettings.Development.json
--- a/Engine/Source/Programs/Horde/HordeServer/appsettings.Development.json
+++ b/Engine/Source/Programs/Horde/HordeServer/appsettings.Development.json
@@ -8,7 +8,7 @@
 		{
 			"Default": "Debug",
 			"Override": {
-				"Horde.Server.Perforce.PerforceService": "Verbose"
+				"HordeServer.Perforce.PerforceService": "Verbose"
 			}
 		}
 	}
diff --git a/Engine/Source/Programs/Horde/HordeServer/appsettings.Local.json b/Engine/Source/Programs/Horde/HordeServer/appsettings.Local.json
--- a/Engine/Source/Programs/Horde/HordeServer/appsettings.Local.json
+++ b/Engine/Source/Programs/Horde/HordeServer/appsettings.Local.json
@@ -28,8 +28,8 @@
 	"Serilog": {
 		"MinimumLevel": {
 			"Override": {
-				"Horde.Server.Agents.Fleet": "Warning",
-				"Horde.Server.Agents.Fleet.Providers.NoOpFleetManager": "Warning"
+				"HordeServer.Agents.Fleet": "Warning",
+				"HordeServer.Agents.Fleet.Providers.NoOpFleetManager": "Warning"
 			}
 		}
 	}
diff --git a/Engine/Source/Programs/Horde/HordeServer/appsettings.json b/Engine/Source/Programs/Horde/HordeServer/appsettings.json
--- a/Engine/Source/Programs/Horde/HordeServer/appsettings.json
+++ b/Engine/Source/Programs/Horde/HordeServer/appsettings.json
@@ -14,12 +14,12 @@
 				"Microsoft": "Warning",
 				"Microsoft.AspNetCore": "Warning",
 				"Microsoft.Hosting": "Information",
-				"Horde.Server.Authentication": "Warning",
+				"HordeServer.Authentication": "Warning",
 				"System.Net.Http.HttpClient": "Warning",
 				"Grpc": "Warning",
-				"Horde.Server.Issues.IssueService": "Debug",
-				"Horde.Server.Storage.ObjectStores.AwsObjectStore": "Information",
-				"Horde.Server.Storage.StorageService": "Debug",
+				"HordeServer.Issues.IssueService": "Debug",
+				"HordeServer.Storage.ObjectStores.AwsObjectStore": "Information",
+				"HordeServer.Storage.StorageService": "Debug",
 				"EpicGames.Horde.Storage": "Debug"
 			}
 		}
