Horde: Add an interface for the redis service.

#rnx

diff --git a/Engine/Source/Programs/Horde/HordeServer/Agents/AgentService.cs b/Engine/Source/Programs/Horde/HordeServer/Agents/AgentService.cs
--- a/Engine/Source/Programs/Horde/HordeServer/Agents/AgentService.cs
+++ b/Engine/Source/Programs/Horde/HordeServer/Agents/AgentService.cs
@@ -76,7 +76,7 @@ namespace HordeServer.Agents
 		readonly AclService _aclService;
 		readonly IDowntimeService _downtimeService;
 		readonly ITaskSource[] _taskSources;
-		readonly RedisService _redisService;
+		readonly IRedisService _redisService;
 		readonly IHostApplicationLifetime _applicationLifetime;
 		readonly IClock _clock;
 		readonly Tracer _tracer;
@@ -115,7 +115,7 @@ namespace HordeServer.Agents
 			IDowntimeService downtimeService,
 			IPoolCollection poolCollection,
 			IEnumerable<ITaskSource> taskSources,
-			RedisService redisService,
+			IRedisService redisService,
 			IHostApplicationLifetime applicationLifetime,
 			IClock clock,
 			Tracer tracer,
diff --git a/Engine/Source/Programs/Horde/HordeServer/Agents/Enrollment/EnrollmentService.cs b/Engine/Source/Programs/Horde/HordeServer/Agents/Enrollment/EnrollmentService.cs
--- a/Engine/Source/Programs/Horde/HordeServer/Agents/Enrollment/EnrollmentService.cs
+++ b/Engine/Source/Programs/Horde/HordeServer/Agents/Enrollment/EnrollmentService.cs
@@ -25,7 +25,7 @@ namespace HordeServer.Agents.Enrollment
 	/// </summary>
 	public sealed class EnrollmentService : IHostedService, IAsyncDisposable
 	{
-		readonly RedisService _redisService;
+		readonly IRedisService _redisService;
 		readonly IClock _clock;
 		readonly RedisChannel _updateChannel;
 		readonly RedisSortedSetKey<string> _keys = new("agents:registration:expire");
@@ -39,7 +39,7 @@ namespace HordeServer.Agents.Enrollment
 		/// <summary>
 		/// Constructor
 		/// </summary>
-		public EnrollmentService(RedisService redisService, IClock clock, IOptionsMonitor<ServerSettings> settings)
+		public EnrollmentService(IRedisService redisService, IClock clock, IOptionsMonitor<ServerSettings> settings)
 		{
 			_redisService = redisService;
 			_clock = clock;
diff --git a/Engine/Source/Programs/Horde/HordeServer/Agents/Fleet/AwsAutoScalingLifecycle.cs b/Engine/Source/Programs/Horde/HordeServer/Agents/Fleet/AwsAutoScalingLifecycle.cs
--- a/Engine/Source/Programs/Horde/HordeServer/Agents/Fleet/AwsAutoScalingLifecycle.cs
+++ b/Engine/Source/Programs/Horde/HordeServer/Agents/Fleet/AwsAutoScalingLifecycle.cs
@@ -64,7 +64,7 @@ public sealed class AwsAutoScalingLifecycleService : IHostedService, IAsyncDispo
 
 	private const string RedisKey = "aws-asg-lifecycle";
 	private readonly AgentService _agentService;
-	private readonly RedisService _redisService;
+	private readonly IRedisService _redisService;
 	private readonly IAgentCollection _agents;
 	private readonly IClock _clock;
 	private readonly Tracer _tracer;
@@ -83,7 +83,7 @@ public sealed class AwsAutoScalingLifecycleService : IHostedService, IAsyncDispo
 	/// </summary>
 	public AwsAutoScalingLifecycleService(
 		AgentService agentService,
-		RedisService redisService,
+		IRedisService redisService,
 		IAgentCollection agents,
 		IClock clock,
 		IOptionsMonitor<ServerSettings> settings,
diff --git a/Engine/Source/Programs/Horde/HordeServer/Agents/Relay/AgentRelayService.cs b/Engine/Source/Programs/Horde/HordeServer/Agents/Relay/AgentRelayService.cs
--- a/Engine/Source/Programs/Horde/HordeServer/Agents/Relay/AgentRelayService.cs
+++ b/Engine/Source/Programs/Horde/HordeServer/Agents/Relay/AgentRelayService.cs
@@ -81,7 +81,7 @@ public sealed class AgentRelayService : RelayRpc.RelayRpcBase, IHostedService, I
 	private static string KeyUsedPorts(ClusterId clusterId) => $"relay/used-ports/{clusterId}";
 	private static string KeyAgents(ClusterId clusterId) => $"relay/agents/{clusterId}";
 
-	private readonly RedisService _redis;
+	private readonly IRedisService _redis;
 	private readonly IClock _clock;
 	private readonly ILogger<AgentRelayService> _logger;
 	private readonly object _lock = new();
@@ -100,7 +100,7 @@ public sealed class AgentRelayService : RelayRpc.RelayRpcBase, IHostedService, I
 	/// <param name="redis"></param>
 	/// <param name="clock"></param>
 	/// <param name="logger"></param>
-	public AgentRelayService(RedisService redis, IClock clock, ILogger<AgentRelayService> logger)
+	public AgentRelayService(IRedisService redis, IClock clock, ILogger<AgentRelayService> logger)
 	{
 		_redis = redis;
 		_clock = clock;
diff --git a/Engine/Source/Programs/Horde/HordeServer/Compute/ComputeService.cs b/Engine/Source/Programs/Horde/HordeServer/Compute/ComputeService.cs
--- a/Engine/Source/Programs/Horde/HordeServer/Compute/ComputeService.cs
+++ b/Engine/Source/Programs/Horde/HordeServer/Compute/ComputeService.cs
@@ -176,7 +176,7 @@ namespace HordeServer.Compute
 		readonly ILogCollection _logCollection;
 		readonly AgentService _agentService;
 		readonly AgentRelayService _agentRelayService;
-		readonly RedisService _redisService;
+		readonly IRedisService _redisService;
 		readonly IOptionsMonitor<ServerSettings> _settings;
 		readonly IOptionsMonitor<GlobalConfig> _globalConfig;
 		readonly IClock _clock;
diff --git a/Engine/Source/Programs/Horde/HordeServer/Configuration/ConfigService.cs b/Engine/Source/Programs/Horde/HordeServer/Configuration/ConfigService.cs
--- a/Engine/Source/Programs/Horde/HordeServer/Configuration/ConfigService.cs
+++ b/Engine/Source/Programs/Horde/HordeServer/Configuration/ConfigService.cs
@@ -112,7 +112,7 @@ namespace HordeServer.Configuration
 		// The current config
 		record class ConfigState(IoHash Hash, GlobalConfig GlobalConfig);
 
-		readonly RedisService _redisService;
+		readonly IRedisService _redisService;
 		readonly ServerSettings _serverSettings;
 		readonly Dictionary<string, IConfigSource> _sources;
 		readonly JsonSerializerOptions _jsonOptions;
@@ -140,7 +140,7 @@ namespace HordeServer.Configuration
 		/// <summary>
 		/// Constructor
 		/// </summary>
-		public ConfigService(RedisService redisService, IOptions<ServerSettings> serverSettings, IEnumerable<IConfigSource> sources, IClock clock, IHealthMonitor<ConfigService> health, ILogger<ConfigService> logger)
+		public ConfigService(IRedisService redisService, IOptions<ServerSettings> serverSettings, IEnumerable<IConfigSource> sources, IClock clock, IHealthMonitor<ConfigService> health, ILogger<ConfigService> logger)
 		{
 			_redisService = redisService;
 			_serverSettings = serverSettings.Value;
diff --git a/Engine/Source/Programs/Horde/HordeServer/Issues/IssueCollection.cs b/Engine/Source/Programs/Horde/HordeServer/Issues/IssueCollection.cs
--- a/Engine/Source/Programs/Horde/HordeServer/Issues/IssueCollection.cs
+++ b/Engine/Source/Programs/Horde/HordeServer/Issues/IssueCollection.cs
@@ -588,7 +588,7 @@ namespace HordeServer.Issues
 			}
 		}
 
-		readonly RedisService _redisService;
+		readonly IRedisService _redisService;
 		readonly IUserCollection _userCollection;
 		readonly ISingletonDocument<IssueLedger> _ledgerSingleton;
 		readonly IMongoCollection<Issue> _issues;
@@ -619,7 +619,7 @@ namespace HordeServer.Issues
 			});
 		}
 
-		public IssueCollection(MongoService mongoService, RedisService redisService, IUserCollection userCollection, IAuditLogFactory<int> auditLogFactory, ITelemetrySink telemetrySink, IOptionsMonitor<GlobalConfig> globalConfig, Tracer tracer, ILogger<IssueCollection> logger)
+		public IssueCollection(MongoService mongoService, IRedisService redisService, IUserCollection userCollection, IAuditLogFactory<int> auditLogFactory, ITelemetrySink telemetrySink, IOptionsMonitor<GlobalConfig> globalConfig, Tracer tracer, ILogger<IssueCollection> logger)
 		{
 			_redisService = redisService;
 			_userCollection = userCollection;
diff --git a/Engine/Source/Programs/Horde/HordeServer/Jobs/Schedules/ScheduleService.cs b/Engine/Source/Programs/Horde/HordeServer/Jobs/Schedules/ScheduleService.cs
--- a/Engine/Source/Programs/Horde/HordeServer/Jobs/Schedules/ScheduleService.cs
+++ b/Engine/Source/Programs/Horde/HordeServer/Jobs/Schedules/ScheduleService.cs
@@ -68,7 +68,7 @@ namespace HordeServer.Jobs.Schedules
 		readonly IClock _clock;
 		readonly Tracer _tracer;
 		readonly ILogger _logger;
-		readonly RedisService _redis;
+		readonly IRedisService _redis;
 		static readonly RedisKey s_baseLockKey = "scheduler/locks";
 		static readonly RedisKey s_tickLockKey = s_baseLockKey.Append("/tick"); // Lock to tick the queue
 		static readonly RedisSortedSetKey<QueueItem> s_queueKey = "scheduler/queue"; // Items to tick, ordered by time
@@ -78,7 +78,7 @@ namespace HordeServer.Jobs.Schedules
 		/// <summary>
 		/// Constructor
 		/// </summary>
-		public ScheduleService(RedisService redis, IGraphCollection graphs, ICommitService commitService, IJobCollection jobCollection, JobService jobService, IDowntimeService downtimeService, IStreamCollection streamCollection, ITemplateCollection templateCollection, MongoService mongoService, IClock clock, IOptionsMonitor<GlobalConfig> globalConfig, Tracer tracer, ILogger<ScheduleService> logger)
+		public ScheduleService(IRedisService redis, IGraphCollection graphs, ICommitService commitService, IJobCollection jobCollection, JobService jobService, IDowntimeService downtimeService, IStreamCollection streamCollection, ITemplateCollection templateCollection, MongoService mongoService, IClock clock, IOptionsMonitor<GlobalConfig> globalConfig, Tracer tracer, ILogger<ScheduleService> logger)
 		{
 			_graphs = graphs;
 			_commitService = commitService;
diff --git a/Engine/Source/Programs/Horde/HordeServer/Logs/LogTailService.cs b/Engine/Source/Programs/Horde/HordeServer/Logs/LogTailService.cs
--- a/Engine/Source/Programs/Horde/HordeServer/Logs/LogTailService.cs
+++ b/Engine/Source/Programs/Horde/HordeServer/Logs/LogTailService.cs
@@ -46,7 +46,7 @@ namespace HordeServer.Logs
 		/// </summary>
 		public static TimeSpan ExpireAfter { get; } = TimeSpan.FromMinutes(2.0);
 
-		readonly RedisService _redisService;
+		readonly IRedisService _redisService;
 
 		readonly ConcurrentDictionary<LogId, AsyncEvent> _notifyLogEvents = new ConcurrentDictionary<LogId, AsyncEvent>();
 		readonly RedisChannel<LogId> _tailStartChannel;
@@ -72,7 +72,7 @@ namespace HordeServer.Logs
 		/// <summary>
 		/// Constructor
 		/// </summary>
-		public LogTailService(RedisService redisService, IClock clock, IOptions<ServerSettings> settings, ILogger<LogTailService> logger)
+		public LogTailService(IRedisService redisService, IClock clock, IOptions<ServerSettings> settings, ILogger<LogTailService> logger)
 			: this(redisService, clock, 32, settings, logger)
 		{
 		}
@@ -80,7 +80,7 @@ namespace HordeServer.Logs
 		/// <summary>
 		/// Constructor
 		/// </summary>
-		public LogTailService(RedisService redisService, IClock clock, int chunkLineCount, IOptions<ServerSettings> settings, ILogger<LogTailService> logger)
+		public LogTailService(IRedisService redisService, IClock clock, int chunkLineCount, IOptions<ServerSettings> settings, ILogger<LogTailService> logger)
 		{
 			if ((chunkLineCount & (chunkLineCount - 1)) != 0)
 			{
diff --git a/Engine/Source/Programs/Horde/HordeServer/Perforce/PerforceLoadBalancer.cs b/Engine/Source/Programs/Horde/HordeServer/Perforce/PerforceLoadBalancer.cs
--- a/Engine/Source/Programs/Horde/HordeServer/Perforce/PerforceLoadBalancer.cs
+++ b/Engine/Source/Programs/Horde/HordeServer/Perforce/PerforceLoadBalancer.cs
@@ -151,7 +151,7 @@ namespace HordeServer.Perforce
 			public List<PerforceServerEntry> Servers { get; set; } = new List<PerforceServerEntry>();
 		}
 
-		readonly RedisService _redisService;
+		readonly IRedisService _redisService;
 		readonly ILeaseCollection _leaseCollection;
 		readonly SingletonDocument<PerforceServerList> _serverListSingleton;
 		readonly Random _random = new Random();
@@ -165,7 +165,7 @@ namespace HordeServer.Perforce
 		/// <summary>
 		/// Constructor
 		/// </summary>
-		public PerforceLoadBalancer(MongoService mongoService, RedisService redisService, ILeaseCollection leaseCollection, IClock clock, HttpClient httpClient, IOptionsMonitor<GlobalConfig> globalConfig, IHealthMonitor<PerforceLoadBalancer> health, Tracer tracer, ILogger<PerforceLoadBalancer> logger)
+		public PerforceLoadBalancer(MongoService mongoService, IRedisService redisService, ILeaseCollection leaseCollection, IClock clock, HttpClient httpClient, IOptionsMonitor<GlobalConfig> globalConfig, IHealthMonitor<PerforceLoadBalancer> health, Tracer tracer, ILogger<PerforceLoadBalancer> logger)
 		{
 			_redisService = redisService;
 			_leaseCollection = leaseCollection;
diff --git a/Engine/Source/Programs/Horde/HordeServer/Server/IRedisService.cs b/Engine/Source/Programs/Horde/HordeServer/Server/IRedisService.cs
new file mode 100644
--- /dev/null
+++ b/Engine/Source/Programs/Horde/HordeServer/Server/IRedisService.cs
@@ -0,0 +1,103 @@
+// Copyright Epic Games, Inc. All Rights Reserved.
+
+using System;
+using System.Threading.Tasks;
+using EpicGames.Redis;
+using StackExchange.Redis;
+
+namespace HordeServer.Server
+{
+	/// <summary>
+	/// Provides access to a Redis database
+	/// </summary>
+	public interface IRedisService
+	{
+		/// <summary>
+		/// Flag for whether the connection is read-only
+		/// </summary>
+		bool ReadOnlyMode { get; }
+
+		/// <summary>
+		/// Connection pool
+		/// </summary>
+		public RedisConnectionPool ConnectionPool { get; }
+	}
+
+	/// <summary>
+	/// Extension methods for <see cref="IRedisService"/>
+	/// </summary>
+	public static class RedisServiceExtensions
+	{
+		/// <summary>
+		/// Get the least-loaded Redis connection from the pool
+		/// Don't store the returned object and try to resolve this as late as possible to ensure load is balanced.
+		/// </summary>
+		/// <returns>A Redis connection multiplexer</returns>
+		public static IConnectionMultiplexer GetConnection(this IRedisService redisService)
+			=> redisService.ConnectionPool.GetConnection();
+
+		/// <summary>
+		/// Get the least-loaded Redis database from the connection pool
+		/// Don't store the returned object and try to resolve this as late as possible to ensure load is balanced.
+		/// </summary>
+		/// <returns>A Redis database</returns>
+		public static IDatabase GetDatabase(this IRedisService redisService)
+			=> redisService.ConnectionPool.GetDatabase();
+
+		/// <summary>
+		/// Publish a message to a channel
+		/// </summary>
+		/// <param name="redisService">The redis service instance</param>
+		/// <param name="channel">Channel to post to</param>
+		/// <param name="message">Message to post to the channel</param>
+		/// <param name="flags">Flags for the request</param>
+		public static Task PublishAsync(this IRedisService redisService, RedisChannel channel, RedisValue message, CommandFlags flags = CommandFlags.None)
+			=> redisService.GetDatabase().PublishAsync(channel, message, flags);
+
+		/// <summary>
+		/// Publish a message to a channel
+		/// </summary>
+		/// <typeparam name="T">Type of elements sent over the channel</typeparam>
+		/// <param name="redisService">The redis service instance</param>
+		/// <param name="channel">Channel to post to</param>
+		/// <param name="message">Message to post to the channel</param>
+		/// <param name="flags">Flags for the request</param>
+		public static Task PublishAsync<T>(this IRedisService redisService, RedisChannel<T> channel, T message, CommandFlags flags = CommandFlags.None)
+			=> redisService.GetDatabase().PublishAsync(channel, message, flags);
+
+		/// <inheritdoc cref="SubscribeAsync{T}(IRedisService, RedisChannel{T}, Action{RedisChannel{T}, T})"/>
+		public static Task<RedisSubscription> SubscribeAsync(this IRedisService redisService, RedisChannel channel, Action<RedisValue> callback)
+			=> SubscribeAsync(redisService, channel, (ch, x) => callback(x));
+
+		/// <inheritdoc cref="SubscribeAsync{T}(IRedisService, RedisChannel{T}, Action{RedisChannel{T}, T})"/>
+		public static Task<RedisSubscription> SubscribeAsync<T>(this IRedisService redisService, RedisChannel<T> channel, Action<T> callback)
+			=> SubscribeAsync(redisService, channel, (ch, x) => callback(x));
+
+		/// <summary>
+		/// Subscribe to notifications on a channel
+		/// </summary>
+		/// <param name="redisService">The redis service instance</param>
+		/// <param name="channel">Channel to monitor</param>
+		/// <param name="callback">Callback for new events</param>
+		/// <returns>Subscription object</returns>
+		public static async Task<RedisSubscription> SubscribeAsync(this IRedisService redisService, RedisChannel channel, Action<RedisChannel, RedisValue> callback)
+		{
+			IConnectionMultiplexer connection = redisService.GetConnection();
+			return await connection.SubscribeAsync(channel, callback);
+		}
+
+		/// <summary>
+		/// Subscribe to notifications on a channel
+		/// </summary>
+		/// <typeparam name="T">Type of elements sent over the channel</typeparam>
+		/// <param name="redisService">The redis service instance</param>
+		/// <param name="channel">Channel to monitor</param>
+		/// <param name="callback">Callback for new events</param>
+		/// <returns>Subscription object</returns>
+		public static async Task<RedisSubscription> SubscribeAsync<T>(this IRedisService redisService, RedisChannel<T> channel, Action<RedisChannel<T>, T> callback)
+		{
+			IConnectionMultiplexer connection = redisService.GetConnection();
+			return await connection.SubscribeAsync(channel, callback);
+		}
+	}
+}
diff --git a/Engine/Source/Programs/Horde/HordeServer/Server/LifetimeService.cs b/Engine/Source/Programs/Horde/HordeServer/Server/LifetimeService.cs
--- a/Engine/Source/Programs/Horde/HordeServer/Server/LifetimeService.cs
+++ b/Engine/Source/Programs/Horde/HordeServer/Server/LifetimeService.cs
@@ -38,7 +38,7 @@ namespace HordeServer.Server
 
 		readonly IHostApplicationLifetime _lifetime;
 		readonly MongoService _mongoService;
-		readonly RedisService _redisService;
+		readonly IRedisService _redisService;
 		readonly ITicker _ticker;
 
 		/// <inheritdoc cref="ServerSettings.ShutdownMemoryThreshold" />
@@ -67,7 +67,7 @@ namespace HordeServer.Server
 		/// <param name="redisService">Redis singleton service</param>
 		/// <param name="clock"></param>
 		/// <param name="logger">Logging interface</param>
-		public LifetimeService(IOptionsMonitor<ServerSettings> settings, IHostApplicationLifetime lifetime, IHostEnvironment env, MongoService mongoService, RedisService redisService, IClock clock, ILogger<LifetimeService> logger)
+		public LifetimeService(IOptionsMonitor<ServerSettings> settings, IHostApplicationLifetime lifetime, IHostEnvironment env, MongoService mongoService, IRedisService redisService, IClock clock, ILogger<LifetimeService> logger)
 		{
 			_shutdownMemoryThreshold = settings.CurrentValue.ShutdownMemoryThreshold;
 			_lifetime = lifetime;
diff --git a/Engine/Source/Programs/Horde/HordeServer/Server/MongoService.cs b/Engine/Source/Programs/Horde/HordeServer/Server/MongoService.cs
--- a/Engine/Source/Programs/Horde/HordeServer/Server/MongoService.cs
+++ b/Engine/Source/Programs/Horde/HordeServer/Server/MongoService.cs
@@ -257,7 +257,7 @@ namespace HordeServer.Server
 		static readonly RedisKey s_schemaLockKey = new RedisKey("server/schema-upgrade/lock");
 
 		readonly MongoClient _client;
-		readonly RedisService _redisService;
+		readonly IRedisService _redisService;
 #pragma warning disable CA2213 // Disposable fields should be disposed
 		readonly SemaphoreSlim _upgradeSema = new SemaphoreSlim(1);
 #pragma warning restore CA2213 // Disposable fields should be disposed
@@ -270,7 +270,7 @@ namespace HordeServer.Server
 		/// <summary>
 		/// Constructor
 		/// </summary>
-		public MongoService(IOptions<ServerSettings> settingsSnapshot, RedisService redisService, MongoCommandTracer mongoTracer, Tracer tracer, ILogger<MongoService> logger, ILoggerFactory loggerFactory)
+		public MongoService(IOptions<ServerSettings> settingsSnapshot, IRedisService redisService, MongoCommandTracer mongoTracer, Tracer tracer, ILogger<MongoService> logger, ILoggerFactory loggerFactory)
 		{
 			if (s_existingInstance != null)
 			{
diff --git a/Engine/Source/Programs/Horde/HordeServer/Server/RedisService.cs b/Engine/Source/Programs/Horde/HordeServer/Server/RedisService.cs
--- a/Engine/Source/Programs/Horde/HordeServer/Server/RedisService.cs
+++ b/Engine/Source/Programs/Horde/HordeServer/Server/RedisService.cs
@@ -20,7 +20,7 @@ namespace HordeServer.Server
 	/// <summary>
 	/// Manages the lifetime of a bundled Redis instance
 	/// </summary>
-	public sealed class RedisService : IHealthCheck, IAsyncDisposable
+	public sealed class RedisService : IRedisService, IHealthCheck, IAsyncDisposable
 	{
 		/// <summary>
 		/// Default Redis port
@@ -83,7 +83,7 @@ namespace HordeServer.Server
 		{
 			try
 			{
-				TimeSpan span = await GetDatabase().PingAsync();
+				TimeSpan span = await ConnectionPool.GetDatabase().PingAsync();
 				return HealthCheckResult.Healthy(data: new Dictionary<string, object> { ["Latency"] = span.ToString() });
 			}
 			catch (Exception ex)
@@ -92,26 +92,6 @@ namespace HordeServer.Server
 			}
 		}
 
-		/// <summary>
-		/// Get the least-loaded Redis connection from the pool
-		/// Don't store the returned object and try to resolve this as late as possible to ensure load is balanced.
-		/// </summary>
-		/// <returns>A Redis connection multiplexer</returns>
-		public IConnectionMultiplexer GetConnection()
-		{
-			return ConnectionPool.GetConnection();
-		}
-
-		/// <summary>
-		/// Get the least-loaded Redis database from the connection pool
-		/// Don't store the returned object and try to resolve this as late as possible to ensure load is balanced.
-		/// </summary>
-		/// <returns>A Redis database</returns>
-		public IDatabase GetDatabase()
-		{
-			return ConnectionPool.GetDatabase();
-		}
-
 		/// <inheritdoc/>
 		public async ValueTask DisposeAsync()
 		{
@@ -167,59 +147,5 @@ namespace HordeServer.Server
 			_redisProcess.Start($"\"{redisConfigFile}\"");
 			return true;
 		}
-
-		/// <summary>
-		/// Publish a message to a channel
-		/// </summary>
-		/// <param name="channel">Channel to post to</param>
-		/// <param name="message">Message to post to the channel</param>
-		/// <param name="flags">Flags for the request</param>
-		public Task PublishAsync(RedisChannel channel, RedisValue message, CommandFlags flags = CommandFlags.None)
-		{
-			return GetDatabase().PublishAsync(channel, message, flags);
-		}
-
-		/// <summary>
-		/// Publish a message to a channel
-		/// </summary>
-		/// <typeparam name="T">Type of elements sent over the channel</typeparam>
-		/// <param name="channel">Channel to post to</param>
-		/// <param name="message">Message to post to the channel</param>
-		/// <param name="flags">Flags for the request</param>
-		public Task PublishAsync<T>(RedisChannel<T> channel, T message, CommandFlags flags = CommandFlags.None)
-		{
-			return GetDatabase().PublishAsync(channel, message, flags);
-		}
-
-		/// <inheritdoc cref="SubscribeAsync{T}(RedisChannel{T}, Action{RedisChannel{T}, T})"/>
-		public Task<RedisSubscription> SubscribeAsync(RedisChannel channel, Action<RedisValue> callback) => SubscribeAsync(channel, (ch, x) => callback(x));
-
-		/// <inheritdoc cref="SubscribeAsync{T}(RedisChannel{T}, Action{RedisChannel{T}, T})"/>
-		public Task<RedisSubscription> SubscribeAsync<T>(RedisChannel<T> channel, Action<T> callback) => SubscribeAsync(channel, (ch, x) => callback(x));
-
-		/// <summary>
-		/// Subscribe to notifications on a channel
-		/// </summary>
-		/// <param name="channel">Channel to monitor</param>
-		/// <param name="callback">Callback for new events</param>
-		/// <returns>Subscription object</returns>
-		public async Task<RedisSubscription> SubscribeAsync(RedisChannel channel, Action<RedisChannel, RedisValue> callback)
-		{
-			IConnectionMultiplexer connection = GetConnection();
-			return await connection.SubscribeAsync(channel, callback);
-		}
-
-		/// <summary>
-		/// Subscribe to notifications on a channel
-		/// </summary>
-		/// <typeparam name="T">Type of elements sent over the channel</typeparam>
-		/// <param name="channel">Channel to monitor</param>
-		/// <param name="callback">Callback for new events</param>
-		/// <returns>Subscription object</returns>
-		public async Task<RedisSubscription> SubscribeAsync<T>(RedisChannel<T> channel, Action<RedisChannel<T>, T> callback)
-		{
-			IConnectionMultiplexer connection = GetConnection();
-			return await connection.SubscribeAsync(channel, callback);
-		}
 	}
 }
diff --git a/Engine/Source/Programs/Horde/HordeServer/Startup.cs b/Engine/Source/Programs/Horde/HordeServer/Startup.cs
--- a/Engine/Source/Programs/Horde/HordeServer/Startup.cs
+++ b/Engine/Source/Programs/Horde/HordeServer/Startup.cs
@@ -464,6 +464,7 @@ namespace HordeServer
 			}
 #pragma warning restore CA2000 // Dispose objects before losing scope
 			services.AddSingleton<RedisService>(sp => redisService);
+			services.AddSingleton<IRedisService>(sp => sp.GetRequiredService<RedisService>());
 			services.AddDataProtection().PersistKeysToStackExchangeRedis(() => redisService.GetDatabase(), "aspnet-data-protection");
 			
 			services.AddResponseCompression(options =>
diff --git a/Engine/Source/Programs/Horde/HordeServer/Storage/StorageService.cs b/Engine/Source/Programs/Horde/HordeServer/Storage/StorageService.cs
--- a/Engine/Source/Programs/Horde/HordeServer/Storage/StorageService.cs
+++ b/Engine/Source/Programs/Horde/HordeServer/Storage/StorageService.cs
@@ -372,7 +372,7 @@ namespace HordeServer.Storage
 			public DateTime LastTime { get; set; }
 		}
 
-		readonly RedisService _redisService;
+		readonly IRedisService _redisService;
 		readonly IClock _clock;
 		readonly BundleCache _bundleCache;
 		readonly IMemoryCache _memoryCache;
@@ -400,7 +400,7 @@ namespace HordeServer.Storage
 		/// <summary>
 		/// Constructor
 		/// </summary>
-		public StorageService(MongoService mongoService, RedisService redisService, IClock clock, BundleCache bundleCache, IMemoryCache memoryCache, IObjectStoreFactory objectStoreFactory, IOptionsMonitor<GlobalConfig> globalConfig, Tracer tracer, ILogger<StorageService> logger)
+		public StorageService(MongoService mongoService, IRedisService redisService, IClock clock, BundleCache bundleCache, IMemoryCache memoryCache, IObjectStoreFactory objectStoreFactory, IOptionsMonitor<GlobalConfig> globalConfig, Tracer tracer, ILogger<StorageService> logger)
 		{
 			_redisService = redisService;
 			_clock = clock;
diff --git a/Engine/Source/Programs/Horde/HordeServer/Utilities/Clock.cs b/Engine/Source/Programs/Horde/HordeServer/Utilities/Clock.cs
--- a/Engine/Source/Programs/Horde/HordeServer/Utilities/Clock.cs
+++ b/Engine/Source/Programs/Horde/HordeServer/Utilities/Clock.cs
@@ -128,7 +128,7 @@ namespace HordeCommon
 			public override string ToString() => _name;
 		}
 
-		readonly RedisService _redis;
+		readonly IRedisService _redis;
 		readonly TimeZoneInfo _timeZone;
 		readonly List<TickerImpl> _tickers = new List<TickerImpl>();
 
@@ -141,7 +141,7 @@ namespace HordeCommon
 		/// <summary>
 		/// Constructor
 		/// </summary>
-		public Clock(RedisService redis, IOptions<ServerSettings> settings)
+		public Clock(IRedisService redis, IOptions<ServerSettings> settings)
 		{
 			_redis = redis;
 
