Horde: Add an interface for MongoService.

diff --git a/Engine/Source/Programs/Horde/HordeServer.Tests/Compute/RedisTaskSchedulerTest.cs b/Engine/Source/Programs/Horde/HordeServer.Tests/Compute/RedisTaskSchedulerTest.cs
--- a/Engine/Source/Programs/Horde/HordeServer.Tests/Compute/RedisTaskSchedulerTest.cs
+++ b/Engine/Source/Programs/Horde/HordeServer.Tests/Compute/RedisTaskSchedulerTest.cs
@@ -19,7 +19,7 @@ namespace HordeServer.Tests.Compute
 		public RedisTaskSchedulerTest()
 		{
 			using ILoggerFactory loggerFactory = LoggerFactory.Create(builder => { builder.AddConsole(); });
-			RedisService redisService = GetRedisServiceSingleton();
+			IRedisService redisService = GetRedisServiceSingleton();
 			_scheduler = new RedisTaskScheduler<string, string>(redisService.ConnectionPool, new RedisKey("myBaseKey"), loggerFactory.CreateLogger<RedisTaskScheduler<string, string>>());
 		}
 
diff --git a/Engine/Source/Programs/Horde/HordeServer.Tests/DatabaseIntegrationTest.cs b/Engine/Source/Programs/Horde/HordeServer.Tests/DatabaseIntegrationTest.cs
--- a/Engine/Source/Programs/Horde/HordeServer.Tests/DatabaseIntegrationTest.cs
+++ b/Engine/Source/Programs/Horde/HordeServer.Tests/DatabaseIntegrationTest.cs
@@ -228,8 +228,8 @@ namespace HordeServer.Tests
 		{
 			base.ConfigureServices(services);
 
-			services.AddSingleton(GetMongoServiceSingleton());
-			services.AddSingleton(GetRedisServiceSingleton());
+			services.AddSingleton<IMongoService>(GetMongoServiceSingleton());
+			services.AddSingleton<IRedisService>(GetRedisServiceSingleton());
 		}
 
 		public override async ValueTask DisposeAsync()
@@ -253,13 +253,13 @@ namespace HordeServer.Tests
 			_loggerFactory.Dispose();
 		}
 
-		public MongoService GetMongoServiceSingleton()
+		public IMongoService GetMongoServiceSingleton()
 		{
 			lock (s_lockObject)
 			{
 				if (_mongoService == null)
 				{
-					RedisService redisService = GetRedisServiceSingleton();
+					IRedisService redisService = GetRedisServiceSingleton();
 					_mongoInstance = new MongoInstance();
 
 					ServerSettings ss = new ServerSettings();
@@ -273,7 +273,7 @@ namespace HordeServer.Tests
 			return _mongoService;
 		}
 
-		public RedisService GetRedisServiceSingleton()
+		public IRedisService GetRedisServiceSingleton()
 		{
 			if (_redisService == null)
 			{
diff --git a/Engine/Source/Programs/Horde/HordeServer.Tests/Logs/LogTailServiceTests.cs b/Engine/Source/Programs/Horde/HordeServer.Tests/Logs/LogTailServiceTests.cs
--- a/Engine/Source/Programs/Horde/HordeServer.Tests/Logs/LogTailServiceTests.cs
+++ b/Engine/Source/Programs/Horde/HordeServer.Tests/Logs/LogTailServiceTests.cs
@@ -71,7 +71,7 @@ namespace HordeServer.Tests.Logs
 		[TestMethod]
 		public async Task SplittingAsync()
 		{
-			RedisService redisService = ServiceProvider.GetRequiredService<RedisService>();
+			IRedisService redisService = ServiceProvider.GetRequiredService<IRedisService>();
 
 			LogTailService tailService = new LogTailService(redisService, ServiceProvider.GetRequiredService<IClock>(), 4, ServiceProvider.GetRequiredService<IOptions<ServerSettings>>(), ServiceProvider.GetRequiredService<ILogger<LogTailService>>());
 			await tailService.StartAsync(CancellationToken.None);
@@ -114,7 +114,7 @@ namespace HordeServer.Tests.Logs
 		[TestMethod]
 		public async Task ExpiryAsync()
 		{
-			RedisService redisService = GetRedisServiceSingleton();
+			IRedisService redisService = GetRedisServiceSingleton();
 
 			LogTailService tailService = ServiceProvider.GetRequiredService<LogTailService>();
 
@@ -173,7 +173,7 @@ namespace HordeServer.Tests.Logs
 			LogTailService tailService = ServiceProvider.GetRequiredService<LogTailService>();
 			await tailService.StartAsync(CancellationToken.None);
 
-			LogTailService tailService2 = new LogTailService(ServiceProvider.GetRequiredService<RedisService>(), ServiceProvider.GetRequiredService<IClock>(), 4, ServiceProvider.GetRequiredService<IOptions<ServerSettings>>(), ServiceProvider.GetRequiredService<ILogger<LogTailService>>());
+			LogTailService tailService2 = new LogTailService(ServiceProvider.GetRequiredService<IRedisService>(), ServiceProvider.GetRequiredService<IClock>(), 4, ServiceProvider.GetRequiredService<IOptions<ServerSettings>>(), ServiceProvider.GetRequiredService<ILogger<LogTailService>>());
 
 			Task<int> task = tailService.WaitForTailNextAsync(_logId, CancellationToken.None);
 			Assert.IsTrue(!task.IsCompleted);
diff --git a/Engine/Source/Programs/Horde/HordeServer.Tests/ServiceAccounts/ServiceAccountCollectionTests.cs b/Engine/Source/Programs/Horde/HordeServer.Tests/ServiceAccounts/ServiceAccountCollectionTests.cs
--- a/Engine/Source/Programs/Horde/HordeServer.Tests/ServiceAccounts/ServiceAccountCollectionTests.cs
+++ b/Engine/Source/Programs/Horde/HordeServer.Tests/ServiceAccounts/ServiceAccountCollectionTests.cs
@@ -19,7 +19,7 @@ namespace HordeServer.Tests.Accounts
 
 		public HordeAccountCollectionTests()
 		{
-			MongoService mongoService = GetMongoServiceSingleton();
+			IMongoService mongoService = GetMongoServiceSingleton();
 			_serviceAccountCollection = new ServiceAccountCollection(mongoService);
 			(_serviceAccount, _token) = _serviceAccountCollection.CreateAsync(new CreateServiceAccountOptions(Description: "myDesc",
 				Claims: new List<IUserClaim> { new UserClaim("myClaim", "myValue") })
diff --git a/Engine/Source/Programs/Horde/HordeServer.Tests/TestSetup.cs b/Engine/Source/Programs/Horde/HordeServer.Tests/TestSetup.cs
--- a/Engine/Source/Programs/Horde/HordeServer.Tests/TestSetup.cs
+++ b/Engine/Source/Programs/Horde/HordeServer.Tests/TestSetup.cs
@@ -110,7 +110,7 @@ namespace HordeServer.Tests
 		public AgentRelayService AgentRelayService => ServiceProvider.GetRequiredService<AgentRelayService>();
 		public ICommitService CommitService => ServiceProvider.GetRequiredService<ICommitService>();
 		public GlobalsService GlobalsService => ServiceProvider.GetRequiredService<GlobalsService>();
-		public MongoService MongoService => ServiceProvider.GetRequiredService<MongoService>();
+		public IMongoService MongoService => ServiceProvider.GetRequiredService<IMongoService>();
 		public ITemplateCollection TemplateCollection => ServiceProvider.GetRequiredService<ITemplateCollection>();
 		internal PerforceServiceStub PerforceService => (PerforceServiceStub)ServiceProvider.GetRequiredService<IPerforceService>();
 		public ISubscriptionCollection SubscriptionCollection => ServiceProvider.GetRequiredService<ISubscriptionCollection>();
diff --git a/Engine/Source/Programs/Horde/HordeServer/Accounts/AccountCollection.cs b/Engine/Source/Programs/Horde/HordeServer/Accounts/AccountCollection.cs
--- a/Engine/Source/Programs/Horde/HordeServer/Accounts/AccountCollection.cs
+++ b/Engine/Source/Programs/Horde/HordeServer/Accounts/AccountCollection.cs
@@ -199,7 +199,7 @@ namespace HordeServer.Accounts
 		/// Constructor
 		/// </summary>
 		/// <param name="mongoService">The database service</param>
-		public AccountCollection(MongoService mongoService)
+		public AccountCollection(IMongoService mongoService)
 		{
 			_accounts = mongoService.GetCollection<AccountDocument>("Accounts", keys => keys.Ascending(x => x.NormalizedLogin), unique: true);
 		}
diff --git a/Engine/Source/Programs/Horde/HordeServer/Agents/AgentCollection.cs b/Engine/Source/Programs/Horde/HordeServer/Agents/AgentCollection.cs
--- a/Engine/Source/Programs/Horde/HordeServer/Agents/AgentCollection.cs
+++ b/Engine/Source/Programs/Horde/HordeServer/Agents/AgentCollection.cs
@@ -221,14 +221,14 @@ namespace HordeServer.Agents
 
 		readonly IMongoCollection<AgentDocument> _agents;
 		readonly IAuditLog<AgentId> _auditLog;
-		readonly RedisService _redisService;
+		readonly IRedisService _redisService;
 		readonly IClock _clock;
 		readonly RedisChannel<AgentId> _updateEventChannel;
 
 		/// <summary>
 		/// Constructor
 		/// </summary>
-		public AgentCollection(MongoService mongoService, RedisService redisService, IClock clock, IAuditLog<AgentId> auditLog)
+		public AgentCollection(IMongoService mongoService, IRedisService redisService, IClock clock, IAuditLog<AgentId> auditLog)
 		{
 			List<MongoIndex<AgentDocument>> indexes = new List<MongoIndex<AgentDocument>>();
 			indexes.Add(keys => keys.Ascending(x => x.Deleted).Ascending(x => x.Id).Ascending(x => x.Pools));
diff --git a/Engine/Source/Programs/Horde/HordeServer/Agents/AgentReportService.cs b/Engine/Source/Programs/Horde/HordeServer/Agents/AgentReportService.cs
--- a/Engine/Source/Programs/Horde/HordeServer/Agents/AgentReportService.cs
+++ b/Engine/Source/Programs/Horde/HordeServer/Agents/AgentReportService.cs
@@ -52,7 +52,7 @@ namespace HordeServer.Agents
 		/// <summary>
 		/// Constructor
 		/// </summary>
-		public AgentReportService(MongoService mongoService, IAgentCollection agentCollection, INotificationService notificationService, IClock clock, ILogger<AgentReportService> logger)
+		public AgentReportService(IMongoService mongoService, IAgentCollection agentCollection, INotificationService notificationService, IClock clock, ILogger<AgentReportService> logger)
 		{
 			_state = new SingletonDocument<AgentReportState>(mongoService);
 			_agentCollection = agentCollection;
diff --git a/Engine/Source/Programs/Horde/HordeServer/Agents/Leases/LeaseCollection.cs b/Engine/Source/Programs/Horde/HordeServer/Agents/Leases/LeaseCollection.cs
--- a/Engine/Source/Programs/Horde/HordeServer/Agents/Leases/LeaseCollection.cs
+++ b/Engine/Source/Programs/Horde/HordeServer/Agents/Leases/LeaseCollection.cs
@@ -80,7 +80,7 @@ namespace HordeServer.Agents.Leases
 		/// <summary>
 		/// Constructor
 		/// </summary>
-		public LeaseCollection(MongoService mongoService, Tracer tracer)
+		public LeaseCollection(IMongoService mongoService, Tracer tracer)
 		{
 			List<MongoIndex<LeaseDocument>> indexes = new List<MongoIndex<LeaseDocument>>();
 			//			indexes.Add(keys => keys.Ascending(x => x.ParentId), sparse: true);
diff --git a/Engine/Source/Programs/Horde/HordeServer/Agents/Pools/PoolCollection.cs b/Engine/Source/Programs/Horde/HordeServer/Agents/Pools/PoolCollection.cs
--- a/Engine/Source/Programs/Horde/HordeServer/Agents/Pools/PoolCollection.cs
+++ b/Engine/Source/Programs/Horde/HordeServer/Agents/Pools/PoolCollection.cs
@@ -233,7 +233,7 @@ namespace HordeServer.Agents.Pools
 		/// </summary>
 		/// <param name="mongoService">The database service instance</param>
 		/// <param name="globalConfig"></param>
-		public PoolCollection(MongoService mongoService, IOptionsMonitor<GlobalConfig> globalConfig)
+		public PoolCollection(IMongoService mongoService, IOptionsMonitor<GlobalConfig> globalConfig)
 		{
 			_poolsV1 = mongoService.GetCollection<PoolDocumentV1>("Pools");
 			_poolsV2 = mongoService.GetCollection<PoolDocumentV2>("PoolsV2");
diff --git a/Engine/Source/Programs/Horde/HordeServer/Agents/Sessions/SessionCollection.cs b/Engine/Source/Programs/Horde/HordeServer/Agents/Sessions/SessionCollection.cs
--- a/Engine/Source/Programs/Horde/HordeServer/Agents/Sessions/SessionCollection.cs
+++ b/Engine/Source/Programs/Horde/HordeServer/Agents/Sessions/SessionCollection.cs
@@ -73,7 +73,7 @@ namespace HordeServer.Agents.Sessions
 		/// Constructor
 		/// </summary>
 		/// <param name="mongoService">The database service</param>
-		public SessionCollection(MongoService mongoService)
+		public SessionCollection(IMongoService mongoService)
 		{
 			List<MongoIndex<SessionDocument>> indexes = new List<MongoIndex<SessionDocument>>();
 			indexes.Add(keys => keys.Ascending(x => x.AgentId).Ascending(x => x.StartTime).Ascending(x => x.FinishTime));
diff --git a/Engine/Source/Programs/Horde/HordeServer/Agents/Telemetry/AgentTelemetryCollection.cs b/Engine/Source/Programs/Horde/HordeServer/Agents/Telemetry/AgentTelemetryCollection.cs
--- a/Engine/Source/Programs/Horde/HordeServer/Agents/Telemetry/AgentTelemetryCollection.cs
+++ b/Engine/Source/Programs/Horde/HordeServer/Agents/Telemetry/AgentTelemetryCollection.cs
@@ -70,7 +70,7 @@ namespace HordeServer.Agents.Telemetry
 		/// <summary>
 		/// Constructor
 		/// </summary>
-		public AgentTelemetryCollection(MongoService mongoService, IClock clock, ILogger<AgentTelemetryCollection> logger)
+		public AgentTelemetryCollection(IMongoService mongoService, IClock clock, ILogger<AgentTelemetryCollection> logger)
 		{
 			_clock = clock;
 
diff --git a/Engine/Source/Programs/Horde/HordeServer/Agents/Utilization/UtilizationDataCollection.cs b/Engine/Source/Programs/Horde/HordeServer/Agents/Utilization/UtilizationDataCollection.cs
--- a/Engine/Source/Programs/Horde/HordeServer/Agents/Utilization/UtilizationDataCollection.cs
+++ b/Engine/Source/Programs/Horde/HordeServer/Agents/Utilization/UtilizationDataCollection.cs
@@ -101,7 +101,7 @@ namespace HordeServer.Agents.Utilization
 		/// Constructor
 		/// </summary>
 		/// <param name="database"></param>
-		public UtilizationDataCollection(MongoService database)
+		public UtilizationDataCollection(IMongoService database)
 		{
 			_utilization = database.GetCollection<UtilizationDocument>("Utilization", keys => keys.Ascending(x => x.FinishTime).Ascending(x => x.StartTime));
 		}
diff --git a/Engine/Source/Programs/Horde/HordeServer/Artifacts/ArtifactCollection.cs b/Engine/Source/Programs/Horde/HordeServer/Artifacts/ArtifactCollection.cs
--- a/Engine/Source/Programs/Horde/HordeServer/Artifacts/ArtifactCollection.cs
+++ b/Engine/Source/Programs/Horde/HordeServer/Artifacts/ArtifactCollection.cs
@@ -100,7 +100,7 @@ namespace HordeServer.Artifacts
 		/// <summary>
 		/// Constructor
 		/// </summary>
-		public ArtifactCollection(MongoService mongoService, IClock clock)
+		public ArtifactCollection(IMongoService mongoService, IClock clock)
 		{
 			List<MongoIndex<Artifact>> indexes = new List<MongoIndex<Artifact>>();
 			indexes.Add(keys => keys.Ascending(x => x.Keys));
diff --git a/Engine/Source/Programs/Horde/HordeServer/Auditing/AuditLog.cs b/Engine/Source/Programs/Horde/HordeServer/Auditing/AuditLog.cs
--- a/Engine/Source/Programs/Horde/HordeServer/Auditing/AuditLog.cs
+++ b/Engine/Source/Programs/Horde/HordeServer/Auditing/AuditLog.cs
@@ -173,7 +173,7 @@ namespace HordeServer.Auditing
 
 		public IAuditLogChannel<TSubject> this[TSubject subject] => new AuditLogChannel(this, subject);
 
-		public AuditLog(MongoService mongoService, string collectionName, string subjectProperty, ILogger logger)
+		public AuditLog(IMongoService mongoService, string collectionName, string subjectProperty, ILogger logger)
 		{
 			List<MongoIndex<AuditLogMessage>> indexes = new List<MongoIndex<AuditLogMessage>>();
 			indexes.Add(builder => builder.Ascending(x => x.Subject).Descending(x => x.TimeUtc));
@@ -294,10 +294,10 @@ namespace HordeServer.Auditing
 
 	class AuditLogFactory<TSubject> : IAuditLogFactory<TSubject>
 	{
-		readonly MongoService _mongoService;
+		readonly IMongoService _mongoService;
 		readonly ILogger<AuditLog<TSubject>> _logger;
 
-		public AuditLogFactory(MongoService mongoService, ILogger<AuditLog<TSubject>> logger)
+		public AuditLogFactory(IMongoService mongoService, ILogger<AuditLog<TSubject>> logger)
 		{
 			_mongoService = mongoService;
 			_logger = logger;
diff --git a/Engine/Source/Programs/Horde/HordeServer/Compute/ComputeService.cs b/Engine/Source/Programs/Horde/HordeServer/Compute/ComputeService.cs
--- a/Engine/Source/Programs/Horde/HordeServer/Compute/ComputeService.cs
+++ b/Engine/Source/Programs/Horde/HordeServer/Compute/ComputeService.cs
@@ -198,7 +198,7 @@ namespace HordeServer.Compute
 			ILogCollection logCollection,
 			AgentService agentService,
 			AgentRelayService agentRelayService,
-			RedisService redisService,
+			IRedisService redisService,
 			IOptionsMonitor<ServerSettings> settings,
 			IOptionsMonitor<GlobalConfig> globalConfig,
 			IClock clock,
diff --git a/Engine/Source/Programs/Horde/HordeServer/Dashboard/DashboardPreviewCollection.cs b/Engine/Source/Programs/Horde/HordeServer/Dashboard/DashboardPreviewCollection.cs
--- a/Engine/Source/Programs/Horde/HordeServer/Dashboard/DashboardPreviewCollection.cs
+++ b/Engine/Source/Programs/Horde/HordeServer/Dashboard/DashboardPreviewCollection.cs
@@ -145,7 +145,7 @@ namespace HordeServer.Dashboard
 		readonly IMongoCollection<PreviewDocument> _previews;
 		readonly ISingletonDocument<PreviewLedger> _ledgerSingleton;
 
-		public DashboardPreviewCollection(MongoService mongoService)
+		public DashboardPreviewCollection(IMongoService mongoService)
 		{
 			_ledgerSingleton = new SingletonDocument<PreviewLedger>(mongoService);
 			_previews = mongoService.GetCollection<PreviewDocument>("DashboardPreviews");
diff --git a/Engine/Source/Programs/Horde/HordeServer/Devices/DeviceCollection.cs b/Engine/Source/Programs/Horde/HordeServer/Devices/DeviceCollection.cs
--- a/Engine/Source/Programs/Horde/HordeServer/Devices/DeviceCollection.cs
+++ b/Engine/Source/Programs/Horde/HordeServer/Devices/DeviceCollection.cs
@@ -446,7 +446,7 @@ namespace HordeServer.Devices
 		/// <summary>
 		/// Constructor
 		/// </summary>
-		public DeviceCollection(MongoService mongoService, ILogger<DeviceCollection> logger)
+		public DeviceCollection(IMongoService mongoService, ILogger<DeviceCollection> logger)
 		{
 			_logger = logger;
 
diff --git a/Engine/Source/Programs/Horde/HordeServer/Devices/DeviceReportService.cs b/Engine/Source/Programs/Horde/HordeServer/Devices/DeviceReportService.cs
--- a/Engine/Source/Programs/Horde/HordeServer/Devices/DeviceReportService.cs
+++ b/Engine/Source/Programs/Horde/HordeServer/Devices/DeviceReportService.cs
@@ -144,7 +144,7 @@ namespace HordeServer.Devices
 		/// <summary>
 		/// Constructor
 		/// </summary>
-		public DeviceReportService(MongoService mongoService, DeviceService deviceService, INotificationService notificationService, IClock clock, IOptionsMonitor<GlobalConfig> globalConfig, ILogger<DeviceReportService> logger)
+		public DeviceReportService(IMongoService mongoService, DeviceService deviceService, INotificationService notificationService, IClock clock, IOptionsMonitor<GlobalConfig> globalConfig, ILogger<DeviceReportService> logger)
 		{
 			_state = new SingletonDocument<DeviceReportState>(mongoService);
 			_deviceService = deviceService;
diff --git a/Engine/Source/Programs/Horde/HordeServer/Issues/IssueCollection.cs b/Engine/Source/Programs/Horde/HordeServer/Issues/IssueCollection.cs
--- a/Engine/Source/Programs/Horde/HordeServer/Issues/IssueCollection.cs
+++ b/Engine/Source/Programs/Horde/HordeServer/Issues/IssueCollection.cs
@@ -619,7 +619,7 @@ namespace HordeServer.Issues
 			});
 		}
 
-		public IssueCollection(MongoService mongoService, IRedisService redisService, IUserCollection userCollection, IAuditLogFactory<int> auditLogFactory, ITelemetrySink telemetrySink, IOptionsMonitor<GlobalConfig> globalConfig, Tracer tracer, ILogger<IssueCollection> logger)
+		public IssueCollection(IMongoService mongoService, IRedisService redisService, IUserCollection userCollection, IAuditLogFactory<int> auditLogFactory, ITelemetrySink telemetrySink, IOptionsMonitor<GlobalConfig> globalConfig, Tracer tracer, ILogger<IssueCollection> logger)
 		{
 			_redisService = redisService;
 			_userCollection = userCollection;
diff --git a/Engine/Source/Programs/Horde/HordeServer/Issues/IssueReportService.cs b/Engine/Source/Programs/Horde/HordeServer/Issues/IssueReportService.cs
--- a/Engine/Source/Programs/Horde/HordeServer/Issues/IssueReportService.cs
+++ b/Engine/Source/Programs/Horde/HordeServer/Issues/IssueReportService.cs
@@ -92,7 +92,7 @@ namespace HordeServer.Issues
 		/// <summary>
 		/// Constructor
 		/// </summary>
-		public IssueReportService(MongoService mongoService, IIssueCollection issueCollection, IGraphCollection graphCollection, IJobCollection jobCollection, INotificationService notificationService, IClock clock, IOptionsMonitor<GlobalConfig> globalConfig, ILogger<IssueReportService> logger)
+		public IssueReportService(IMongoService mongoService, IIssueCollection issueCollection, IGraphCollection graphCollection, IJobCollection jobCollection, INotificationService notificationService, IClock clock, IOptionsMonitor<GlobalConfig> globalConfig, ILogger<IssueReportService> logger)
 		{
 			_state = new SingletonDocument<IssueReportState>(mongoService);
 			_issueCollection = issueCollection;
diff --git a/Engine/Source/Programs/Horde/HordeServer/Issues/IssueTagService.cs b/Engine/Source/Programs/Horde/HordeServer/Issues/IssueTagService.cs
--- a/Engine/Source/Programs/Horde/HordeServer/Issues/IssueTagService.cs
+++ b/Engine/Source/Programs/Horde/HordeServer/Issues/IssueTagService.cs
@@ -41,7 +41,7 @@ namespace HordeServer.Issues
 		/// <summary>
 		/// Constructor
 		/// </summary>
-		public IssueTagService(MongoService mongoService, ICommitService commitService, IIssueCollection issueCollection, IClock clock, IOptionsMonitor<GlobalConfig> globalConfig, ILogger<IssueTagService> logger)
+		public IssueTagService(IMongoService mongoService, ICommitService commitService, IIssueCollection issueCollection, IClock clock, IOptionsMonitor<GlobalConfig> globalConfig, ILogger<IssueTagService> logger)
 		{
 			_state = new SingletonDocument<State>(mongoService);
 			_commitService = commitService;
diff --git a/Engine/Source/Programs/Horde/HordeServer/Jobs/Artifacts/ArtifactCollectionV1.cs b/Engine/Source/Programs/Horde/HordeServer/Jobs/Artifacts/ArtifactCollectionV1.cs
--- a/Engine/Source/Programs/Horde/HordeServer/Jobs/Artifacts/ArtifactCollectionV1.cs
+++ b/Engine/Source/Programs/Horde/HordeServer/Jobs/Artifacts/ArtifactCollectionV1.cs
@@ -65,7 +65,7 @@ namespace HordeServer.Jobs.Artifacts
 		/// <summary>
 		/// Constructor
 		/// </summary>
-		public ArtifactCollectionV1(MongoService mongoService, IObjectStore<ArtifactCollectionV1> objectStore)
+		public ArtifactCollectionV1(IMongoService mongoService, IObjectStore<ArtifactCollectionV1> objectStore)
 		{
 			_objectStore = objectStore;
 
diff --git a/Engine/Source/Programs/Horde/HordeServer/Jobs/Bisect/BisectTaskCollection.cs b/Engine/Source/Programs/Horde/HordeServer/Jobs/Bisect/BisectTaskCollection.cs
--- a/Engine/Source/Programs/Horde/HordeServer/Jobs/Bisect/BisectTaskCollection.cs
+++ b/Engine/Source/Programs/Horde/HordeServer/Jobs/Bisect/BisectTaskCollection.cs
@@ -93,7 +93,7 @@ namespace HordeServer.Jobs.Bisect
 		readonly IMongoCollection<BisectTaskDoc> _bisectTasks;
 		readonly MongoIndex<BisectTaskDoc> _runningIndex;
 
-		public BisectTaskCollection(Tracer tracer, MongoService mongoService)
+		public BisectTaskCollection(Tracer tracer, IMongoService mongoService)
 		{
 			List<MongoIndex<BisectTaskDoc>> indexes = new List<MongoIndex<BisectTaskDoc>>();
 			indexes.Add(keys => keys.Ascending(x => x.Id).Ascending(x => x.InitialJobId));
diff --git a/Engine/Source/Programs/Horde/HordeServer/Jobs/Graphs/GraphCollection.cs b/Engine/Source/Programs/Horde/HordeServer/Jobs/Graphs/GraphCollection.cs
--- a/Engine/Source/Programs/Horde/HordeServer/Jobs/Graphs/GraphCollection.cs
+++ b/Engine/Source/Programs/Horde/HordeServer/Jobs/Graphs/GraphCollection.cs
@@ -431,7 +431,7 @@ namespace HordeServer.Jobs.Graphs
 		/// <summary>
 		/// Constructor
 		/// </summary>
-		public GraphCollection(MongoService mongoService)
+		public GraphCollection(IMongoService mongoService)
 		{
 			_graphs = mongoService.GetCollection<GraphDocument>("Graphs");
 		}
diff --git a/Engine/Source/Programs/Horde/HordeServer/Jobs/JobCollection.cs b/Engine/Source/Programs/Horde/HordeServer/Jobs/JobCollection.cs
--- a/Engine/Source/Programs/Horde/HordeServer/Jobs/JobCollection.cs
+++ b/Engine/Source/Programs/Horde/HordeServer/Jobs/JobCollection.cs
@@ -738,7 +738,7 @@ namespace HordeServer.Jobs
 		/// <summary>
 		/// Constructor
 		/// </summary>
-		public JobCollection(MongoService mongoService, IClock clock, IGraphCollection graphCollection, ILogCollection logCollection, ITelemetrySink telemetrySink, IOptionsMonitor<GlobalConfig> globalConfig, Tracer tracer, ILogger<JobCollection> logger)
+		public JobCollection(IMongoService mongoService, IClock clock, IGraphCollection graphCollection, ILogCollection logCollection, ITelemetrySink telemetrySink, IOptionsMonitor<GlobalConfig> globalConfig, Tracer tracer, ILogger<JobCollection> logger)
 		{
 			_clock = clock;
 			_graphCollection = graphCollection;
diff --git a/Engine/Source/Programs/Horde/HordeServer/Jobs/JobStepRefCollection.cs b/Engine/Source/Programs/Horde/HordeServer/Jobs/JobStepRefCollection.cs
--- a/Engine/Source/Programs/Horde/HordeServer/Jobs/JobStepRefCollection.cs
+++ b/Engine/Source/Programs/Horde/HordeServer/Jobs/JobStepRefCollection.cs
@@ -110,7 +110,7 @@ namespace HordeServer.Jobs
 		/// <summary>
 		/// Constructor
 		/// </summary>
-		public JobStepRefCollection(MongoService mongoService, ITelemetrySink telemetrySink, IOptionsMonitor<GlobalConfig> globalConfig)
+		public JobStepRefCollection(IMongoService mongoService, ITelemetrySink telemetrySink, IOptionsMonitor<GlobalConfig> globalConfig)
 		{
 			List<MongoIndex<JobStepRef>> indexes = new List<MongoIndex<JobStepRef>>();
 			indexes.Add(keys => keys.Ascending(x => x.StreamId).Ascending(x => x.TemplateId).Ascending(x => x.Name).Descending(x => x.Change));
diff --git a/Engine/Source/Programs/Horde/HordeServer/Jobs/Schedules/ScheduleService.cs b/Engine/Source/Programs/Horde/HordeServer/Jobs/Schedules/ScheduleService.cs
--- a/Engine/Source/Programs/Horde/HordeServer/Jobs/Schedules/ScheduleService.cs
+++ b/Engine/Source/Programs/Horde/HordeServer/Jobs/Schedules/ScheduleService.cs
@@ -78,7 +78,7 @@ namespace HordeServer.Jobs.Schedules
 		/// <summary>
 		/// Constructor
 		/// </summary>
-		public ScheduleService(IRedisService redis, IGraphCollection graphs, ICommitService commitService, IJobCollection jobCollection, JobService jobService, IDowntimeService downtimeService, IStreamCollection streamCollection, ITemplateCollection templateCollection, MongoService mongoService, IClock clock, IOptionsMonitor<GlobalConfig> globalConfig, Tracer tracer, ILogger<ScheduleService> logger)
+		public ScheduleService(IRedisService redis, IGraphCollection graphs, ICommitService commitService, IJobCollection jobCollection, JobService jobService, IDowntimeService downtimeService, IStreamCollection streamCollection, ITemplateCollection templateCollection, IMongoService mongoService, IClock clock, IOptionsMonitor<GlobalConfig> globalConfig, Tracer tracer, ILogger<ScheduleService> logger)
 		{
 			_graphs = graphs;
 			_commitService = commitService;
diff --git a/Engine/Source/Programs/Horde/HordeServer/Jobs/Templates/TemplateCollection.cs b/Engine/Source/Programs/Horde/HordeServer/Jobs/Templates/TemplateCollection.cs
--- a/Engine/Source/Programs/Horde/HordeServer/Jobs/Templates/TemplateCollection.cs
+++ b/Engine/Source/Programs/Horde/HordeServer/Jobs/Templates/TemplateCollection.cs
@@ -91,7 +91,7 @@ namespace HordeServer.Jobs.Templates
 		/// Constructor
 		/// </summary>
 		/// <param name="mongoService">The database service singleton</param>
-		public TemplateCollection(MongoService mongoService)
+		public TemplateCollection(IMongoService mongoService)
 		{
 			_templates = mongoService.GetCollection<TemplateDocument>("Templates");
 
diff --git a/Engine/Source/Programs/Horde/HordeServer/Jobs/TestData/TestDataCollection.cs b/Engine/Source/Programs/Horde/HordeServer/Jobs/TestData/TestDataCollection.cs
--- a/Engine/Source/Programs/Horde/HordeServer/Jobs/TestData/TestDataCollection.cs
+++ b/Engine/Source/Programs/Horde/HordeServer/Jobs/TestData/TestDataCollection.cs
@@ -403,7 +403,7 @@ namespace HordeServer.Jobs.TestData
 		/// <summary>
 		/// Constructor
 		/// </summary>
-		public TestDataCollection(MongoService mongoService, Tracer tracer, ILogger<TestDataCollection> logger)
+		public TestDataCollection(IMongoService mongoService, Tracer tracer, ILogger<TestDataCollection> logger)
 		{
 			_tracer = tracer;
 			_logger = logger;
diff --git a/Engine/Source/Programs/Horde/HordeServer/Jobs/Timing/JobTimingCollection.cs b/Engine/Source/Programs/Horde/HordeServer/Jobs/Timing/JobTimingCollection.cs
--- a/Engine/Source/Programs/Horde/HordeServer/Jobs/Timing/JobTimingCollection.cs
+++ b/Engine/Source/Programs/Horde/HordeServer/Jobs/Timing/JobTimingCollection.cs
@@ -81,7 +81,7 @@ namespace HordeServer.Jobs.Timing
 		/// Constructor
 		/// </summary>
 		/// <param name="mongoService">The database service singleton</param>
-		public JobTimingCollection(MongoService mongoService)
+		public JobTimingCollection(IMongoService mongoService)
 		{
 			_collection = mongoService.GetCollection<JobTimingDocument>("JobTiming");
 		}
diff --git a/Engine/Source/Programs/Horde/HordeServer/Logs/LogCollection.cs b/Engine/Source/Programs/Horde/HordeServer/Logs/LogCollection.cs
--- a/Engine/Source/Programs/Horde/HordeServer/Logs/LogCollection.cs
+++ b/Engine/Source/Programs/Horde/HordeServer/Logs/LogCollection.cs
@@ -270,7 +270,7 @@ namespace HordeServer.Logs
 		/// <summary>
 		/// Constructor
 		/// </summary>
-		public LogCollection(MongoService mongoService, ILogStorage storage, LogTailService logTailService, StorageService storageService, Tracer tracer, ILogger<LogCollection> logger)
+		public LogCollection(IMongoService mongoService, ILogStorage storage, LogTailService logTailService, StorageService storageService, Tracer tracer, ILogger<LogCollection> logger)
 		{
 			_logCollection = mongoService.GetCollection<LogDocument>("LogFiles");
 			_storage = storage;
diff --git a/Engine/Source/Programs/Horde/HordeServer/Notifications/NotificationService.cs b/Engine/Source/Programs/Horde/HordeServer/Notifications/NotificationService.cs
--- a/Engine/Source/Programs/Horde/HordeServer/Notifications/NotificationService.cs
+++ b/Engine/Source/Programs/Horde/HordeServer/Notifications/NotificationService.cs
@@ -147,7 +147,7 @@ namespace HordeServer.Notifications
 			ILogCollection logCollection,
 			Meter meter,
 			IMemoryCache cache,
-			RedisService redisService,
+			IRedisService redisService,
 			ConfigService configService,
 			IClock clock)
 		{
diff --git a/Engine/Source/Programs/Horde/HordeServer/Notifications/NotificationTriggerCollection.cs b/Engine/Source/Programs/Horde/HordeServer/Notifications/NotificationTriggerCollection.cs
--- a/Engine/Source/Programs/Horde/HordeServer/Notifications/NotificationTriggerCollection.cs
+++ b/Engine/Source/Programs/Horde/HordeServer/Notifications/NotificationTriggerCollection.cs
@@ -67,7 +67,7 @@ namespace HordeServer.Notifications
 		/// </summary>
 		/// <param name="mongoService">The database singleton</param>
 		/// <param name="userCollection"></param>
-		public NotificationTriggerCollection(MongoService mongoService, IUserCollection userCollection)
+		public NotificationTriggerCollection(IMongoService mongoService, IUserCollection userCollection)
 		{
 			_triggers = mongoService.GetCollection<TriggerDocument>("NotificationTriggers");
 			_userCollection = userCollection;
diff --git a/Engine/Source/Programs/Horde/HordeServer/Notifications/Sinks/SlackNotificationSink.cs b/Engine/Source/Programs/Horde/HordeServer/Notifications/Sinks/SlackNotificationSink.cs
--- a/Engine/Source/Programs/Horde/HordeServer/Notifications/Sinks/SlackNotificationSink.cs
+++ b/Engine/Source/Programs/Horde/HordeServer/Notifications/Sinks/SlackNotificationSink.cs
@@ -211,7 +211,7 @@ namespace HordeServer.Notifications.Sinks
 			}
 		}
 
-		readonly RedisService _redisService;
+		readonly IRedisService _redisService;
 		readonly IssueService _issueService;
 		readonly IUserCollection _userCollection;
 		readonly ILogCollection _logCollection;
@@ -246,7 +246,7 @@ namespace HordeServer.Notifications.Sinks
 		/// <summary>
 		/// Constructor
 		/// </summary>
-		public SlackNotificationSink(MongoService mongoService, RedisService redisService, IssueService issueService, IUserCollection userCollection, ILogCollection logCollection, IExternalIssueService externalIssueService, IWebHostEnvironment environment, IOptions<ServerSettings> settings, IClock clock, IOptionsMonitor<GlobalConfig> globalConfig, ILogger<SlackNotificationSink> logger)
+		public SlackNotificationSink(IMongoService mongoService, RedisService redisService, IssueService issueService, IUserCollection userCollection, ILogCollection logCollection, IExternalIssueService externalIssueService, IWebHostEnvironment environment, IOptions<ServerSettings> settings, IClock clock, IOptionsMonitor<GlobalConfig> globalConfig, ILogger<SlackNotificationSink> logger)
 		{
 			_redisService = redisService;
 			_issueService = issueService;
diff --git a/Engine/Source/Programs/Horde/HordeServer/Notifications/SubscriptionCollection.cs b/Engine/Source/Programs/Horde/HordeServer/Notifications/SubscriptionCollection.cs
--- a/Engine/Source/Programs/Horde/HordeServer/Notifications/SubscriptionCollection.cs
+++ b/Engine/Source/Programs/Horde/HordeServer/Notifications/SubscriptionCollection.cs
@@ -182,7 +182,7 @@ namespace HordeServer.Notifications
 		/// Constructor
 		/// </summary>
 		/// <param name="mongoService">The database service</param>
-		public SubscriptionCollection(MongoService mongoService)
+		public SubscriptionCollection(IMongoService mongoService)
 		{
 			List<MongoIndex<Subscription>> indexes = new List<MongoIndex<Subscription>>();
 			indexes.Add(keys => keys.Ascending(x => x.Event));
diff --git a/Engine/Source/Programs/Horde/HordeServer/Perforce/PerforceLoadBalancer.cs b/Engine/Source/Programs/Horde/HordeServer/Perforce/PerforceLoadBalancer.cs
--- a/Engine/Source/Programs/Horde/HordeServer/Perforce/PerforceLoadBalancer.cs
+++ b/Engine/Source/Programs/Horde/HordeServer/Perforce/PerforceLoadBalancer.cs
@@ -165,7 +165,7 @@ namespace HordeServer.Perforce
 		/// <summary>
 		/// Constructor
 		/// </summary>
-		public PerforceLoadBalancer(MongoService mongoService, IRedisService redisService, ILeaseCollection leaseCollection, IClock clock, HttpClient httpClient, IOptionsMonitor<GlobalConfig> globalConfig, IHealthMonitor<PerforceLoadBalancer> health, Tracer tracer, ILogger<PerforceLoadBalancer> logger)
+		public PerforceLoadBalancer(IMongoService mongoService, IRedisService redisService, ILeaseCollection leaseCollection, IClock clock, HttpClient httpClient, IOptionsMonitor<GlobalConfig> globalConfig, IHealthMonitor<PerforceLoadBalancer> health, Tracer tracer, ILogger<PerforceLoadBalancer> logger)
 		{
 			_redisService = redisService;
 			_leaseCollection = leaseCollection;
diff --git a/Engine/Source/Programs/Horde/HordeServer/Perforce/PerforceServiceCache.cs b/Engine/Source/Programs/Horde/HordeServer/Perforce/PerforceServiceCache.cs
--- a/Engine/Source/Programs/Horde/HordeServer/Perforce/PerforceServiceCache.cs
+++ b/Engine/Source/Programs/Horde/HordeServer/Perforce/PerforceServiceCache.cs
@@ -222,8 +222,8 @@ namespace HordeServer.Perforce
 			}
 		}
 
-		readonly MongoService _mongoService;
-		readonly RedisService _redisService;
+		readonly IMongoService _mongoService;
+		readonly IRedisService _redisService;
 		readonly IDowntimeService _downtimeService;
 		readonly IMongoCollection<CachedCommitDoc> _commits;
 		readonly IOptionsMonitor<GlobalConfig> _globalConfig;
@@ -237,7 +237,7 @@ namespace HordeServer.Perforce
 		/// <summary>
 		/// Constructor
 		/// </summary>
-		public PerforceServiceCache(PerforceLoadBalancer loadBalancer, MongoService mongoService, RedisService redisService, IDowntimeService downtimeService, IUserCollection userCollection, IClock clock, IOptions<ServerSettings> settings, IOptionsMonitor<GlobalConfig> globalConfig, Tracer tracer, ILogger<PerforceService> logger)
+		public PerforceServiceCache(PerforceLoadBalancer loadBalancer, IMongoService mongoService, IRedisService redisService, IDowntimeService downtimeService, IUserCollection userCollection, IClock clock, IOptions<ServerSettings> settings, IOptionsMonitor<GlobalConfig> globalConfig, Tracer tracer, ILogger<PerforceService> logger)
 			: base(loadBalancer, userCollection, settings, globalConfig, tracer, logger)
 		{
 			_mongoService = mongoService;
diff --git a/Engine/Source/Programs/Horde/HordeServer/Replicators/ReplicatorCollection.cs b/Engine/Source/Programs/Horde/HordeServer/Replicators/ReplicatorCollection.cs
--- a/Engine/Source/Programs/Horde/HordeServer/Replicators/ReplicatorCollection.cs
+++ b/Engine/Source/Programs/Horde/HordeServer/Replicators/ReplicatorCollection.cs
@@ -116,7 +116,7 @@ namespace HordeServer.Replicators
 		/// <summary>
 		/// Constructor
 		/// </summary>
-		public ReplicatorCollection(MongoService mongoService, IClock clock)
+		public ReplicatorCollection(IMongoService mongoService, IClock clock)
 		{
 			_documents = mongoService.GetCollection<ReplicatorDoc>("Replicators");
 			_clock = clock;
diff --git a/Engine/Source/Programs/Horde/HordeServer/Server/DebugController.cs b/Engine/Source/Programs/Horde/HordeServer/Server/DebugController.cs
--- a/Engine/Source/Programs/Horde/HordeServer/Server/DebugController.cs
+++ b/Engine/Source/Programs/Horde/HordeServer/Server/DebugController.cs
@@ -44,7 +44,7 @@ namespace HordeServer.Server
 	{
 		private static readonly Random s_random = new();
 
-		private readonly MongoService _mongoService;
+		private readonly IMongoService _mongoService;
 		private readonly ConfigService _configService;
 		private readonly AgentRelayService _agentRelayService;
 		private readonly ILogCollection _logCollection;
@@ -55,7 +55,7 @@ namespace HordeServer.Server
 		/// Constructor
 		/// </summary>
 		public DebugController(
-			MongoService mongoService,
+			IMongoService mongoService,
 			ConfigService configService,
 			AgentRelayService agentRelayService,
 			ILogCollection logCollection,
diff --git a/Engine/Source/Programs/Horde/HordeServer/Server/GlobalsService.cs b/Engine/Source/Programs/Horde/HordeServer/Server/GlobalsService.cs
--- a/Engine/Source/Programs/Horde/HordeServer/Server/GlobalsService.cs
+++ b/Engine/Source/Programs/Horde/HordeServer/Server/GlobalsService.cs
@@ -65,7 +65,7 @@ namespace HordeServer.Server
 			}
 		}
 
-		readonly MongoService _mongoService;
+		readonly IMongoService _mongoService;
 		readonly string _jwtIssuer;
 
 		/// <summary>
@@ -73,7 +73,7 @@ namespace HordeServer.Server
 		/// </summary>
 		/// <param name="mongoService"></param>
 		/// <param name="settings">Global settings instance</param>
-		public GlobalsService(MongoService mongoService, IOptions<ServerSettings> settings)
+		public GlobalsService(IMongoService mongoService, IOptions<ServerSettings> settings)
 		{
 			_mongoService = mongoService;
 
diff --git a/Engine/Source/Programs/Horde/HordeServer/Server/IMongoService.cs b/Engine/Source/Programs/Horde/HordeServer/Server/IMongoService.cs
new file mode 100644
--- /dev/null
+++ b/Engine/Source/Programs/Horde/HordeServer/Server/IMongoService.cs
@@ -0,0 +1,101 @@
+// Copyright Epic Games, Inc. All Rights Reserved.
+
+using System;
+using System.Collections.Generic;
+using System.Threading;
+using System.Threading.Tasks;
+using HordeServer.Utilities;
+using Microsoft.Extensions.Diagnostics.HealthChecks;
+using MongoDB.Driver;
+
+namespace HordeServer.Server
+{
+	/// <summary>
+	/// Interface for the Horde MongoDB instance
+	/// </summary>
+	public interface IMongoService : IHealthCheck
+	{
+		/// <summary>
+		/// The database instance
+		/// </summary>
+		IMongoDatabase Database { get; }
+
+		/// <summary>
+		/// Access the database in a read-only mode (don't create indices or modify content)
+		/// </summary>
+		public bool ReadOnlyMode { get; }
+
+		/// <summary>
+		/// Get the MongoDB client
+		/// </summary>
+		/// <returns>A MongoDB client instance</returns>
+		MongoClient GetClient();
+
+		/// <summary>
+		/// Get a MongoDB collection from database
+		/// </summary>
+		/// <param name="name">Name of collection</param>
+		/// <typeparam name="T">A MongoDB document</typeparam>
+		/// <returns></returns>
+		IMongoCollection<T> GetCollection<T>(string name);
+
+		/// <summary>
+		/// Get a MongoDB collection from database with a single index
+		/// </summary>
+		/// <param name="name">Name of collection</param>
+		/// <param name="keysFunc">Method to configure keys for the collection</param>
+		/// <param name="unique">Whether a unique index is required</param>
+		/// <typeparam name="T">A MongoDB document</typeparam>
+		/// <returns></returns>
+		IMongoCollection<T> GetCollection<T>(string name, Func<IndexKeysDefinitionBuilder<T>, IndexKeysDefinition<T>> keysFunc, bool unique = false);
+
+		/// <summary>
+		/// Get a MongoDB collection from database
+		/// </summary>
+		/// <param name="name">Name of collection</param>
+		/// <param name="indexes">Indexes for the collection</param>
+		/// <typeparam name="T">A MongoDB document</typeparam>
+		/// <returns></returns>
+		IMongoCollection<T> GetCollection<T>(string name, IEnumerable<MongoIndex<T>> indexes);
+
+		/// <summary>
+		/// Gets a singleton document by id
+		/// </summary>
+		/// <returns>The document</returns>
+		Task<T> GetSingletonAsync<T>(CancellationToken cancellationToken) where T : SingletonBase, new();
+
+		/// <summary>
+		/// Gets a singleton document by id
+		/// </summary>
+		/// <param name="constructor">Method to use to construct a new object</param>
+		/// <param name="cancellationToken">Cancellation token for the operation</param>
+		/// <returns>The document</returns>
+		Task<T> GetSingletonAsync<T>(Func<T> constructor, CancellationToken cancellationToken) where T : SingletonBase, new();
+
+		/// <summary>
+		/// Updates a singleton
+		/// </summary>
+		/// <typeparam name="T"></typeparam>
+		/// <param name="updater"></param>
+		/// <param name="cancellationToken">Cancellation token for the operation</param>
+		/// <returns></returns>
+		Task UpdateSingletonAsync<T>(Action<T> updater, CancellationToken cancellationToken) where T : SingletonBase, new();
+
+		/// <summary>
+		/// Updates a singleton
+		/// </summary>
+		/// <typeparam name="T"></typeparam>
+		/// <param name="updater"></param>
+		/// <param name="cancellationToken">Cancellation token for the operation</param>
+		/// <returns></returns>
+		Task UpdateSingletonAsync<T>(Func<T, bool> updater, CancellationToken cancellationToken) where T : SingletonBase, new();
+
+		/// <summary>
+		/// Attempts to update a singleton object
+		/// </summary>
+		/// <param name="singletonObject">The singleton object</param>
+		/// <param name="cancellationToken">Cancellation token for the operation</param>
+		/// <returns>True if the singleton document was updated</returns>
+		Task<bool> TryUpdateSingletonAsync<T>(T singletonObject, CancellationToken cancellationToken) where T : SingletonBase;
+	}
+}
diff --git a/Engine/Source/Programs/Horde/HordeServer/Server/IRedisService.cs b/Engine/Source/Programs/Horde/HordeServer/Server/IRedisService.cs
--- a/Engine/Source/Programs/Horde/HordeServer/Server/IRedisService.cs
+++ b/Engine/Source/Programs/Horde/HordeServer/Server/IRedisService.cs
@@ -3,6 +3,7 @@
 using System;
 using System.Threading.Tasks;
 using EpicGames.Redis;
+using Microsoft.Extensions.Diagnostics.HealthChecks;
 using StackExchange.Redis;
 
 namespace HordeServer.Server
@@ -10,7 +11,7 @@ namespace HordeServer.Server
 	/// <summary>
 	/// Provides access to a Redis database
 	/// </summary>
-	public interface IRedisService
+	public interface IRedisService : IHealthCheck
 	{
 		/// <summary>
 		/// Flag for whether the connection is read-only
diff --git a/Engine/Source/Programs/Horde/HordeServer/Server/LifetimeService.cs b/Engine/Source/Programs/Horde/HordeServer/Server/LifetimeService.cs
--- a/Engine/Source/Programs/Horde/HordeServer/Server/LifetimeService.cs
+++ b/Engine/Source/Programs/Horde/HordeServer/Server/LifetimeService.cs
@@ -37,7 +37,7 @@ namespace HordeServer.Server
 		readonly CancellationTokenRegistration _registration;
 
 		readonly IHostApplicationLifetime _lifetime;
-		readonly MongoService _mongoService;
+		readonly IMongoService _mongoService;
 		readonly IRedisService _redisService;
 		readonly ITicker _ticker;
 
@@ -67,7 +67,7 @@ namespace HordeServer.Server
 		/// <param name="redisService">Redis singleton service</param>
 		/// <param name="clock"></param>
 		/// <param name="logger">Logging interface</param>
-		public LifetimeService(IOptionsMonitor<ServerSettings> settings, IHostApplicationLifetime lifetime, IHostEnvironment env, MongoService mongoService, IRedisService redisService, IClock clock, ILogger<LifetimeService> logger)
+		public LifetimeService(IOptionsMonitor<ServerSettings> settings, IHostApplicationLifetime lifetime, IHostEnvironment env, IMongoService mongoService, IRedisService redisService, IClock clock, ILogger<LifetimeService> logger)
 		{
 			_shutdownMemoryThreshold = settings.CurrentValue.ShutdownMemoryThreshold;
 			_lifetime = lifetime;
diff --git a/Engine/Source/Programs/Horde/HordeServer/Server/MongoIndex.cs b/Engine/Source/Programs/Horde/HordeServer/Server/MongoIndex.cs
new file mode 100644
--- /dev/null
+++ b/Engine/Source/Programs/Horde/HordeServer/Server/MongoIndex.cs
@@ -0,0 +1,164 @@
+// Copyright Epic Games, Inc. All Rights Reserved.
+
+using System;
+using System.Collections.Generic;
+using System.Diagnostics;
+using System.Text;
+using MongoDB.Bson;
+using MongoDB.Bson.Serialization;
+using MongoDB.Bson.Serialization.Attributes;
+using MongoDB.Driver;
+
+namespace HordeServer.Server
+{
+	/// <summary>
+	/// Constrained set of parameters for building a mongo index
+	/// </summary>
+	[DebuggerDisplay("{Name}")]
+	public class MongoIndex : IEquatable<MongoIndex>
+	{
+		/// <summary>
+		/// Name of the index
+		/// </summary>
+		[BsonElement("name"), BsonRequired]
+		public string Name { get; private set; }
+
+		/// <summary>
+		/// Keys for the index
+		/// </summary>
+		[BsonElement("key"), BsonRequired]
+		public BsonDocument KeysDocument { get; private set; }
+
+		/// <summary>
+		/// The index should be unique
+		/// </summary>
+		[BsonElement("unique")]
+		public bool Unique { get; private set; }
+
+		/// <summary>
+		/// Whether to create a sparse index
+		/// </summary>
+		[BsonElement("sparse")]
+		public bool Sparse { get; private set; }
+
+		/// <summary>
+		/// Constructor
+		/// </summary>
+		[BsonConstructor]
+		private MongoIndex()
+		{
+			Name = String.Empty;
+			KeysDocument = new BsonDocument();
+		}
+
+		/// <summary>
+		/// Constructor
+		/// </summary>
+		/// <param name="name">Name of the index</param>
+		/// <param name="keysDocument">Keys for the index</param>
+		/// <param name="unique">Whether the index should be unique</param>
+		/// <param name="sparse">Whether to create a sparse index</param>
+		public MongoIndex(string? name, BsonDocument keysDocument, bool unique, bool sparse)
+		{
+			Name = name ?? GetDefaultName(keysDocument);
+			KeysDocument = keysDocument;
+			Unique = unique;
+			Sparse = sparse;
+		}
+
+		/// <summary>
+		/// Constructor
+		/// </summary>
+		/// <param name="keysFunc">Callback to configure keys for this document</param>
+		/// <param name="unique">Whether the index should be unique</param>
+		/// <param name="sparse">Whether to create a sparse index</param>
+		public static MongoIndex<T> Create<T>(Func<IndexKeysDefinitionBuilder<T>, IndexKeysDefinition<T>> keysFunc, bool unique = false, bool sparse = false)
+		{
+			return new MongoIndex<T>(null, keysFunc(Builders<T>.IndexKeys), unique, sparse);
+		}
+
+		/// <summary>
+		/// Constructor
+		/// </summary>
+		/// <param name="name">Name of the index</param>
+		/// <param name="keysFunc">Callback to configure keys for this document</param>
+		/// <param name="unique">Whether the index should be unique</param>
+		/// <param name="sparse">Whether to create a sparse index</param>
+		public static MongoIndex<T> Create<T>(string name, Func<IndexKeysDefinitionBuilder<T>, IndexKeysDefinition<T>> keysFunc, bool unique = false, bool sparse = false)
+		{
+			return new MongoIndex<T>(name, keysFunc(Builders<T>.IndexKeys), unique, sparse);
+		}
+
+		/// <inheritdoc/>
+		public override bool Equals(object? obj) => obj is MongoIndex other && Equals(other);
+
+		/// <inheritdoc/>
+		public override int GetHashCode() => Name.GetHashCode(StringComparison.Ordinal);
+
+		/// <inheritdoc/>
+		public bool Equals(MongoIndex? other) => other != null && Name.Equals(other.Name, StringComparison.Ordinal) && KeysDocument.Equals(other.KeysDocument) && Sparse == other.Sparse && Unique == other.Unique;
+
+		/// <summary>
+		/// Gets the default name for an index based on its keys
+		/// </summary>
+		/// <param name="keys">Keys for the index</param>
+		/// <returns>Name of the index</returns>
+		protected static string GetDefaultName(BsonDocument keys)
+		{
+			StringBuilder name = new StringBuilder();
+			foreach (BsonElement element in keys.Elements)
+			{
+				if (name.Length > 0)
+				{
+					name.Append('_');
+				}
+				name.Append(element.Name);
+				name.Append('_');
+				name.Append(element.Value.ToString());
+			}
+			return name.ToString();
+		}
+	}
+
+	/// <summary>
+	/// Strongly typed index document
+	/// </summary>
+	/// <typeparam name="T"></typeparam>
+	public class MongoIndex<T> : MongoIndex
+	{
+		/// <summary>
+		/// Keys for the index
+		/// </summary>
+		[BsonIgnore]
+		public IndexKeysDefinition<T> Keys { get; }
+
+		/// <summary>
+		/// Constructor
+		/// </summary>
+		/// <param name="name">Name of the index</param>
+		/// <param name="keys">Keys for the index</param>
+		/// <param name="unique">Whether the index should be unique</param>
+		/// <param name="sparse">Whether to create a sparse index</param>
+		public MongoIndex(string? name, IndexKeysDefinition<T> keys, bool unique, bool sparse)
+			: base(name, keys.Render(BsonSerializer.LookupSerializer<T>(), BsonSerializer.SerializerRegistry), unique, sparse)
+		{
+			Keys = keys;
+		}
+	}
+
+	/// <summary>
+	/// Extension methods for <see cref="MongoIndex{T}"/>
+	/// </summary>
+	static class MongoIndexExtensions
+	{
+		public static void Add<T>(this List<MongoIndex<T>> list, Func<IndexKeysDefinitionBuilder<T>, IndexKeysDefinition<T>> keyFunc, bool unique = false, bool sparse = false)
+		{
+			list.Add(MongoIndex.Create(keyFunc, unique, sparse));
+		}
+
+		public static void Add<T>(this List<MongoIndex<T>> list, string name, Func<IndexKeysDefinitionBuilder<T>, IndexKeysDefinition<T>> keyFunc, bool unique = false, bool sparse = false)
+		{
+			list.Add(MongoIndex.Create(name, keyFunc, unique, sparse));
+		}
+	}
+}
diff --git a/Engine/Source/Programs/Horde/HordeServer/Server/MongoService.cs b/Engine/Source/Programs/Horde/HordeServer/Server/MongoService.cs
--- a/Engine/Source/Programs/Horde/HordeServer/Server/MongoService.cs
+++ b/Engine/Source/Programs/Horde/HordeServer/Server/MongoService.cs
@@ -33,157 +33,6 @@ using StackExchange.Redis;
 
 namespace HordeServer.Server
 {
-	/// <summary>
-	/// Constrained set of parameters for building a mongo index
-	/// </summary>
-	[DebuggerDisplay("{Name}")]
-	public class MongoIndex : IEquatable<MongoIndex>
-	{
-		/// <summary>
-		/// Name of the index
-		/// </summary>
-		[BsonElement("name"), BsonRequired]
-		public string Name { get; private set; }
-
-		/// <summary>
-		/// Keys for the index
-		/// </summary>
-		[BsonElement("key"), BsonRequired]
-		public BsonDocument KeysDocument { get; private set; }
-
-		/// <summary>
-		/// The index should be unique
-		/// </summary>
-		[BsonElement("unique")]
-		public bool Unique { get; private set; }
-
-		/// <summary>
-		/// Whether to create a sparse index
-		/// </summary>
-		[BsonElement("sparse")]
-		public bool Sparse { get; private set; }
-
-		/// <summary>
-		/// Constructor
-		/// </summary>
-		[BsonConstructor]
-		private MongoIndex()
-		{
-			Name = String.Empty;
-			KeysDocument = new BsonDocument();
-		}
-
-		/// <summary>
-		/// Constructor
-		/// </summary>
-		/// <param name="name">Name of the index</param>
-		/// <param name="keysDocument">Keys for the index</param>
-		/// <param name="unique">Whether the index should be unique</param>
-		/// <param name="sparse">Whether to create a sparse index</param>
-		public MongoIndex(string? name, BsonDocument keysDocument, bool unique, bool sparse)
-		{
-			Name = name ?? GetDefaultName(keysDocument);
-			KeysDocument = keysDocument;
-			Unique = unique;
-			Sparse = sparse;
-		}
-
-		/// <summary>
-		/// Constructor
-		/// </summary>
-		/// <param name="keysFunc">Callback to configure keys for this document</param>
-		/// <param name="unique">Whether the index should be unique</param>
-		/// <param name="sparse">Whether to create a sparse index</param>
-		public static MongoIndex<T> Create<T>(Func<IndexKeysDefinitionBuilder<T>, IndexKeysDefinition<T>> keysFunc, bool unique = false, bool sparse = false)
-		{
-			return new MongoIndex<T>(null, keysFunc(Builders<T>.IndexKeys), unique, sparse);
-		}
-
-		/// <summary>
-		/// Constructor
-		/// </summary>
-		/// <param name="name">Name of the index</param>
-		/// <param name="keysFunc">Callback to configure keys for this document</param>
-		/// <param name="unique">Whether the index should be unique</param>
-		/// <param name="sparse">Whether to create a sparse index</param>
-		public static MongoIndex<T> Create<T>(string name, Func<IndexKeysDefinitionBuilder<T>, IndexKeysDefinition<T>> keysFunc, bool unique = false, bool sparse = false)
-		{
-			return new MongoIndex<T>(name, keysFunc(Builders<T>.IndexKeys), unique, sparse);
-		}
-
-		/// <inheritdoc/>
-		public override bool Equals(object? obj) => obj is MongoIndex other && Equals(other);
-
-		/// <inheritdoc/>
-		public override int GetHashCode() => Name.GetHashCode(StringComparison.Ordinal);
-
-		/// <inheritdoc/>
-		public bool Equals(MongoIndex? other) => other != null && Name.Equals(other.Name, StringComparison.Ordinal) && KeysDocument.Equals(other.KeysDocument) && Sparse == other.Sparse && Unique == other.Unique;
-
-		/// <summary>
-		/// Gets the default name for an index based on its keys
-		/// </summary>
-		/// <param name="keys">Keys for the index</param>
-		/// <returns>Name of the index</returns>
-		protected static string GetDefaultName(BsonDocument keys)
-		{
-			StringBuilder name = new StringBuilder();
-			foreach (BsonElement element in keys.Elements)
-			{
-				if (name.Length > 0)
-				{
-					name.Append('_');
-				}
-				name.Append(element.Name);
-				name.Append('_');
-				name.Append(element.Value.ToString());
-			}
-			return name.ToString();
-		}
-	}
-
-	/// <summary>
-	/// Strongly typed index document
-	/// </summary>
-	/// <typeparam name="T"></typeparam>
-	public class MongoIndex<T> : MongoIndex
-	{
-		/// <summary>
-		/// Keys for the index
-		/// </summary>
-		[BsonIgnore]
-		public IndexKeysDefinition<T> Keys { get; }
-
-		/// <summary>
-		/// Constructor
-		/// </summary>
-		/// <param name="name">Name of the index</param>
-		/// <param name="keys">Keys for the index</param>
-		/// <param name="unique">Whether the index should be unique</param>
-		/// <param name="sparse">Whether to create a sparse index</param>
-		public MongoIndex(string? name, IndexKeysDefinition<T> keys, bool unique, bool sparse)
-			: base(name, keys.Render(BsonSerializer.LookupSerializer<T>(), BsonSerializer.SerializerRegistry), unique, sparse)
-		{
-			Keys = keys;
-		}
-	}
-
-	/// <summary>
-	/// Extension methods for <see cref="MongoIndex{T}"/>
-	/// </summary>
-	static class MongoIndexExtensions
-	{
-		public static void Add<T>(this List<MongoIndex<T>> list, Func<IndexKeysDefinitionBuilder<T>, IndexKeysDefinition<T>> keyFunc, bool unique = false, bool sparse = false)
-		{
-			list.Add(MongoIndex.Create(keyFunc, unique, sparse));
-		}
-
-		public static void Add<T>(this List<MongoIndex<T>> list, string name, Func<IndexKeysDefinitionBuilder<T>, IndexKeysDefinition<T>> keyFunc, bool unique = false, bool sparse = false)
-		{
-			list.Add(MongoIndex.Create(name, keyFunc, unique, sparse));
-		}
-	}
-
 	/// <summary>
 	/// Task used to update indexes
 	/// </summary>
@@ -192,7 +41,7 @@ namespace HordeServer.Server
 	/// <summary>
 	/// Singleton for accessing the database
 	/// </summary>
-	public sealed class MongoService : IHealthCheck, IAsyncDisposable
+	public sealed class MongoService : IMongoService, IHealthCheck, IAsyncDisposable
 	{
 		/// <summary>
 		/// The database instance
diff --git a/Engine/Source/Programs/Horde/HordeServer/Server/Notices/NoticeCollection.cs b/Engine/Source/Programs/Horde/HordeServer/Server/Notices/NoticeCollection.cs
--- a/Engine/Source/Programs/Horde/HordeServer/Server/Notices/NoticeCollection.cs
+++ b/Engine/Source/Programs/Horde/HordeServer/Server/Notices/NoticeCollection.cs
@@ -57,7 +57,7 @@ namespace HordeServer.Server.Notices
 		/// <summary>
 		/// Constructor
 		/// </summary>
-		public NoticeCollection(MongoService mongoService)
+		public NoticeCollection(IMongoService mongoService)
 		{
 			_notices = mongoService.GetCollection<NoticeDocument>("Notices");
 		}
diff --git a/Engine/Source/Programs/Horde/HordeServer/Server/RedisService.cs b/Engine/Source/Programs/Horde/HordeServer/Server/RedisService.cs
--- a/Engine/Source/Programs/Horde/HordeServer/Server/RedisService.cs
+++ b/Engine/Source/Programs/Horde/HordeServer/Server/RedisService.cs
@@ -20,7 +20,7 @@ namespace HordeServer.Server
 	/// <summary>
 	/// Manages the lifetime of a bundled Redis instance
 	/// </summary>
-	public sealed class RedisService : IRedisService, IHealthCheck, IAsyncDisposable
+	public sealed class RedisService : IRedisService, IAsyncDisposable
 	{
 		/// <summary>
 		/// Default Redis port
diff --git a/Engine/Source/Programs/Horde/HordeServer/Server/ServerStatusService.cs b/Engine/Source/Programs/Horde/HordeServer/Server/ServerStatusService.cs
--- a/Engine/Source/Programs/Horde/HordeServer/Server/ServerStatusService.cs
+++ b/Engine/Source/Programs/Horde/HordeServer/Server/ServerStatusService.cs
@@ -100,8 +100,8 @@ public class ServerStatusService : IHostedService
 	public const int MaxHistoryLength = 10;
 
 	private readonly IClock _clock;
-	private readonly MongoService _mongoService;
-	private readonly RedisService _redis;
+	private readonly IMongoService _mongoService;
+	private readonly IRedisService _redisService;
 
 	private readonly IHealthMonitor<MongoService> _mongoDbHealth;
 	private readonly ITicker _mongoDbHealthTicker;
@@ -118,10 +118,10 @@ public class ServerStatusService : IHostedService
 	/// <param name="clock"></param>
 	/// <param name="mongoService"></param>
 	/// <param name="logger"></param>
-	public ServerStatusService(MongoService mongoService, RedisService redisService, IClock clock, ILogger<ServerStatusService> logger)
+	public ServerStatusService(IMongoService mongoService, IRedisService redisService, IClock clock, ILogger<ServerStatusService> logger)
 	{
 		_mongoService = mongoService;
-		_redis = redisService;
+		_redisService = redisService;
 		_clock = clock;
 
 		_mongoDbHealth = new HealthMonitor<MongoService>(this, "MongoDB");
@@ -159,7 +159,7 @@ public class ServerStatusService : IHostedService
 	/// </summary>
 	internal async ValueTask UpdateRedisHealthAsync(CancellationToken cancellationToken)
 	{
-		HealthCheckResult result = await _redis.CheckHealthAsync(new HealthCheckContext(), cancellationToken);
+		HealthCheckResult result = await _redisService.CheckHealthAsync(new HealthCheckContext(), cancellationToken);
 		await _redisHealth.UpdateAsync(result.Status, result.Description);
 	}
 
@@ -174,7 +174,7 @@ public class ServerStatusService : IHostedService
 	public async Task ReportAsync(Type type, string name, HealthStatus result, string? message = null, DateTimeOffset? timestamp = null)
 	{
 		string id = type.Name;
-		IDatabase redis = _redis.GetDatabase();
+		IDatabase redis = _redisService.GetDatabase();
 		SubsystemStatus status = await GetSubsystemStatusFromRedisAsync(redis, id, name);
 		SubsystemStatusUpdate update = new(result, message, timestamp ?? _clock.UtcNow);
 		status.Updates.Add(update);
@@ -186,7 +186,7 @@ public class ServerStatusService : IHostedService
 		}
 
 		string data = JsonSerializer.Serialize(status);
-		if (!_redis.ReadOnlyMode)
+		if (!_redisService.ReadOnlyMode)
 		{
 			await redis.HashSetAsync(RedisHashKey(), id, data);
 		}
@@ -216,7 +216,7 @@ public class ServerStatusService : IHostedService
 	/// <returns>A list of statuses</returns>
 	public async Task<IReadOnlyList<SubsystemStatus>> GetSubsystemStatusesAsync()
 	{
-		HashEntry[] entries = await _redis.GetDatabase().HashGetAllAsync(RedisHashKey());
+		HashEntry[] entries = await _redisService.GetDatabase().HashGetAllAsync(RedisHashKey());
 		List<SubsystemStatus> subsystems = [];
 		foreach (HashEntry entry in entries)
 		{
diff --git a/Engine/Source/Programs/Horde/HordeServer/ServiceAccounts/ServiceAccountCollection.cs b/Engine/Source/Programs/Horde/HordeServer/ServiceAccounts/ServiceAccountCollection.cs
--- a/Engine/Source/Programs/Horde/HordeServer/ServiceAccounts/ServiceAccountCollection.cs
+++ b/Engine/Source/Programs/Horde/HordeServer/ServiceAccounts/ServiceAccountCollection.cs
@@ -89,7 +89,7 @@ namespace HordeServer.ServiceAccounts
 		/// Constructor
 		/// </summary>
 		/// <param name="mongoService">The database service</param>
-		public ServiceAccountCollection(MongoService mongoService)
+		public ServiceAccountCollection(IMongoService mongoService)
 		{
 			_accounts = mongoService.GetCollection<ServiceAccountDocument>("ServiceAccounts", keys => keys.Ascending(x => x.SecretToken));
 		}
diff --git a/Engine/Source/Programs/Horde/HordeServer/Startup.cs b/Engine/Source/Programs/Horde/HordeServer/Startup.cs
--- a/Engine/Source/Programs/Horde/HordeServer/Startup.cs
+++ b/Engine/Source/Programs/Horde/HordeServer/Startup.cs
@@ -571,6 +571,7 @@ namespace HordeServer
 			services.AddSingleton<ComputeService>();
 			services.AddSingleton<MongoCommandTracer>();
 			services.AddSingleton<MongoService>();
+			services.AddSingleton<IMongoService>(sp => sp.GetRequiredService<MongoService>());
 			services.AddSingleton<GlobalsService>();
 			services.AddSingleton<ICommitService, CommitService>();
 			services.AddSingleton<IClock, Clock>();
diff --git a/Engine/Source/Programs/Horde/HordeServer/Storage/StorageService.cs b/Engine/Source/Programs/Horde/HordeServer/Storage/StorageService.cs
--- a/Engine/Source/Programs/Horde/HordeServer/Storage/StorageService.cs
+++ b/Engine/Source/Programs/Horde/HordeServer/Storage/StorageService.cs
@@ -400,7 +400,7 @@ namespace HordeServer.Storage
 		/// <summary>
 		/// Constructor
 		/// </summary>
-		public StorageService(MongoService mongoService, IRedisService redisService, IClock clock, BundleCache bundleCache, IMemoryCache memoryCache, IObjectStoreFactory objectStoreFactory, IOptionsMonitor<GlobalConfig> globalConfig, Tracer tracer, ILogger<StorageService> logger)
+		public StorageService(IMongoService mongoService, IRedisService redisService, IClock clock, BundleCache bundleCache, IMemoryCache memoryCache, IObjectStoreFactory objectStoreFactory, IOptionsMonitor<GlobalConfig> globalConfig, Tracer tracer, ILogger<StorageService> logger)
 		{
 			_redisService = redisService;
 			_clock = clock;
diff --git a/Engine/Source/Programs/Horde/HordeServer/Streams/StreamCollection.cs b/Engine/Source/Programs/Horde/HordeServer/Streams/StreamCollection.cs
--- a/Engine/Source/Programs/Horde/HordeServer/Streams/StreamCollection.cs
+++ b/Engine/Source/Programs/Horde/HordeServer/Streams/StreamCollection.cs
@@ -216,7 +216,7 @@ namespace HordeServer.Streams
 		/// </summary>
 		/// <param name="mongoService">The database service instance</param>
 		/// <param name="clock"></param>
-		public StreamCollection(MongoService mongoService, IClock clock)
+		public StreamCollection(IMongoService mongoService, IClock clock)
 		{
 			_streams = mongoService.GetCollection<StreamDoc>("Streams");
 			_clock = clock;
diff --git a/Engine/Source/Programs/Horde/HordeServer/Tasks/ConformTaskSource.cs b/Engine/Source/Programs/Horde/HordeServer/Tasks/ConformTaskSource.cs
--- a/Engine/Source/Programs/Horde/HordeServer/Tasks/ConformTaskSource.cs
+++ b/Engine/Source/Programs/Horde/HordeServer/Tasks/ConformTaskSource.cs
@@ -51,7 +51,7 @@ namespace HordeServer.Tasks
 		/// <summary>
 		/// Constructor
 		/// </summary>
-		public ConformTaskSource(MongoService mongoService, IAgentCollection agentCollection, PoolService poolService, ILogCollection logCollection, PerforceLoadBalancer perforceLoadBalancer, IClock clock, IOptionsMonitor<ServerSettings> settings, IOptionsMonitor<GlobalConfig> globalConfig, ILogger<ConformTaskSource> logger)
+		public ConformTaskSource(IMongoService mongoService, IAgentCollection agentCollection, PoolService poolService, ILogCollection logCollection, PerforceLoadBalancer perforceLoadBalancer, IClock clock, IOptionsMonitor<ServerSettings> settings, IOptionsMonitor<GlobalConfig> globalConfig, ILogger<ConformTaskSource> logger)
 		{
 			_agentCollection = agentCollection;
 			_poolService = poolService;
diff --git a/Engine/Source/Programs/Horde/HordeServer/Telemetry/Metrics/MetricCollection.cs b/Engine/Source/Programs/Horde/HordeServer/Telemetry/Metrics/MetricCollection.cs
--- a/Engine/Source/Programs/Horde/HordeServer/Telemetry/Metrics/MetricCollection.cs
+++ b/Engine/Source/Programs/Horde/HordeServer/Telemetry/Metrics/MetricCollection.cs
@@ -63,7 +63,7 @@ namespace HordeServer.Telemetry.Metrics
 
 		Dictionary<SampleKey, List<double>> _queuedSamples = new Dictionary<SampleKey, List<double>>();
 
-		public MetricCollection(MongoService mongoService, IClock clock, IOptionsMonitor<GlobalConfig> globalConfig, ILogger<MetricCollection> logger)
+		public MetricCollection(IMongoService mongoService, IClock clock, IOptionsMonitor<GlobalConfig> globalConfig, ILogger<MetricCollection> logger)
 		{
 			List<MongoIndex<MetricDocument>> indexes = new List<MongoIndex<MetricDocument>>();
 			indexes.Add(MongoIndex.Create<MetricDocument>(keys => keys.Ascending(x => x.TelemetryStoreId).Descending(x => x.Time).Ascending(x => x.MetricId).Ascending(x => x.Group)));
diff --git a/Engine/Source/Programs/Horde/HordeServer/Telemetry/Sinks/MongoTelemetrySink.cs b/Engine/Source/Programs/Horde/HordeServer/Telemetry/Sinks/MongoTelemetrySink.cs
--- a/Engine/Source/Programs/Horde/HordeServer/Telemetry/Sinks/MongoTelemetrySink.cs
+++ b/Engine/Source/Programs/Horde/HordeServer/Telemetry/Sinks/MongoTelemetrySink.cs
@@ -62,7 +62,7 @@ namespace HordeServer.Telemetry.Sinks
 		/// <summary>
 		/// Constructor
 		/// </summary>
-		public MongoTelemetrySink(MongoService mongoService, IClock clock, IOptions<ServerSettings> serverSettings, ILogger<MongoTelemetrySink> logger)
+		public MongoTelemetrySink(IMongoService mongoService, IClock clock, IOptions<ServerSettings> serverSettings, ILogger<MongoTelemetrySink> logger)
 		{
 			_config = serverSettings.Value.Telemetry.Select(x => x as MongoTelemetryConfig).FirstOrDefault(x => x != null);
 			_collection = mongoService.GetCollection<EventDocument>("Telemetry", builder => builder.Ascending(x => x.TelemetryStoreId).Descending(x => x.Id));
diff --git a/Engine/Source/Programs/Horde/HordeServer/Tools/ToolCollection.cs b/Engine/Source/Programs/Horde/HordeServer/Tools/ToolCollection.cs
--- a/Engine/Source/Programs/Horde/HordeServer/Tools/ToolCollection.cs
+++ b/Engine/Source/Programs/Horde/HordeServer/Tools/ToolCollection.cs
@@ -206,7 +206,7 @@ namespace HordeServer.Tools
 		/// <summary>
 		/// Constructor
 		/// </summary>
-		public ToolCollection(MongoService mongoService, StorageService storageService, BundleCache cache, FileObjectStoreFactory fileObjectStoreFactory, IClock clock, IOptionsMonitor<GlobalConfig> globalConfig, ILogger<ToolCollection> logger)
+		public ToolCollection(IMongoService mongoService, StorageService storageService, BundleCache cache, FileObjectStoreFactory fileObjectStoreFactory, IClock clock, IOptionsMonitor<GlobalConfig> globalConfig, ILogger<ToolCollection> logger)
 		{
 			_tools = mongoService.GetCollection<ToolDocument>("Tools");
 			_storageService = storageService;
diff --git a/Engine/Source/Programs/Horde/HordeServer/Ugs/UgsMetadataCollection.cs b/Engine/Source/Programs/Horde/HordeServer/Ugs/UgsMetadataCollection.cs
--- a/Engine/Source/Programs/Horde/HordeServer/Ugs/UgsMetadataCollection.cs
+++ b/Engine/Source/Programs/Horde/HordeServer/Ugs/UgsMetadataCollection.cs
@@ -94,7 +94,7 @@ namespace HordeServer.Ugs
 		/// <summary>
 		/// Constructor
 		/// </summary>
-		public UgsMetadataCollection(MongoService mongoService, ILogger<UgsMetadataCollection> logger)
+		public UgsMetadataCollection(IMongoService mongoService, ILogger<UgsMetadataCollection> logger)
 		{
 			List<MongoIndex<UgsMetadataDocument>> indexes = new List<MongoIndex<UgsMetadataDocument>>();
 			indexes.Add(keys => keys.Ascending(x => x.Stream).Descending(x => x.Change).Ascending(x => x.Project), unique: true);
diff --git a/Engine/Source/Programs/Horde/HordeServer/Users/UserCollection.cs b/Engine/Source/Programs/Horde/HordeServer/Users/UserCollection.cs
--- a/Engine/Source/Programs/Horde/HordeServer/Users/UserCollection.cs
+++ b/Engine/Source/Programs/Horde/HordeServer/Users/UserCollection.cs
@@ -94,7 +94,7 @@ namespace HordeServer.Users
 		/// Constructor
 		/// </summary>
 		/// <param name="mongoService"></param>
-		public UserCollectionV1(MongoService mongoService)
+		public UserCollectionV1(IMongoService mongoService)
 		{
 			_users = mongoService.GetCollection<UserDocument>("Users", keys => keys.Ascending(x => x.PrimaryClaim), unique: true);
 		}
diff --git a/Engine/Source/Programs/Horde/HordeServer/Users/UserCollectionV2.cs b/Engine/Source/Programs/Horde/HordeServer/Users/UserCollectionV2.cs
--- a/Engine/Source/Programs/Horde/HordeServer/Users/UserCollectionV2.cs
+++ b/Engine/Source/Programs/Horde/HordeServer/Users/UserCollectionV2.cs
@@ -213,7 +213,7 @@ namespace HordeServer.Users
 		/// </summary>
 		/// <param name="mongoService"></param>
 		/// <param name="logger"></param>
-		public UserCollectionV2(MongoService mongoService, ILogger<UserCollectionV2> logger)
+		public UserCollectionV2(IMongoService mongoService, ILogger<UserCollectionV2> logger)
 		{
 			_logger = logger;
 
diff --git a/Engine/Source/Programs/Horde/HordeServer/Utilities/SingletonDocument.cs b/Engine/Source/Programs/Horde/HordeServer/Utilities/SingletonDocument.cs
--- a/Engine/Source/Programs/Horde/HordeServer/Utilities/SingletonDocument.cs
+++ b/Engine/Source/Programs/Horde/HordeServer/Utilities/SingletonDocument.cs
@@ -120,13 +120,13 @@ namespace HordeServer.Utilities
 		/// <summary>
 		/// The database service instance
 		/// </summary>
-		readonly MongoService _mongoService;
+		readonly IMongoService _mongoService;
 
 		/// <summary>
 		/// Constructor
 		/// </summary>
 		/// <param name="mongoService">The database service instance</param>
-		public SingletonDocument(MongoService mongoService)
+		public SingletonDocument(IMongoService mongoService)
 		{
 			_mongoService = mongoService;
 		}
