Horde: Fix static analysis warnings.

diff --git a/Engine/Source/Programs/Horde/HordeServer.Tests/Artifacts/ArtifactTests.cs b/Engine/Source/Programs/Horde/HordeServer.Tests/Artifacts/ArtifactTests.cs
--- a/Engine/Source/Programs/Horde/HordeServer.Tests/Artifacts/ArtifactTests.cs
+++ b/Engine/Source/Programs/Horde/HordeServer.Tests/Artifacts/ArtifactTests.cs
@@ -5,15 +5,14 @@ using System.Collections.Generic;
 using System.Linq;
 using System.Threading;
 using System.Threading.Tasks;
+using EpicGames.Horde.Acls;
 using EpicGames.Horde.Artifacts;
 using EpicGames.Horde.Streams;
-using HordeServer.Acls;
+using HordeCommon;
 using HordeServer.Artifacts;
 using HordeServer.Server;
-using HordeCommon;
 using Microsoft.Extensions.DependencyInjection;
 using Microsoft.VisualStudio.TestTools.UnitTesting;
-using EpicGames.Horde.Acls;
 
 namespace HordeServer.Tests.Artifacts
 {
diff --git a/Engine/Source/Programs/Horde/HordeServer/Server/MongoService.cs b/Engine/Source/Programs/Horde/HordeServer/Server/MongoService.cs
--- a/Engine/Source/Programs/Horde/HordeServer/Server/MongoService.cs
+++ b/Engine/Source/Programs/Horde/HordeServer/Server/MongoService.cs
@@ -24,7 +24,6 @@ using Microsoft.Extensions.Logging;
 using Microsoft.Extensions.Options;
 using MongoDB.Bson;
 using MongoDB.Bson.Serialization;
-using MongoDB.Bson.Serialization.Attributes;
 using MongoDB.Driver;
 using MongoDB.Driver.Core.Events;
 using MongoDB.Driver.Linq;
@@ -192,7 +191,7 @@ namespace HordeServer.Server
 				mongoSettings.LinqProvider = LinqProvider.V2;
 				mongoSettings.SslSettings = new SslSettings();
 				mongoSettings.SslSettings.ServerCertificateValidationCallback = CertificateValidationCallBack;
-				
+
 				// Increase defaults to handle larger deployments of Horde interfacing with MongoDB
 				// Making room for slower running commands potentially hogging the queue 
 				mongoSettings.MinConnectionPoolSize = 50; // Default is 0
@@ -540,7 +539,7 @@ namespace HordeServer.Server
 			_logger.LogError("TLS certificate validation failed ({Errors}).{AdditionalInfo}", sslPolicyErrors, StringUtils.Indent(builder.ToString(), "    "));
 			return false;
 		}
-		
+
 		/// <summary>
 		/// Get the MongoDB client
 		/// </summary>
@@ -549,7 +548,7 @@ namespace HordeServer.Server
 		{
 			return _client;
 		}
-		
+
 		/// <summary>
 		/// Get a MongoDB collection from database
 		/// </summary>
diff --git a/Engine/Source/Programs/Horde/HordeServer/Server/RedisService.cs b/Engine/Source/Programs/Horde/HordeServer/Server/RedisService.cs
--- a/Engine/Source/Programs/Horde/HordeServer/Server/RedisService.cs
+++ b/Engine/Source/Programs/Horde/HordeServer/Server/RedisService.cs
@@ -13,7 +13,6 @@ using EpicGames.Redis;
 using Microsoft.Extensions.Diagnostics.HealthChecks;
 using Microsoft.Extensions.Logging;
 using Microsoft.Extensions.Options;
-using StackExchange.Redis;
 
 namespace HordeServer.Server
 {
