Horde: Move PerforceConfigSource into the Perforce namespace.

#rnx

diff --git a/Engine/Source/Programs/Horde/HordeServer/Configuration/ConfigNode.cs b/Engine/Source/Programs/Horde/HordeServer/Configuration/ConfigNode.cs
--- a/Engine/Source/Programs/Horde/HordeServer/Configuration/ConfigNode.cs
+++ b/Engine/Source/Programs/Horde/HordeServer/Configuration/ConfigNode.cs
@@ -96,13 +96,14 @@ namespace HordeServer.Configuration
 		{
 			if (path.StartsWith("//", StringComparison.Ordinal))
 			{
-				if (baseUri.Scheme == PerforceConfigSource.Scheme)
+				const string PerforceScheme = "perforce";
+				if (baseUri.Scheme == PerforceScheme)
 				{
-					return new Uri($"{PerforceConfigSource.Scheme}://{baseUri.Host}{path}");
+					return new Uri($"{PerforceScheme}://{baseUri.Host}{path}");
 				}
 				else
 				{
-					return new Uri($"{PerforceConfigSource.Scheme}://{PerforceConnectionSettings.Default}{path}");
+					return new Uri($"{PerforceScheme}://{PerforceConnectionSettings.Default}{path}");
 				}
 			}
 			return new Uri(baseUri, path);
diff --git a/Engine/Source/Programs/Horde/HordeServer/Configuration/ConfigService.cs b/Engine/Source/Programs/Horde/HordeServer/Configuration/ConfigService.cs
--- a/Engine/Source/Programs/Horde/HordeServer/Configuration/ConfigService.cs
+++ b/Engine/Source/Programs/Horde/HordeServer/Configuration/ConfigService.cs
@@ -25,11 +25,6 @@ using StackExchange.Redis;
 
 namespace HordeServer.Configuration
 {
-	/// <summary>
-	/// Information about the updated config
-	/// </summary>
-	public record class ConfigUpdateInfo(List<string> Status, HashSet<UserId> Authors, Exception? Exception);
-
 	/// <summary>
 	/// Service which processes runtime configuration data.
 	/// </summary>
diff --git a/Engine/Source/Programs/Horde/HordeServer/Configuration/IConfigSource.cs b/Engine/Source/Programs/Horde/HordeServer/Configuration/IConfigSource.cs
--- a/Engine/Source/Programs/Horde/HordeServer/Configuration/IConfigSource.cs
+++ b/Engine/Source/Programs/Horde/HordeServer/Configuration/IConfigSource.cs
@@ -4,19 +4,19 @@ using System;
 using System.Collections.Concurrent;
 using System.Collections.Generic;
 using System.IO;
-using System.Linq;
 using System.Threading;
 using System.Threading.Tasks;
 using EpicGames.Core;
 using EpicGames.Horde.Users;
-using EpicGames.Perforce;
 using HordeServer.Users;
-using Microsoft.Extensions.Caching.Memory;
-using Microsoft.Extensions.Logging;
-using Microsoft.Extensions.Options;
 
 namespace HordeServer.Configuration
 {
+	/// <summary>
+	/// Information about the updated config
+	/// </summary>
+	public record class ConfigUpdateInfo(List<string> Status, HashSet<UserId> Authors, Exception? Exception);
+
 	/// <summary>
 	/// Source for reading config files
 	/// </summary>
@@ -237,210 +237,4 @@ namespace HordeServer.Configuration
 		public Task GetUpdateInfoAsync(IReadOnlyDictionary<Uri, string> files, IReadOnlyDictionary<Uri, string>? prevFiles, ConfigUpdateInfo updateInfo, CancellationToken cancellationToken)
 			=> Task.CompletedTask;
 	}
-
-	/// <summary>
-	/// Perforce cluster config file source
-	/// </summary>
-	public sealed class PerforceConfigSource : IConfigSource
-	{
-		class ConfigFileImpl : IConfigFile
-		{
-			public Uri Uri { get; }
-			public int Change { get; }
-			public string Revision { get; }
-			public IUser? Author { get; }
-
-			readonly PerforceConfigSource _owner;
-
-			public ConfigFileImpl(Uri uri, int change, IUser? author, PerforceConfigSource owner)
-			{
-				Uri = uri;
-				Change = change;
-				Revision = $"{change}";
-				Author = author;
-				_owner = owner;
-			}
-
-			public ValueTask<ReadOnlyMemory<byte>> ReadAsync(CancellationToken cancellationToken) => _owner.ReadAsync(Uri, Change, cancellationToken);
-		}
-
-		/// <summary>
-		/// Name of the scheme for this source
-		/// </summary>
-		public const string Scheme = "perforce";
-
-		/// <inheritdoc/>
-		string IConfigSource.Scheme => Scheme;
-
-		/// <inheritdoc/>
-		public TimeSpan UpdateInterval => TimeSpan.FromMinutes(1.0);
-
-		readonly IOptionsMonitor<ServerSettings> _settings;
-		readonly IUserCollection _userCollection;
-		readonly IMemoryCache _cache;
-		readonly ILogger _logger;
-
-		/// <summary>
-		/// Constructor
-		/// </summary>
-		public PerforceConfigSource(IOptionsMonitor<ServerSettings> settings, IUserCollection userCollection, IMemoryCache cache, ILogger<PerforceConfigSource> logger)
-		{
-			_settings = settings;
-			_userCollection = userCollection;
-			_cache = cache;
-			_logger = logger;
-		}
-
-		/// <inheritdoc/>
-		public async Task<IConfigFile[]> GetFilesAsync(Uri[] uris, CancellationToken cancellationToken)
-		{
-			Dictionary<Uri, IConfigFile> results = new Dictionary<Uri, IConfigFile>();
-			foreach (IGrouping<string, Uri> group in uris.GroupBy(x => x.Host))
-			{
-				using (IPerforceConnection perforce = await ConnectAsync(group.Key, cancellationToken))
-				{
-					try
-					{
-						FileSpecList fileSpec = group.Select(x => x.AbsolutePath).Distinct(StringComparer.OrdinalIgnoreCase).ToList();
-
-						List<FStatRecord> records = await perforce.FStatAsync(FStatOptions.ShortenOutput, fileSpec, cancellationToken).ToListAsync(cancellationToken);
-						records.RemoveAll(x => x.HeadAction == FileAction.Delete || x.HeadAction == FileAction.MoveDelete);
-
-						Dictionary<string, FStatRecord> absolutePathToRecord = records.ToDictionary(x => x.DepotFile ?? String.Empty, x => x, StringComparer.OrdinalIgnoreCase);
-						foreach (Uri uri in group)
-						{
-							FStatRecord? record;
-							if (!absolutePathToRecord.TryGetValue(uri.AbsolutePath, out record))
-							{
-								throw new FileNotFoundException($"Unable to read {uri}. No matching files found.");
-							}
-
-							IUser? author = await GetAuthorAsync(perforce, group.Key, record.HeadChange, cancellationToken);
-							results[uri] = new ConfigFileImpl(uri, record.HeadChange, author, this);
-						}
-					}
-					catch (PerforceException ex) when (ex.Error != null && ex.Error.Generic == PerforceGenericCode.Config)
-					{
-						throw new PerforceException($"{ex.Message} [server: {perforce.Settings.ServerAndPort}, user: {perforce.Settings.UserName}]", ex);
-					}
-				}
-			}
-			return uris.ConvertAll(x => results[x]);
-		}
-
-		async ValueTask<IUser?> GetAuthorAsync(IPerforceConnection perforce, string host, int change, CancellationToken cancellationToken)
-		{
-			string cacheKey = $"{nameof(PerforceConfigSource)}:author:{host}@{change}";
-			if (!_cache.TryGetValue(cacheKey, out string? author))
-			{
-				ChangeRecord record = await perforce.GetChangeAsync(GetChangeOptions.None, change, cancellationToken);
-				using (ICacheEntry entry = _cache.CreateEntry(cacheKey))
-				{
-					entry.SetSlidingExpiration(TimeSpan.FromHours(1.0));
-					entry.SetSize(256);
-					entry.SetValue(record.User);
-				}
-			}
-			return (author != null) ? await _userCollection.FindUserByLoginAsync(author, cancellationToken) : null;
-		}
-
-		async ValueTask<ReadOnlyMemory<byte>> ReadAsync(Uri uri, int change, CancellationToken cancellationToken)
-		{
-			string cacheKey = $"{nameof(PerforceConfigSource)}:data:{uri}@{change}";
-			if (_cache.TryGetValue(cacheKey, out ReadOnlyMemory<byte> data))
-			{
-				_logger.LogInformation("Read {Uri}@{Change} from cache ({Key})", uri, change, cacheKey);
-			}
-			else
-			{
-				_logger.LogInformation("Reading {Uri} at CL {Change} from Perforce", uri, change);
-				using (IPerforceConnection perforce = await ConnectAsync(uri.Host, cancellationToken))
-				{
-					PerforceResponse<PrintRecord<byte[]>> response = await perforce.TryPrintDataAsync($"{uri.AbsolutePath}@{change}", cancellationToken);
-					response.EnsureSuccess();
-					data = response.Data.Contents!;
-				}
-				using (ICacheEntry entry = _cache.CreateEntry(cacheKey))
-				{
-					entry.SetSlidingExpiration(TimeSpan.FromHours(1.0));
-					entry.SetSize(data.Length);
-					entry.SetValue(data);
-				}
-			}
-			return data;
-		}
-
-		async Task<IPerforceConnection> ConnectAsync(string host, CancellationToken cancellationToken)
-		{
-			_ = cancellationToken;
-
-			ServerSettings settings = _settings.CurrentValue;
-
-			PerforceConnectionId connectionId = new PerforceConnectionId();
-			if (!String.IsNullOrEmpty(host))
-			{
-				connectionId = new PerforceConnectionId(host);
-			}
-
-			PerforceConnectionSettings? connectionSettings = settings.Perforce.FirstOrDefault(x => x.Id == connectionId);
-			if (connectionSettings == null)
-			{
-				if (connectionId == PerforceConnectionSettings.Default)
-				{
-					connectionSettings = new PerforceConnectionSettings();
-				}
-				else
-				{
-					throw new InvalidOperationException($"No Perforce connection settings defined for '{connectionId}'.");
-				}
-			}
-
-			IPerforceConnection connection = await PerforceConnection.CreateAsync(connectionSettings.ToPerforceSettings(), _logger);
-			if (connectionSettings.Credentials != null && !String.IsNullOrEmpty(connectionSettings.Credentials.Password) && String.IsNullOrEmpty(connectionSettings.Credentials.Ticket))
-			{
-				await connection.LoginAsync(connectionSettings.Credentials.Password, cancellationToken);
-			}
-			return connection;
-		}
-
-		/// <inheritdoc/>
-		public async Task GetUpdateInfoAsync(IReadOnlyDictionary<Uri, string> files, IReadOnlyDictionary<Uri, string>? prevFiles, ConfigUpdateInfo updateInfo, CancellationToken cancellationToken)
-		{
-			foreach (IGrouping<string, KeyValuePair<Uri, string>> group in files.Where(x => x.Key.Scheme == PerforceConfigSource.Scheme).GroupBy(x => x.Key.Host))
-			{
-				// Figure out the most recent changelist
-				int change = group.Select(x => Int32.Parse(x.Value)).Max();
-				updateInfo.Status.Add($"Perforce changelist ({group.Key}): {change}");
-
-				// Find all the users that have committed changes
-				if (prevFiles != null)
-				{
-					using IPerforceConnection perforce = await ConnectAsync(group.Key, cancellationToken);
-					foreach ((Uri uri, string revision) in group)
-					{
-						string? prevRevision;
-						if (prevFiles.TryGetValue(uri, out prevRevision) && prevRevision != revision)
-						{
-							int prevRevisionNum = Int32.Parse(prevRevision);
-							string fileSpec = $"{uri.LocalPath}@{prevRevisionNum + 1},{revision}";
-							await FindAuthorsAsync(perforce, fileSpec, updateInfo.Authors, cancellationToken);
-						}
-					}
-				}
-			}
-		}
-
-		async Task FindAuthorsAsync(IPerforceConnection perforce, string fileSpec, HashSet<UserId> authors, CancellationToken cancellationToken)
-		{
-			List<ChangesRecord> records = await perforce.GetChangesAsync(ChangesOptions.None, -1, ChangeStatus.Submitted, fileSpec, cancellationToken);
-			foreach (ChangesRecord record in records)
-			{
-				IUser? user = await _userCollection.FindUserByLoginAsync(record.User, cancellationToken);
-				if (user != null)
-				{
-					authors.Add(user.Id);
-				}
-			}
-		}
-	}
 }
diff --git a/Engine/Source/Programs/Horde/HordeServer/Notifications/Sinks/SlackNotificationSink.cs b/Engine/Source/Programs/Horde/HordeServer/Notifications/Sinks/SlackNotificationSink.cs
--- a/Engine/Source/Programs/Horde/HordeServer/Notifications/Sinks/SlackNotificationSink.cs
+++ b/Engine/Source/Programs/Horde/HordeServer/Notifications/Sinks/SlackNotificationSink.cs
@@ -37,6 +37,7 @@ using HordeServer.Issues.External;
 using HordeServer.Jobs;
 using HordeServer.Jobs.Graphs;
 using HordeServer.Logs;
+using HordeServer.Perforce;
 using HordeServer.Server;
 using HordeServer.Streams;
 using HordeServer.Users;
diff --git a/Engine/Source/Programs/Horde/HordeServer/Perforce/PerforceConfigSource.cs b/Engine/Source/Programs/Horde/HordeServer/Perforce/PerforceConfigSource.cs
new file mode 100644
--- /dev/null
+++ b/Engine/Source/Programs/Horde/HordeServer/Perforce/PerforceConfigSource.cs
@@ -0,0 +1,225 @@
+// Copyright Epic Games, Inc. All Rights Reserved.
+
+using System;
+using System.Collections.Generic;
+using System.IO;
+using System.Linq;
+using System.Threading;
+using System.Threading.Tasks;
+using EpicGames.Core;
+using EpicGames.Horde.Users;
+using EpicGames.Perforce;
+using HordeServer.Configuration;
+using HordeServer.Users;
+using Microsoft.Extensions.Caching.Memory;
+using Microsoft.Extensions.Logging;
+using Microsoft.Extensions.Options;
+
+namespace HordeServer.Perforce
+{
+	/// <summary>
+	/// Perforce cluster config file source
+	/// </summary>
+	public sealed class PerforceConfigSource : IConfigSource
+	{
+		class ConfigFileImpl : IConfigFile
+		{
+			public Uri Uri { get; }
+			public int Change { get; }
+			public string Revision { get; }
+			public IUser? Author { get; }
+
+			readonly PerforceConfigSource _owner;
+
+			public ConfigFileImpl(Uri uri, int change, IUser? author, PerforceConfigSource owner)
+			{
+				Uri = uri;
+				Change = change;
+				Revision = $"{change}";
+				Author = author;
+				_owner = owner;
+			}
+
+			public ValueTask<ReadOnlyMemory<byte>> ReadAsync(CancellationToken cancellationToken) => _owner.ReadAsync(Uri, Change, cancellationToken);
+		}
+
+		/// <summary>
+		/// Name of the scheme for this source
+		/// </summary>
+		public const string Scheme = "perforce";
+
+		/// <inheritdoc/>
+		string IConfigSource.Scheme => Scheme;
+
+		/// <inheritdoc/>
+		public TimeSpan UpdateInterval => TimeSpan.FromMinutes(1.0);
+
+		readonly IOptionsMonitor<ServerSettings> _settings;
+		readonly IUserCollection _userCollection;
+		readonly IMemoryCache _cache;
+		readonly ILogger _logger;
+
+		/// <summary>
+		/// Constructor
+		/// </summary>
+		public PerforceConfigSource(IOptionsMonitor<ServerSettings> settings, IUserCollection userCollection, IMemoryCache cache, ILogger<PerforceConfigSource> logger)
+		{
+			_settings = settings;
+			_userCollection = userCollection;
+			_cache = cache;
+			_logger = logger;
+		}
+
+		/// <inheritdoc/>
+		public async Task<IConfigFile[]> GetFilesAsync(Uri[] uris, CancellationToken cancellationToken)
+		{
+			Dictionary<Uri, IConfigFile> results = new Dictionary<Uri, IConfigFile>();
+			foreach (IGrouping<string, Uri> group in uris.GroupBy(x => x.Host))
+			{
+				using (IPerforceConnection perforce = await ConnectAsync(group.Key, cancellationToken))
+				{
+					try
+					{
+						FileSpecList fileSpec = group.Select(x => x.AbsolutePath).Distinct(StringComparer.OrdinalIgnoreCase).ToList();
+
+						List<FStatRecord> records = await perforce.FStatAsync(FStatOptions.ShortenOutput, fileSpec, cancellationToken).ToListAsync(cancellationToken);
+						records.RemoveAll(x => x.HeadAction == FileAction.Delete || x.HeadAction == FileAction.MoveDelete);
+
+						Dictionary<string, FStatRecord> absolutePathToRecord = records.ToDictionary(x => x.DepotFile ?? String.Empty, x => x, StringComparer.OrdinalIgnoreCase);
+						foreach (Uri uri in group)
+						{
+							FStatRecord? record;
+							if (!absolutePathToRecord.TryGetValue(uri.AbsolutePath, out record))
+							{
+								throw new FileNotFoundException($"Unable to read {uri}. No matching files found.");
+							}
+
+							IUser? author = await GetAuthorAsync(perforce, group.Key, record.HeadChange, cancellationToken);
+							results[uri] = new ConfigFileImpl(uri, record.HeadChange, author, this);
+						}
+					}
+					catch (PerforceException ex) when (ex.Error != null && ex.Error.Generic == PerforceGenericCode.Config)
+					{
+						throw new PerforceException($"{ex.Message} [server: {perforce.Settings.ServerAndPort}, user: {perforce.Settings.UserName}]", ex);
+					}
+				}
+			}
+			return uris.ConvertAll(x => results[x]);
+		}
+
+		async ValueTask<IUser?> GetAuthorAsync(IPerforceConnection perforce, string host, int change, CancellationToken cancellationToken)
+		{
+			string cacheKey = $"{nameof(PerforceConfigSource)}:author:{host}@{change}";
+			if (!_cache.TryGetValue(cacheKey, out string? author))
+			{
+				ChangeRecord record = await perforce.GetChangeAsync(GetChangeOptions.None, change, cancellationToken);
+				using (ICacheEntry entry = _cache.CreateEntry(cacheKey))
+				{
+					entry.SetSlidingExpiration(TimeSpan.FromHours(1.0));
+					entry.SetSize(256);
+					entry.SetValue(record.User);
+				}
+			}
+			return (author != null) ? await _userCollection.FindUserByLoginAsync(author, cancellationToken) : null;
+		}
+
+		async ValueTask<ReadOnlyMemory<byte>> ReadAsync(Uri uri, int change, CancellationToken cancellationToken)
+		{
+			string cacheKey = $"{nameof(PerforceConfigSource)}:data:{uri}@{change}";
+			if (_cache.TryGetValue(cacheKey, out ReadOnlyMemory<byte> data))
+			{
+				_logger.LogInformation("Read {Uri}@{Change} from cache ({Key})", uri, change, cacheKey);
+			}
+			else
+			{
+				_logger.LogInformation("Reading {Uri} at CL {Change} from Perforce", uri, change);
+				using (IPerforceConnection perforce = await ConnectAsync(uri.Host, cancellationToken))
+				{
+					PerforceResponse<PrintRecord<byte[]>> response = await perforce.TryPrintDataAsync($"{uri.AbsolutePath}@{change}", cancellationToken);
+					response.EnsureSuccess();
+					data = response.Data.Contents!;
+				}
+				using (ICacheEntry entry = _cache.CreateEntry(cacheKey))
+				{
+					entry.SetSlidingExpiration(TimeSpan.FromHours(1.0));
+					entry.SetSize(data.Length);
+					entry.SetValue(data);
+				}
+			}
+			return data;
+		}
+
+		async Task<IPerforceConnection> ConnectAsync(string host, CancellationToken cancellationToken)
+		{
+			_ = cancellationToken;
+
+			ServerSettings settings = _settings.CurrentValue;
+
+			PerforceConnectionId connectionId = new PerforceConnectionId();
+			if (!String.IsNullOrEmpty(host))
+			{
+				connectionId = new PerforceConnectionId(host);
+			}
+
+			PerforceConnectionSettings? connectionSettings = settings.Perforce.FirstOrDefault(x => x.Id == connectionId);
+			if (connectionSettings == null)
+			{
+				if (connectionId == PerforceConnectionSettings.Default)
+				{
+					connectionSettings = new PerforceConnectionSettings();
+				}
+				else
+				{
+					throw new InvalidOperationException($"No Perforce connection settings defined for '{connectionId}'.");
+				}
+			}
+
+			IPerforceConnection connection = await PerforceConnection.CreateAsync(connectionSettings.ToPerforceSettings(), _logger);
+			if (connectionSettings.Credentials != null && !String.IsNullOrEmpty(connectionSettings.Credentials.Password) && String.IsNullOrEmpty(connectionSettings.Credentials.Ticket))
+			{
+				await connection.LoginAsync(connectionSettings.Credentials.Password, cancellationToken);
+			}
+			return connection;
+		}
+
+		/// <inheritdoc/>
+		public async Task GetUpdateInfoAsync(IReadOnlyDictionary<Uri, string> files, IReadOnlyDictionary<Uri, string>? prevFiles, ConfigUpdateInfo updateInfo, CancellationToken cancellationToken)
+		{
+			foreach (IGrouping<string, KeyValuePair<Uri, string>> group in files.Where(x => x.Key.Scheme == PerforceConfigSource.Scheme).GroupBy(x => x.Key.Host))
+			{
+				// Figure out the most recent changelist
+				int change = group.Select(x => Int32.Parse(x.Value)).Max();
+				updateInfo.Status.Add($"Perforce changelist ({group.Key}): {change}");
+
+				// Find all the users that have committed changes
+				if (prevFiles != null)
+				{
+					using IPerforceConnection perforce = await ConnectAsync(group.Key, cancellationToken);
+					foreach ((Uri uri, string revision) in group)
+					{
+						string? prevRevision;
+						if (prevFiles.TryGetValue(uri, out prevRevision) && prevRevision != revision)
+						{
+							int prevRevisionNum = Int32.Parse(prevRevision);
+							string fileSpec = $"{uri.LocalPath}@{prevRevisionNum + 1},{revision}";
+							await FindAuthorsAsync(perforce, fileSpec, updateInfo.Authors, cancellationToken);
+						}
+					}
+				}
+			}
+		}
+
+		async Task FindAuthorsAsync(IPerforceConnection perforce, string fileSpec, HashSet<UserId> authors, CancellationToken cancellationToken)
+		{
+			List<ChangesRecord> records = await perforce.GetChangesAsync(ChangesOptions.None, -1, ChangeStatus.Submitted, fileSpec, cancellationToken);
+			foreach (ChangesRecord record in records)
+			{
+				IUser? user = await _userCollection.FindUserByLoginAsync(record.User, cancellationToken);
+				if (user != null)
+				{
+					authors.Add(user.Id);
+				}
+			}
+		}
+	}
+}
