Horde: Remove a few more dependencies from ConfigService to particular implementation details.

diff --git a/Engine/Source/Programs/Horde/HordeServer/Configuration/ConfigContext.cs b/Engine/Source/Programs/Horde/HordeServer/Configuration/ConfigContext.cs
--- a/Engine/Source/Programs/Horde/HordeServer/Configuration/ConfigContext.cs
+++ b/Engine/Source/Programs/Horde/HordeServer/Configuration/ConfigContext.cs
@@ -81,7 +81,6 @@ namespace HordeServer.Configuration
 			Sources = sources;
 			ScopeStack.Push("$");
 			Logger = logger;
-			MacroScopes.Add(new Dictionary<string, string> { ["HordeDir"] = ServerApp.AppDir.FullName });
 
 			_getMacroValue = GetMacroValue;
 		}
diff --git a/Engine/Source/Programs/Horde/HordeServer/Configuration/ConfigException.cs b/Engine/Source/Programs/Horde/HordeServer/Configuration/ConfigException.cs
--- a/Engine/Source/Programs/Horde/HordeServer/Configuration/ConfigException.cs
+++ b/Engine/Source/Programs/Horde/HordeServer/Configuration/ConfigException.cs
@@ -23,7 +23,7 @@ namespace HordeServer.Configuration
 		/// <param name="context">Current parse context for the error</param>
 		/// <param name="message">Description of the error</param>
 		/// <param name="innerException">Inner exception details</param>
-		internal ConfigException(ConfigContext context, string message, Exception? innerException = null)
+		public ConfigException(ConfigContext context, string message, Exception? innerException = null)
 			: base(message, innerException)
 		{
 			_context = context;
@@ -32,6 +32,6 @@ namespace HordeServer.Configuration
 		/// <summary>
 		/// Gets the parser context when this exception was thrown. This is not exposed as a public property to avoid serializing the whole thing to Serilog.
 		/// </summary>
-		internal ConfigContext GetContext() => _context;
+		public ConfigContext GetContext() => _context;
 	}
 }
diff --git a/Engine/Source/Programs/Horde/HordeServer/Configuration/ConfigNode.cs b/Engine/Source/Programs/Horde/HordeServer/Configuration/ConfigNode.cs
--- a/Engine/Source/Programs/Horde/HordeServer/Configuration/ConfigNode.cs
+++ b/Engine/Source/Programs/Horde/HordeServer/Configuration/ConfigNode.cs
@@ -103,7 +103,7 @@ namespace HordeServer.Configuration
 				}
 				else
 				{
-					return new Uri($"{PerforceScheme}://{PerforceConnectionSettings.Default}{path}");
+					return new Uri($"{PerforceScheme}://default{path}");
 				}
 			}
 			return new Uri(baseUri, path);
diff --git a/Engine/Source/Programs/Horde/HordeServer/Configuration/ConfigService.cs b/Engine/Source/Programs/Horde/HordeServer/Configuration/ConfigService.cs
--- a/Engine/Source/Programs/Horde/HordeServer/Configuration/ConfigService.cs
+++ b/Engine/Source/Programs/Horde/HordeServer/Configuration/ConfigService.cs
@@ -234,6 +234,13 @@ namespace HordeServer.Configuration
 				=> Task.CompletedTask;
 		}
 
+		ConfigContext CreateConfigContext(IReadOnlyDictionary<string, IConfigSource> sources, ILogger logger)
+		{
+			ConfigContext context = new ConfigContext(_jsonOptions, sources, logger);
+			context.MacroScopes.Add(new Dictionary<string, string> { ["HordeDir"] = ServerApp.AppDir.FullName });
+			return context;
+		}
+
 		/// <summary>
 		/// Validate a new set of config files. Parses and runs PostLoad methods on them.
 		/// </summary>
@@ -251,7 +258,7 @@ namespace HordeServer.Configuration
 				overrideSources.Add(schema, new OverrideConfigSource(source, overrideFiles));
 			}
 
-			ConfigContext context = new ConfigContext(_jsonOptions, overrideSources, NullLogger.Instance);
+			ConfigContext context = CreateConfigContext(overrideSources, NullLogger.Instance);
 			try
 			{
 				Uri globalConfigUri = GetGlobalConfigUri();
@@ -551,7 +558,7 @@ namespace HordeServer.Configuration
 		async Task<ConfigSnapshot> CreateSnapshotAsync(CancellationToken cancellationToken)
 		{
 			// Read the config files
-			ConfigContext context = new ConfigContext(_jsonOptions, _sources, _logger);
+			ConfigContext context = CreateConfigContext(_sources, _logger);
 			try
 			{
 				ConfigSnapshot snapshot = new ConfigSnapshot();
