Horde: Create a library of core Horde server functionality for use by plugins.

diff --git a/Engine/Source/Programs/Horde/HordeServer.Shared/Auditing/IAuditLog.cs b/Engine/Source/Programs/Horde/HordeServer.Shared/Auditing/IAuditLog.cs
new file mode 100644
--- /dev/null
+++ b/Engine/Source/Programs/Horde/HordeServer.Shared/Auditing/IAuditLog.cs
@@ -0,0 +1,150 @@
+// Copyright Epic Games, Inc. All Rights Reserved.
+
+using System;
+using System.Collections.Generic;
+using System.IO.Pipelines;
+using System.Text;
+using System.Threading;
+using System.Threading.Tasks;
+using Microsoft.Extensions.Logging;
+
+namespace HordeServer.Auditing
+{
+	/// <summary>
+	/// Message from an audit log
+	/// </summary>
+	public interface IAuditLogMessage
+	{
+		/// <summary>
+		/// Timestamp for the event
+		/// </summary>
+		public DateTime TimeUtc { get; }
+
+		/// <summary>
+		/// Severity of the message
+		/// </summary>
+		public LogLevel Level { get; }
+
+		/// <summary>
+		/// The message payload. Should be an encoded JSON object, with format/properties fields.
+		/// </summary>
+		public string Data { get; }
+	}
+
+	/// <summary>
+	/// Channel for a particular entity
+	/// </summary>
+	public interface IAuditLogChannel : ILogger
+	{
+		/// <summary>
+		/// Finds messages matching certain criteria
+		/// </summary>
+		/// <param name="minTime"></param>
+		/// <param name="maxTime"></param>
+		/// <param name="index"></param>
+		/// <param name="count"></param>
+		/// <param name="cancellationToken">Cancellation token for the operation</param>
+		/// <returns></returns>
+		IAsyncEnumerable<IAuditLogMessage> FindAsync(DateTime? minTime = null, DateTime? maxTime = null, int? index = null, int? count = null, CancellationToken cancellationToken = default);
+
+		/// <summary>
+		/// Deletes messages between a given time range for a particular object
+		/// </summary>
+		/// <param name="minTime">Minimum time to remove</param>
+		/// <param name="maxTime">Maximum time to remove</param>
+		/// <param name="cancellationToken">Cancellation token for the operation</param>
+		/// <returns>Async task</returns>
+		Task<long> DeleteAsync(DateTime? minTime = null, DateTime? maxTime = null, CancellationToken cancellationToken = default);
+
+		/// <summary>
+		/// Flush all writes to this log
+		/// </summary>
+		/// <param name="cancellationToken">Cancellation token for the operation</param>
+		Task FlushAsync(CancellationToken cancellationToken = default);
+	}
+
+	/// <summary>
+	/// Channel for a particular entity
+	/// </summary>
+	public interface IAuditLogChannel<TSubject> : IAuditLogChannel
+	{
+		/// <summary>
+		/// Identifier for the subject of this channel
+		/// </summary>
+		TSubject Subject { get; }
+	}
+
+	/// <summary>
+	/// Message in an audit log
+	/// </summary>
+	/// <typeparam name="TTargetId">Type of entity that the log is for</typeparam>
+	public interface IAuditLogMessage<TTargetId> : IAuditLogMessage
+	{
+		/// <summary>
+		/// Unique id for the entity
+		/// </summary>
+		public TTargetId Subject { get; }
+	}
+
+	/// <summary>
+	/// Interface for a collection of log messages for a particular document type
+	/// </summary>
+	public interface IAuditLog<TSubject>
+	{
+		/// <summary>
+		/// Get the channel for a particular subject
+		/// </summary>
+		/// <param name="subject"></param>
+		/// <returns></returns>
+		IAuditLogChannel<TSubject> this[TSubject subject] { get; }
+	}
+
+	/// <summary>
+	/// Factory for instantiating audit log instances
+	/// </summary>
+	/// <typeparam name="TSubject"></typeparam>
+	public interface IAuditLogFactory<TSubject>
+	{
+		/// <summary>
+		/// Create a new audit log instance, with the given database name
+		/// </summary>
+		/// <param name="collectionName"></param>
+		/// <param name="subjectProperty"></param>
+		/// <returns></returns>
+		IAuditLog<TSubject> Create(string collectionName, string subjectProperty);
+	}
+
+	/// <summary>
+	/// Extension methods for audit log channels
+	/// </summary>
+	static class AuditLogExtensions
+	{
+		/// <summary>
+		/// Retrieve historical information about a specific agent
+		/// </summary>
+		/// <param name="channel">Channel to query</param>
+		/// <param name="bodyWriter">Writer for Json data</param>
+		/// <param name="minTime">Minimum time for records to return</param>
+		/// <param name="maxTime">Maximum time for records to return</param>
+		/// <param name="index">Offset of the first result</param>
+		/// <param name="count">Number of records to return</param>
+		/// <param name="cancellationToken">Cancellation token for the operation</param>
+		/// <returns>Information about the requested agent</returns>
+		public static async Task FindAsync<T>(this IAuditLogChannel<T> channel, PipeWriter bodyWriter, DateTime? minTime = null, DateTime? maxTime = null, int index = 0, int count = 50, CancellationToken cancellationToken = default)
+		{
+			string prefix = "{\n\t\"entries\":\n\t[";
+			await bodyWriter.WriteAsync(Encoding.UTF8.GetBytes(prefix), cancellationToken);
+
+			string separator = "";
+			await foreach (IAuditLogMessage message in channel.FindAsync(minTime, maxTime, index, count, cancellationToken))
+			{
+				string line = $"{separator}\n\t\t{message.Data}";
+				await bodyWriter.WriteAsync(Encoding.UTF8.GetBytes(line), cancellationToken);
+				separator = ",";
+			}
+
+			string suffix = "\n\t]\n}";
+			await bodyWriter.WriteAsync(Encoding.UTF8.GetBytes(suffix), cancellationToken);
+		}
+	}
+}
diff --git a/Engine/Source/Programs/Horde/HordeServer.Shared/Configuration/ConfigContext.cs b/Engine/Source/Programs/Horde/HordeServer.Shared/Configuration/ConfigContext.cs
new file mode 100644
--- /dev/null
+++ b/Engine/Source/Programs/Horde/HordeServer.Shared/Configuration/ConfigContext.cs
@@ -0,0 +1,224 @@
+// Copyright Epic Games, Inc. All Rights Reserved.
+
+using System;
+using System.Collections.Generic;
+using System.Diagnostics;
+using System.Diagnostics.CodeAnalysis;
+using System.Text.Json;
+using System.Threading;
+using System.Threading.Tasks;
+using EpicGames.Core;
+using Microsoft.Extensions.Logging;
+
+namespace HordeServer.Configuration
+{
+	using JsonObject = System.Text.Json.Nodes.JsonObject;
+
+	/// <summary>
+	/// Context for reading a tree of config files
+	/// </summary>
+	[DebuggerDisplay("{CurrentFile}")]
+	public class ConfigContext
+	{
+		/// <summary>
+		/// Options for serializing config files
+		/// </summary>
+		public JsonSerializerOptions JsonOptions { get; }
+
+		/// <summary>
+		/// Stack of included files
+		/// </summary>
+		public Stack<IConfigFile> IncludeStack { get; } = new Stack<IConfigFile>();
+
+		/// <summary>
+		/// Stack of properties
+		/// </summary>
+		public Stack<string> ScopeStack { get; } = new Stack<string>();
+
+		/// <summary>
+		/// Map of property path to the file declaring a value for it
+		/// </summary>
+		public Dictionary<string, Uri> PropertyPathToFile { get; } = new Dictionary<string, Uri>(StringComparer.OrdinalIgnoreCase);
+
+		/// <summary>
+		/// Sources to read config files from
+		/// </summary>
+		public IReadOnlyDictionary<string, IConfigSource> Sources { get; }
+
+		/// <summary>
+		/// Map of macro name to value 
+		/// </summary>
+		public List<Dictionary<string, string>> MacroScopes { get; } = new List<Dictionary<string, string>>();
+
+		/// <summary>
+		/// Tracks files read as part of the configuration
+		/// </summary>
+		public Dictionary<Uri, IConfigFile> Files { get; } = new Dictionary<Uri, IConfigFile>();
+
+		/// <summary>
+		/// Logger for config messages
+		/// </summary>
+		public ILogger Logger { get; }
+
+		/// <summary>
+		/// Uri of the current file
+		/// </summary>
+		public Uri CurrentFile => (IncludeStack.Count > 0) ? IncludeStack.Peek().Uri : null!;
+
+		/// <summary>
+		/// Current property scope
+		/// </summary>
+		public string CurrentScope => ScopeStack.Peek();
+
+		readonly Func<string, string?> _getMacroValue;
+
+		/// <summary>
+		/// Constructor
+		/// </summary>
+		public ConfigContext(JsonSerializerOptions jsonOptions, IReadOnlyDictionary<string, IConfigSource> sources, ILogger logger)
+		{
+			JsonOptions = jsonOptions;
+			Sources = sources;
+			ScopeStack.Push("$");
+			Logger = logger;
+
+			_getMacroValue = GetMacroValue;
+		}
+
+		/// <summary>
+		/// Marks a property as defined in the current file
+		/// </summary>
+		/// <param name="name"></param>
+		public void AddProperty(string name)
+		{
+			if (!TryAddProperty(name, out Uri? otherFile))
+			{
+				throw new ConfigException(this, $"Property {CurrentScope}.{name} was already defined in {otherFile}.");
+			}
+		}
+
+		/// <summary>
+		/// Marks a property as defined in the current file
+		/// </summary>
+		/// <param name="name">Name of the property within the current scope</param>
+		/// <param name="otherFile">If the property is not added, the file that previously defined it</param>
+		public bool TryAddProperty(string name, [NotNullWhen(false)] out Uri? otherFile)
+		{
+			string propertyPath = $"{CurrentScope}.{name}";
+
+			Uri currentFile = CurrentFile;
+			if (PropertyPathToFile.TryAdd(propertyPath, currentFile))
+			{
+				otherFile = null;
+				return true;
+			}
+			else
+			{
+				otherFile = PropertyPathToFile[propertyPath];
+				return false;
+			}
+		}
+
+		/// <summary>
+		/// Pushes a scope to the property stack
+		/// </summary>
+		/// <param name="name"></param>
+		public void EnterScope(string name)
+		{
+			ScopeStack.Push($"{CurrentScope}.{name}");
+		}
+
+		/// <summary>
+		/// Pops a scope from the property stack
+		/// </summary>
+		public void LeaveScope()
+		{
+			ScopeStack.Pop();
+		}
+
+		/// <summary>
+		/// Expand macros in a string property
+		/// </summary>
+		public string ExpandMacros(string text) => StringUtils.ExpandProperties(text, _getMacroValue);
+
+		/// <summary>
+		/// Gets the value of a named macro
+		/// </summary>
+		string? GetMacroValue(string name)
+		{
+			for (int idx = MacroScopes.Count - 1; idx >= 0; idx--)
+			{
+				if (MacroScopes[idx].TryGetValue(name, out string? value))
+				{
+					return value;
+				}
+			}
+			return null;
+		}
+
+		/// <summary>
+		/// Reads the contents of a file using the appropriate source
+		/// </summary>
+		/// <param name="uri">Uri of the file to read. The scheme indicates the source to read from.</param>
+		/// <param name="cancellationToken">Cancellation token for the operation</param>
+		/// <returns>Information about the file</returns>
+		public async ValueTask<IConfigFile> ReadFileAsync(Uri uri, CancellationToken cancellationToken)
+		{
+			IConfigSource? source = Sources[uri.Scheme];
+			if (source == null)
+			{
+				throw new ConfigException(this, $"Invalid/unknown scheme for config file {uri}");
+			}
+
+			IConfigFile? file;
+			if (!Files.TryGetValue(uri, out file))
+			{
+				file = await source.GetAsync(uri, cancellationToken);
+				Files.Add(uri, file);
+			}
+
+			return file;
+		}
+
+		/// <summary>
+		/// Reads an object from a particular URL
+		/// </summary>
+		/// <typeparam name="T">Type of object to read</typeparam>
+		/// <param name="uri">Location of the file to read</param>
+		/// <param name="cancellationToken">Cancellation token for the operation</param>
+		/// <returns></returns>
+		public async Task<T> ReadAsync<T>(Uri uri, CancellationToken cancellationToken) where T : class, new()
+		{
+			IConfigFile file = await ReadFileAsync(uri, cancellationToken);
+			JsonObject obj = await ParseFileAsync(file, cancellationToken);
+
+			IncludeStack.Push(file);
+
+			ObjectConfigNode type = new ObjectConfigNode(typeof(T));
+			obj = await type.PreprocessAsync(obj, this, cancellationToken);
+
+			IncludeStack.Pop();
+
+			return JsonSerializer.Deserialize<T>(obj, JsonOptions) ?? new T();
+		}
+
+		/// <summary>
+		/// Parses a config file as a json object
+		/// </summary>
+		/// <param name="file">File to parse</param>
+		/// <param name="cancellationToken">Cancellation token for the operation</param>
+		/// <returns>The parsed file</returns>
+		public async ValueTask<JsonObject> ParseFileAsync(IConfigFile file, CancellationToken cancellationToken = default)
+		{
+			ReadOnlyMemory<byte> data = await file.ReadAsync(cancellationToken);
+
+			JsonObject? obj = JsonSerializer.Deserialize<JsonObject>(data.Span, JsonOptions);
+			if (obj == null)
+			{
+				throw new ConfigException(this, $"Config file {file.Uri} contains a null object.");
+			}
+
+			return obj;
+		}
+	}
+}
diff --git a/Engine/Source/Programs/Horde/HordeServer.Shared/Configuration/ConfigException.cs b/Engine/Source/Programs/Horde/HordeServer.Shared/Configuration/ConfigException.cs
new file mode 100644
--- /dev/null
+++ b/Engine/Source/Programs/Horde/HordeServer.Shared/Configuration/ConfigException.cs
@@ -0,0 +1,37 @@
+// Copyright Epic Games, Inc. All Rights Reserved.
+
+using System;
+using System.Collections.Generic;
+
+namespace HordeServer.Configuration
+{
+	/// <summary>
+	/// Exception thrown when reading config files
+	/// </summary>
+	public sealed class ConfigException : Exception
+	{
+		readonly ConfigContext _context;
+
+		/// <summary>
+		/// Stack of properties
+		/// </summary>
+		public IEnumerable<string> ScopeStack => _context.ScopeStack;
+
+		/// <summary>
+		/// Constructor
+		/// </summary>
+		/// <param name="context">Current parse context for the error</param>
+		/// <param name="message">Description of the error</param>
+		/// <param name="innerException">Inner exception details</param>
+		public ConfigException(ConfigContext context, string message, Exception? innerException = null)
+			: base(message, innerException)
+		{
+			_context = context;
+		}
+
+		/// <summary>
+		/// Gets the parser context when this exception was thrown. This is not exposed as a public property to avoid serializing the whole thing to Serilog.
+		/// </summary>
+		public ConfigContext GetContext() => _context;
+	}
+}
diff --git a/Engine/Source/Programs/Horde/HordeServer.Shared/Configuration/ConfigMacro.cs b/Engine/Source/Programs/Horde/HordeServer.Shared/Configuration/ConfigMacro.cs
new file mode 100644
--- /dev/null
+++ b/Engine/Source/Programs/Horde/HordeServer.Shared/Configuration/ConfigMacro.cs
@@ -0,0 +1,22 @@
+// Copyright Epic Games, Inc. All Rights Reserved.
+
+using System;
+
+namespace HordeServer.Configuration
+{
+	/// <summary>
+	/// Declares a config macro
+	/// </summary>
+	public class ConfigMacro
+	{
+		/// <summary>
+		/// Name of the macro property
+		/// </summary>
+		public string Name { get; set; } = String.Empty;
+
+		/// <summary>
+		/// Value for the macro property
+		/// </summary>
+		public string Value { get; set; } = String.Empty;
+	}
+}
diff --git a/Engine/Source/Programs/Horde/HordeServer.Shared/Configuration/ConfigNode.cs b/Engine/Source/Programs/Horde/HordeServer.Shared/Configuration/ConfigNode.cs
new file mode 100644
--- /dev/null
+++ b/Engine/Source/Programs/Horde/HordeServer.Shared/Configuration/ConfigNode.cs
@@ -0,0 +1,628 @@
+// Copyright Epic Games, Inc. All Rights Reserved.
+
+using System;
+using System.Collections.Generic;
+using System.Reflection;
+using System.Text.Json;
+using System.Text.Json.Nodes;
+using System.Text.Json.Serialization;
+using System.Threading;
+using System.Threading.Tasks;
+
+namespace HordeServer.Configuration
+{
+	/// <summary>
+	/// Attribute used to mark <see cref="Uri"/> properties that include other config files
+	/// </summary>
+	[AttributeUsage(AttributeTargets.Property)]
+	public sealed class ConfigIncludeAttribute : Attribute
+	{
+	}
+
+	/// <summary>
+	/// Specifies that a class is the root for including other files
+	/// </summary>
+	[AttributeUsage(AttributeTargets.Class)]
+	public sealed class ConfigIncludeRootAttribute : Attribute
+	{
+	}
+
+	/// <summary>
+	/// Attribute used to mark <see cref="Uri"/> properties that are relative to their containing file
+	/// </summary>
+	[AttributeUsage(AttributeTargets.Property)]
+	public sealed class ConfigRelativePathAttribute : Attribute
+	{
+	}
+
+	/// <summary>
+	/// Attribute used to mark <see cref="Uri"/> properties that are relative to their containing file
+	/// </summary>
+	[AttributeUsage(AttributeTargets.Class)]
+	public sealed class ConfigMacroScopeAttribute : Attribute
+	{
+	}
+
+	/// <summary>
+	/// Node in the preprocessor parse tree
+	/// </summary>
+	public abstract class ConfigNode
+	{
+		/// <summary>
+		/// Parses macro definitions from this object
+		/// </summary>
+		/// <param name="node">Node to parse macros from</param>
+		/// <param name="context">Context for the preprocessor</param>
+		/// <param name="macros">Macros parsed from the boject</param>
+		public abstract void ParseMacros(JsonNode? node, ConfigContext context, Dictionary<string, string> macros);
+
+		/// <summary>
+		/// Preprocesses a json document
+		/// </summary>
+		/// <param name="node">Node to preprocess</param>
+		/// <param name="context">Context for the preprocessor</param>
+		/// <param name="cancellationToken">Cancellation token for the operation</param>
+		/// <returns>Preprocessed node</returns>
+		public Task<JsonNode?> PreprocessAsync(JsonNode? node, ConfigContext context, CancellationToken cancellationToken)
+			=> PreprocessAsync(node, null, context, cancellationToken);
+
+		/// <summary>
+		/// Preprocesses a json document (possibly merging with an existing property)
+		/// </summary>
+		/// <param name="node">Node to preprocess</param>
+		/// <param name="existingNode">Optional existing node to merge with. Can be modified.</param>
+		/// <param name="context">Context for the preprocessor</param>
+		/// <param name="cancellationToken">Cancellation token for the operation</param>
+		/// <returns>The preprocessed node. May be existingNode if changes are merged with it.</returns>
+		public abstract Task<JsonNode?> PreprocessAsync(JsonNode? node, JsonNode? existingNode, ConfigContext context, CancellationToken cancellationToken);
+
+		/// <summary>
+		/// Traverse the property tree starting with 'node' and process any include directives, merging the results into the given target object.
+		/// </summary>
+		/// <param name="node"></param>
+		/// <param name="includes">Receives the included files</param>
+		/// <param name="context"></param>
+		/// <param name="cancellationToken"></param>
+		/// <returns></returns>
+		public abstract Task ParseIncludesAsync(JsonNode? node, List<IConfigFile> includes, ConfigContext context, CancellationToken cancellationToken);
+
+		/// <summary>
+		/// Combine a relative path with a base URI to produce a new URI
+		/// </summary>
+		/// <param name="baseUri">Base uri to rebase relative to</param>
+		/// <param name="path">Relative path</param>
+		/// <returns>Absolute URI</returns>
+		public static Uri CombinePaths(Uri baseUri, string path)
+		{
+			if (path.StartsWith("//", StringComparison.Ordinal))
+			{
+				const string PerforceScheme = "perforce";
+				if (baseUri.Scheme == PerforceScheme)
+				{
+					return new Uri($"{PerforceScheme}://{baseUri.Host}{path}");
+				}
+				else
+				{
+					return new Uri($"{PerforceScheme}://default{path}");
+				}
+			}
+			return new Uri(baseUri, path);
+		}
+
+		/// <summary>
+		/// Helper method to expand all macros within a node without performing any other processing on it
+		/// </summary>
+		protected static JsonNode? ExpandMacros(JsonNode? node, ConfigContext context)
+		{
+			if (node == null)
+			{
+				return node;
+			}
+			else if (node is JsonObject obj)
+			{
+				JsonObject result = new JsonObject();
+				foreach ((string propertyName, JsonNode? propertyNode) in obj)
+				{
+					result[propertyName] = ExpandMacros(propertyNode, context);
+				}
+				return result;
+			}
+			else if (node is JsonArray arr)
+			{
+				JsonArray result = new JsonArray();
+				foreach (JsonNode? elementNode in arr)
+				{
+					result.Add(ExpandMacros(elementNode, context));
+				}
+				return result;
+			}
+			else if (node is JsonValue val && val.GetValueKind() == JsonValueKind.String)
+			{
+				string strValue = ((string?)val)!;
+				string expandedStrValue = context.ExpandMacros(strValue);
+				return JsonValue.Create(expandedStrValue);
+			}
+			else
+			{
+				return node.DeepClone();
+			}
+		}
+	}
+
+	/// <summary>
+	/// Implementation of <see cref="ConfigNode"/> for scalar types
+	/// </summary>
+	public class ScalarConfigNode : ConfigNode
+	{
+		/// <summary>
+		/// Whether this node represents an include directive
+		/// </summary>
+		public bool Include { get; set; }
+
+		/// <summary>
+		/// Specifies that this node is a string which should be made into an absolute path using the path of the current file.
+		/// </summary>
+		public bool RelativeToContainingFile { get; set; }
+
+		/// <summary>
+		/// Constructor
+		/// </summary>
+		public ScalarConfigNode(bool include = false, bool relativePath = false)
+		{
+			Include = include;
+			RelativeToContainingFile = relativePath;
+		}
+
+		/// <inheritdoc/>
+		public override void ParseMacros(JsonNode? jsonNode, ConfigContext context, Dictionary<string, string> macros)
+		{ }
+
+		/// <inheritdoc/>
+		public override Task<JsonNode?> PreprocessAsync(JsonNode? node, JsonNode? existingNode, ConfigContext context, CancellationToken cancellationToken)
+		{
+			JsonNode? result;
+			if (node == null)
+			{
+				result = null;
+			}
+			else if (RelativeToContainingFile && node is JsonValue value && value.GetValueKind() == JsonValueKind.String)
+			{
+				result = JsonValue.Create(CombinePaths(context.CurrentFile, value.ToString()).AbsoluteUri);
+			}
+			else
+			{
+				result = ExpandMacros(node, context);
+			}
+
+			return Task.FromResult(result);
+		}
+
+		/// <inheritdoc/>
+		public override async Task ParseIncludesAsync(JsonNode? node, List<IConfigFile> includes, ConfigContext context, CancellationToken cancellationToken)
+		{
+			if (Include)
+			{
+				string? path = (string?)node;
+
+				Uri uri = CombinePaths(context.CurrentFile, context.ExpandMacros(path!));
+
+				IConfigFile file = await context.ReadFileAsync(uri, cancellationToken);
+				includes.Add(file);
+			}
+		}
+	}
+
+	/// <summary>
+	/// Property containing a binary resource
+	/// </summary>
+	public class ResourceConfigNode : ConfigNode
+	{
+		/// <inheritdoc/>
+		public override Task ParseIncludesAsync(JsonNode? node, List<IConfigFile> includes, ConfigContext context, CancellationToken cancellationToken)
+			=> Task.CompletedTask;
+
+		/// <inheritdoc/>
+		public override void ParseMacros(JsonNode? node, ConfigContext context, Dictionary<string, string> macros)
+		{ }
+
+		/// <inheritdoc/>
+		public override async Task<JsonNode?> PreprocessAsync(JsonNode? node, JsonNode? existingNode, ConfigContext context, CancellationToken cancellationToken)
+		{
+			Uri uri = CombinePaths(context.CurrentFile, JsonSerializer.Deserialize<string>(node, context.JsonOptions) ?? String.Empty);
+			IConfigFile file = await context.ReadFileAsync(uri, cancellationToken);
+
+			ConfigResource resource = new ConfigResource();
+			resource.Path = uri.AbsoluteUri;
+			resource.Data = await file.ReadAsync(cancellationToken);
+
+			return JsonSerializer.SerializeToNode(resource, context.JsonOptions);
+		}
+	}
+
+	/// <summary>
+	/// Array of Json values
+	/// </summary>
+	public class ArrayConfigNode : ConfigNode
+	{
+		/// <summary>
+		/// Type of each element of the array
+		/// </summary>
+		public ConfigNode ElementType { get; }
+
+		/// <summary>
+		/// Constructor
+		/// </summary>
+		public ArrayConfigNode(ConfigNode elementType)
+		{
+			ElementType = elementType;
+		}
+
+		/// <inheritdoc/>
+		public override void ParseMacros(JsonNode? node, ConfigContext context, Dictionary<string, string> macros)
+		{
+			if (node is JsonArray arrayNode)
+			{
+				foreach (JsonNode? elementNode in arrayNode)
+				{
+					ElementType.ParseMacros(elementNode, context, macros);
+				}
+			}
+		}
+
+		/// <inheritdoc/>
+		public override async Task<JsonNode?> PreprocessAsync(JsonNode? node, JsonNode? existingNode, ConfigContext context, CancellationToken cancellationToken)
+		{
+			JsonArray targetArray = ((JsonArray?)existingNode) ?? new JsonArray();
+			foreach (JsonNode? element in (JsonArray)node!)
+			{
+				context.EnterScope($"[{targetArray.Count}]");
+
+				JsonNode? elementValue = await ElementType.PreprocessAsync(element, context, cancellationToken);
+				targetArray.Add(elementValue);
+
+				context.LeaveScope();
+			}
+			return targetArray;
+		}
+
+		/// <inheritdoc/>
+		public override async Task ParseIncludesAsync(JsonNode? node, List<IConfigFile> includes, ConfigContext context, CancellationToken cancellationToken)
+		{
+			if (node is JsonArray arrayNode)
+			{
+				foreach (JsonObject? element in arrayNode)
+				{
+					await ElementType.ParseIncludesAsync(element, includes, context, cancellationToken);
+				}
+			}
+		}
+	}
+
+	/// <summary>
+	/// Arbitary mapping of string values to keys
+	/// </summary>
+	public class DictionaryConfigNode : ConfigNode
+	{
+		/// <summary>
+		/// Type of values in the dictionary
+		/// </summary>
+		public ConfigNode ValueType { get; }
+
+		/// <summary>
+		/// Constructor
+		/// </summary>
+		public DictionaryConfigNode(ConfigNode valueType)
+			=> ValueType = valueType;
+
+		/// <inheritdoc/>
+		public override void ParseMacros(JsonNode? jsonNode, ConfigContext context, Dictionary<string, string> macros)
+		{ }
+
+		/// <inheritdoc/>
+		public override async Task<JsonNode?> PreprocessAsync(JsonNode? node, JsonNode? existingNode, ConfigContext context, CancellationToken cancellationToken)
+		{
+			JsonObject? targetObject = ((JsonObject?)existingNode) ?? new JsonObject();
+			foreach ((string key, JsonNode? element) in (JsonObject)node!)
+			{
+				context.EnterScope($"[{key}]");
+
+				JsonNode? elementValue = await ValueType.PreprocessAsync(element, context, cancellationToken);
+				targetObject[key] = elementValue;
+
+				context.LeaveScope();
+			}
+			return targetObject;
+		}
+
+		/// <inheritdoc/>
+		public override async Task ParseIncludesAsync(JsonNode? node, List<IConfigFile> includes, ConfigContext context, CancellationToken cancellationToken)
+		{
+			if (node is JsonObject obj && ValueType is ObjectConfigNode classElementType)
+			{
+				foreach ((_, JsonNode? value) in obj)
+				{
+					await classElementType.ParseIncludesAsync(value, includes, context, cancellationToken);
+				}
+			}
+		}
+	}
+
+	/// <summary>
+	/// Handles macro objects
+	/// </summary>
+	public class MacroConfigNode : ConfigNode
+	{
+		/// <inheritdoc/>
+		public override Task ParseIncludesAsync(JsonNode? node, List<IConfigFile> includes, ConfigContext context, CancellationToken cancellationToken)
+			=> Task.CompletedTask;
+
+		/// <inheritdoc/>
+		public override void ParseMacros(JsonNode? node, ConfigContext context, Dictionary<string, string> macros)
+		{
+			ConfigMacro? macro = JsonSerializer.Deserialize<ConfigMacro>(node, context.JsonOptions);
+			if (macro != null)
+			{
+				macros.Add(macro.Name, macro.Value);
+			}
+		}
+
+		/// <inheritdoc/>
+		public override Task<JsonNode?> PreprocessAsync(JsonNode? node, JsonNode? existingNode, ConfigContext context, CancellationToken cancellationToken)
+			=> Task.FromResult<JsonNode?>(node?.DeepClone());
+	}
+
+	/// <summary>
+	/// Implementation of <see cref="ConfigNode"/> to handle class types
+	/// </summary>
+	public class ObjectConfigNode : ConfigNode
+	{
+		/// <summary>
+		/// Whether this object should be treated as a root for nested include directies
+		/// </summary>
+		public bool IncludeRoot { get; set; }
+
+		/// <summary>
+		/// Declares a new macro scope
+		/// </summary>
+		public bool MacroScope { get; set; }
+
+		/// <summary>
+		/// Properties within this object
+		/// </summary>
+		public Dictionary<string, ConfigNode> Properties { get; }
+
+		/// <summary>
+		/// Constructor
+		/// </summary>
+		public ObjectConfigNode(bool isIncludeRoot, bool isMacroScope, IEnumerable<KeyValuePair<string, ConfigNode>> properties)
+		{
+			IncludeRoot = isIncludeRoot;
+			MacroScope = isMacroScope;
+			Properties = new Dictionary<string, ConfigNode>(properties, StringComparer.OrdinalIgnoreCase);
+		}
+
+		/// <summary>
+		/// Constructor
+		/// </summary>
+		/// <param name="type">Type to construct from</param>
+		public ObjectConfigNode(Type type)
+			: this(type, new Dictionary<Type, ConfigNode>())
+		{ }
+
+		ObjectConfigNode(Type type, Dictionary<Type, ConfigNode> recursiveTypes)
+		{
+			recursiveTypes.Add(type, this);
+
+			IncludeRoot = type.GetCustomAttribute<ConfigIncludeRootAttribute>() != null;
+			MacroScope = type.GetCustomAttribute<ConfigMacroScopeAttribute>() != null;
+			Properties = new Dictionary<string, ConfigNode>(StringComparer.OrdinalIgnoreCase);
+
+			// Find all the direct include properties
+			PropertyInfo[] propertyInfos = type.GetProperties(BindingFlags.Instance | BindingFlags.Public | BindingFlags.GetProperty);
+			foreach (PropertyInfo propertyInfo in propertyInfos)
+			{
+				if (propertyInfo.GetCustomAttribute<JsonIgnoreAttribute>() == null)
+				{
+					ConfigNode? propertyType = CreateTypeForProperty(propertyInfo, recursiveTypes);
+					if (propertyType != null)
+					{
+						string name = propertyInfo.GetCustomAttribute<JsonPropertyNameAttribute>()?.Name ?? propertyInfo.Name;
+						Properties.Add(name, propertyType);
+					}
+				}
+			}
+		}
+
+		bool IsDefault()
+			=> !IncludeRoot && !MacroScope && Properties.Count == 0;
+
+		static ConfigNode? CreateTypeForProperty(PropertyInfo propertyInfo, Dictionary<Type, ConfigNode> recursiveTypes)
+		{
+			Type propertyType = propertyInfo.PropertyType;
+			if (!propertyType.IsClass || propertyType == typeof(string))
+			{
+				bool include = propertyInfo.GetCustomAttribute<ConfigIncludeAttribute>() != null;
+				bool relativePath = propertyInfo.GetCustomAttribute<ConfigRelativePathAttribute>() != null;
+				if (include || relativePath)
+				{
+					return new ScalarConfigNode(include, relativePath);
+				}
+				else
+				{
+					return null;
+				}
+			}
+			return CreateType(propertyType, recursiveTypes);
+		}
+
+		static ConfigNode? CreateType(Type type, Dictionary<Type, ConfigNode> recursiveTypes)
+		{
+			ConfigNode? value;
+			if (!type.IsClass || type == typeof(string) || type == typeof(JsonNode))
+			{
+				value = null;
+			}
+			else if (type == typeof(ConfigMacro))
+			{
+				value = new MacroConfigNode();
+			}
+			else if (type.IsAssignableTo(typeof(ConfigResource)))
+			{
+				value = new ResourceConfigNode();
+			}
+			else if (type.IsGenericType && type.GetGenericTypeDefinition() == typeof(List<>))
+			{
+				ConfigNode? elementType = CreateType(type.GetGenericArguments()[0], recursiveTypes);
+				value = (elementType == null) ? null : new ArrayConfigNode(elementType);
+			}
+			else if (type.IsGenericType && type.GetGenericTypeDefinition() == typeof(Dictionary<,>))
+			{
+				ConfigNode? elementType = CreateType(type.GetGenericArguments()[1], recursiveTypes);
+				value = (elementType == null) ? null : new DictionaryConfigNode(elementType);
+			}
+			else if (!recursiveTypes.TryGetValue(type, out value))
+			{
+				ObjectConfigNode objValue = new ObjectConfigNode(type);
+				value = objValue.IsDefault() ? null : objValue;
+			}
+			return value;
+		}
+
+		/// <summary>
+		/// Preprocesses an object
+		/// </summary>
+		/// <param name="obj">Object to process</param>
+		/// <param name="context">Preprocessor context</param>
+		/// <param name="cancellationToken">Cancellation token for the operation</param>
+		public async Task<JsonObject> PreprocessAsync(JsonObject obj, ConfigContext context, CancellationToken cancellationToken)
+			=> (JsonObject)(await PreprocessAsync(obj, null, context, cancellationToken))!;
+
+		/// <inheritdoc/>
+		public override void ParseMacros(JsonNode? jsonNode, ConfigContext context, Dictionary<string, string> macros)
+		{
+			if (jsonNode is JsonObject jsonObject)
+			{
+				foreach ((string name, JsonNode? node) in jsonObject)
+				{
+					if (node != null && Properties.TryGetValue(name, out ConfigNode? property))
+					{
+						if (property is not ObjectConfigNode propertyObj || !propertyObj.MacroScope)
+						{
+							property.ParseMacros(node, context, macros);
+						}
+					}
+				}
+			}
+		}
+
+		/// <inheritdoc/>
+		public override async Task<JsonNode?> PreprocessAsync(JsonNode? node, JsonNode? existingNode, ConfigContext context, CancellationToken cancellationToken)
+		{
+			JsonObject? obj = (JsonObject?)node;
+			JsonObject? target = (JsonObject?)existingNode;
+			if (obj == null)
+			{
+				return target;
+			}
+
+			// Before parsing properties into this object, read all the includes recursively
+			if (IncludeRoot)
+			{
+				// Find all the includes
+				List<IConfigFile> includes = new List<IConfigFile>();
+				await ParseIncludesInternalAsync(obj, includes, context, cancellationToken);
+
+				// Find all the files, merge them into target
+				foreach (IConfigFile include in includes)
+				{
+					context.IncludeStack.Push(include);
+
+					JsonObject includedJsonObject = await context.ParseFileAsync(include, cancellationToken);
+					target = (JsonObject?)await PreprocessAsync(includedJsonObject, target, context, cancellationToken);
+
+					context.IncludeStack.Pop();
+				}
+			}
+
+			// Ensure that the target object is valid so we can write properties into it
+			target ??= new JsonObject();
+
+			// Parse all the macros for this scope
+			if (MacroScope)
+			{
+				Dictionary<string, string> macros = new Dictionary<string, string>(StringComparer.OrdinalIgnoreCase);
+				ParseMacros(obj, context, macros);
+				context.MacroScopes.Add(macros);
+			}
+
+			// Parse all the properties into this object
+			foreach ((string name, JsonNode? newNode) in obj)
+			{
+				if (newNode is JsonValue)
+				{
+					context.AddProperty(name);
+				}
+
+				if (Properties.TryGetValue(name, out ConfigNode? property))
+				{
+					context.EnterScope(name);
+					target[name] = await property.PreprocessAsync(newNode, target[name], context, cancellationToken);
+					context.LeaveScope();
+				}
+				else
+				{
+					target[name] = Merge(ExpandMacros(newNode, context), target[name]);
+				}
+			}
+
+			// Parse all the macros for this scope
+			if (MacroScope)
+			{
+				context.MacroScopes.RemoveAt(context.MacroScopes.Count - 1);
+			}
+
+			return target;
+		}
+
+		static JsonNode? Merge(JsonNode? source, JsonNode? target)
+		{
+			if (source is JsonObject sourceObj && target is JsonObject targetObj)
+			{
+				foreach ((string name, JsonNode? node) in sourceObj)
+				{
+					targetObj[name] = Merge(node, targetObj[name]);
+				}
+				return target;
+			}
+			else if (source is JsonArray sourceArr && target is JsonArray targetArr)
+			{
+				foreach (JsonNode? node in sourceArr)
+				{
+					targetArr.Add(node?.DeepClone());
+				}
+				return target;
+			}
+			return source;
+		}
+
+		/// <inheritdoc/>
+		public override async Task ParseIncludesAsync(JsonNode? node, List<IConfigFile> includes, ConfigContext context, CancellationToken cancellationToken)
+		{
+			if (!IncludeRoot && node is JsonObject obj)
+			{
+				await ParseIncludesInternalAsync(obj, includes, context, cancellationToken);
+			}
+		}
+
+		async Task ParseIncludesInternalAsync(JsonObject obj, List<IConfigFile> includes, ConfigContext context, CancellationToken cancellationToken)
+		{
+			foreach ((string name, JsonNode? propertyNode) in obj)
+			{
+				if (Properties.TryGetValue(name, out ConfigNode? property))
+				{
+					await property.ParseIncludesAsync(propertyNode, includes, context, cancellationToken);
+				}
+			}
+		}
+	}
+}
diff --git a/Engine/Source/Programs/Horde/HordeServer.Shared/Configuration/ConfigObject.cs b/Engine/Source/Programs/Horde/HordeServer.Shared/Configuration/ConfigObject.cs
new file mode 100644
--- /dev/null
+++ b/Engine/Source/Programs/Horde/HordeServer.Shared/Configuration/ConfigObject.cs
@@ -0,0 +1,183 @@
+// Copyright Epic Games, Inc. All Rights Reserved.
+
+using System;
+using System.Collections;
+using System.Collections.Generic;
+using System.Linq;
+using System.Reflection;
+using EpicGames.Core;
+
+namespace HordeServer.Configuration
+{
+	/// <summary>
+	/// Possible methods for merging config values
+	/// </summary>
+	public enum ConfigMergeStrategy
+	{
+		/// <summary>
+		/// Default strategy; replace with the base value if the current value is null
+		/// </summary>
+		Default,
+
+		/// <summary>
+		/// Append the contents of this list to the base list
+		/// </summary>
+		Append,
+
+		/// <summary>
+		/// Recursively merge object properties
+		/// </summary>
+		Recursive,
+	}
+
+	/// <summary>
+	/// Attribute used to mark object properties whose child properties should be merged individually
+	/// </summary>
+	[AttributeUsage(AttributeTargets.Property)]
+	public sealed class ConfigMergeStrategyAttribute : Attribute
+	{
+		/// <summary>
+		/// Strategy for merging this property
+		/// </summary>
+		public ConfigMergeStrategy Strategy { get; }
+
+		/// <summary>
+		/// Constructor
+		/// </summary>
+		public ConfigMergeStrategyAttribute(ConfigMergeStrategy strategy) => Strategy = strategy;
+	}
+
+	/// <summary>
+	/// Helper methods for merging default values
+	/// </summary>
+	public static class ConfigObject
+	{
+		/// <summary>
+		/// Merges any unassigned properties from the source object to the target
+		/// </summary>
+		/// <param name="objects">Set of objects to use for merging</param>
+		public static void MergeDefaults<TKey, TValue>(IEnumerable<(TKey Key, TKey? BaseKey, TValue Value)> objects)
+			where TKey : notnull
+			where TValue : class
+		{
+			// Convert the objects to a dictionary
+			List<(TKey Key, TKey BaseKey, TValue Value)> remainingObjects = new List<(TKey Key, TKey BaseKey, TValue Value)>();
+
+			// Find all the objects with no base
+			Dictionary<TKey, TValue> handledValues = new Dictionary<TKey, TValue>();
+			foreach ((TKey key, TKey? baseKey, TValue value) in objects)
+			{
+#pragma warning disable CA1508 // Avoid dead conditional code (false positive due to generics)
+				if (baseKey == null || Equals(baseKey, default(TKey)))
+				{
+					handledValues.Add(key, value);
+				}
+				else
+				{
+					remainingObjects.Add((key, baseKey, value));
+				}
+#pragma warning restore CA1508 // Avoid dead conditional code
+			}
+
+			// Iteratively merge objects with their base
+			for (int lastRemainingObjectCount = 0; remainingObjects.Count != lastRemainingObjectCount;)
+			{
+				lastRemainingObjectCount = remainingObjects.Count;
+				for (int idx = remainingObjects.Count - 1; idx >= 0; idx--)
+				{
+					(TKey key, TKey baseKey, TValue value) = remainingObjects[idx];
+					if (handledValues.TryGetValue(baseKey, out TValue? baseValue))
+					{
+						MergeDefaults(value, baseValue);
+						handledValues.Add(key, value);
+						remainingObjects.RemoveAt(idx);
+					}
+				}
+			}
+
+			// Check we were able to merge everything
+			if (remainingObjects.Count > 0)
+			{
+				HashSet<TKey> validKeys = new HashSet<TKey>(objects.Select(x => x.Key));
+				foreach ((TKey key, TKey baseKey, TValue value) in remainingObjects)
+				{
+					if (!validKeys.Contains(baseKey!))
+					{
+						throw new Exception($"{key} has invalid/missing base {baseKey}");
+					}
+				}
+
+				List<TKey> circularKeys = objects.Where(x => !handledValues.ContainsKey(x.Key)).Select(x => x.Key).ToList();
+				if (circularKeys.Count == 1)
+				{
+					throw new Exception($"{circularKeys[0]} has a circular dependency on itself");
+				}
+				else
+				{
+					throw new Exception($"{StringUtils.FormatList(circularKeys.Select(x => x.ToString() ?? "(unknown)"))} have circular dependencies.");
+				}
+			}
+		}
+
+		/// <summary>
+		/// Merges any unassigned properties from the source object to the target
+		/// </summary>
+		/// <param name="target">Target object to merge into</param>
+		/// <param name="source">Source object to merge from</param>
+		public static void MergeDefaults<T>(T target, T source) where T : class
+		{
+			MergeDefaults(typeof(T), target, source);
+		}
+
+		/// <summary>
+		/// Merges any unassigned properties from the source object to the target
+		/// </summary>
+		/// <param name="type">Type of the referenced object</param>
+		/// <param name="target">Target object to merge into</param>
+		/// <param name="source">Source object to merge from</param>
+		public static void MergeDefaults(Type type, object target, object source)
+		{
+			foreach (PropertyInfo propertyInfo in type.GetProperties(BindingFlags.Instance | BindingFlags.Public))
+			{
+				object? targetValue = propertyInfo.GetValue(target);
+				object? sourceValue = propertyInfo.GetValue(source);
+
+				ConfigMergeStrategy strategy = propertyInfo.GetCustomAttribute<ConfigMergeStrategyAttribute>()?.Strategy ?? ConfigMergeStrategy.Default;
+				switch (strategy)
+				{
+					case ConfigMergeStrategy.Default:
+						if (targetValue == null)
+						{
+							propertyInfo.SetValue(target, sourceValue);
+						}
+						break;
+					case ConfigMergeStrategy.Append:
+						if (targetValue == null)
+						{
+							propertyInfo.SetValue(target, sourceValue);
+						}
+						else if (sourceValue != null)
+						{
+							IList sourceList = (IList)sourceValue;
+							IList targetList = (IList)targetValue;
+							for (int idx = 0; idx < sourceList.Count; idx++)
+							{
+								targetList.Insert(idx, sourceList[idx]);
+							}
+						}
+						break;
+					case ConfigMergeStrategy.Recursive:
+						if (targetValue == null)
+						{
+							propertyInfo.SetValue(target, sourceValue);
+						}
+						else if (sourceValue != null)
+						{
+							MergeDefaults(propertyInfo.PropertyType, targetValue, sourceValue);
+						}
+						break;
+				}
+			}
+		}
+	}
+}
diff --git a/Engine/Source/Programs/Horde/HordeServer.Shared/Configuration/ConfigResource.cs b/Engine/Source/Programs/Horde/HordeServer.Shared/Configuration/ConfigResource.cs
new file mode 100644
--- /dev/null
+++ b/Engine/Source/Programs/Horde/HordeServer.Shared/Configuration/ConfigResource.cs
@@ -0,0 +1,84 @@
+// Copyright Epic Games, Inc. All Rights Reserved.
+
+using System;
+using System.Text.Json;
+using System.Text.Json.Serialization;
+using EpicGames.Core;
+
+namespace HordeServer.Configuration
+{
+	/// <summary>
+	/// Special resource type in configuration files which stores data from an external source
+	/// </summary>
+	[JsonSchemaString]
+	[JsonConverter(typeof(ConfigResourceSerializer))]
+	public sealed class ConfigResource
+	{
+		/// <summary>
+		/// Path to the resource
+		/// </summary>
+		[ConfigRelativePath]
+		public string? Path { get; set; }
+
+		/// <summary>
+		/// Data for the resource. This will be stored externally to the config data.
+		/// </summary>
+		public ReadOnlyMemory<byte> Data { get; set; } = ReadOnlyMemory<byte>.Empty;
+	}
+
+	/// <summary>
+	/// Serializer for <see cref="ConfigResource"/> objects
+	/// </summary>
+	class ConfigResourceSerializer : JsonConverter<ConfigResource>
+	{
+		/// <inheritdoc/>
+		public override ConfigResource Read(ref Utf8JsonReader reader, Type typeToConvert, JsonSerializerOptions options)
+		{
+			ConfigResource resource = new ConfigResource();
+			if (reader.TokenType == JsonTokenType.String)
+			{
+				resource.Path = reader.GetString();
+			}
+			else if (reader.TokenType == JsonTokenType.StartObject)
+			{
+				while (reader.Read() && reader.TokenType != JsonTokenType.EndObject)
+				{
+					if (reader.TokenType == JsonTokenType.PropertyName)
+					{
+						if (reader.ValueTextEquals(nameof(ConfigResource.Path)))
+						{
+							reader.Read();
+							resource.Path = reader.GetString();
+						}
+						else if (reader.ValueTextEquals(nameof(ConfigResource.Data)))
+						{
+							reader.Read();
+							resource.Data = reader.GetBytesFromBase64();
+						}
+					}
+				}
+			}
+			else
+			{
+				throw new InvalidOperationException();
+			}
+			return resource;
+		}
+
+		/// <inheritdoc/>
+		public override void Write(Utf8JsonWriter writer, ConfigResource resource, JsonSerializerOptions options)
+		{
+			if (resource.Data.Length == 0)
+			{
+				writer.WriteStringValue(resource.Path);
+			}
+			else
+			{
+				writer.WriteStartObject();
+				writer.WriteString(nameof(ConfigResource.Path), resource.Path);
+				writer.WriteBase64String(nameof(ConfigResource.Data), resource.Data.Span);
+				writer.WriteEndObject();
+			}
+		}
+	}
+}
diff --git a/Engine/Source/Programs/Horde/HordeServer.Shared/Configuration/IConfigFile.cs b/Engine/Source/Programs/Horde/HordeServer.Shared/Configuration/IConfigFile.cs
new file mode 100644
--- /dev/null
+++ b/Engine/Source/Programs/Horde/HordeServer.Shared/Configuration/IConfigFile.cs
@@ -0,0 +1,63 @@
+// Copyright Epic Games, Inc. All Rights Reserved.
+
+using System;
+using System.Threading;
+using System.Threading.Tasks;
+using HordeServer.Users;
+
+namespace HordeServer.Configuration
+{
+	/// <summary>
+	/// Accessor for a specific revision of a config file. Provides metadata about the current revision, and allows reading its data.
+	/// </summary>
+	public interface IConfigFile
+	{
+		/// <summary>
+		/// URI of the config file
+		/// </summary>
+		Uri Uri { get; }
+
+		/// <summary>
+		/// String used to identify a specific version of the config data. Used for ordinal comparisons, otherwise opaque.
+		/// </summary>
+		string Revision { get; }
+
+		/// <summary>
+		/// Author of this revision, if known
+		/// </summary>
+		IUser? Author { get; }
+
+		/// <summary>
+		/// Reads data for the config file
+		/// </summary>
+		/// <param name="cancellationToken">Cancellation token for the operation</param>
+		/// <returns>UTF-8 encoded data for the config file</returns>
+		ValueTask<ReadOnlyMemory<byte>> ReadAsync(CancellationToken cancellationToken);
+	}
+
+	/// <summary>
+	/// Extension methods for config files
+	/// </summary>
+	public static class ConfigFileExtensions
+	{
+		/// <summary>
+		/// Formats a the URI for a config file in a format that is more readable for users
+		/// </summary>
+		/// <param name="file">Config file to get a path for</param>
+		public static string GetUserFormattedPath(this IConfigFile file)
+		{
+			const string DefaultPerforcePrefix = "perforce://default//";
+
+			string path = file.Uri.ToString();
+			if (path.StartsWith(DefaultPerforcePrefix, StringComparison.OrdinalIgnoreCase))
+			{
+				path = path.Substring(DefaultPerforcePrefix.Length - 2);
+			}
+			if (!String.IsNullOrEmpty(file.Revision))
+			{
+				path = $"{path}@{file.Revision}";
+			}
+			return path;
+		}
+	}
+}
diff --git a/Engine/Source/Programs/Horde/HordeServer.Shared/Configuration/IConfigSource.cs b/Engine/Source/Programs/Horde/HordeServer.Shared/Configuration/IConfigSource.cs
new file mode 100644
--- /dev/null
+++ b/Engine/Source/Programs/Horde/HordeServer.Shared/Configuration/IConfigSource.cs
@@ -0,0 +1,240 @@
+// Copyright Epic Games, Inc. All Rights Reserved.
+
+using System;
+using System.Collections.Concurrent;
+using System.Collections.Generic;
+using System.IO;
+using System.Threading;
+using System.Threading.Tasks;
+using EpicGames.Core;
+using EpicGames.Horde.Users;
+using HordeServer.Users;
+
+namespace HordeServer.Configuration
+{
+	/// <summary>
+	/// Information about the updated config
+	/// </summary>
+	public record class ConfigUpdateInfo(List<string> Status, HashSet<UserId> Authors, Exception? Exception);
+
+	/// <summary>
+	/// Source for reading config files
+	/// </summary>
+	public interface IConfigSource
+	{
+		/// <summary>
+		/// URI scheme for this config source
+		/// </summary>
+		string Scheme { get; }
+
+		/// <summary>
+		/// Update interval for this source
+		/// </summary>
+		TimeSpan UpdateInterval { get; }
+
+		/// <summary>
+		/// Reads a config file from this source
+		/// </summary>
+		/// <param name="uris">Locations of the config files to query</param>
+		/// <param name="cancellationToken">Cancellation token for the operation</param>
+		/// <returns>Config file data</returns>
+		Task<IConfigFile[]> GetFilesAsync(Uri[] uris, CancellationToken cancellationToken);
+
+		/// <summary>
+		/// Gets summary infomration for sending notifications
+		/// </summary>
+		/// <param name="files">New files for the configuration</param>
+		/// <param name="prevFiles">Previous set of files for the configuration</param>
+		/// <param name="updateInfo">Information about the update</param>
+		/// <param name="cancellationToken">Cancellation token for the operation</param>
+		Task GetUpdateInfoAsync(IReadOnlyDictionary<Uri, string> files, IReadOnlyDictionary<Uri, string>? prevFiles, ConfigUpdateInfo updateInfo, CancellationToken cancellationToken);
+	}
+
+	/// <summary>
+	/// Extension methods for <see cref="IConfigSource"/>
+	/// </summary>
+	public static class ConfigSource
+	{
+		/// <summary>
+		/// Gets a single config file from a source
+		/// </summary>
+		/// <param name="source">Source to query</param>
+		/// <param name="uri">Location of the config file to query</param>
+		/// <param name="cancellationToken">Cancellation token for the operation</param>
+		/// <returns>Config file data</returns>
+		public static async Task<IConfigFile> GetAsync(this IConfigSource source, Uri uri, CancellationToken cancellationToken)
+		{
+			IConfigFile[] result = await source.GetFilesAsync(new[] { uri }, cancellationToken);
+			return result[0];
+		}
+	}
+
+	/// <summary>
+	/// In-memory config file source
+	/// </summary>
+	public sealed class InMemoryConfigSource : IConfigSource
+	{
+		class ConfigFileRevisionImpl : IConfigFile
+		{
+			public Uri Uri { get; }
+			public string Revision { get; }
+			public ReadOnlyMemory<byte> Data { get; }
+			public IUser? Author => null;
+
+			public ConfigFileRevisionImpl(Uri uri, string version, ReadOnlyMemory<byte> data)
+			{
+				Uri = uri;
+				Revision = version;
+				Data = data;
+			}
+
+			public ValueTask<ReadOnlyMemory<byte>> ReadAsync(CancellationToken cancellationToken) => new ValueTask<ReadOnlyMemory<byte>>(Data);
+		}
+
+		readonly Dictionary<Uri, ConfigFileRevisionImpl> _files = new Dictionary<Uri, ConfigFileRevisionImpl>();
+
+		/// <summary>
+		/// Name of the scheme for this source
+		/// </summary>
+		public const string Scheme = "memory";
+
+		/// <inheritdoc/>
+		string IConfigSource.Scheme => Scheme;
+
+		/// <inheritdoc/>
+		public TimeSpan UpdateInterval => TimeSpan.FromSeconds(1.0);
+
+		/// <summary>
+		/// Manually adds a new config file
+		/// </summary>
+		/// <param name="path">Path to the config file</param>
+		/// <param name="data">Config file data</param>
+		public void Add(Uri path, ReadOnlyMemory<byte> data)
+		{
+			_files.Add(path, new ConfigFileRevisionImpl(path, "v1", data));
+		}
+
+		/// <inheritdoc/>
+		public Task<IConfigFile[]> GetFilesAsync(Uri[] uris, CancellationToken cancellationToken)
+		{
+			IConfigFile[] result = new IConfigFile[uris.Length];
+			for (int idx = 0; idx < uris.Length; idx++)
+			{
+				ConfigFileRevisionImpl? configFile;
+				if (!_files.TryGetValue(uris[idx], out configFile))
+				{
+					throw new FileNotFoundException($"Config file {uris[idx]} not found.");
+				}
+				result[idx] = configFile;
+			}
+			return Task.FromResult(result);
+		}
+
+		/// <inheritdoc/>
+		public Task GetUpdateInfoAsync(IReadOnlyDictionary<Uri, string> files, IReadOnlyDictionary<Uri, string>? prevFiles, ConfigUpdateInfo updateInfo, CancellationToken cancellationToken)
+			=> Task.CompletedTask;
+	}
+
+	/// <summary>
+	/// Config file source which reads from the filesystem
+	/// </summary>
+	public sealed class FileConfigSource : IConfigSource
+	{
+		class ConfigFileImpl : IConfigFile
+		{
+			public Uri Uri { get; }
+			public string Revision { get; }
+			public DateTime LastWriteTimeUtc { get; }
+			public ReadOnlyMemory<byte> Data { get; }
+			public IUser? Author => null;
+
+			public ConfigFileImpl(Uri uri, DateTime lastWriteTimeUtc, ReadOnlyMemory<byte> data)
+			{
+				Uri = uri;
+				Revision = $"timestamp={lastWriteTimeUtc.Ticks}";
+				LastWriteTimeUtc = lastWriteTimeUtc;
+				Data = data;
+			}
+
+			public ValueTask<ReadOnlyMemory<byte>> ReadAsync(CancellationToken cancellationToken) => new ValueTask<ReadOnlyMemory<byte>>(Data);
+		}
+
+		/// <summary>
+		/// Name of the scheme for this source
+		/// </summary>
+		public const string Scheme = "file";
+
+		/// <inheritdoc/>
+		string IConfigSource.Scheme => Scheme;
+
+		/// <inheritdoc/>
+		public TimeSpan UpdateInterval => TimeSpan.FromSeconds(5.0);
+
+		readonly DirectoryReference _baseDir;
+		readonly ConcurrentDictionary<FileReference, ConfigFileImpl> _files = new ConcurrentDictionary<FileReference, ConfigFileImpl>();
+
+		/// <summary>
+		/// Constructor
+		/// </summary>
+		public FileConfigSource()
+			: this(DirectoryReference.GetCurrentDirectory())
+		{
+		}
+
+		/// <summary>
+		/// Constructor
+		/// </summary>
+		/// <param name="baseDir">Base directory for resolving relative paths</param>
+		public FileConfigSource(DirectoryReference baseDir)
+		{
+			_baseDir = baseDir;
+		}
+
+		/// <inheritdoc/>
+		public async Task<IConfigFile[]> GetFilesAsync(Uri[] uris, CancellationToken cancellationToken)
+		{
+			IConfigFile[] files = new IConfigFile[uris.Length];
+			for (int idx = 0; idx < uris.Length; idx++)
+			{
+				Uri uri = uris[idx];
+				FileReference localPath = FileReference.Combine(_baseDir, uri.LocalPath);
+
+				ConfigFileImpl? file;
+				for (; ; )
+				{
+					if (_files.TryGetValue(localPath, out file))
+					{
+						if (FileReference.GetLastWriteTimeUtc(localPath) == file.LastWriteTimeUtc)
+						{
+							break;
+						}
+						else
+						{
+							_files.TryRemove(new KeyValuePair<FileReference, ConfigFileImpl>(localPath, file));
+						}
+					}
+
+					using (FileStream stream = FileReference.Open(localPath, FileMode.Open, FileAccess.Read, FileShare.Read))
+					{
+						using MemoryStream memoryStream = new MemoryStream();
+						await stream.CopyToAsync(memoryStream, cancellationToken);
+						DateTime lastWriteTime = FileReference.GetLastWriteTimeUtc(localPath);
+						file = new ConfigFileImpl(uri, lastWriteTime, memoryStream.ToArray());
+					}
+
+					if (_files.TryAdd(localPath, file))
+					{
+						break;
+					}
+				}
+
+				files[idx] = file;
+			}
+			return files;
+		}
+
+		/// <inheritdoc/>
+		public Task GetUpdateInfoAsync(IReadOnlyDictionary<Uri, string> files, IReadOnlyDictionary<Uri, string>? prevFiles, ConfigUpdateInfo updateInfo, CancellationToken cancellationToken)
+			=> Task.CompletedTask;
+	}
+}
diff --git a/Engine/Source/Programs/Horde/HordeServer.Shared/HordeServer.Shared.csproj b/Engine/Source/Programs/Horde/HordeServer.Shared/HordeServer.Shared.csproj
new file mode 100644
--- /dev/null
+++ b/Engine/Source/Programs/Horde/HordeServer.Shared/HordeServer.Shared.csproj
@@ -0,0 +1,25 @@
+﻿<Project Sdk="Microsoft.NET.Sdk">
+
+  <PropertyGroup>
+    <TargetFramework>net8.0</TargetFramework>
+    <ImplicitUsings>enable</ImplicitUsings>
+    <Nullable>enable</Nullable>
+    <RootNamespace>HordeServer</RootNamespace>
+  </PropertyGroup>
+
+  <ItemGroup>
+    <FrameworkReference Include="Microsoft.AspNetCore.App" />
+  </ItemGroup>
+
+  <ItemGroup>
+    <ProjectReference Include="..\..\Shared\EpicGames.Core\EpicGames.Core.csproj" />
+    <ProjectReference Include="..\..\Shared\EpicGames.Horde\EpicGames.Horde.csproj" />
+    <ProjectReference Include="..\..\Shared\EpicGames.Redis\EpicGames.Redis.csproj" />
+  </ItemGroup>
+
+  <ItemGroup>
+    <PackageReference Include="MongoDB.Driver" Version="2.22.0" />
+    <PackageReference Include="OpenTelemetry" Version="1.8.0" />
+  </ItemGroup>
+
+</Project>
diff --git a/Engine/Source/Programs/Horde/HordeServer.Shared/Server/IMongoService.cs b/Engine/Source/Programs/Horde/HordeServer.Shared/Server/IMongoService.cs
new file mode 100644
--- /dev/null
+++ b/Engine/Source/Programs/Horde/HordeServer.Shared/Server/IMongoService.cs
@@ -0,0 +1,101 @@
+// Copyright Epic Games, Inc. All Rights Reserved.
+
+using System;
+using System.Collections.Generic;
+using System.Threading;
+using System.Threading.Tasks;
+using HordeServer.Utilities;
+using Microsoft.Extensions.Diagnostics.HealthChecks;
+using MongoDB.Driver;
+
+namespace HordeServer.Server
+{
+	/// <summary>
+	/// Interface for the Horde MongoDB instance
+	/// </summary>
+	public interface IMongoService : IHealthCheck
+	{
+		/// <summary>
+		/// The database instance
+		/// </summary>
+		IMongoDatabase Database { get; }
+
+		/// <summary>
+		/// Access the database in a read-only mode (don't create indices or modify content)
+		/// </summary>
+		public bool ReadOnlyMode { get; }
+
+		/// <summary>
+		/// Get the MongoDB client
+		/// </summary>
+		/// <returns>A MongoDB client instance</returns>
+		MongoClient GetClient();
+
+		/// <summary>
+		/// Get a MongoDB collection from database
+		/// </summary>
+		/// <param name="name">Name of collection</param>
+		/// <typeparam name="T">A MongoDB document</typeparam>
+		/// <returns></returns>
+		IMongoCollection<T> GetCollection<T>(string name);
+
+		/// <summary>
+		/// Get a MongoDB collection from database with a single index
+		/// </summary>
+		/// <param name="name">Name of collection</param>
+		/// <param name="keysFunc">Method to configure keys for the collection</param>
+		/// <param name="unique">Whether a unique index is required</param>
+		/// <typeparam name="T">A MongoDB document</typeparam>
+		/// <returns></returns>
+		IMongoCollection<T> GetCollection<T>(string name, Func<IndexKeysDefinitionBuilder<T>, IndexKeysDefinition<T>> keysFunc, bool unique = false);
+
+		/// <summary>
+		/// Get a MongoDB collection from database
+		/// </summary>
+		/// <param name="name">Name of collection</param>
+		/// <param name="indexes">Indexes for the collection</param>
+		/// <typeparam name="T">A MongoDB document</typeparam>
+		/// <returns></returns>
+		IMongoCollection<T> GetCollection<T>(string name, IEnumerable<MongoIndex<T>> indexes);
+
+		/// <summary>
+		/// Gets a singleton document by id
+		/// </summary>
+		/// <returns>The document</returns>
+		Task<T> GetSingletonAsync<T>(CancellationToken cancellationToken) where T : SingletonBase, new();
+
+		/// <summary>
+		/// Gets a singleton document by id
+		/// </summary>
+		/// <param name="constructor">Method to use to construct a new object</param>
+		/// <param name="cancellationToken">Cancellation token for the operation</param>
+		/// <returns>The document</returns>
+		Task<T> GetSingletonAsync<T>(Func<T> constructor, CancellationToken cancellationToken) where T : SingletonBase, new();
+
+		/// <summary>
+		/// Updates a singleton
+		/// </summary>
+		/// <typeparam name="T"></typeparam>
+		/// <param name="updater"></param>
+		/// <param name="cancellationToken">Cancellation token for the operation</param>
+		/// <returns></returns>
+		Task UpdateSingletonAsync<T>(Action<T> updater, CancellationToken cancellationToken) where T : SingletonBase, new();
+
+		/// <summary>
+		/// Updates a singleton
+		/// </summary>
+		/// <typeparam name="T"></typeparam>
+		/// <param name="updater"></param>
+		/// <param name="cancellationToken">Cancellation token for the operation</param>
+		/// <returns></returns>
+		Task UpdateSingletonAsync<T>(Func<T, bool> updater, CancellationToken cancellationToken) where T : SingletonBase, new();
+
+		/// <summary>
+		/// Attempts to update a singleton object
+		/// </summary>
+		/// <param name="singletonObject">The singleton object</param>
+		/// <param name="cancellationToken">Cancellation token for the operation</param>
+		/// <returns>True if the singleton document was updated</returns>
+		Task<bool> TryUpdateSingletonAsync<T>(T singletonObject, CancellationToken cancellationToken) where T : SingletonBase;
+	}
+}
diff --git a/Engine/Source/Programs/Horde/HordeServer.Shared/Server/IRedisService.cs b/Engine/Source/Programs/Horde/HordeServer.Shared/Server/IRedisService.cs
new file mode 100644
--- /dev/null
+++ b/Engine/Source/Programs/Horde/HordeServer.Shared/Server/IRedisService.cs
@@ -0,0 +1,104 @@
+// Copyright Epic Games, Inc. All Rights Reserved.
+
+using System;
+using System.Threading.Tasks;
+using EpicGames.Redis;
+using Microsoft.Extensions.Diagnostics.HealthChecks;
+using StackExchange.Redis;
+
+namespace HordeServer.Server
+{
+	/// <summary>
+	/// Provides access to a Redis database
+	/// </summary>
+	public interface IRedisService : IHealthCheck
+	{
+		/// <summary>
+		/// Flag for whether the connection is read-only
+		/// </summary>
+		bool ReadOnlyMode { get; }
+
+		/// <summary>
+		/// Connection pool
+		/// </summary>
+		public RedisConnectionPool ConnectionPool { get; }
+	}
+
+	/// <summary>
+	/// Extension methods for <see cref="IRedisService"/>
+	/// </summary>
+	public static class RedisServiceExtensions
+	{
+		/// <summary>
+		/// Get the least-loaded Redis connection from the pool
+		/// Don't store the returned object and try to resolve this as late as possible to ensure load is balanced.
+		/// </summary>
+		/// <returns>A Redis connection multiplexer</returns>
+		public static IConnectionMultiplexer GetConnection(this IRedisService redisService)
+			=> redisService.ConnectionPool.GetConnection();
+
+		/// <summary>
+		/// Get the least-loaded Redis database from the connection pool
+		/// Don't store the returned object and try to resolve this as late as possible to ensure load is balanced.
+		/// </summary>
+		/// <returns>A Redis database</returns>
+		public static IDatabase GetDatabase(this IRedisService redisService)
+			=> redisService.ConnectionPool.GetDatabase();
+
+		/// <summary>
+		/// Publish a message to a channel
+		/// </summary>
+		/// <param name="redisService">The redis service instance</param>
+		/// <param name="channel">Channel to post to</param>
+		/// <param name="message">Message to post to the channel</param>
+		/// <param name="flags">Flags for the request</param>
+		public static Task PublishAsync(this IRedisService redisService, RedisChannel channel, RedisValue message, CommandFlags flags = CommandFlags.None)
+			=> redisService.GetDatabase().PublishAsync(channel, message, flags);
+
+		/// <summary>
+		/// Publish a message to a channel
+		/// </summary>
+		/// <typeparam name="T">Type of elements sent over the channel</typeparam>
+		/// <param name="redisService">The redis service instance</param>
+		/// <param name="channel">Channel to post to</param>
+		/// <param name="message">Message to post to the channel</param>
+		/// <param name="flags">Flags for the request</param>
+		public static Task PublishAsync<T>(this IRedisService redisService, RedisChannel<T> channel, T message, CommandFlags flags = CommandFlags.None)
+			=> redisService.GetDatabase().PublishAsync(channel, message, flags);
+
+		/// <inheritdoc cref="SubscribeAsync{T}(IRedisService, RedisChannel{T}, Action{RedisChannel{T}, T})"/>
+		public static Task<RedisSubscription> SubscribeAsync(this IRedisService redisService, RedisChannel channel, Action<RedisValue> callback)
+			=> SubscribeAsync(redisService, channel, (ch, x) => callback(x));
+
+		/// <inheritdoc cref="SubscribeAsync{T}(IRedisService, RedisChannel{T}, Action{RedisChannel{T}, T})"/>
+		public static Task<RedisSubscription> SubscribeAsync<T>(this IRedisService redisService, RedisChannel<T> channel, Action<T> callback)
+			=> SubscribeAsync(redisService, channel, (ch, x) => callback(x));
+
+		/// <summary>
+		/// Subscribe to notifications on a channel
+		/// </summary>
+		/// <param name="redisService">The redis service instance</param>
+		/// <param name="channel">Channel to monitor</param>
+		/// <param name="callback">Callback for new events</param>
+		/// <returns>Subscription object</returns>
+		public static async Task<RedisSubscription> SubscribeAsync(this IRedisService redisService, RedisChannel channel, Action<RedisChannel, RedisValue> callback)
+		{
+			IConnectionMultiplexer connection = redisService.GetConnection();
+			return await connection.SubscribeAsync(channel, callback);
+		}
+
+		/// <summary>
+		/// Subscribe to notifications on a channel
+		/// </summary>
+		/// <typeparam name="T">Type of elements sent over the channel</typeparam>
+		/// <param name="redisService">The redis service instance</param>
+		/// <param name="channel">Channel to monitor</param>
+		/// <param name="callback">Callback for new events</param>
+		/// <returns>Subscription object</returns>
+		public static async Task<RedisSubscription> SubscribeAsync<T>(this IRedisService redisService, RedisChannel<T> channel, Action<RedisChannel<T>, T> callback)
+		{
+			IConnectionMultiplexer connection = redisService.GetConnection();
+			return await connection.SubscribeAsync(channel, callback);
+		}
+	}
+}
diff --git a/Engine/Source/Programs/Horde/HordeServer.Shared/Server/MongoIndex.cs b/Engine/Source/Programs/Horde/HordeServer.Shared/Server/MongoIndex.cs
new file mode 100644
--- /dev/null
+++ b/Engine/Source/Programs/Horde/HordeServer.Shared/Server/MongoIndex.cs
@@ -0,0 +1,164 @@
+// Copyright Epic Games, Inc. All Rights Reserved.
+
+using System;
+using System.Collections.Generic;
+using System.Diagnostics;
+using System.Text;
+using MongoDB.Bson;
+using MongoDB.Bson.Serialization;
+using MongoDB.Bson.Serialization.Attributes;
+using MongoDB.Driver;
+
+namespace HordeServer.Server
+{
+	/// <summary>
+	/// Constrained set of parameters for building a mongo index
+	/// </summary>
+	[DebuggerDisplay("{Name}")]
+	public class MongoIndex : IEquatable<MongoIndex>
+	{
+		/// <summary>
+		/// Name of the index
+		/// </summary>
+		[BsonElement("name"), BsonRequired]
+		public string Name { get; private set; }
+
+		/// <summary>
+		/// Keys for the index
+		/// </summary>
+		[BsonElement("key"), BsonRequired]
+		public BsonDocument KeysDocument { get; private set; }
+
+		/// <summary>
+		/// The index should be unique
+		/// </summary>
+		[BsonElement("unique")]
+		public bool Unique { get; private set; }
+
+		/// <summary>
+		/// Whether to create a sparse index
+		/// </summary>
+		[BsonElement("sparse")]
+		public bool Sparse { get; private set; }
+
+		/// <summary>
+		/// Constructor
+		/// </summary>
+		[BsonConstructor]
+		private MongoIndex()
+		{
+			Name = String.Empty;
+			KeysDocument = new BsonDocument();
+		}
+
+		/// <summary>
+		/// Constructor
+		/// </summary>
+		/// <param name="name">Name of the index</param>
+		/// <param name="keysDocument">Keys for the index</param>
+		/// <param name="unique">Whether the index should be unique</param>
+		/// <param name="sparse">Whether to create a sparse index</param>
+		public MongoIndex(string? name, BsonDocument keysDocument, bool unique, bool sparse)
+		{
+			Name = name ?? GetDefaultName(keysDocument);
+			KeysDocument = keysDocument;
+			Unique = unique;
+			Sparse = sparse;
+		}
+
+		/// <summary>
+		/// Constructor
+		/// </summary>
+		/// <param name="keysFunc">Callback to configure keys for this document</param>
+		/// <param name="unique">Whether the index should be unique</param>
+		/// <param name="sparse">Whether to create a sparse index</param>
+		public static MongoIndex<T> Create<T>(Func<IndexKeysDefinitionBuilder<T>, IndexKeysDefinition<T>> keysFunc, bool unique = false, bool sparse = false)
+		{
+			return new MongoIndex<T>(null, keysFunc(Builders<T>.IndexKeys), unique, sparse);
+		}
+
+		/// <summary>
+		/// Constructor
+		/// </summary>
+		/// <param name="name">Name of the index</param>
+		/// <param name="keysFunc">Callback to configure keys for this document</param>
+		/// <param name="unique">Whether the index should be unique</param>
+		/// <param name="sparse">Whether to create a sparse index</param>
+		public static MongoIndex<T> Create<T>(string name, Func<IndexKeysDefinitionBuilder<T>, IndexKeysDefinition<T>> keysFunc, bool unique = false, bool sparse = false)
+		{
+			return new MongoIndex<T>(name, keysFunc(Builders<T>.IndexKeys), unique, sparse);
+		}
+
+		/// <inheritdoc/>
+		public override bool Equals(object? obj) => obj is MongoIndex other && Equals(other);
+
+		/// <inheritdoc/>
+		public override int GetHashCode() => Name.GetHashCode(StringComparison.Ordinal);
+
+		/// <inheritdoc/>
+		public bool Equals(MongoIndex? other) => other != null && Name.Equals(other.Name, StringComparison.Ordinal) && KeysDocument.Equals(other.KeysDocument) && Sparse == other.Sparse && Unique == other.Unique;
+
+		/// <summary>
+		/// Gets the default name for an index based on its keys
+		/// </summary>
+		/// <param name="keys">Keys for the index</param>
+		/// <returns>Name of the index</returns>
+		protected static string GetDefaultName(BsonDocument keys)
+		{
+			StringBuilder name = new StringBuilder();
+			foreach (BsonElement element in keys.Elements)
+			{
+				if (name.Length > 0)
+				{
+					name.Append('_');
+				}
+				name.Append(element.Name);
+				name.Append('_');
+				name.Append(element.Value.ToString());
+			}
+			return name.ToString();
+		}
+	}
+
+	/// <summary>
+	/// Strongly typed index document
+	/// </summary>
+	/// <typeparam name="T"></typeparam>
+	public class MongoIndex<T> : MongoIndex
+	{
+		/// <summary>
+		/// Keys for the index
+		/// </summary>
+		[BsonIgnore]
+		public IndexKeysDefinition<T> Keys { get; }
+
+		/// <summary>
+		/// Constructor
+		/// </summary>
+		/// <param name="name">Name of the index</param>
+		/// <param name="keys">Keys for the index</param>
+		/// <param name="unique">Whether the index should be unique</param>
+		/// <param name="sparse">Whether to create a sparse index</param>
+		public MongoIndex(string? name, IndexKeysDefinition<T> keys, bool unique, bool sparse)
+			: base(name, keys.Render(BsonSerializer.LookupSerializer<T>(), BsonSerializer.SerializerRegistry), unique, sparse)
+		{
+			Keys = keys;
+		}
+	}
+
+	/// <summary>
+	/// Extension methods for <see cref="MongoIndex{T}"/>
+	/// </summary>
+	static class MongoIndexExtensions
+	{
+		public static void Add<T>(this List<MongoIndex<T>> list, Func<IndexKeysDefinitionBuilder<T>, IndexKeysDefinition<T>> keyFunc, bool unique = false, bool sparse = false)
+		{
+			list.Add(MongoIndex.Create(keyFunc, unique, sparse));
+		}
+
+		public static void Add<T>(this List<MongoIndex<T>> list, string name, Func<IndexKeysDefinitionBuilder<T>, IndexKeysDefinition<T>> keyFunc, bool unique = false, bool sparse = false)
+		{
+			list.Add(MongoIndex.Create(name, keyFunc, unique, sparse));
+		}
+	}
+}
diff --git a/Engine/Source/Programs/Horde/HordeServer.Shared/Users/IAvatar.cs b/Engine/Source/Programs/Horde/HordeServer.Shared/Users/IAvatar.cs
new file mode 100644
--- /dev/null
+++ b/Engine/Source/Programs/Horde/HordeServer.Shared/Users/IAvatar.cs
@@ -0,0 +1,30 @@
+// Copyright Epic Games, Inc. All Rights Reserved.
+
+namespace HordeServer.Users
+{
+	/// <summary>
+	/// Avatar for the user
+	/// </summary>
+	public interface IAvatar
+	{
+		/// <summary>
+		/// Url for a 24x24 pixel image
+		/// </summary>
+		public string? Image24 { get; set; }
+
+		/// <summary>
+		/// Url for a 32x32 pixel image
+		/// </summary>
+		public string? Image32 { get; set; }
+
+		/// <summary>
+		/// Url for a 48x48 pixel image
+		/// </summary>
+		public string? Image48 { get; set; }
+
+		/// <summary>
+		/// Url for a 72x72 pixel image
+		/// </summary>
+		public string? Image72 { get; set; }
+	}
+}
diff --git a/Engine/Source/Programs/Horde/HordeServer.Shared/Users/IAvatarService.cs b/Engine/Source/Programs/Horde/HordeServer.Shared/Users/IAvatarService.cs
new file mode 100644
--- /dev/null
+++ b/Engine/Source/Programs/Horde/HordeServer.Shared/Users/IAvatarService.cs
@@ -0,0 +1,31 @@
+// Copyright Epic Games, Inc. All Rights Reserved.
+
+using System.Threading;
+using System.Threading.Tasks;
+
+namespace HordeServer.Users
+{
+	/// <summary>
+	/// Provides avatar images
+	/// </summary>
+	public interface IAvatarService
+	{
+		/// <summary>
+		/// Gets a users's avatar
+		/// </summary>
+		/// <param name="user"></param>
+		/// <param name="cancellationToken">CAncellation token for the operation</param>
+		/// <returns></returns>
+		Task<IAvatar?> GetAvatarAsync(IUser user, CancellationToken cancellationToken);
+	}
+
+	/// <summary>
+	/// Default implementation of <see cref="IAvatarService"/>
+	/// </summary>
+	public class NullAvatarService : IAvatarService
+	{
+		/// <inheritdoc/>
+		public Task<IAvatar?> GetAvatarAsync(IUser user, CancellationToken cancellationToken) => Task.FromResult<IAvatar?>(null);
+	}
+}
+
diff --git a/Engine/Source/Programs/Horde/HordeServer.Shared/Users/IUser.cs b/Engine/Source/Programs/Horde/HordeServer.Shared/Users/IUser.cs
new file mode 100644
--- /dev/null
+++ b/Engine/Source/Programs/Horde/HordeServer.Shared/Users/IUser.cs
@@ -0,0 +1,224 @@
+// Copyright Epic Games, Inc. All Rights Reserved.
+
+using System;
+using System.Collections.Generic;
+using System.Linq;
+using EpicGames.Core;
+using EpicGames.Horde.Jobs;
+using EpicGames.Horde.Jobs.Bisect;
+using EpicGames.Horde.Jobs.Templates;
+using EpicGames.Horde.Streams;
+using EpicGames.Horde.Users;
+using MongoDB.Bson;
+
+namespace HordeServer.Users
+{
+	/// <summary>
+	/// Known user ids
+	/// </summary>
+	public static class KnownUsers
+	{
+		/// <summary>
+		/// The system user. Used for automated processes.
+		/// </summary>
+		public static UserId System { get; } = UserId.Parse("6170b423b94a2c7c2d6b6f87");
+	}
+
+	/// <summary>
+	/// Document which collates information about a user, and their personal settings
+	/// </summary>
+	public interface IUser
+	{
+		/// <summary>
+		/// The user id
+		/// </summary>
+		public UserId Id { get; }
+
+		/// <summary>
+		/// Full name of the user
+		/// </summary>
+		public string Name { get; }
+
+		/// <summary>
+		/// The user's login id
+		/// </summary>
+		public string Login { get; }
+
+		/// <summary>
+		/// The user's email
+		/// </summary>
+		public string? Email { get; }
+	}
+
+	/// <summary>
+	/// Claims for a particular user
+	/// </summary>
+	public interface IUserClaims
+	{
+		/// <summary>
+		/// The user id
+		/// </summary>
+		public UserId UserId { get; }
+
+		/// <summary>
+		/// Claims for this user (on last login)
+		/// </summary>
+		public IReadOnlyList<IUserClaim> Claims { get; }
+	}
+
+	/// <summary>
+	/// 
+	/// </summary>
+	public interface IUserJobTemplateSettings
+	{
+		/// <summary>
+		/// The stream the job was run in
+		/// </summary>
+		public StreamId StreamId { get; }
+
+		/// <summary>
+		/// The template id of the job
+		/// </summary>
+		public TemplateId TemplateId { get; }
+
+		/// <summary>
+		/// The hash of the template definition
+		/// </summary>
+		public string TemplateHash { get; }
+
+		/// <summary>
+		/// The arguments defined when creating the job
+		/// </summary>
+		public IReadOnlyList<string> Arguments { get; }
+
+		/// <summary>
+		/// The last time the job template was used
+		/// </summary>
+		public DateTime UpdateTimeUtc { get; }
+	}
+
+	/// <summary>
+	/// Settings for updating job template user preference
+	/// </summary>
+	public class UpdateUserJobTemplateOptions
+	{
+		/// <summary>
+		/// The stream the job was run in
+		/// </summary>
+		public StreamId StreamId { get; set; }
+
+		/// <summary>
+		/// The template id of the job
+		/// </summary>
+		public TemplateId TemplateId { get; set; }
+
+		/// <summary>
+		/// The hash of the template definition
+		/// </summary>
+		public string TemplateHash { get; set; } = String.Empty;
+
+		/// <summary>
+		/// The arguments defined when creating the job
+		/// </summary>
+		public IReadOnlyList<string> Arguments { get; set; } = new List<string>();
+	}
+
+	/// <summary>
+	/// User settings document
+	/// </summary>
+	public interface IUserSettings
+	{
+		/// <summary>
+		/// The user id
+		/// </summary>
+		public UserId UserId { get; }
+
+		/// <summary>
+		/// Whether to enable experimental features
+		/// </summary>
+		public bool EnableExperimentalFeatures { get; }
+
+		/// <summary>
+		/// Whether to always tag CL descriptions for preflights
+		/// </summary>
+		public bool AlwaysTagPreflightCL { get; }
+
+		/// <summary>
+		/// Opaque settings dictionary for the dashboard
+		/// </summary>
+		public BsonValue DashboardSettings { get; }
+
+		/// <summary>
+		/// List of pinned jobs
+		/// </summary>
+		public IReadOnlyList<JobId> PinnedJobIds { get; }
+
+		/// <summary>
+		/// List of pinned bisection tasks
+		/// </summary>
+		public IReadOnlyList<BisectTaskId> PinnedBisectTaskIds { get; }
+
+		/// <summary>
+		/// List of job template preferences
+		/// </summary>
+		public IReadOnlyList<IUserJobTemplateSettings>? JobTemplateSettings { get; }
+	}
+
+	/// <summary>
+	/// Extension methods
+	/// </summary>
+	public static class UserExtensions
+	{
+		/// <summary>
+		/// Creates an API response object
+		/// </summary>
+		public static GetUserResponse ToApiResponse(this IUser user, IAvatar? avatar, IUserClaims? claims, IUserSettings? settings)
+		{
+			GetUserResponse response = new GetUserResponse(user.Id, user.Name);
+			response.Email = user.Email;
+
+			response.Image24 = avatar?.Image24;
+			response.Image32 = avatar?.Image32;
+			response.Image48 = avatar?.Image48;
+			response.Image72 = avatar?.Image72;
+
+			response.Claims = claims?.Claims.Select(x => new GetUserClaimResponse(x.Type, x.Value)).ToList();
+
+			if (settings != null)
+			{
+				response.EnableExperimentalFeatures = settings.EnableExperimentalFeatures;
+				response.AlwaysTagPreflightCL = settings.AlwaysTagPreflightCL;
+
+				response.DashboardSettings = BsonTypeMapper.MapToDotNetValue(settings.DashboardSettings);
+				response.PinnedJobIds = settings.PinnedJobIds.ConvertAll(x => x.ToString());
+				response.PinnedBisectTaskIds = settings.PinnedBisectTaskIds.ConvertAll(x => x.ToString());
+
+				if (settings.JobTemplateSettings != null && settings.JobTemplateSettings.Count > 0)
+				{
+					response.JobTemplateSettings = new List<GetJobTemplateSettingsResponse>();
+					for (int i = 0; i < settings.JobTemplateSettings.Count; i++)
+					{
+						response.JobTemplateSettings.Add(settings.JobTemplateSettings[i].ToApiResponse());
+					}
+				}
+			}
+			return response;
+		}
+
+		/// <summary>
+		/// Creates an API thin user response
+		/// </summary>
+		public static GetThinUserInfoResponse ToThinApiResponse(this IUser user)
+		{
+			return new GetThinUserInfoResponse(user.Id, user.Name, user.Email, user.Login);
+		}
+
+		/// <summary>
+		/// Creates an API settings response
+		/// </summary>
+		public static GetJobTemplateSettingsResponse ToApiResponse(this IUserJobTemplateSettings settings)
+		{
+			return new GetJobTemplateSettingsResponse(settings.StreamId.ToString(), settings.TemplateId.ToString(), settings.TemplateHash.ToString(), settings.Arguments.ToList(), settings.UpdateTimeUtc);
+		}
+	}
+}
diff --git a/Engine/Source/Programs/Horde/HordeServer.Shared/Users/IUserClaim.cs b/Engine/Source/Programs/Horde/HordeServer.Shared/Users/IUserClaim.cs
new file mode 100644
--- /dev/null
+++ b/Engine/Source/Programs/Horde/HordeServer.Shared/Users/IUserClaim.cs
@@ -0,0 +1,70 @@
+// Copyright Epic Games, Inc. All Rights Reserved.
+
+using System.Security.Claims;
+
+namespace HordeServer.Users
+{
+	/// <summary>
+	/// Claim for a user
+	/// </summary>
+	public interface IUserClaim
+	{
+		/// <summary>
+		/// Name of the claim
+		/// </summary>
+		public string Type { get; }
+
+		/// <summary>
+		/// Value of the claim
+		/// </summary>
+		public string Value { get; }
+	}
+
+	/// <summary>
+	/// New claim document
+	/// </summary>
+	public class UserClaim : IUserClaim
+	{
+		/// <inheritdoc/>
+		public string Type { get; set; }
+
+		/// <inheritdoc/>
+		public string Value { get; set; }
+
+		/// <summary>
+		/// Constructor
+		/// </summary>
+		/// <param name="type"></param>
+		/// <param name="value"></param>
+		public UserClaim(string type, string value)
+		{
+			Type = type;
+			Value = value;
+		}
+
+		/// <summary>
+		/// Constructor
+		/// </summary>
+		/// <param name="other">Claim to construct from</param>
+		public UserClaim(IUserClaim other)
+			: this(other.Type, other.Value)
+		{
+		}
+
+		/// <summary>
+		/// Constructs a UserClaim from a Claim object
+		/// </summary>
+		/// <param name="claim">Claim object</param>
+		/// <returns></returns>
+		public static UserClaim FromClaim(Claim claim)
+		{
+			return new UserClaim(claim.Type, claim.Value);
+		}
+
+		/// <summary>
+		/// Conversion operator from NET claims
+		/// </summary>
+		/// <param name="claim"></param>
+		public static implicit operator UserClaim(Claim claim) => FromClaim(claim);
+	}
+}
diff --git a/Engine/Source/Programs/Horde/HordeServer.Shared/Users/IUserCollection.cs b/Engine/Source/Programs/Horde/HordeServer.Shared/Users/IUserCollection.cs
new file mode 100644
--- /dev/null
+++ b/Engine/Source/Programs/Horde/HordeServer.Shared/Users/IUserCollection.cs
@@ -0,0 +1,140 @@
+// Copyright Epic Games, Inc. All Rights Reserved.
+
+using System.Collections.Generic;
+using System.Security.Claims;
+using System.Threading;
+using System.Threading.Tasks;
+using EpicGames.Horde.Jobs;
+using EpicGames.Horde.Jobs.Bisect;
+using EpicGames.Horde.Users;
+using HordeServer.Utilities;
+using MongoDB.Bson;
+
+namespace HordeServer.Users
+{
+	/// <summary>
+	/// Manages user documents
+	/// </summary>
+	public interface IUserCollection
+	{
+		/// <summary>
+		/// Gets a user by unique id
+		/// </summary>
+		/// <param name="id">Id of the user</param>
+		/// <param name="cancellationToken">Cancellation token for the operation</param>
+		/// <returns>The user information</returns>
+		Task<IUser?> GetUserAsync(UserId id, CancellationToken cancellationToken = default);
+
+		/// <summary>
+		/// Gets a cached user by unique id
+		/// </summary>
+		/// <param name="id">Id of the user</param>
+		/// <param name="cancellationToken">Cancellation token for the operation</param>
+		/// <returns>The user information</returns>
+		ValueTask<IUser?> GetCachedUserAsync(UserId? id, CancellationToken cancellationToken = default);
+
+		/// <summary>
+		/// Gets a user by unique id
+		/// </summary>
+		/// <param name="ids">Ids of the users</param>
+		/// <param name="nameRegex">Name regex to match for the users</param>
+		/// <param name="index">Maximum number of results</param>
+		/// <param name="count">Number of results to return</param>
+		/// <param name="cancellationToken">Cancellation token for the operation</param>
+		/// <returns>The user information</returns>
+		Task<IReadOnlyList<IUser>> FindUsersAsync(IEnumerable<UserId>? ids = null, string? nameRegex = null, int? index = null, int? count = null, CancellationToken cancellationToken = default);
+
+		/// <summary>
+		/// Gets a user by login
+		/// </summary>
+		/// <param name="login">Login for the user</param>
+		/// <param name="cancellationToken">Cancellation token for the operation</param>
+		/// <returns>The user information</returns>
+		Task<IUser?> FindUserByLoginAsync(string login, CancellationToken cancellationToken = default);
+
+		/// <summary>
+		/// Gets a user by email address
+		/// </summary>
+		/// <param name="email">Email for the user</param>
+		/// <param name="cancellationToken">Cancellation token for the operation</param>
+		/// <returns>The user information</returns>
+		Task<IUser?> FindUserByEmailAsync(string email, CancellationToken cancellationToken = default);
+
+		/// <summary>
+		/// Find or add a user with the given claims. Claims will be updated if the user exists.
+		/// </summary>
+		/// <param name="login">Login id of the user</param>
+		/// <param name="name">Full name of the user</param>
+		/// <param name="email">Email address of the user</param>
+		/// <param name="cancellationToken">Cancellation token for the operation</param>
+		/// <returns>The user document</returns>
+		Task<IUser> FindOrAddUserByLoginAsync(string login, string? name = null, string? email = null, CancellationToken cancellationToken = default);
+
+		/// <summary>
+		/// Gets the claims for a user
+		/// </summary>
+		/// <param name="userId"></param>
+		/// <param name="cancellationToken">Cancellation token for the operation</param>
+		/// <returns></returns>
+		Task<IUserClaims> GetClaimsAsync(UserId userId, CancellationToken cancellationToken = default);
+
+		/// <summary>
+		/// Update the claims for a user
+		/// </summary>
+		/// <param name="userId"></param>
+		/// <param name="claims"></param>
+		/// <param name="cancellationToken">Cancellation token for the operation</param>
+		/// <returns></returns>
+		Task UpdateClaimsAsync(UserId userId, IEnumerable<IUserClaim> claims, CancellationToken cancellationToken = default);
+
+		/// <summary>
+		/// Get settings for a user
+		/// </summary>
+		/// <param name="userId"></param>
+		/// <param name="cancellationToken">Cancellation token for the operation</param>
+		/// <returns></returns>
+		Task<IUserSettings> GetSettingsAsync(UserId userId, CancellationToken cancellationToken = default);
+
+		/// <summary>
+		/// Update a user
+		/// </summary>
+		/// <param name="userId">The user to update</param>
+		/// <param name="enableExperimentalFeatures"></param>
+		/// <param name="alwaysTagPreflightCL"></param>
+		/// <param name="dashboardSettings">Opaque settings object for the dashboard</param>
+		/// <param name="addPinnedJobIds"></param>
+		/// <param name="removePinnedJobIds"></param>
+		/// <param name="templateOptions"></param>
+		/// <param name="addBisectTaskIds"></param>
+		/// <param name="removeBisectTaskIds"></param>
+		/// <param name="cancellationToken">Cancellation token for the operation</param>
+		/// <returns>Updated user object</returns>
+		Task UpdateSettingsAsync(UserId userId, bool? enableExperimentalFeatures = null, bool? alwaysTagPreflightCL = null, BsonValue? dashboardSettings = null, IEnumerable<JobId>? addPinnedJobIds = null, IEnumerable<JobId>? removePinnedJobIds = null, UpdateUserJobTemplateOptions? templateOptions = null, IEnumerable<BisectTaskId>? addBisectTaskIds = null, IEnumerable<BisectTaskId>? removeBisectTaskIds = null, CancellationToken cancellationToken = default);
+	}
+
+	/// <summary>
+	/// Extension methods for the user collect
+	/// </summary>
+	static class UserCollectionExtensions
+	{
+		/// <summary>
+		/// Gets a particular user info from the collection
+		/// </summary>
+		/// <param name="userCollection"></param>
+		/// <param name="principal"></param>
+		/// <param name="cancellationToken">Cancellation token for the operation</param>
+		/// <returns></returns>
+		public static Task<IUser?> GetUserAsync(this IUserCollection userCollection, ClaimsPrincipal principal, CancellationToken cancellationToken = default)
+		{
+			UserId? userId = principal.GetUserId();
+			if (userId != null)
+			{
+				return userCollection.GetUserAsync(userId.Value, cancellationToken);
+			}
+			else
+			{
+				return Task.FromResult<IUser?>(null);
+			}
+		}
+	}
+}
diff --git a/Engine/Source/Programs/Horde/HordeServer.Shared/Utilities/AsyncCachedValue.cs b/Engine/Source/Programs/Horde/HordeServer.Shared/Utilities/AsyncCachedValue.cs
new file mode 100644
--- /dev/null
+++ b/Engine/Source/Programs/Horde/HordeServer.Shared/Utilities/AsyncCachedValue.cs
@@ -0,0 +1,174 @@
+// Copyright Epic Games, Inc. All Rights Reserved.
+
+using System;
+using System.Diagnostics;
+using System.Threading;
+using System.Threading.Tasks;
+using EpicGames.Core;
+
+namespace HordeServer.Utilities
+{
+	/// <summary>
+	/// Caches a value and asynchronously updates it after a period of time
+	/// </summary>
+	/// <typeparam name="T"></typeparam>
+	public sealed class AsyncCachedValue<T> : IAsyncDisposable
+	{
+		class State
+		{
+			public readonly T Value;
+			readonly Stopwatch _timer;
+			public Task<State>? _next;
+
+			public TimeSpan Elapsed => _timer.Elapsed;
+
+			public State(T value)
+			{
+				Value = value;
+				_timer = Stopwatch.StartNew();
+			}
+		}
+
+		CancellationTokenSource _cancellationTokenSource = new CancellationTokenSource();
+
+		/// <summary>
+		/// The current state
+		/// </summary>
+		Task<State>? _current = null;
+
+		/// <summary>
+		/// Generator for the new value
+		/// </summary>
+		readonly Func<CancellationToken, Task<T>> _generator;
+
+		/// <summary>
+		/// Time at which to start to refresh the value
+		/// </summary>
+		readonly TimeSpan _minRefreshTime;
+
+		/// <summary>
+		/// Time at which to wait for the value to refresh
+		/// </summary>
+		readonly TimeSpan _maxRefreshTime;
+
+		/// <summary>
+		/// Default constructor
+		/// </summary>
+		public AsyncCachedValue(Func<CancellationToken, Task<T>> generator, TimeSpan refreshTime)
+			: this(generator, refreshTime * 0.75, refreshTime)
+		{
+		}
+
+		/// <summary>
+		/// Default constructor
+		/// </summary>
+		public AsyncCachedValue(Func<CancellationToken, Task<T>> generator, TimeSpan minRefreshTime, TimeSpan maxRefreshTime)
+		{
+			_generator = generator;
+			_minRefreshTime = minRefreshTime;
+			_maxRefreshTime = maxRefreshTime;
+		}
+
+		/// <inheritdoc/>
+		public async ValueTask DisposeAsync()
+		{
+			if (_current != null)
+			{
+				await _cancellationTokenSource.CancelAsync();
+				await _current.IgnoreCanceledExceptionsAsync();
+				_current = null;
+			}
+
+			if (_cancellationTokenSource != null)
+			{
+				_cancellationTokenSource.Dispose();
+				_cancellationTokenSource = null!;
+			}
+		}
+
+		/// <summary>
+		/// Invalidates the current value
+		/// </summary>
+		public void Invalidate()
+		{
+			_current = null;
+		}
+
+		/// <summary>
+		/// Tries to get the current value
+		/// </summary>
+		/// <param name="cancellationToken">Cancellation token for the request</param>
+		/// <returns>The cached value, if valid</returns>
+		public Task<T> GetAsync(CancellationToken cancellationToken = default)
+		{
+			return GetAsync(_maxRefreshTime, cancellationToken);
+		}
+
+		/// <summary>
+		/// Tries to get the current value
+		/// </summary>
+		/// <returns>The cached value, if valid</returns>
+		public Task<T> GetAsync(TimeSpan maxAge, CancellationToken cancellationToken = default)
+		{
+			Task<T> task = GetInternalAsync(maxAge);
+			if (cancellationToken.CanBeCanceled)
+			{
+				// The returned task object is shared, so we don't want to cancel computation of it; just the wait for a result.
+				task = WrapCancellationAsync(task, cancellationToken);
+			}
+			return task;
+		}
+
+		static async Task<T> WrapCancellationAsync(Task<T> task, CancellationToken cancellationToken)
+		{
+			await Task.WhenAny(task, Task.Delay(-1, cancellationToken));
+			cancellationToken.ThrowIfCancellationRequested();
+			return await task;
+		}
+
+		async Task<T> GetInternalAsync(TimeSpan maxAge)
+		{
+			Task<State> stateTask = CreateOrGetStateTaskAsync(ref _current);
+
+			State state = await stateTask;
+			if (state._next != null && state._next.IsCompleted)
+			{
+				_ = Interlocked.CompareExchange(ref _current, state._next, stateTask);
+				state = await state._next;
+			}
+			if (state.Elapsed > maxAge)
+			{
+				state = await CreateOrGetStateTaskAsync(ref state._next);
+			}
+			if (state.Elapsed > _minRefreshTime)
+			{
+				_ = CreateOrGetStateTaskAsync(ref state._next);
+			}
+
+			return state.Value;
+		}
+
+		Task<State> CreateOrGetStateTaskAsync(ref Task<State>? stateTask)
+		{
+			for (; ; )
+			{
+				Task<State>? currentStateTask = stateTask;
+				if (currentStateTask != null)
+				{
+					return currentStateTask;
+				}
+
+				Task<Task<State>> innerNewStateTask = new Task<Task<State>>(() => CreateStateAsync());
+				if (Interlocked.CompareExchange(ref stateTask, innerNewStateTask.Unwrap(), null) == null)
+				{
+					innerNewStateTask.Start();
+				}
+			}
+		}
+
+		async Task<State> CreateStateAsync()
+		{
+			return new State(await _generator(_cancellationTokenSource.Token));
+		}
+	}
+}
diff --git a/Engine/Source/Programs/Horde/HordeServer.Shared/Utilities/BinaryIdConverter.cs b/Engine/Source/Programs/Horde/HordeServer.Shared/Utilities/BinaryIdConverter.cs
new file mode 100644
--- /dev/null
+++ b/Engine/Source/Programs/Horde/HordeServer.Shared/Utilities/BinaryIdConverter.cs
@@ -0,0 +1,41 @@
+// Copyright Epic Games, Inc. All Rights Reserved.
+
+using System;
+using System.Reflection;
+using EpicGames.Horde;
+using MongoDB.Bson.Serialization;
+using MongoDB.Bson.Serialization.Serializers;
+
+namespace HordeServer.Utilities
+{
+	/// <summary>
+	/// Class which serializes object id types to BSON
+	/// </summary>
+	sealed class BinaryIdBsonSerializer<TValue, TConverter> : SerializerBase<TValue> where TValue : struct where TConverter : BinaryIdConverter<TValue>, new()
+	{
+		readonly TConverter _converter = new TConverter();
+
+		/// <inheritdoc/>
+		public override TValue Deserialize(BsonDeserializationContext context, BsonDeserializationArgs args) => _converter.FromBinaryId(BinaryIdUtils.FromObjectId(context.Reader.ReadObjectId()));
+
+		/// <inheritdoc/>
+		public override void Serialize(BsonSerializationContext context, BsonSerializationArgs args, TValue value) => context.Writer.WriteObjectId(BinaryIdUtils.ToObjectId(_converter.ToBinaryId(value)));
+	}
+
+	/// <summary>
+	/// Class which serializes object id types to BSON
+	/// </summary>
+	sealed class BinaryIdBsonSerializationProvider : BsonSerializationProviderBase
+	{
+		/// <inheritdoc/>
+		public override IBsonSerializer? GetSerializer(Type type, IBsonSerializerRegistry serializerRegistry)
+		{
+			BinaryIdConverterAttribute? attribute = type.GetCustomAttribute<BinaryIdConverterAttribute>();
+			if (attribute == null)
+			{
+				return null;
+			}
+			return (IBsonSerializer?)Activator.CreateInstance(typeof(BinaryIdBsonSerializer<,>).MakeGenericType(type, attribute.ConverterType));
+		}
+	}
+}
diff --git a/Engine/Source/Programs/Horde/HordeServer.Shared/Utilities/BinaryIdUtils.cs b/Engine/Source/Programs/Horde/HordeServer.Shared/Utilities/BinaryIdUtils.cs
new file mode 100644
--- /dev/null
+++ b/Engine/Source/Programs/Horde/HordeServer.Shared/Utilities/BinaryIdUtils.cs
@@ -0,0 +1,38 @@
+// Copyright Epic Games, Inc. All Rights Reserved.
+
+using EpicGames.Horde;
+using MongoDB.Bson;
+
+namespace HordeServer.Utilities
+{
+	/// <summary>
+	/// Utility methods for <see cref="BinaryId"/>
+	/// </summary>
+	static class BinaryIdUtils
+	{
+		/// <summary>
+		/// Creates a new BinaryId
+		/// </summary>
+		public static BinaryId CreateNew()
+		{
+			ObjectId objectId = ObjectId.GenerateNewId();
+			return new BinaryId(objectId.ToByteArray());
+		}
+
+		/// <summary>
+		/// Creates a BinaryId from an ObjectId
+		/// </summary>
+		public static BinaryId FromObjectId(ObjectId objectId)
+		{
+			return new BinaryId(objectId.ToByteArray());
+		}
+
+		/// <summary>
+		/// Creates an ObjectId from a BinaryId
+		/// </summary>
+		public static ObjectId ToObjectId(BinaryId binaryId)
+		{
+			return new ObjectId(binaryId.ToByteArray());
+		}
+	}
+}
diff --git a/Engine/Source/Programs/Horde/HordeServer.Shared/Utilities/HordeClaimTypes.cs b/Engine/Source/Programs/Horde/HordeServer.Shared/Utilities/HordeClaimTypes.cs
new file mode 100644
--- /dev/null
+++ b/Engine/Source/Programs/Horde/HordeServer.Shared/Utilities/HordeClaimTypes.cs
@@ -0,0 +1,157 @@
+// Copyright Epic Games, Inc. All Rights Reserved.
+
+using System.Security.Claims;
+using EpicGames.Horde.Accounts;
+using EpicGames.Horde.Users;
+using HordeServer.Users;
+
+namespace HordeServer.Utilities
+{
+	/// <summary>
+	/// Claim types that are specific to horde
+	/// </summary>
+	public static class HordeClaimTypes
+	{
+		/// <summary>
+		/// Version number for auth claims. Can be updated to force a re-login.
+		/// </summary>
+		public const string CurrentVersion = "1";
+
+		/// <summary>
+		/// Base URI for all Horde claims.
+		/// </summary>
+		const string Prefix = "http://epicgames.com/ue/horde/";
+
+		/// <summary>
+		/// Version number for the auth header
+		/// </summary>
+		public const string Version = Prefix + "version";
+
+		/// <summary>
+		/// Claim for a particular role. This is reserved for well-known roles defined by Horde itself.
+		/// </summary>
+		public const string Role = Prefix + "role";
+
+		/// <summary>
+		/// Claim type reserved for a particular group within the Horde account system. Values are user defined.
+		/// </summary>
+		public const string Group = Prefix + "group";
+
+		/// <summary>
+		/// Claim for a particular agent
+		/// </summary>
+		public const string Agent = Prefix + "agent";
+
+		/// <summary>
+		/// Claim for an agent's enrollment key
+		/// </summary>
+		public const string AgentEnrollmentKey = Prefix + "agent-enrollment-key";
+
+		/// <summary>
+		/// Claim for a particular account id
+		/// </summary>
+		public const string AccountId = Prefix + "account-id";
+
+		/// <summary>
+		/// Claim for a particular session
+		/// </summary>
+		public const string AgentSessionId = Prefix + "session";
+
+		/// <summary>
+		/// User name for horde
+		/// </summary>
+		public const string User = Prefix + "user";
+
+		/// <summary>
+		/// Unique id of the horde user (see <see cref="IUserCollection"/>)
+		/// </summary>
+		public const string UserId = Prefix + "user-id-v3";
+
+		/// <summary>
+		/// Claim for downloading artifacts from a particular job
+		/// </summary>
+		public const string JobArtifacts = Prefix + "job-artifacts";
+
+		/// <summary>
+		/// Claim for the Perforce username
+		/// </summary>
+		public const string PerforceUser = Prefix + "perforce-user";
+
+		/// <summary>
+		/// Claim for the external issue username
+		/// </summary>
+		public const string ExternalIssueUser = Prefix + "external-issue-user";
+
+		/// <summary>
+		/// Claim identifying an agent as executing a particular lease
+		/// </summary>
+		public const string Lease = Prefix + "lease";
+
+		/// <summary>
+		/// Claim allowing a certain namespace and/or path to be read from. See <see cref="WriteNamespace"/>.
+		/// </summary>
+		public const string ReadNamespace = Prefix + "read-namespace";
+
+		/// <summary>
+		/// Claim allowing a certain namespace and/or path to be written to. Value may be a namespace name, or a subpath (eg. ns:a/b will only allow blobs and refs with an a/b prefix to be written).
+		/// </summary>
+		public const string WriteNamespace = Prefix + "write-namespace";
+	}
+
+	/// <summary>
+	/// Extension methods for getting Horde claims from a principal
+	/// </summary>
+	public static class HordeClaimExtensions
+	{
+		/// <summary>
+		/// Gets the Horde account id from a principal
+		/// </summary>
+		/// <param name="principal"></param>
+		/// <returns></returns>
+		public static AccountId? GetAccountId(this ClaimsPrincipal principal)
+		{
+			string? idValue = principal.FindFirstValue(HordeClaimTypes.AccountId);
+			if (idValue == null)
+			{
+				return null;
+			}
+			else
+			{
+				return AccountId.Parse(idValue);
+			}
+		}
+
+		/// <summary>
+		/// Gets the Horde user id from a principal
+		/// </summary>
+		/// <param name="principal"></param>
+		/// <returns></returns>
+		public static UserId? GetUserId(this ClaimsPrincipal principal)
+		{
+			string? idValue = principal.FindFirstValue(HordeClaimTypes.UserId);
+			if (idValue == null)
+			{
+				return null;
+			}
+			else
+			{
+				return UserId.Parse(idValue);
+			}
+		}
+
+		/// <summary>
+		/// Gets the external issue username for the given principal
+		/// </summary>
+		/// <param name="user">The principal to get the external issue user for</param>
+		/// <returns>Jira user name</returns>
+		public static string? GetExternalIssueUser(this ClaimsPrincipal user)
+		{
+			Claim? claim = user.FindFirst(HordeClaimTypes.ExternalIssueUser);
+			if (claim == null)
+			{
+				return null;
+			}
+			return claim.Value;
+		}
+	}
+}
diff --git a/Engine/Source/Programs/Horde/HordeServer.Shared/Utilities/HordeControllerBase.cs b/Engine/Source/Programs/Horde/HordeServer.Shared/Utilities/HordeControllerBase.cs
new file mode 100644
--- /dev/null
+++ b/Engine/Source/Programs/Horde/HordeServer.Shared/Utilities/HordeControllerBase.cs
@@ -0,0 +1,350 @@
+// Copyright Epic Games, Inc. All Rights Reserved.
+
+using EpicGames.Core;
+using EpicGames.Horde.Accounts;
+using EpicGames.Horde.Acls;
+using EpicGames.Horde.Agents;
+using EpicGames.Horde.Agents.Leases;
+using EpicGames.Horde.Artifacts;
+using EpicGames.Horde.Compute;
+using EpicGames.Horde.Jobs;
+using EpicGames.Horde.Jobs.Bisect;
+using EpicGames.Horde.Jobs.Templates;
+using EpicGames.Horde.Projects;
+using EpicGames.Horde.Secrets;
+using EpicGames.Horde.Storage;
+using EpicGames.Horde.Streams;
+using EpicGames.Horde.Telemetry;
+using EpicGames.Horde.Tools;
+using Microsoft.AspNetCore.Http;
+using Microsoft.AspNetCore.Mvc;
+using Microsoft.Extensions.Logging;
+
+namespace HordeServer.Utilities
+{
+	/// <summary>
+	/// Base class for Horde controllers
+	/// </summary>
+	public abstract class HordeControllerBase : ControllerBase
+	{
+		/// <summary>
+		/// Create a response indicating a problem with the request
+		/// </summary>
+		/// <param name="message">Standard structured-logging format string</param>
+		/// <param name="args">Arguments for the format string</param>
+		[NonAction]
+		protected ActionResult BadRequest(string message, params object[] args)
+		{
+			return BadRequest(LogEvent.Create(LogLevel.Error, message, args));
+		}
+
+		/// <summary>
+		/// Create a response indicating a problem with the request
+		/// </summary>
+		/// <param name="eventId">Id identifying the error</param>
+		/// <param name="message">Standard structured-logging format string</param>
+		/// <param name="args">Arguments for the format string</param>
+		[NonAction]
+		protected ActionResult BadRequest(EventId eventId, string message, params object[] args)
+		{
+			return BadRequest(LogEvent.Create(LogLevel.Error, eventId, message, args));
+		}
+
+		/// <summary>
+		/// Returns a 403 (forbidden) response with the given action
+		/// </summary>
+		[NonAction]
+		protected ActionResult Forbid(AclAction action)
+		{
+			return Forbid("User does not have {Action} permission", action);
+		}
+
+		/// <summary>
+		/// Returns a 403 (forbidden) response with the given action and object
+		/// </summary>
+		[NonAction]
+		protected ActionResult Forbid(AclAction action, AgentId agentId)
+		{
+			return Forbid(action, "agent {AgentId}", agentId);
+		}
+
+		/// <summary>
+		/// Returns a 403 (forbidden) response with the given action and object
+		/// </summary>
+		[NonAction]
+		protected ActionResult Forbid(AclAction action, AclScopeName scopeName)
+		{
+			return Forbid(action, "scope {ScopeName}", scopeName);
+		}
+
+		/// <summary>
+		/// Returns a 403 (forbidden) response with the given action and object
+		/// </summary>
+		[NonAction]
+		protected ActionResult Forbid(AclAction action, ClusterId clusterId)
+		{
+			return Forbid(action, "cluster {ClusterId}", clusterId);
+		}
+
+		/// <summary>
+		/// Returns a 403 (forbidden) response with the given action and object
+		/// </summary>
+		[NonAction]
+		protected ActionResult Forbid(AclAction action, JobId jobId)
+		{
+			return Forbid(action, "job {JobId}", jobId);
+		}
+
+		/// <summary>
+		/// Returns a 403 (forbidden) response with the given action and object
+		/// </summary>
+		[NonAction]
+		protected ActionResult Forbid(AclAction action, NamespaceId namespaceId)
+		{
+			return Forbid(action, "namespace {NamespaceId}", namespaceId);
+		}
+
+		/// <summary>
+		/// Returns a 403 (forbidden) response with the given action and object
+		/// </summary>
+		[NonAction]
+		protected ActionResult Forbid(AclAction action, ProjectId projectId)
+		{
+			return Forbid(action, "project {ProjectId}", projectId);
+		}
+
+		/// <summary>
+		/// Returns a 403 (forbidden) response with the given action and object
+		/// </summary>
+		[NonAction]
+		protected ActionResult Forbid(AclAction action, SecretId secretId)
+		{
+			return Forbid(action, "secret {SecretId}", secretId);
+		}
+
+		/// <summary>
+		/// Returns a 403 (forbidden) response with the given action and object
+		/// </summary>
+		[NonAction]
+		protected ActionResult Forbid(AclAction action, StreamId streamId)
+		{
+			return Forbid(action, "stream {StreamId}", streamId);
+		}
+
+		/// <summary>
+		/// Returns a 403 (forbidden) response with the given action and object
+		/// </summary>
+		[NonAction]
+		protected ActionResult Forbid(AclAction action, ToolId toolId)
+		{
+			return Forbid(action, "tool {ToolId}", toolId);
+		}
+
+		/// <summary>
+		/// Returns a 403 (forbidden) response with the given action and object
+		/// </summary>
+		[NonAction]
+		protected ActionResult Forbid(AclAction action, TemplateId templateId)
+		{
+			return Forbid(action, "template {TemplateId}", templateId);
+		}
+
+		/// <summary>
+		/// Returns a 403 (forbidden) response with the given action and object
+		/// </summary>
+		[NonAction]
+		private ActionResult Forbid(AclAction action, string objectMessage, object obj)
+		{
+			return Forbid($"User does not have {{Action}} permission for {objectMessage}", action, obj);
+		}
+
+		/// <summary>
+		/// Returns a 403 response with the given log event
+		/// </summary>
+		[NonAction]
+		protected ActionResult Forbid(string message, params object[] args)
+		{
+			return StatusCode(StatusCodes.Status403Forbidden, LogEvent.Create(LogLevel.Error, message, args));
+		}
+
+		/// <summary>
+		/// Returns a 404 response for the given object
+		/// </summary>
+		[NonAction]
+		protected ActionResult NotFound(AccountId accountId)
+		{
+			return NotFound("Account {AccountId} not found", accountId);
+		}
+
+		/// <summary>
+		/// Returns a 404 response for the given object
+		/// </summary>
+		[NonAction]
+		protected ActionResult NotFound(AgentId agentId)
+		{
+			return NotFound("Agent {AgentId} not found", agentId);
+		}
+
+		/// <summary>
+		/// Returns a 404 response for the given object
+		/// </summary>
+		[NonAction]
+		protected ActionResult NotFound(AgentId agentId, LeaseId leaseId)
+		{
+			return NotFound("Lease {LeaseId} not found for agent {AgentId}", leaseId, agentId);
+		}
+
+		/// <summary>
+		/// Returns a 404 response for the given object
+		/// </summary>
+		[NonAction]
+		protected ActionResult NotFound(ArtifactId artifactId)
+		{
+			return NotFound("Artifact {ArtifactId} not found", artifactId);
+		}
+
+		/// <summary>
+		/// Returns a 404 response for the given object
+		/// </summary>
+		[NonAction]
+		protected ActionResult NotFound(ClusterId clusterId)
+		{
+			return NotFound("Cluster {ClusterId} not found", clusterId);
+		}
+
+		/// <summary>
+		/// Returns a 404 response for the given object
+		/// </summary>
+		[NonAction]
+		protected ActionResult NotFound(JobId jobId)
+		{
+			return NotFound("Job {JobId} not found", jobId);
+		}
+
+		/// <summary>
+		/// Returns a 404 response for the given object
+		/// </summary>
+		[NonAction]
+		protected ActionResult NotFound(JobId jobId, int groupIdx)
+		{
+			return NotFound("Group {GroupIdx} on job {JobId} not found", groupIdx, jobId);
+		}
+
+		/// <summary>
+		/// Returns a 404 response for the given object
+		/// </summary>
+		[NonAction]
+		protected ActionResult NotFound(JobId jobId, int groupIdx, int nodeIdx)
+		{
+			return NotFound("Node {NodeIdx} not found on job {JobId} group {GroupIdx}", nodeIdx, jobId, groupIdx);
+		}
+
+		/// <summary>
+		/// Returns a 404 response for the given object
+		/// </summary>
+		[NonAction]
+		protected ActionResult NotFound(JobId jobId, JobStepBatchId batchId)
+		{
+			return NotFound("Batch {BatchId} not found on job {JobId}", batchId, jobId);
+		}
+
+		/// <summary>
+		/// Returns a 404 response for the given object
+		/// </summary>
+		[NonAction]
+		protected ActionResult NotFound(JobId jobId, JobStepBatchId batchId, JobStepId stepId)
+		{
+			return NotFound("Step {StepId} not found on job {JobId} batch {BatchId}", stepId, jobId, batchId);
+		}
+
+		/// <summary>
+		/// Returns a 404 response for the given object
+		/// </summary>
+		[NonAction]
+		protected ActionResult NotFound(NamespaceId namespaceId)
+		{
+			return NotFound("Namespace {NamespaceId} not found", namespaceId);
+		}
+
+		/// <summary>
+		/// Returns a 404 response for the given object
+		/// </summary>
+		[NonAction]
+		protected ActionResult NotFound(NamespaceId namespaceId, RefName refName)
+		{
+			return NotFound("Ref {RefName} in namespace {NamespaceId} not found", refName, namespaceId);
+		}
+
+		/// <summary>
+		/// Returns a 404 response for the given object
+		/// </summary>
+		[NonAction]
+		protected ActionResult NotFound(BisectTaskId bisectTaskId)
+		{
+			return NotFound("Bisect task {BisectTaskId} not found", bisectTaskId);
+		}
+
+		/// <summary>
+		/// Returns a 404 response for the given object
+		/// </summary>
+		[NonAction]
+		protected ActionResult NotFound(SecretId secretId)
+		{
+			return NotFound("Secret {SecretId} not found", secretId);
+		}
+
+		/// <summary>
+		/// Returns a 404 response for the given object
+		/// </summary>
+		[NonAction]
+		protected ActionResult NotFound(StreamId streamId)
+		{
+			return NotFound("Stream {StreamId} not found", streamId);
+		}
+
+		/// <summary>
+		/// Returns a 404 response for the given object
+		/// </summary>
+		[NonAction]
+		protected ActionResult NotFound(StreamId streamId, TemplateId templateId)
+		{
+			return NotFound("Template {TemplateId} not found on stream {StreamId}", templateId, streamId);
+		}
+
+		/// <summary>
+		/// Returns a 404 response for the given object
+		/// </summary>
+		[NonAction]
+		protected ActionResult NotFound(TelemetryStoreId telemetryStoreId)
+		{
+			return NotFound("Telemetry store {TelemetryStoreId} not found", telemetryStoreId);
+		}
+
+		/// <summary>
+		/// Returns a 404 response for the given object
+		/// </summary>
+		[NonAction]
+		protected ActionResult NotFound(ToolId toolId)
+		{
+			return NotFound("Tool {ToolId} not found", toolId);
+		}
+
+		/// <summary>
+		/// Returns a 404 response for the given object
+		/// </summary>
+		[NonAction]
+		protected ActionResult NotFound(ToolId toolId, ToolDeploymentId deploymentId)
+		{
+			return NotFound("Deployment {DeploymentId} not found on tool {ToolId}", deploymentId, toolId);
+		}
+
+		/// <summary>
+		/// Returns a 404 response for the given object
+		/// </summary>
+		[NonAction]
+		protected ActionResult NotFound(string message, params object[] args)
+		{
+			return NotFound(LogEvent.Create(LogLevel.Error, message, args));
+		}
+	}
+}
diff --git a/Engine/Source/Programs/Horde/HordeServer.Shared/Utilities/MongoBufferedWriter.cs b/Engine/Source/Programs/Horde/HordeServer.Shared/Utilities/MongoBufferedWriter.cs
new file mode 100644
--- /dev/null
+++ b/Engine/Source/Programs/Horde/HordeServer.Shared/Utilities/MongoBufferedWriter.cs
@@ -0,0 +1,130 @@
+// Copyright Epic Games, Inc. All Rights Reserved.
+
+using System;
+using System.Collections.Concurrent;
+using System.Collections.Generic;
+using System.Threading;
+using System.Threading.Tasks;
+using EpicGames.Core;
+using MongoDB.Driver;
+using Microsoft.Extensions.Logging;
+
+namespace HordeServer.Utilities
+{
+	/// <summary>
+	/// Utility class for buffering document writes to a Mongo collection
+	/// </summary>
+	public sealed class MongoBufferedWriter<TDocument> : IAsyncDisposable
+		where TDocument : class
+	{
+		readonly IMongoCollection<TDocument> _collection;
+		readonly int _flushCount;
+		readonly TimeSpan _flushTime;
+		readonly ConcurrentQueue<TDocument> _queue = new ConcurrentQueue<TDocument>();
+		readonly BackgroundTask _backgroundTask;
+		readonly AsyncEvent _newDataEvent = new AsyncEvent();
+		readonly AsyncEvent _flushEvent = new AsyncEvent();
+		readonly ILogger _logger;
+
+		public MongoBufferedWriter(IMongoCollection<TDocument> collection, ILogger logger)
+			: this(collection, 50, TimeSpan.FromSeconds(5.0), logger)
+		{
+		}
+
+		/// <summary>
+		/// Constructor
+		/// </summary>
+		public MongoBufferedWriter(IMongoCollection<TDocument> collection, int flushCount, TimeSpan flushTime, ILogger logger)
+		{
+			_collection = collection;
+			_flushCount = flushCount;
+			_flushTime = flushTime;
+			_backgroundTask = new BackgroundTask(BackgroundFlushAsync);
+			_logger = logger;
+		}
+
+		/// <inheritdoc/>
+		public async ValueTask DisposeAsync()
+		{
+			await _backgroundTask.DisposeAsync();
+		}
+
+		/// <summary>
+		/// Start the background task to periodically flush data to the DB
+		/// </summary>
+		public ValueTask StartAsync()
+		{
+			_backgroundTask.Start();
+			return new ValueTask();
+		}
+
+		/// <summary>
+		/// Stops the background task
+		/// </summary>
+		public async ValueTask StopAsync(CancellationToken cancellationToken)
+		{
+			await _backgroundTask.StopAsync(cancellationToken);
+			await FlushAsync(cancellationToken);
+		}
+
+		// Flushes the sink in the background
+		async Task BackgroundFlushAsync(CancellationToken cancellationToken)
+		{
+			Task newDataTask = _newDataEvent.Task;
+			Task flushTask = _flushEvent.Task;
+
+			while (!cancellationToken.IsCancellationRequested)
+			{
+				try
+				{
+					await newDataTask.WaitAsync(cancellationToken);
+					await Task.WhenAny(flushTask, Task.Delay(_flushTime, cancellationToken));
+
+					newDataTask = _newDataEvent.Task;
+					flushTask = _flushEvent.Task;
+
+					await FlushAsync(cancellationToken);
+				}
+				catch (OperationCanceledException)
+				{
+					break;
+				}
+				catch (Exception ex)
+				{
+					_logger.LogError(ex, "Exception in buffered writer: {Message}", ex.Message);
+					await Task.Delay(TimeSpan.FromSeconds(30.0), cancellationToken);
+				}
+			}
+		}
+
+		/// <inheritdoc/>
+		public async ValueTask FlushAsync(CancellationToken cancellationToken)
+		{
+			// Copy all the event documents from the queue
+			List<TDocument> documents = new List<TDocument>(_queue.Count);
+			while (_queue.TryDequeue(out TDocument? document))
+			{
+				documents.Add(document);
+			}
+
+			// Insert them into the database
+			if (documents.Count > 0)
+			{
+				_logger.LogDebug("Writing {NumEvents} new telemetry events to {CollectionName}.", documents.Count, _collection.CollectionNamespace.CollectionName);
+				await _collection.InsertManyAsync(documents, cancellationToken: cancellationToken);
+			}
+		}
+
+		/// <inheritdoc/>
+		public void Write(TDocument document)
+		{
+			_queue.Enqueue(document);
+			_newDataEvent.Set();
+
+			if (_queue.Count > _flushCount)
+			{
+				_flushEvent.Set();
+			}
+		}
+	}
+}
diff --git a/Engine/Source/Programs/Horde/HordeServer.Shared/Utilities/MongoExtensions.cs b/Engine/Source/Programs/Horde/HordeServer.Shared/Utilities/MongoExtensions.cs
new file mode 100644
--- /dev/null
+++ b/Engine/Source/Programs/Horde/HordeServer.Shared/Utilities/MongoExtensions.cs
@@ -0,0 +1,355 @@
+// Copyright Epic Games, Inc. All Rights Reserved.
+
+using System;
+using System.Collections.Generic;
+using System.Linq;
+using System.Linq.Expressions;
+using System.Runtime.CompilerServices;
+using System.Text.Json;
+using System.Threading;
+using System.Threading.Tasks;
+using MongoDB.Bson;
+using MongoDB.Bson.Serialization;
+using MongoDB.Driver;
+
+namespace HordeServer.Utilities
+{
+	/// <summary>
+	/// Extension methods for MongoDB
+	/// </summary>
+	public static class MongoExtensions
+	{
+		/// <summary>
+		/// Maps a constructor for a type into a classmap
+		/// </summary>
+		public static void MapConstructor<T>(this BsonClassMap classMap, Expression<Func<T>> generator, params string[] argumentNames)
+		{
+			NewExpression newExpr = (NewExpression)generator.Body;
+			classMap.MapConstructor(newExpr.Constructor, argumentNames);
+		}
+
+		/// <summary>
+		/// Rounds a time value to its BSON equivalent (ie. milliseconds since Unix Epoch).
+		/// </summary>
+		/// <param name="time"></param>
+		/// <returns></returns>
+		public static DateTime RoundToBsonDateTime(DateTime time)
+		{
+			return BsonUtils.ToDateTimeFromMillisecondsSinceEpoch(BsonUtils.ToMillisecondsSinceEpoch(time));
+		}
+
+		/// <summary>
+		/// Renders a filter definition to a document using the default serializer registry
+		/// </summary>
+		public static BsonDocument Render<T>(this FilterDefinition<T> filter)
+		{
+			IBsonSerializerRegistry serializerRegistry = BsonSerializer.SerializerRegistry;
+			IBsonSerializer<T> documentSerializer = serializerRegistry.GetSerializer<T>();
+			return filter.Render(documentSerializer, serializerRegistry, MongoDB.Driver.Linq.LinqProvider.V2);
+		}
+
+		/// <summary>
+		/// Renders a filter definition to a document using the default serializer registry
+		/// </summary>
+		public static BsonValue Render<T>(this UpdateDefinition<T> update)
+		{
+			IBsonSerializerRegistry serializerRegistry = BsonSerializer.SerializerRegistry;
+			IBsonSerializer<T> documentSerializer = serializerRegistry.GetSerializer<T>();
+			return update.Render(documentSerializer, serializerRegistry, MongoDB.Driver.Linq.LinqProvider.V2);
+		}
+
+		/// <summary>
+		/// Executes a query with a particular index hint
+		/// </summary>
+		/// <typeparam name="TDoc"></typeparam>
+		/// <typeparam name="TResult"></typeparam>
+		/// <param name="collection"></param>
+		/// <param name="filter"></param>
+		/// <param name="indexHint"></param>
+		/// <param name="processAsync"></param>
+		/// <returns></returns>
+		public static async Task<TResult> FindWithHintAsync<TDoc, TResult>(this IMongoCollection<TDoc> collection, FilterDefinition<TDoc> filter, string? indexHint, Func<IFindFluent<TDoc, TDoc>, Task<TResult>> processAsync)
+		{
+			FindOptions? findOptions = null;
+			if (indexHint != null)
+			{
+				findOptions = new FindOptions { Hint = new BsonString(indexHint) };
+			}
+			try
+			{
+				return await processAsync(collection.Find(filter, findOptions));
+			}
+			catch (MongoCommandException ex) when (indexHint != null && ex.Code == 2)
+			{
+				return await processAsync(collection.Find(filter));
+			}
+		}
+
+		/// <summary>
+		/// Filters the documents returned from a search
+		/// </summary>
+		/// <param name="query">The query to filter</param>
+		/// <param name="index">Index of the first document to return</param>
+		/// <param name="count">Number of documents to return</param>
+		/// <returns>New query</returns>
+		public static IFindFluent<TDocument, TProjection> Range<TDocument, TProjection>(this IFindFluent<TDocument, TProjection> query, int? index, int? count)
+		{
+			if (index != null && index.Value != 0)
+			{
+				query = query.Skip(index.Value);
+			}
+			if (count != null)
+			{
+				query = query.Limit(count.Value);
+			}
+			return query;
+		}
+
+		/// <summary>
+		/// Filters the documents returned from a search
+		/// </summary>
+		/// <param name="query">The query to filter</param>
+		/// <param name="index">Index of the first document to return</param>
+		/// <param name="count">Number of documents to return</param>
+		/// <returns>New query</returns>
+		public static IAggregateFluent<TDocument> Range<TDocument>(this IAggregateFluent<TDocument> query, int? index, int? count)
+		{
+			if (index != null && index.Value != 0)
+			{
+				query = query.Skip(index.Value);
+			}
+			if (count != null)
+			{
+				query = query.Limit(count.Value);
+			}
+			return query;
+		}
+
+		/// <summary>
+		/// Attempts to insert a document into a collection, handling the error case that a document with the given key already exists
+		/// </summary>
+		/// <typeparam name="TDocument"></typeparam>
+		/// <param name="collection">Collection to insert into</param>
+		/// <param name="newDocument">The document to insert</param>
+		/// <param name="cancellationToken">Cancellation token for the operation</param>
+		/// <returns>True if the document was inserted, false if it already exists</returns>
+		public static async Task<bool> InsertOneIgnoreDuplicatesAsync<TDocument>(this IMongoCollection<TDocument> collection, TDocument newDocument, CancellationToken cancellationToken = default)
+		{
+			try
+			{
+				await collection.InsertOneAsync(newDocument, cancellationToken: cancellationToken);
+				return true;
+			}
+			catch (MongoWriteException ex)
+			{
+				if (ex.WriteError.Category == ServerErrorCategory.DuplicateKey)
+				{
+					return false;
+				}
+				else
+				{
+					throw;
+				}
+			}
+		}
+
+		/// <summary>
+		/// Attempts to insert a document into a collection, handling the error case that a document with the given key already exists
+		/// </summary>
+		/// <typeparam name="TDocument"></typeparam>
+		/// <param name="collection">Collection to insert into</param>
+		/// <param name="newDocuments">The document to insert</param>
+		/// <param name="cancellationToken">Cancellation token for the operation</param>
+		public static async Task InsertManyIgnoreDuplicatesAsync<TDocument>(this IMongoCollection<TDocument> collection, List<TDocument> newDocuments, CancellationToken cancellationToken = default)
+		{
+			try
+			{
+				if (newDocuments.Count > 0)
+				{
+					await collection.InsertManyAsync(newDocuments, new InsertManyOptions { IsOrdered = false }, cancellationToken);
+				}
+			}
+			catch (MongoWriteException ex)
+			{
+				if (ex.WriteError.Category != ServerErrorCategory.DuplicateKey)
+				{
+					throw;
+				}
+			}
+			catch (MongoBulkWriteException ex)
+			{
+				if (ex.WriteErrors.Any(x => x.Category != ServerErrorCategory.DuplicateKey))
+				{
+					throw;
+				}
+			}
+		}
+
+		/// <summary>
+		/// Sets a field to a value, or unsets it if the value is null
+		/// </summary>
+		/// <typeparam name="TDocument">The document type</typeparam>
+		/// <typeparam name="TField">Type of the field to set</typeparam>
+		/// <param name="updateBuilder">Update builder</param>
+		/// <param name="field">Expression for the field to set</param>
+		/// <param name="value">New value to set</param>
+		/// <returns>Update defintiion</returns>
+		public static UpdateDefinition<TDocument> SetOrUnsetNull<TDocument, TField>(this UpdateDefinitionBuilder<TDocument> updateBuilder, Expression<Func<TDocument, TField?>> field, TField? value) where TField : struct
+		{
+			if (value.HasValue)
+			{
+				return updateBuilder.Set(field, value.Value);
+			}
+			else
+			{
+				return updateBuilder.Unset(new ExpressionFieldDefinition<TDocument, TField?>(field));
+			}
+		}
+
+		/// <summary>
+		/// Sets a field to a value, or unsets it if the value is null
+		/// </summary>
+		/// <typeparam name="TDocument">The document type</typeparam>
+		/// <typeparam name="TField">Type of the field to set</typeparam>
+		/// <param name="update">Update builder</param>
+		/// <param name="field">Expression for the field to set</param>
+		/// <param name="value">New value to set</param>
+		/// <returns>Update defintiion</returns>
+		public static UpdateDefinition<TDocument> SetOrUnsetNull<TDocument, TField>(this UpdateDefinition<TDocument> update, Expression<Func<TDocument, TField?>> field, TField? value) where TField : struct
+		{
+			if (value.HasValue)
+			{
+				return update.Set(field, value.Value);
+			}
+			else
+			{
+				return update.Unset(new ExpressionFieldDefinition<TDocument, TField?>(field));
+			}
+		}
+
+		/// <summary>
+		/// Sets a field to a value, or unsets it if the value is null
+		/// </summary>
+		/// <typeparam name="TDocument">The document type</typeparam>
+		/// <param name="update">Update builder</param>
+		/// <param name="field">Expression for the field to set</param>
+		/// <param name="value">New value to set</param>
+		/// <returns>Update defintiion</returns>
+		public static UpdateDefinition<TDocument> SetOrUnsetBool<TDocument>(this UpdateDefinition<TDocument> update, Expression<Func<TDocument, bool>> field, bool value)
+		{
+			if (value)
+			{
+				return update.Set(field, value);
+			}
+			else
+			{
+				return update.Unset(new ExpressionFieldDefinition<TDocument>(field));
+			}
+		}
+
+		/// <summary>
+		/// Sets a field to a value, or unsets it if the value is null
+		/// </summary>
+		/// <typeparam name="TDocument">The document type</typeparam>
+		/// <param name="updateBuilder">Update builder</param>
+		/// <param name="field">Expression for the field to set</param>
+		/// <param name="value">New value to set</param>
+		/// <returns>Update defintiion</returns>
+		public static UpdateDefinition<TDocument> SetOrUnsetBool<TDocument>(this UpdateDefinitionBuilder<TDocument> updateBuilder, Expression<Func<TDocument, bool>> field, bool value)
+		{
+			if (value)
+			{
+				return updateBuilder.Set(field, value);
+			}
+			else
+			{
+				return updateBuilder.Unset(new ExpressionFieldDefinition<TDocument>(field));
+			}
+		}
+
+		/// <summary>
+		/// Sets a field to a value, or unsets it if the value is null
+		/// </summary>
+		/// <typeparam name="TDocument">The document type</typeparam>
+		/// <typeparam name="TField">Type of the field to set</typeparam>
+		/// <param name="updateBuilder">Update builder</param>
+		/// <param name="field">Expression for the field to set</param>
+		/// <param name="value">New value to set</param>
+		/// <returns>Update defintiion</returns>
+		public static UpdateDefinition<TDocument> SetOrUnsetNullRef<TDocument, TField>(this UpdateDefinitionBuilder<TDocument> updateBuilder, Expression<Func<TDocument, TField?>> field, TField? value) where TField : class
+		{
+			if (value != null)
+			{
+				return updateBuilder.Set(field, value);
+			}
+			else
+			{
+				return updateBuilder.Unset(new ExpressionFieldDefinition<TDocument, TField?>(field));
+			}
+		}
+
+		/// <summary>
+		/// Creates a filter definition from a linq expression. This is not generally explicitly castable, so expose it as a FilterDefinitionBuilder method.
+		/// </summary>
+		/// <typeparam name="TDocument">The document type</typeparam>
+		/// <param name="filter">The filter builder</param>
+		/// <param name="expression">Expression to parse</param>
+		/// <returns>New filter definition</returns>
+#pragma warning disable IDE0060
+		public static FilterDefinition<TDocument> Expr<TDocument>(this FilterDefinitionBuilder<TDocument> filter, Expression<Func<TDocument, bool>> expression)
+		{
+			return expression;
+		}
+#pragma warning restore IDE0060
+
+		/// <summary>
+		/// Creates an async enumerator from the given query
+		/// </summary>
+		/// <typeparam name="T">Document type</typeparam>
+		/// <param name="source">Query source</param>
+		/// <param name="cancellationToken">Cancellation token for the operation</param>
+		/// <returns></returns>
+#pragma warning disable IDE1006 // Naming Styles
+#pragma warning disable VSTHRD200 // Use Async suffix
+		public static async IAsyncEnumerable<T> ToAsyncEnumerable<T>(this IAsyncCursorSource<T> source, [EnumeratorCancellation] CancellationToken cancellationToken)
+#pragma warning restore VSTHRD200 // Use Async suffix
+#pragma warning restore IDE1006 // Naming Styles
+		{
+			using (IAsyncCursor<T> cursor = await source.ToCursorAsync(cancellationToken))
+			{
+				while (await cursor.MoveNextAsync(cancellationToken))
+				{
+					foreach (T value in cursor.Current)
+					{
+						yield return value;
+						cancellationToken.ThrowIfCancellationRequested();
+					}
+				}
+			}
+		}
+
+		/// <summary>
+		/// Converts a JsonElement to a BsonValue
+		/// </summary>
+		/// <param name="element">The element to convert</param>
+		/// <returns>Bson value</returns>
+		public static BsonValue ToBsonValue(this JsonElement element)
+		{
+			switch (element.ValueKind)
+			{
+				case JsonValueKind.True:
+					return true;
+				case JsonValueKind.False:
+					return false;
+				case JsonValueKind.String:
+					return element.GetString();
+				case JsonValueKind.Array:
+					return new BsonArray(element.EnumerateArray().Select(x => x.ToBsonValue()));
+				case JsonValueKind.Object:
+					return new BsonDocument(element.EnumerateObject().Select(x => new BsonElement(x.Name, x.Value.ToBsonValue())));
+				default:
+					return BsonNull.Value;
+			}
+		}
+	}
+}
diff --git a/Engine/Source/Programs/Horde/HordeServer.Shared/Utilities/RedisExtensions.cs b/Engine/Source/Programs/Horde/HordeServer.Shared/Utilities/RedisExtensions.cs
new file mode 100644
--- /dev/null
+++ b/Engine/Source/Programs/Horde/HordeServer.Shared/Utilities/RedisExtensions.cs
@@ -0,0 +1,31 @@
+// Copyright Epic Games, Inc. All Rights Reserved.
+
+using System.Threading.Tasks;
+using StackExchange.Redis;
+
+namespace HordeServer.Utilities;
+
+/// <summary>
+/// Extensions for Redis classes
+/// </summary>
+public static class RedisExtensions
+{
+	/// <summary>
+	/// Await given tasks and swallow any task cancellation exceptions.
+	/// Useful to do after (failing) transaction to ensure no trailing tasks are left unawaited
+	/// </summary>
+	/// <param name="transaction"></param>
+	/// <param name="tasks">Tasks to await</param>
+	public static async Task WaitAndIgnoreCancellationsAsync(this ITransaction transaction, params Task[] tasks)
+	{
+		_ = transaction;
+		try
+		{
+			await Task.WhenAll(tasks);
+		}
+		catch (TaskCanceledException)
+		{
+			/* Ignore */
+		}
+	}
+}
\ No newline at end of file
diff --git a/Engine/Source/Programs/Horde/HordeServer.Shared/Utilities/SingletonDocument.cs b/Engine/Source/Programs/Horde/HordeServer.Shared/Utilities/SingletonDocument.cs
new file mode 100644
--- /dev/null
+++ b/Engine/Source/Programs/Horde/HordeServer.Shared/Utilities/SingletonDocument.cs
@@ -0,0 +1,174 @@
+// Copyright Epic Games, Inc. All Rights Reserved.
+
+using System;
+using System.Threading;
+using System.Threading.Tasks;
+using EpicGames.Horde;
+using HordeServer.Server;
+using MongoDB.Bson.Serialization.Attributes;
+
+namespace HordeServer.Utilities
+{
+	/// <summary>
+	/// Identifier for a pool
+	/// </summary>
+	/// <param name="Id">Id to construct from</param>
+	[StringIdConverter(typeof(SingletonIdConverter))]
+	public record struct SingletonId(StringId Id)
+	{
+		/// <summary>
+		/// Constructor
+		/// </summary>
+		public SingletonId(string id) : this(new StringId(id)) { }
+
+		/// <inheritdoc/>
+		public override string ToString() => Id.ToString();
+	}
+
+	/// <summary>
+	/// Converter to and from <see cref="StringId"/> instances.
+	/// </summary>
+	class SingletonIdConverter : StringIdConverter<SingletonId>
+	{
+		/// <inheritdoc/>
+		public override SingletonId FromStringId(StringId id) => new SingletonId(id);
+
+		/// <inheritdoc/>
+		public override StringId ToStringId(SingletonId value) => value.Id;
+	}
+
+	/// <summary>
+	/// Base class for singleton documents
+	/// </summary>
+	public abstract class SingletonBase
+	{
+		/// <summary>
+		/// Unique id for this singleton
+		/// </summary>
+		[BsonId]
+		public SingletonId Id { get; set; }
+
+		/// <summary>
+		/// The revision index of this document
+		/// </summary>
+		public int Revision { get; set; }
+
+		/// <summary>
+		/// Callback to allow the singleton to fix up itself after being read
+		/// </summary>
+		public virtual void PostLoad()
+		{
+		}
+	}
+
+	/// <summary>
+	/// Attribute specifying the unique id for a singleton document
+	/// </summary>
+	[AttributeUsage(AttributeTargets.Class)]
+	public sealed class SingletonDocumentAttribute : Attribute
+	{
+		/// <summary>
+		/// Unique id for the singleton document
+		/// </summary>
+		public string Id { get; }
+
+		/// <summary>
+		/// ObjectId for the legacy document
+		/// </summary>
+		public string? LegacyId { get; }
+
+		/// <summary>
+		/// Constructor
+		/// </summary>
+		/// <param name="id">Unique id for the singleton document. Should be a valid StringId.</param>
+		/// <param name="legacyId">Legacy identifier for the document</param>
+		public SingletonDocumentAttribute(string id, string? legacyId = null)
+		{
+			Id = id;
+			LegacyId = legacyId;
+		}
+	}
+
+	/// <summary>
+	/// Interface for the getting and setting the singleton
+	/// </summary>
+	/// <typeparam name="T">Type of document</typeparam>
+	public interface ISingletonDocument<T>
+	{
+		/// <summary>
+		/// Gets the current document
+		/// </summary>
+		/// <param name="cancellationToken">Cancellation token for the operation</param>
+		/// <returns>The current document</returns>
+		Task<T> GetAsync(CancellationToken cancellationToken);
+
+		/// <summary>
+		/// Attempts to update the document
+		/// </summary>
+		/// <param name="value">New state of the document</param>
+		/// <param name="cancellationToken">Cancellation token for the operation</param>
+		/// <returns>True if the document was updated, false otherwise</returns>
+		Task<bool> TryUpdateAsync(T value, CancellationToken cancellationToken);
+	}
+
+	/// <summary>
+	/// Concrete implementation of <see cref="ISingletonDocument{T}"/>
+	/// </summary>
+	/// <typeparam name="T">The document type</typeparam>
+	public class SingletonDocument<T> : ISingletonDocument<T> where T : SingletonBase, new()
+	{
+		/// <summary>
+		/// The database service instance
+		/// </summary>
+		readonly IMongoService _mongoService;
+
+		/// <summary>
+		/// Constructor
+		/// </summary>
+		/// <param name="mongoService">The database service instance</param>
+		public SingletonDocument(IMongoService mongoService)
+		{
+			_mongoService = mongoService;
+		}
+
+		/// <inheritdoc/>
+		public async Task<T> GetAsync(CancellationToken cancellationToken)
+		{
+			return await _mongoService.GetSingletonAsync<T>(cancellationToken);
+		}
+
+		/// <inheritdoc/>
+		public Task<bool> TryUpdateAsync(T value, CancellationToken cancellationToken)
+		{
+			return _mongoService.TryUpdateSingletonAsync<T>(value, cancellationToken);
+		}
+	}
+
+	/// <summary>
+	/// Extension methods for singletons
+	/// </summary>
+	public static class SingletonDocumentExtensions
+	{
+		/// <summary>
+		/// Update a singleton
+		/// </summary>
+		/// <typeparam name="T"></typeparam>
+		/// <param name="singleton"></param>
+		/// <param name="updateAction"></param>
+		/// <param name="cancellationToken">Cancellation token for the operation</param>
+		/// <returns></returns>
+		public static async Task<T> UpdateAsync<T>(this ISingletonDocument<T> singleton, Action<T> updateAction, CancellationToken cancellationToken) where T : SingletonBase, new()
+		{
+			for (; ; )
+			{
+				T value = await singleton.GetAsync(cancellationToken);
+				updateAction(value);
+
+				if (await singleton.TryUpdateAsync(value, cancellationToken))
+				{
+					return value;
+				}
+			}
+		}
+	}
+}
diff --git a/Engine/Source/Programs/Horde/HordeServer.Shared/Utilities/StringIdConverter.cs b/Engine/Source/Programs/Horde/HordeServer.Shared/Utilities/StringIdConverter.cs
new file mode 100644
--- /dev/null
+++ b/Engine/Source/Programs/Horde/HordeServer.Shared/Utilities/StringIdConverter.cs
@@ -0,0 +1,58 @@
+// Copyright Epic Games, Inc. All Rights Reserved.
+
+using System;
+using System.Reflection;
+using EpicGames.Core;
+using EpicGames.Horde;
+using MongoDB.Bson.Serialization;
+using MongoDB.Bson.Serialization.Serializers;
+using ProtoBuf;
+
+namespace HordeServer.Utilities
+{
+	/// <summary>
+	/// Class which serializes object id types to BSON
+	/// </summary>
+	sealed class StringIdBsonSerializer<TValue, TConverter> : SerializerBase<TValue> where TValue : struct where TConverter : StringIdConverter<TValue>, new()
+	{
+		readonly TConverter _converter = new TConverter();
+
+		/// <inheritdoc/>
+		public override TValue Deserialize(BsonDeserializationContext context, BsonDeserializationArgs args) => _converter.FromStringId(new StringId(new Utf8String(context.Reader.ReadString()), StringId.Validate.None));
+
+		/// <inheritdoc/>
+		public override void Serialize(BsonSerializationContext context, BsonSerializationArgs args, TValue value) => context.Writer.WriteString(_converter.ToStringId(value).ToString());
+	}
+
+	/// <summary>
+	/// Class which serializes object id types to BSON
+	/// </summary>
+	sealed class StringIdBsonSerializationProvider : BsonSerializationProviderBase
+	{
+		/// <inheritdoc/>
+		public override IBsonSerializer? GetSerializer(Type type, IBsonSerializerRegistry serializerRegistry)
+		{
+			StringIdConverterAttribute? attribute = type.GetCustomAttribute<StringIdConverterAttribute>();
+			if (attribute == null)
+			{
+				return null;
+			}
+			return (IBsonSerializer?)Activator.CreateInstance(typeof(StringIdBsonSerializer<,>).MakeGenericType(type, attribute.ConverterType));
+		}
+	}
+
+	/// <summary>
+	/// Surrogate type for serializing StringId types to ProtoBuf
+	/// </summary>
+	[ProtoContract]
+	struct StringIdProto<TValue, TConverter> where TValue : struct where TConverter : StringIdConverter<TValue>, new()
+	{
+		static readonly TConverter s_converter = new TConverter();
+
+		[ProtoMember(1)]
+		public string? Id { get; set; }
+
+		public static implicit operator TValue(StringIdProto<TValue, TConverter> source) => s_converter.FromStringId(new StringId(source.Id!));
+		public static implicit operator StringIdProto<TValue, TConverter>(TValue source) => new StringIdProto<TValue, TConverter> { Id = s_converter.ToStringId(source).ToString() };
+	}
+}
diff --git a/Engine/Source/Programs/Horde/HordeServer/Auditing/IAuditLog.cs b/Engine/Source/Programs/Horde/HordeServer/Auditing/IAuditLog.cs
deleted file mode 100644
diff --git a/Engine/Source/Programs/Horde/HordeServer/Configuration/ConfigContext.cs b/Engine/Source/Programs/Horde/HordeServer/Configuration/ConfigContext.cs
deleted file mode 100644
diff --git a/Engine/Source/Programs/Horde/HordeServer/Configuration/ConfigException.cs b/Engine/Source/Programs/Horde/HordeServer/Configuration/ConfigException.cs
deleted file mode 100644
diff --git a/Engine/Source/Programs/Horde/HordeServer/Configuration/ConfigMacro.cs b/Engine/Source/Programs/Horde/HordeServer/Configuration/ConfigMacro.cs
deleted file mode 100644
diff --git a/Engine/Source/Programs/Horde/HordeServer/Configuration/ConfigNode.cs b/Engine/Source/Programs/Horde/HordeServer/Configuration/ConfigNode.cs
deleted file mode 100644
diff --git a/Engine/Source/Programs/Horde/HordeServer/Configuration/ConfigObject.cs b/Engine/Source/Programs/Horde/HordeServer/Configuration/ConfigObject.cs
deleted file mode 100644
diff --git a/Engine/Source/Programs/Horde/HordeServer/Configuration/ConfigResource.cs b/Engine/Source/Programs/Horde/HordeServer/Configuration/ConfigResource.cs
deleted file mode 100644
diff --git a/Engine/Source/Programs/Horde/HordeServer/Configuration/IConfigFile.cs b/Engine/Source/Programs/Horde/HordeServer/Configuration/IConfigFile.cs
deleted file mode 100644
diff --git a/Engine/Source/Programs/Horde/HordeServer/Configuration/IConfigSource.cs b/Engine/Source/Programs/Horde/HordeServer/Configuration/IConfigSource.cs
deleted file mode 100644
diff --git a/Engine/Source/Programs/Horde/HordeServer/HordeServer.csproj b/Engine/Source/Programs/Horde/HordeServer/HordeServer.csproj
--- a/Engine/Source/Programs/Horde/HordeServer/HordeServer.csproj
+++ b/Engine/Source/Programs/Horde/HordeServer/HordeServer.csproj
@@ -1,4 +1,4 @@
-﻿<Project Sdk="Microsoft.NET.Sdk.Web">
+<Project Sdk="Microsoft.NET.Sdk.Web">
 
   <PropertyGroup>
     <TargetFramework>net8.0</TargetFramework>
@@ -92,7 +92,6 @@
 
   <ItemGroup>
     <None Remove="HordeServer.xml" />
-    <None Remove="HordeServer.xml" />
   </ItemGroup>
 
   <ItemGroup>
@@ -135,6 +134,7 @@
     <ProjectReference Include="..\..\Shared\EpicGames.Redis\EpicGames.Redis.csproj" />
     <ProjectReference Include="..\..\Shared\EpicGames.Serialization\EpicGames.Serialization.csproj" />
     <ProjectReference Include="..\..\Shared\EpicGames.Slack\EpicGames.Slack.csproj" />
+    <ProjectReference Include="..\HordeServer.Shared\HordeServer.Shared.csproj" />
   </ItemGroup>
 
   <ItemGroup>
diff --git a/Engine/Source/Programs/Horde/HordeServer/Server/IMongoService.cs b/Engine/Source/Programs/Horde/HordeServer/Server/IMongoService.cs
deleted file mode 100644
diff --git a/Engine/Source/Programs/Horde/HordeServer/Server/IRedisService.cs b/Engine/Source/Programs/Horde/HordeServer/Server/IRedisService.cs
deleted file mode 100644
diff --git a/Engine/Source/Programs/Horde/HordeServer/Server/MongoIndex.cs b/Engine/Source/Programs/Horde/HordeServer/Server/MongoIndex.cs
deleted file mode 100644
diff --git a/Engine/Source/Programs/Horde/HordeServer/Users/IAvatar.cs b/Engine/Source/Programs/Horde/HordeServer/Users/IAvatar.cs
deleted file mode 100644
diff --git a/Engine/Source/Programs/Horde/HordeServer/Users/IAvatarService.cs b/Engine/Source/Programs/Horde/HordeServer/Users/IAvatarService.cs
deleted file mode 100644
diff --git a/Engine/Source/Programs/Horde/HordeServer/Users/IUser.cs b/Engine/Source/Programs/Horde/HordeServer/Users/IUser.cs
deleted file mode 100644
diff --git a/Engine/Source/Programs/Horde/HordeServer/Users/IUserClaim.cs b/Engine/Source/Programs/Horde/HordeServer/Users/IUserClaim.cs
deleted file mode 100644
diff --git a/Engine/Source/Programs/Horde/HordeServer/Users/IUserCollection.cs b/Engine/Source/Programs/Horde/HordeServer/Users/IUserCollection.cs
deleted file mode 100644
diff --git a/Engine/Source/Programs/Horde/HordeServer/Utilities/AsyncCachedValue.cs b/Engine/Source/Programs/Horde/HordeServer/Utilities/AsyncCachedValue.cs
deleted file mode 100644
diff --git a/Engine/Source/Programs/Horde/HordeServer/Utilities/BinaryIdConverter.cs b/Engine/Source/Programs/Horde/HordeServer/Utilities/BinaryIdConverter.cs
deleted file mode 100644
diff --git a/Engine/Source/Programs/Horde/HordeServer/Utilities/BinaryIdUtils.cs b/Engine/Source/Programs/Horde/HordeServer/Utilities/BinaryIdUtils.cs
deleted file mode 100644
diff --git a/Engine/Source/Programs/Horde/HordeServer/Utilities/HordeClaimTypes.cs b/Engine/Source/Programs/Horde/HordeServer/Utilities/HordeClaimTypes.cs
deleted file mode 100644
diff --git a/Engine/Source/Programs/Horde/HordeServer/Utilities/HordeControllerBase.cs b/Engine/Source/Programs/Horde/HordeServer/Utilities/HordeControllerBase.cs
deleted file mode 100644
diff --git a/Engine/Source/Programs/Horde/HordeServer/Utilities/MongoBufferedWriter.cs b/Engine/Source/Programs/Horde/HordeServer/Utilities/MongoBufferedWriter.cs
deleted file mode 100644
diff --git a/Engine/Source/Programs/Horde/HordeServer/Utilities/MongoExtensions.cs b/Engine/Source/Programs/Horde/HordeServer/Utilities/MongoExtensions.cs
deleted file mode 100644
diff --git a/Engine/Source/Programs/Horde/HordeServer/Utilities/RedisExtensions.cs b/Engine/Source/Programs/Horde/HordeServer/Utilities/RedisExtensions.cs
deleted file mode 100644
diff --git a/Engine/Source/Programs/Horde/HordeServer/Utilities/SingletonDocument.cs b/Engine/Source/Programs/Horde/HordeServer/Utilities/SingletonDocument.cs
deleted file mode 100644
diff --git a/Engine/Source/Programs/Horde/HordeServer/Utilities/StringIdConverter.cs b/Engine/Source/Programs/Horde/HordeServer/Utilities/StringIdConverter.cs
deleted file mode 100644
