Horde: Fix some compile errors.

#rnx

diff --git a/Engine/Source/Programs/Horde/HordeServer.Shared/Auditing/IAuditLog.cs b/Engine/Source/Programs/Horde/HordeServer.Shared/Auditing/IAuditLog.cs
--- a/Engine/Source/Programs/Horde/HordeServer.Shared/Auditing/IAuditLog.cs
+++ b/Engine/Source/Programs/Horde/HordeServer.Shared/Auditing/IAuditLog.cs
@@ -117,7 +117,7 @@ namespace HordeServer.Auditing
 	/// <summary>
 	/// Extension methods for audit log channels
 	/// </summary>
-	static class AuditLogExtensions
+	public static class AuditLogExtensions
 	{
 		/// <summary>
 		/// Retrieve historical information about a specific agent
diff --git a/Engine/Source/Programs/Horde/HordeServer.Shared/Server/MongoIndex.cs b/Engine/Source/Programs/Horde/HordeServer.Shared/Server/MongoIndex.cs
--- a/Engine/Source/Programs/Horde/HordeServer.Shared/Server/MongoIndex.cs
+++ b/Engine/Source/Programs/Horde/HordeServer.Shared/Server/MongoIndex.cs
@@ -149,7 +149,7 @@ namespace HordeServer.Server
 	/// <summary>
 	/// Extension methods for <see cref="MongoIndex{T}"/>
 	/// </summary>
-	static class MongoIndexExtensions
+	public static class MongoIndexExtensions
 	{
 		public static void Add<T>(this List<MongoIndex<T>> list, Func<IndexKeysDefinitionBuilder<T>, IndexKeysDefinition<T>> keyFunc, bool unique = false, bool sparse = false)
 		{
diff --git a/Engine/Source/Programs/Horde/HordeServer.Shared/Users/IUserCollection.cs b/Engine/Source/Programs/Horde/HordeServer.Shared/Users/IUserCollection.cs
--- a/Engine/Source/Programs/Horde/HordeServer.Shared/Users/IUserCollection.cs
+++ b/Engine/Source/Programs/Horde/HordeServer.Shared/Users/IUserCollection.cs
@@ -115,7 +115,7 @@ namespace HordeServer.Users
 	/// <summary>
 	/// Extension methods for the user collect
 	/// </summary>
-	static class UserCollectionExtensions
+	public static class UserCollectionExtensions
 	{
 		/// <summary>
 		/// Gets a particular user info from the collection
diff --git a/Engine/Source/Programs/Horde/HordeServer.Shared/Utilities/BinaryIdConverter.cs b/Engine/Source/Programs/Horde/HordeServer.Shared/Utilities/BinaryIdConverter.cs
--- a/Engine/Source/Programs/Horde/HordeServer.Shared/Utilities/BinaryIdConverter.cs
+++ b/Engine/Source/Programs/Horde/HordeServer.Shared/Utilities/BinaryIdConverter.cs
@@ -11,7 +11,7 @@ namespace HordeServer.Utilities
 	/// <summary>
 	/// Class which serializes object id types to BSON
 	/// </summary>
-	sealed class BinaryIdBsonSerializer<TValue, TConverter> : SerializerBase<TValue> where TValue : struct where TConverter : BinaryIdConverter<TValue>, new()
+	public sealed class BinaryIdBsonSerializer<TValue, TConverter> : SerializerBase<TValue> where TValue : struct where TConverter : BinaryIdConverter<TValue>, new()
 	{
 		readonly TConverter _converter = new TConverter();
 
@@ -25,7 +25,7 @@ namespace HordeServer.Utilities
 	/// <summary>
 	/// Class which serializes object id types to BSON
 	/// </summary>
-	sealed class BinaryIdBsonSerializationProvider : BsonSerializationProviderBase
+	public sealed class BinaryIdBsonSerializationProvider : BsonSerializationProviderBase
 	{
 		/// <inheritdoc/>
 		public override IBsonSerializer? GetSerializer(Type type, IBsonSerializerRegistry serializerRegistry)
diff --git a/Engine/Source/Programs/Horde/HordeServer.Shared/Utilities/BinaryIdUtils.cs b/Engine/Source/Programs/Horde/HordeServer.Shared/Utilities/BinaryIdUtils.cs
--- a/Engine/Source/Programs/Horde/HordeServer.Shared/Utilities/BinaryIdUtils.cs
+++ b/Engine/Source/Programs/Horde/HordeServer.Shared/Utilities/BinaryIdUtils.cs
@@ -8,7 +8,7 @@ namespace HordeServer.Utilities
 	/// <summary>
 	/// Utility methods for <see cref="BinaryId"/>
 	/// </summary>
-	static class BinaryIdUtils
+	public static class BinaryIdUtils
 	{
 		/// <summary>
 		/// Creates a new BinaryId
diff --git a/Engine/Source/Programs/Horde/HordeServer.Shared/Utilities/StringIdConverter.cs b/Engine/Source/Programs/Horde/HordeServer.Shared/Utilities/StringIdConverter.cs
--- a/Engine/Source/Programs/Horde/HordeServer.Shared/Utilities/StringIdConverter.cs
+++ b/Engine/Source/Programs/Horde/HordeServer.Shared/Utilities/StringIdConverter.cs
@@ -13,7 +13,7 @@ namespace HordeServer.Utilities
 	/// <summary>
 	/// Class which serializes object id types to BSON
 	/// </summary>
-	sealed class StringIdBsonSerializer<TValue, TConverter> : SerializerBase<TValue> where TValue : struct where TConverter : StringIdConverter<TValue>, new()
+	public sealed class StringIdBsonSerializer<TValue, TConverter> : SerializerBase<TValue> where TValue : struct where TConverter : StringIdConverter<TValue>, new()
 	{
 		readonly TConverter _converter = new TConverter();
 
@@ -27,7 +27,7 @@ namespace HordeServer.Utilities
 	/// <summary>
 	/// Class which serializes object id types to BSON
 	/// </summary>
-	sealed class StringIdBsonSerializationProvider : BsonSerializationProviderBase
+	public sealed class StringIdBsonSerializationProvider : BsonSerializationProviderBase
 	{
 		/// <inheritdoc/>
 		public override IBsonSerializer? GetSerializer(Type type, IBsonSerializerRegistry serializerRegistry)
@@ -45,7 +45,7 @@ namespace HordeServer.Utilities
 	/// Surrogate type for serializing StringId types to ProtoBuf
 	/// </summary>
 	[ProtoContract]
-	struct StringIdProto<TValue, TConverter> where TValue : struct where TConverter : StringIdConverter<TValue>, new()
+	public struct StringIdProto<TValue, TConverter> where TValue : struct where TConverter : StringIdConverter<TValue>, new()
 	{
 		static readonly TConverter s_converter = new TConverter();
 
