Horde: Moving more types into HordeServer.Shared.

#rnx

diff --git a/Engine/Source/Programs/Horde/HordeServer.Shared/Acls/AclConfig.cs b/Engine/Source/Programs/Horde/HordeServer.Shared/Acls/AclConfig.cs
new file mode 100644
--- /dev/null
+++ b/Engine/Source/Programs/Horde/HordeServer.Shared/Acls/AclConfig.cs
@@ -0,0 +1,417 @@
+// Copyright Epic Games, Inc. All Rights Reserved.
+
+using System;
+using System.Collections.Generic;
+using System.ComponentModel.DataAnnotations;
+using System.Linq;
+using System.Security.Claims;
+using System.Text.Json.Serialization;
+using EpicGames.Horde.Acls;
+using HordeServer.Server;
+using HordeServer.Utilities;
+
+#pragma warning disable CA2227 // Change 'X' to be read-only by removing the property setter
+
+namespace HordeServer.Acls
+{
+	/// <summary>
+	/// Parameters to update an ACL
+	/// </summary>
+	public class AclConfig
+	{
+		/// <summary>
+		/// The parent scope object
+		/// </summary>
+		[JsonIgnore]
+		public AclConfig? Parent { get; set; }
+
+		/// <summary>
+		/// Name of this scope
+		/// </summary>
+		[JsonIgnore]
+		public AclScopeName ScopeName { get; set; }
+
+		/// <summary>
+		/// ACLs which are parented to this
+		/// </summary>
+		[JsonIgnore]
+		public List<AclConfig>? Children { get; set; }
+
+		/// <summary>
+		/// Entries to replace the existing ACL
+		/// </summary>
+		public List<AclEntryConfig>? Entries { get; set; }
+
+		/// <summary>
+		/// Defines profiles which allow grouping sets of actions into named collections
+		/// </summary>
+		public List<AclProfileConfig>? Profiles { get; set; }
+
+		/// <summary>
+		/// Whether to inherit permissions from the parent ACL
+		/// </summary>
+		public bool? Inherit { get; set; }
+
+		/// <summary>
+		/// List of exceptions to the inherited setting
+		/// </summary>
+		public List<AclAction>? Exceptions { get; set; }
+
+		IReadOnlyDictionary<AclProfileId, AclProfileConfig> _profileLookup = null!;
+
+		/// <summary>
+		/// Tests whether a user is authorized to perform the given actions
+		/// </summary>
+		/// <param name="action">Action that is being performed. This should be a single flag.</param>
+		/// <param name="user">The principal to authorize</param>
+		/// <returns>True/false if the action is allowed or denied, null if there is no specific setting for this user</returns>
+		public bool Authorize(AclAction action, ClaimsPrincipal user)
+		{
+			if (user.HasAdminClaim())
+			{
+				return true;
+			}
+
+			for (AclConfig? next = this; next != null; next = next.Parent)
+			{
+				bool? result = next.AuthorizeSingleScope(action, user);
+				if (result.HasValue)
+				{
+					return result.Value;
+				}
+			}
+
+			return false;
+		}
+
+		bool? AuthorizeSingleScope(AclAction action, ClaimsPrincipal user)
+		{
+			// Check if there's a specific entry for this action
+			if (Entries != null)
+			{
+				foreach (AclEntryConfig entry in Entries)
+				{
+					if (entry.ComputedActions.Contains(action) && user.HasClaim(entry.Claim.Type, entry.Claim.Value))
+					{
+						return true;
+					}
+				}
+			}
+
+			// Otherwise check if we're prevented from inheriting permissions
+			if (Inherit ?? true)
+			{
+				if (Exceptions != null && Exceptions.Contains(action))
+				{
+					return false;
+				}
+			}
+			else
+			{
+				if (Exceptions == null || !Exceptions.Contains(action))
+				{
+					return false;
+				}
+			}
+
+			// Otherwise allow to propagate up the hierarchy
+			return null;
+		}
+
+		/// <summary>
+		/// Called after the config file has been read
+		/// </summary>
+		public void PostLoad(AclConfig parentScope, string scopeNameSuffix)
+		{
+			PostLoad(parentScope, parentScope.ScopeName.Append(scopeNameSuffix));
+		}
+
+		/// <summary>
+		/// Called after the config file has been read
+		/// </summary>
+		public void PostLoad(AclConfig? parentScope, AclScopeName scopeName)
+		{
+			Parent = parentScope;
+			ScopeName = scopeName;
+			Children = null;
+
+			if (parentScope == null)
+			{
+				_profileLookup = new Dictionary<AclProfileId, AclProfileConfig>();
+			}
+			else
+			{
+				parentScope.Children ??= new List<AclConfig>();
+				parentScope.Children.Add(this);
+
+				_profileLookup = parentScope._profileLookup;
+			}
+
+			if (Profiles != null && Profiles.Count > 0)
+			{
+				Dictionary<AclProfileId, AclProfileConfig> newProfileLookup = new Dictionary<AclProfileId, AclProfileConfig>(_profileLookup);
+
+				foreach (AclProfileConfig profileConfig in Profiles)
+				{
+					newProfileLookup.Add(profileConfig.Id, profileConfig);
+				}
+				_profileLookup = newProfileLookup;
+
+				HashSet<AclProfileId> visited = new HashSet<AclProfileId>();
+				foreach (AclProfileConfig profileConfig in Profiles)
+				{
+					profileConfig.PostLoad(ScopeName, newProfileLookup, visited);
+				}
+			}
+
+			if (Entries != null)
+			{
+				foreach (AclEntryConfig entryConfig in Entries)
+				{
+					entryConfig.PostLoad(ScopeName, _profileLookup);
+				}
+			}
+		}
+
+		/// <summary>
+		/// Find all entitlements for a user
+		/// </summary>
+		public Dictionary<AclScopeName, HashSet<AclAction>> FindEntitlements(Predicate<AclClaimConfig> predicate)
+		{
+			Dictionary<AclScopeName, HashSet<AclAction>> scopeToActions = new Dictionary<AclScopeName, HashSet<AclAction>>();
+			FindEntitlements(predicate, scopeToActions);
+			return scopeToActions;
+		}
+
+		/// <summary>
+		/// Find all entitlements for a user
+		/// </summary>
+		public void FindEntitlements(Predicate<AclClaimConfig> predicate, Dictionary<AclScopeName, HashSet<AclAction>> scopeToActions)
+		{
+			if (Entries != null)
+			{
+				foreach (AclEntryConfig entry in Entries)
+				{
+					if (predicate(entry.Claim))
+					{
+						HashSet<AclAction>? actions;
+						if (!scopeToActions.TryGetValue(ScopeName, out actions))
+						{
+							actions = new HashSet<AclAction>();
+							scopeToActions.Add(ScopeName, actions);
+						}
+						actions.UnionWith(entry.ComputedActions);
+					}
+				}
+			}
+			if (Children != null)
+			{
+				foreach (AclConfig childAclConfig in Children)
+				{
+					childAclConfig.FindEntitlements(predicate, scopeToActions);
+				}
+			}
+		}
+	}
+
+	/// <summary>
+	/// Individual entry in an ACL
+	/// </summary>
+	public class AclEntryConfig
+	{
+		/// <summary>
+		/// Name of the user or group
+		/// </summary>
+		[Required]
+		public AclClaimConfig Claim { get; set; }
+
+		/// <summary>
+		/// Array of actions to allow
+		/// </summary>
+		public List<AclAction>? Actions { get; set; }
+
+		/// <summary>
+		/// List of profiles to grant
+		/// </summary>
+		public List<AclProfileId>? Profiles { get; set; }
+
+		/// <summary>
+		/// List of all actions inherited from all profiles
+		/// </summary>
+		[JsonIgnore]
+		public HashSet<AclAction> ComputedActions { get; set; } = null!;
+
+		/// <summary>
+		/// Private constructor for serialization
+		/// </summary>
+		public AclEntryConfig()
+		{
+			Claim = new AclClaimConfig();
+		}
+
+		/// <summary>
+		/// Constructor
+		/// </summary>
+		/// <param name="claim">The claim this entry applies to</param>
+		/// <param name="actions">List of allowed operations</param>
+		/// <param name="profiles">Profiles to inherit</param>
+		public AclEntryConfig(AclClaimConfig claim, IEnumerable<AclAction>? actions, IEnumerable<AclProfileId>? profiles = null)
+		{
+			Claim = claim;
+			if (actions != null)
+			{
+				Actions = new List<AclAction>(actions);
+			}
+			if (profiles != null)
+			{
+				Profiles = new List<AclProfileId>(profiles);
+			}
+		}
+
+		internal void PostLoad(AclScopeName scopeName, IReadOnlyDictionary<AclProfileId, AclProfileConfig> profileLookup)
+		{
+			HashSet<AclAction> computedActions = new HashSet<AclAction>();
+			if (Actions != null)
+			{
+				computedActions.UnionWith(Actions);
+			}
+			if (Profiles != null)
+			{
+				computedActions.UnionWith(Profiles.SelectMany(profileId => GetActionsForProfile(scopeName, profileId, profileLookup)));
+			}
+			ComputedActions = computedActions;
+		}
+
+		static IEnumerable<AclAction> GetActionsForProfile(AclScopeName scopeName, AclProfileId profileId, IReadOnlyDictionary<AclProfileId, AclProfileConfig> profileLookup)
+		{
+			AclProfileConfig? profileConfig;
+			if (!profileLookup.TryGetValue(profileId, out profileConfig))
+			{
+				throw new Exception($"Undefined profile '{profileId}' referenced from {scopeName}");
+			}
+			return profileConfig.ComputedActions;
+		}
+	}
+
+	/// <summary>
+	/// Configuration for an ACL profile. This defines a preset group of actions which can be given to a user via an ACL entry.
+	/// </summary>
+	public class AclProfileConfig
+	{
+		/// <summary>
+		/// Identifier for this profile
+		/// </summary>
+		public AclProfileId Id { get; set; }
+
+		/// <summary>
+		/// Actions to include
+		/// </summary>
+		public List<AclAction>? Actions { get; set; }
+
+		/// <summary>
+		/// Actions to exclude from the inherited actions
+		/// </summary>
+		public List<AclAction>? ExcludeActions { get; set; }
+
+		/// <summary>
+		/// Other profiles to extend from
+		/// </summary>
+		public List<AclProfileId>? Extends { get; set; }
+
+		/// <summary>
+		/// Computed list of actions after considering base profiles etc... Fixed up by calling PostLoad().
+		/// </summary>
+		[JsonIgnore]
+		internal HashSet<AclAction> ComputedActions { get; set; } = null!;
+
+		internal void PostLoad(AclScopeName scopeName, Dictionary<AclProfileId, AclProfileConfig> profileLookup, HashSet<AclProfileId> visited)
+		{
+			if (ComputedActions == null)
+			{
+				if (!visited.Add(Id))
+				{
+					throw new Exception($"Recursive profile definition for '{Id}' in {scopeName}");
+				}
+
+				HashSet<AclAction> computedActions = new HashSet<AclAction>();
+				if (Extends != null)
+				{
+					computedActions.UnionWith(Extends.SelectMany(profileId => GetActionsForProfile(scopeName, profileId, profileLookup, visited)));
+				}
+				if (Actions != null)
+				{
+					computedActions.UnionWith(Actions);
+				}
+				if (ExcludeActions != null)
+				{
+					computedActions.ExceptWith(ExcludeActions);
+				}
+				ComputedActions = computedActions;
+			}
+		}
+
+		static IEnumerable<AclAction> GetActionsForProfile(AclScopeName scopeName, AclProfileId profileId, Dictionary<AclProfileId, AclProfileConfig> profileLookup, HashSet<AclProfileId> visited)
+		{
+			AclProfileConfig? profileConfig;
+			if (!profileLookup.TryGetValue(profileId, out profileConfig))
+			{
+				throw new Exception($"Undefined profile '{profileId}'");
+			}
+
+			profileConfig.PostLoad(scopeName, profileLookup, visited);
+			return profileConfig.ComputedActions!;
+		}
+	}
+
+	/// <summary>
+	/// New claim to create
+	/// </summary>
+	public class AclClaimConfig
+	{
+		/// <summary>
+		/// The claim type
+		/// </summary>
+		[Required]
+		public string Type { get; set; } = null!;
+
+		/// <summary>
+		/// The claim value
+		/// </summary>
+		[Required]
+		public string Value { get; set; } = null!;
+
+		/// <summary>
+		/// Constructor
+		/// </summary>
+		public AclClaimConfig()
+		{
+			Type = String.Empty;
+			Value = String.Empty;
+		}
+
+		/// <summary>
+		/// Constructor
+		/// </summary>
+		/// <param name="claim">The claim object</param>
+		public AclClaimConfig(Claim claim)
+			: this(claim.Type, claim.Value)
+		{
+		}
+
+		/// <summary>
+		/// Constructor
+		/// </summary>
+		/// <param name="type">The claim type</param>
+		/// <param name="value">The claim value</param>
+		public AclClaimConfig(string type, string value)
+		{
+			Type = type;
+			Value = value;
+		}
+
+		/// <summary>
+		/// Converts this object to a regular <see cref="Claim"/> object.
+		/// </summary>
+		public Claim ToClaim() => new Claim(Type, Value);
+	}
+}
diff --git a/Engine/Source/Programs/Horde/HordeServer.Shared/Acls/AclProfileId.cs b/Engine/Source/Programs/Horde/HordeServer.Shared/Acls/AclProfileId.cs
new file mode 100644
--- /dev/null
+++ b/Engine/Source/Programs/Horde/HordeServer.Shared/Acls/AclProfileId.cs
@@ -0,0 +1,44 @@
+// Copyright Epic Games, Inc. All Rights Reserved.
+
+using System.ComponentModel;
+using EpicGames.Core;
+using EpicGames.Horde;
+
+namespace HordeServer.Acls
+{
+	/// <summary>
+	/// Identifier for a profile; a group of actions which can be assigned to a user as a whole.
+	/// </summary>
+	/// <param name="Id">Id to construct from</param>
+	[LogValueType]
+	[JsonSchemaString]
+	[TypeConverter(typeof(StringIdTypeConverter<AclProfileId, AclProfileIdConverter>))]
+	[StringIdConverter(typeof(AclProfileIdConverter))]
+	public record struct AclProfileId(StringId Id)
+	{
+		/// <summary>
+		/// Constructor
+		/// </summary>
+		public AclProfileId(string id) : this(new StringId(id))
+		{
+		}
+
+		/// <inheritdoc cref="StringId.IsEmpty"/>
+		public bool IsEmpty => Id.IsEmpty;
+
+		/// <inheritdoc/>
+		public override string ToString() => Id.ToString();
+	}
+
+	/// <summary>
+	/// Converter to and from <see cref="StringId"/> instances.
+	/// </summary>
+	class AclProfileIdConverter : StringIdConverter<AclProfileId>
+	{
+		/// <inheritdoc/>
+		public override AclProfileId FromStringId(StringId id) => new AclProfileId(id);
+
+		/// <inheritdoc/>
+		public override StringId ToStringId(AclProfileId value) => value.Id;
+	}
+}
diff --git a/Engine/Source/Programs/Horde/HordeServer.Shared/Configuration/ConfigInclude.cs b/Engine/Source/Programs/Horde/HordeServer.Shared/Configuration/ConfigInclude.cs
new file mode 100644
--- /dev/null
+++ b/Engine/Source/Programs/Horde/HordeServer.Shared/Configuration/ConfigInclude.cs
@@ -0,0 +1,18 @@
+// Copyright Epic Games, Inc. All Rights Reserved.
+
+using System.ComponentModel.DataAnnotations;
+
+namespace HordeServer.Configuration
+{
+	/// <summary>
+	/// Directive to merge config data from another source
+	/// </summary>
+	public class ConfigInclude
+	{
+		/// <summary>
+		/// Path to the config data to be included. May be relative to the including file's location.
+		/// </summary>
+		[Required, ConfigInclude, ConfigRelativePath]
+		public string Path { get; set; } = null!;
+	}
+}
diff --git a/Engine/Source/Programs/Horde/HordeServer.Shared/Utilities/HordeClaimTypes.cs b/Engine/Source/Programs/Horde/HordeServer.Shared/Utilities/HordeClaimTypes.cs
--- a/Engine/Source/Programs/Horde/HordeServer.Shared/Utilities/HordeClaimTypes.cs
+++ b/Engine/Source/Programs/Horde/HordeServer.Shared/Utilities/HordeClaimTypes.cs
@@ -153,5 +153,14 @@ namespace HordeServer.Utilities
 			}
 			return claim.Value;
 		}
+
+		/// <summary>
+		/// Determine if a user has the administrator claim
+		/// </summary>
+		/// <param name="user">The principal to get the external issue user for</param>
+		public static bool HasAdminClaim(this ClaimsPrincipal user)
+		{
+			return user.HasClaim(HordeClaims.AdminClaim.Type, HordeClaims.AdminClaim.Value);
+		}
 	}
 }
diff --git a/Engine/Source/Programs/Horde/HordeServer.Shared/Utilities/HordeClaims.cs b/Engine/Source/Programs/Horde/HordeServer.Shared/Utilities/HordeClaims.cs
new file mode 100644
--- /dev/null
+++ b/Engine/Source/Programs/Horde/HordeServer.Shared/Utilities/HordeClaims.cs
@@ -0,0 +1,69 @@
+﻿// Copyright Epic Games, Inc. All Rights Reserved.
+
+using EpicGames.Horde.Agents;
+using EpicGames.Horde.Agents.Sessions;
+using HordeServer.Acls;
+
+namespace HordeServer.Utilities
+{
+	/// <summary>
+	/// Predetermined claim values
+	/// </summary>
+	public static class HordeClaims
+	{
+		/// <summary>
+		/// Claim which gives unfettered access to all Horde endpoints
+		/// </summary>
+		public static AclClaimConfig AdminClaim { get; } = new AclClaimConfig(HordeClaimTypes.Role, "admin");
+
+		/// <summary>
+		/// Name of the role that can be used to administer agents
+		/// </summary>
+		public static AclClaimConfig AgentRegistrationClaim { get; } = new AclClaimConfig(HordeClaimTypes.Role, "agent-registration");
+
+		/// <summary>
+		/// Name of the role used to upload software
+		/// </summary>
+		public static AclClaimConfig DownloadSoftwareClaim { get; } = new AclClaimConfig(HordeClaimTypes.Role, "download-software");
+
+		/// <summary>
+		/// Name of the role used to upload tools
+		/// </summary>
+		public static AclClaimConfig UploadToolsClaim { get; } = new AclClaimConfig(HordeClaimTypes.Role, "upload-software");
+
+		/// <summary>
+		/// Name of the role used to upload software
+		/// </summary>
+		public static AclClaimConfig ConfigureProjectsClaim { get; } = new AclClaimConfig(HordeClaimTypes.Role, "configure-projects");
+
+		/// <summary>
+		/// Name of the role used to upload software
+		/// </summary>
+		public static AclClaimConfig StartChainedJobClaim { get; } = new AclClaimConfig(HordeClaimTypes.Role, "start-chained-job");
+
+		/// <summary>
+		/// Role for all agents
+		/// </summary>
+		public static AclClaimConfig AgentRoleClaim { get; } = new AclClaimConfig(HordeClaimTypes.Role, "agent");
+
+		/// <summary>
+		/// Gets the role for a specific agent
+		/// </summary>
+		/// <param name="agentId">The session id</param>
+		/// <returns>New claim instance</returns>
+		public static AclClaimConfig GetAgentClaim(AgentId agentId)
+		{
+			return new AclClaimConfig(HordeClaimTypes.Agent, agentId.ToString());
+		}
+
+		/// <summary>
+		/// Gets the role for a specific agent session
+		/// </summary>
+		/// <param name="sessionId">The session id</param>
+		/// <returns>New claim instance</returns>
+		public static AclClaimConfig GetSessionClaim(SessionId sessionId)
+		{
+			return new AclClaimConfig(HordeClaimTypes.AgentSessionId, sessionId.ToString());
+		}
+	}
+}
diff --git a/Engine/Source/Programs/Horde/HordeServer.Shared/Utilities/IntervalJsonConverter.cs b/Engine/Source/Programs/Horde/HordeServer.Shared/Utilities/IntervalJsonConverter.cs
new file mode 100644
--- /dev/null
+++ b/Engine/Source/Programs/Horde/HordeServer.Shared/Utilities/IntervalJsonConverter.cs
@@ -0,0 +1,110 @@
+﻿// Copyright Epic Games, Inc. All Rights Reserved.
+
+using System;
+using System.Text;
+using System.Text.Json;
+using System.Text.Json.Serialization;
+
+namespace HordeServer.Utilities
+{
+	/// <summary>
+	/// Converts TimeSpan intervals from formats like "30m", "1h30m", etc...
+	/// </summary>
+	public class IntervalJsonConverter : JsonConverter<TimeSpan>
+	{
+		record struct TimeSuffix(string Text, long Ticks);
+
+		static readonly TimeSuffix[] s_timeSuffixes =
+		{
+			new TimeSuffix("w", TimeSpan.TicksPerDay * 7),
+			new TimeSuffix("d", TimeSpan.TicksPerDay),
+			new TimeSuffix("h", TimeSpan.TicksPerHour),
+			new TimeSuffix("m", TimeSpan.TicksPerMinute),
+			new TimeSuffix("s", TimeSpan.TicksPerSecond)
+		};
+
+		/// <inheritdoc/>
+		public override TimeSpan Read(ref Utf8JsonReader reader, Type typeToConvert, JsonSerializerOptions options)
+		{
+			return Parse(reader.GetString());
+		}
+
+		/// <summary>
+		/// Parse a string as a time interval
+		/// </summary>
+		public static TimeSpan Parse(string? time)
+		{
+			if (String.IsNullOrEmpty(time))
+			{
+				throw new FormatException("TimeSpan may not be empty");
+			}
+
+			if (TimeSpan.TryParse(time, out TimeSpan result))
+			{
+				return result;
+			}
+
+			long ticks = 0;
+			int suffixIdx = 0;
+			for (ReadOnlySpan<char> span = time.AsSpan(); span.Length > 0;)
+			{
+				long value = 0;
+
+				// Parse the value
+				int length = 0;
+				for (; length < span.Length && span[length] >= '0' && span[length] <= '9'; length++)
+				{
+					value = (value * 10) + (span[length] - '0');
+				}
+				if (length == 0)
+				{
+					throw new FormatException($"Unable to parse '{time}' as a TimeSpan. Invalid value.");
+				}
+				span = span.Slice(length);
+
+				// Parse the suffix
+				length = 0;
+				while (length < span.Length && !(span[length] >= '0' && span[length] <= '9'))
+				{
+					length++;
+				}
+				if (length != 0) // Allow '3m30' to use the next suffix in order.
+				{
+					while (suffixIdx < s_timeSuffixes.Length && !s_timeSuffixes[suffixIdx].Text.AsSpan().SequenceEqual(span[0..length]))
+					{
+						suffixIdx++;
+					}
+					if (suffixIdx == s_timeSuffixes.Length)
+					{
+						throw new FormatException($"Unable to parse '{time}' as a TimeSpan. Invalid suffix.");
+					}
+					span = span.Slice(length);
+				}
+
+				ticks += s_timeSuffixes[suffixIdx++].Ticks * value;
+			}
+
+			return new TimeSpan(ticks);
+		}
+
+		/// <inheritdoc/>
+		public override void Write(Utf8JsonWriter writer, TimeSpan value, JsonSerializerOptions options)
+		{
+			long ticks = value.Ticks;
+
+			StringBuilder builder = new StringBuilder();
+			for (int suffixIdx = 0; suffixIdx < s_timeSuffixes.Length; suffixIdx++)
+			{
+				if (ticks >= s_timeSuffixes[suffixIdx].Ticks)
+				{
+					long timeValue = ticks / s_timeSuffixes[suffixIdx].Ticks;
+					builder.Append(timeValue);
+					builder.Append(s_timeSuffixes[suffixIdx].Text);
+					ticks -= timeValue * s_timeSuffixes[suffixIdx].Ticks;
+				}
+			}
+
+			writer.WriteStringValue(builder.ToString());
+		}
+	}
+}
diff --git a/Engine/Source/Programs/Horde/HordeServer/Acls/AclConfig.cs b/Engine/Source/Programs/Horde/HordeServer/Acls/AclConfig.cs
deleted file mode 100644
diff --git a/Engine/Source/Programs/Horde/HordeServer/Acls/AclProfileId.cs b/Engine/Source/Programs/Horde/HordeServer/Acls/AclProfileId.cs
deleted file mode 100644
diff --git a/Engine/Source/Programs/Horde/HordeServer/Acls/AclService.cs b/Engine/Source/Programs/Horde/HordeServer/Acls/AclService.cs
--- a/Engine/Source/Programs/Horde/HordeServer/Acls/AclService.cs
+++ b/Engine/Source/Programs/Horde/HordeServer/Acls/AclService.cs
@@ -81,11 +81,6 @@ namespace HordeServer.Acls
 
 	internal static class ClaimExtensions
 	{
-		public static bool HasAdminClaim(this ClaimsPrincipal user)
-		{
-			return user.HasClaim(HordeClaims.AdminClaim.Type, HordeClaims.AdminClaim.Value);
-		}
-
 		public static bool HasAgentClaim(this ClaimsPrincipal user, AgentId agentId)
 		{
 			return user.HasClaim(HordeClaimTypes.Agent, agentId.ToString());
diff --git a/Engine/Source/Programs/Horde/HordeServer/Dockerfile b/Engine/Source/Programs/Horde/HordeServer/Dockerfile
--- a/Engine/Source/Programs/Horde/HordeServer/Dockerfile
+++ b/Engine/Source/Programs/Horde/HordeServer/Dockerfile
@@ -40,6 +40,7 @@ COPY Source/Programs/Shared/EpicGames.Serialization.Tests/*.csproj ./Source/Prog
 COPY Source/Programs/Horde/Drivers/JobDriver/*.csproj ./Source/Programs/Horde/Drivers/JobDriver/
 COPY Source/Programs/Horde/HordeAgent/*.csproj ./Source/Programs/Horde/HordeAgent/
 COPY Source/Programs/Horde/HordeServer/*.csproj ./Source/Programs/Horde/HordeServer/
+COPY Source/Programs/Horde/HordeServer.Shared/*.csproj ./Source/Programs/Horde/HordeServer.Shared/
 COPY Source/Programs/Horde/HordeServer.Tests/*.csproj ./Source/Programs/Horde/HordeServer.Tests/
 RUN dotnet restore Source/Programs/Horde/HordeServer/HordeServer.csproj -p:WithHordeStorage=false
 
@@ -48,6 +49,7 @@ COPY Binaries/DotNET/EpicGames.Perforce.Native/ Binaries/DotNET/EpicGames.Perfor
 COPY Source/Programs/Shared/ Source/Programs/Shared/
 COPY Source/Programs/Horde/.editorconfig Source/Programs/Horde/.editorconfig
 COPY Source/Programs/Horde/HordeServer/ Source/Programs/Horde/HordeServer/
+COPY Source/Programs/Horde/HordeServer.Shared/ Source/Programs/Horde/HordeServer.Shared/
 COPY Source/Programs/AutomationTool/AutomationUtils/Matchers/ Source/Programs/AutomationTool/AutomationUtils/Matchers/
 COPY Source/Programs/UnrealBuildTool/Matchers/ Source/Programs/UnrealBuildTool/Matchers/
 ARG msbuild_args
diff --git a/Engine/Source/Programs/Horde/HordeServer/Server/GlobalConfig.cs b/Engine/Source/Programs/Horde/HordeServer/Server/GlobalConfig.cs
--- a/Engine/Source/Programs/Horde/HordeServer/Server/GlobalConfig.cs
+++ b/Engine/Source/Programs/Horde/HordeServer/Server/GlobalConfig.cs
@@ -27,9 +27,16 @@ using EpicGames.Serialization;
 using HordeServer.Acls;
 using HordeServer.Agents;
 using HordeServer.Agents.Pools;
+using HordeServer.Agents.Sessions;
+using HordeServer.Agents.Software;
+using HordeServer.Artifacts;
 using HordeServer.Configuration;
 using HordeServer.Dashboard;
 using HordeServer.Devices;
+using HordeServer.Jobs;
+using HordeServer.Jobs.Bisect;
+using HordeServer.Logs;
+using HordeServer.Notifications;
 using HordeServer.Perforce;
 using HordeServer.Projects;
 using HordeServer.Secrets;
@@ -79,18 +86,6 @@ namespace HordeServer.Server
 #pragma warning restore CA1069
 #pragma warning restore CA1027
 
-	/// <summary>
-	/// Directive to merge config data from another source
-	/// </summary>
-	public class ConfigInclude
-	{
-		/// <summary>
-		/// Path to the config data to be included. May be relative to the including file's location.
-		/// </summary>
-		[Required, ConfigInclude, ConfigRelativePath]
-		public string Path { get; set; } = null!;
-	}
-
 	/// <summary>
 	/// Configuration for an artifact
 	/// </summary>
@@ -298,7 +293,7 @@ namespace HordeServer.Server
 		{
 			ServerSettings = serverSettings;
 
-			AclConfig defaultAcl = AclConfig.CreateRoot();
+			AclConfig defaultAcl = CreateRootAcl();
 			Acl.PostLoad(defaultAcl, defaultAcl.ScopeName);
 
 			Streams = Projects.SelectMany(x => x.Streams).ToList();
@@ -382,6 +377,69 @@ namespace HordeServer.Server
 			}
 		}
 
+		/// <summary>
+		/// Creates the default root ACL
+		/// </summary>
+		static AclConfig CreateRootAcl()
+		{
+			AclConfig defaultAcl = new AclConfig();
+
+			defaultAcl.Entries = new List<AclEntryConfig>();
+			defaultAcl.Entries.Add(new AclEntryConfig(new AclClaimConfig(ClaimTypes.Role, "internal:AgentRegistration"), new[] { AgentAclAction.CreateAgent, SessionAclAction.CreateSession }));
+			defaultAcl.Entries.Add(new AclEntryConfig(HordeClaims.AgentRegistrationClaim, new[] { AgentAclAction.CreateAgent, SessionAclAction.CreateSession, AgentAclAction.UpdateAgent, AgentSoftwareAclAction.DownloadSoftware, PoolAclAction.CreatePool, PoolAclAction.UpdatePool, PoolAclAction.ViewPool, PoolAclAction.DeletePool, PoolAclAction.ListPools, StreamAclAction.ViewStream, ProjectAclAction.ViewProject, JobAclAction.ViewJob, ServerAclAction.ViewCosts }));
+			defaultAcl.Entries.Add(new AclEntryConfig(HordeClaims.AgentRoleClaim, new[] { ProjectAclAction.ViewProject, StreamAclAction.ViewStream, LogAclAction.CreateEvent, AgentSoftwareAclAction.DownloadSoftware }));
+			defaultAcl.Entries.Add(new AclEntryConfig(HordeClaims.DownloadSoftwareClaim, new[] { AgentSoftwareAclAction.DownloadSoftware }));
+			defaultAcl.Entries.Add(new AclEntryConfig(HordeClaims.UploadToolsClaim, new[] { AgentSoftwareAclAction.UploadSoftware, ToolAclAction.UploadTool }));
+			defaultAcl.Entries.Add(new AclEntryConfig(HordeClaims.ConfigureProjectsClaim, new[] { ProjectAclAction.CreateProject, ProjectAclAction.UpdateProject, ProjectAclAction.ViewProject, StreamAclAction.CreateStream, StreamAclAction.UpdateStream, StreamAclAction.ViewStream }));
+			defaultAcl.Entries.Add(new AclEntryConfig(HordeClaims.StartChainedJobClaim, new[] { JobAclAction.CreateJob, JobAclAction.ExecuteJob, JobAclAction.UpdateJob, JobAclAction.ViewJob, StreamAclAction.ViewTemplate, StreamAclAction.ViewStream }));
+
+			defaultAcl.Profiles = new List<AclProfileConfig>();
+			defaultAcl.Profiles.Add(new AclProfileConfig
+			{
+				Id = new AclProfileId("default-read"),
+				Actions = new List<AclAction>
+				{
+					AgentAclAction.ListAgents,
+					AgentAclAction.ViewAgent,
+					ArtifactAclAction.DownloadArtifact,
+					ArtifactAclAction.ReadArtifact,
+					BisectTaskAclAction.ViewBisectTask,
+					DeviceAclAction.DeviceRead,
+					JobAclAction.ViewJob,
+					LogAclAction.ViewEvent,
+					LogAclAction.ViewLog,
+					NotificationAclAction.CreateSubscription,
+					PoolAclAction.ListPools,
+					PoolAclAction.ViewPool,
+					ProjectAclAction.ViewProject,
+					ServerAclAction.IssueBearerToken,
+					StreamAclAction.ViewChanges,
+					StreamAclAction.ViewStream,
+					StreamAclAction.ViewTemplate,
+				}
+			});
+			defaultAcl.Profiles.Add(new AclProfileConfig
+			{
+				Id = new AclProfileId("default-run"),
+				Extends = new List<AclProfileId>
+				{
+					new AclProfileId("default-read")
+				},
+				Actions = new List<AclAction>
+				{
+					JobAclAction.CreateJob,
+					JobAclAction.UpdateJob,
+					JobAclAction.RetryJobStep,
+					DeviceAclAction.DeviceWrite,
+					BisectTaskAclAction.CreateBisectTask,
+					BisectTaskAclAction.UpdateBisectTask,
+				}
+			});
+
+			defaultAcl.PostLoad(null, AclScopeName.Root);
+			return defaultAcl;
+		}
+
 		static void BuildAclScopeLookup(AclConfig acl, Dictionary<AclScopeName, AclConfig> aclLookup)
 		{
 			aclLookup.Add(acl.ScopeName, acl);
diff --git a/Engine/Source/Programs/Horde/HordeServer/Utilities/HordeClaims.cs b/Engine/Source/Programs/Horde/HordeServer/Utilities/HordeClaims.cs
deleted file mode 100644
diff --git a/Engine/Source/Programs/Horde/HordeServer/Utilities/IntervalJsonConverter.cs b/Engine/Source/Programs/Horde/HordeServer/Utilities/IntervalJsonConverter.cs
deleted file mode 100644
