Horde: Fix static analysis warnings.

#rnx

diff --git a/Engine/Source/Programs/Horde/HordeServer/Dashboard/DashboardConfig.cs b/Engine/Source/Programs/Horde/HordeServer/Dashboard/DashboardConfig.cs
--- a/Engine/Source/Programs/Horde/HordeServer/Dashboard/DashboardConfig.cs
+++ b/Engine/Source/Programs/Horde/HordeServer/Dashboard/DashboardConfig.cs
@@ -4,11 +4,9 @@ using System.Collections.Generic;
 using EpicGames.Core;
 using EpicGames.Horde.Common;
 using HordeServer.Configuration;
-using HordeServer.Server;
 
 namespace HordeServer.Dashboard
 {
-
 	/// <summary>
 	/// Configuration for dashboard features
 	/// </summary>
diff --git a/Engine/Source/Programs/Horde/HordeServer/Server/Notices/NoticesController.cs b/Engine/Source/Programs/Horde/HordeServer/Server/Notices/NoticesController.cs
--- a/Engine/Source/Programs/Horde/HordeServer/Server/Notices/NoticesController.cs
+++ b/Engine/Source/Programs/Horde/HordeServer/Server/Notices/NoticesController.cs
@@ -8,7 +8,6 @@ using System.Threading.Tasks;
 using EpicGames.Core;
 using EpicGames.Horde.Server.Notices;
 using EpicGames.Horde.Users;
-using HordeServer.Acls;
 using HordeServer.Users;
 using HordeServer.Utilities;
 using Microsoft.AspNetCore.Authorization;
