Horde: Move default json serializers into HordeServer.Shared.

#rnx

diff --git a/Engine/Source/Programs/Horde/HordeServer.Shared/Utilities/DateTimeJsonConverter.cs b/Engine/Source/Programs/Horde/HordeServer.Shared/Utilities/DateTimeJsonConverter.cs
new file mode 100644
--- /dev/null
+++ b/Engine/Source/Programs/Horde/HordeServer.Shared/Utilities/DateTimeJsonConverter.cs
@@ -0,0 +1,39 @@
+// Copyright Epic Games, Inc. All Rights Reserved.
+
+using System;
+using System.Collections.Generic;
+using System.Diagnostics;
+using System.Globalization;
+using System.Linq;
+using System.Text;
+using System.Text.Json;
+using System.Text.Json.Serialization;
+using System.Threading.Tasks;
+
+namespace HordeServer.Utilities
+{
+	/// <summary>
+	/// Converter for <see cref="DateTime"/> values to JSON
+	/// </summary>
+	public class DateTimeJsonConverter : JsonConverter<DateTime>
+	{
+		/// <inheritdoc/>
+		public override DateTime Read(ref Utf8JsonReader reader, Type typeToConvert, JsonSerializerOptions options)
+		{
+			Debug.Assert(typeToConvert == typeof(DateTime));
+
+			string? str = reader.GetString();
+			if (str == null)
+			{
+				throw new InvalidDataException("Unable to parse DateTime");
+			}
+			return DateTime.Parse(str, CultureInfo.CurrentCulture);
+		}
+
+		/// <inheritdoc/>
+		public override void Write(Utf8JsonWriter writer, DateTime dateTime, JsonSerializerOptions options)
+		{
+			writer.WriteStringValue(dateTime.ToUniversalTime().ToString("yyyy'-'MM'-'dd'T'HH':'mm':'ssZ", CultureInfo.CurrentCulture));
+		}
+	}
+}
diff --git a/Engine/Source/Programs/Horde/HordeServer.Shared/Utilities/JsonUtils.cs b/Engine/Source/Programs/Horde/HordeServer.Shared/Utilities/JsonUtils.cs
new file mode 100644
--- /dev/null
+++ b/Engine/Source/Programs/Horde/HordeServer.Shared/Utilities/JsonUtils.cs
@@ -0,0 +1,46 @@
+// Copyright Epic Games, Inc. All Rights Reserved.
+
+using System.Diagnostics;
+using System.Globalization;
+using System.Text.Json;
+using System.Text.Json.Serialization;
+using EpicGames.Core;
+using EpicGames.Horde;
+using EpicGames.Horde.Jobs;
+using MongoDB.Bson;
+
+namespace HordeServer.Utilities
+{
+	/// <summary>
+	/// Helper methods for json serialization
+	/// </summary>
+	public static class JsonUtils
+	{
+		/// <summary>
+		/// Default JSON serialization options
+		/// </summary>
+		public static JsonSerializerOptions DefaultSerializerOptions { get; } = CreateDefaultJsonSerializerOptions();
+
+		static JsonSerializerOptions CreateDefaultJsonSerializerOptions()
+		{
+			JsonSerializerOptions options = new JsonSerializerOptions();
+			JsonUtils.ConfigureJsonSerializer(options);
+			options.MakeReadOnly(true);
+			return options;
+		}
+
+		/// <summary>
+		/// Configure a json serializer with all standard Horde converters
+		/// </summary>
+		public static void ConfigureJsonSerializer(JsonSerializerOptions options)
+		{
+			HordeHttpClient.ConfigureJsonSerializer(options);
+			options.Converters.Add(new ObjectIdJsonConverter());
+			options.Converters.Add(new ObjectIdJsonConverterFactory());
+			options.Converters.Add(new JsonKnownTypesConverterFactory());
+			options.Converters.Add(new SubResourceIdJsonConverterFactory());
+			options.Converters.Add(new DateTimeJsonConverter());
+			options.Converters.Add(new TimeSpanJsonConverter());
+		}
+	}
+}
diff --git a/Engine/Source/Programs/Horde/HordeServer.Shared/Utilities/ObjectIdConverter.cs b/Engine/Source/Programs/Horde/HordeServer.Shared/Utilities/ObjectIdConverter.cs
new file mode 100644
--- /dev/null
+++ b/Engine/Source/Programs/Horde/HordeServer.Shared/Utilities/ObjectIdConverter.cs
@@ -0,0 +1,180 @@
+// Copyright Epic Games, Inc. All Rights Reserved.
+
+using System;
+using System.ComponentModel;
+using System.Globalization;
+using System.Reflection;
+using System.Text.Json;
+using System.Text.Json.Serialization;
+using MongoDB.Bson;
+using MongoDB.Bson.Serialization;
+using MongoDB.Bson.Serialization.Serializers;
+
+namespace HordeServer.Utilities
+{
+	/// <summary>
+	/// Base class for converting to and from types containing an object id. Useful pattern for reducing boilerplate with strongly typed records.
+	/// </summary>
+	/// <typeparam name="T"></typeparam>
+	public abstract class ObjectIdConverter<T> where T : struct
+	{
+		/// <summary>
+		/// Converts a type to a <see cref="ObjectId"/>
+		/// </summary>
+		public abstract ObjectId ToObjectId(T value);
+
+		/// <summary>
+		/// Constructs a type from an object id
+		/// </summary>
+		public abstract T FromObjectId(ObjectId id);
+	}
+
+	/// <summary>
+	/// Attribute declaring an object id converter for a particular type
+	/// </summary>
+	[AttributeUsage(AttributeTargets.Struct)]
+	public sealed class ObjectIdConverterAttribute : Attribute
+	{
+		/// <summary>
+		/// The converter type
+		/// </summary>
+		public Type ConverterType { get; }
+
+		/// <summary>
+		/// Constructor
+		/// </summary>
+		public ObjectIdConverterAttribute(Type converterType) => ConverterType = converterType;
+	}
+
+	/// <summary>
+	/// Class which serializes types with a <see cref="ObjectIdConverter{T}"/> to Json
+	/// </summary>
+	public sealed class ObjectIdTypeConverter<TValue, TConverter> : TypeConverter where TValue : struct where TConverter : ObjectIdConverter<TValue>, new()
+	{
+		readonly TConverter _converter = new TConverter();
+
+		/// <inheritdoc/>
+		public override bool CanConvertFrom(ITypeDescriptorContext? context, Type sourceType)
+		{
+			return sourceType == typeof(string);
+		}
+
+		/// <inheritdoc/>
+		public override object? ConvertFrom(ITypeDescriptorContext? context, CultureInfo? culture, object value)
+		{
+			if (value is string str)
+			{
+				return _converter.FromObjectId(ObjectId.Parse(str));
+			}
+			return null;
+		}
+
+		/// <inheritdoc/>
+		public override bool CanConvertTo(ITypeDescriptorContext? context, Type? destinationType) => destinationType == typeof(string);
+
+		/// <inheritdoc/>
+		public override object? ConvertTo(ITypeDescriptorContext? context, CultureInfo? culture, object? value, Type destinationType)
+		{
+			if (destinationType == typeof(string))
+			{
+				return _converter.ToObjectId((TValue)value!).ToString();
+			}
+			return null;
+		}
+	}
+
+	/// <summary>
+	/// Class which serializes ObjectId types to Json
+	/// </summary>
+	public class ObjectIdJsonConverter : JsonConverter<ObjectId>
+	{
+		/// <inheritdoc/>
+		public override ObjectId Read(ref Utf8JsonReader reader, Type typeToConvert, JsonSerializerOptions options)
+		{
+			string? str = reader.GetString();
+			if (str == null)
+			{
+				throw new InvalidDataException("Unable to parse object id");
+			}
+			if (str.Length == 0)
+			{
+				return ObjectId.Empty;
+			}
+			return ObjectId.Parse(str);
+		}
+
+		/// <inheritdoc/>
+		public override void Write(Utf8JsonWriter writer, ObjectId objectId, JsonSerializerOptions options)
+		{
+			writer.WriteStringValue(objectId.ToString());
+		}
+	}
+
+	/// <summary>
+	/// Class which serializes types with a <see cref="ObjectIdConverter{T}"/> to Json
+	/// </summary>
+	public sealed class ObjectIdJsonConverter<TValue, TConverter> : JsonConverter<TValue> where TValue : struct where TConverter : ObjectIdConverter<TValue>, new()
+	{
+		readonly TConverter _converter = new TConverter();
+
+		/// <inheritdoc/>
+		public override TValue Read(ref Utf8JsonReader reader, Type typeToConvert, JsonSerializerOptions options) 
+			=> _converter.FromObjectId(ObjectId.Parse(reader.GetString()));
+
+		/// <inheritdoc/>
+		public override void Write(Utf8JsonWriter writer, TValue value, JsonSerializerOptions options) 
+			=> writer.WriteStringValue(_converter.ToObjectId(value).ToString());
+	}
+
+	/// <summary>
+	/// Creates constructors for types with a <see cref="ObjectIdConverter{T}"/> to Json
+	/// </summary>
+	public sealed class ObjectIdJsonConverterFactory : JsonConverterFactory
+	{
+		/// <inheritdoc/>
+		public override bool CanConvert(Type typeToConvert) => typeToConvert.GetCustomAttribute<ObjectIdConverterAttribute>() != null;
+
+		/// <inheritdoc/>
+		public override JsonConverter? CreateConverter(Type type, JsonSerializerOptions options)
+		{
+			ObjectIdConverterAttribute? attribute = type.GetCustomAttribute<ObjectIdConverterAttribute>();
+			if (attribute == null)
+			{
+				return null;
+			}
+			return (JsonConverter?)Activator.CreateInstance(typeof(ObjectIdJsonConverter<,>).MakeGenericType(type, attribute.ConverterType));
+		}
+	}
+
+	/// <summary>
+	/// Class which serializes object id types to BSON
+	/// </summary>
+	public sealed class ObjectIdBsonSerializer<TValue, TConverter> : SerializerBase<TValue> where TValue : struct where TConverter : ObjectIdConverter<TValue>, new()
+	{
+		readonly TConverter _converter = new TConverter();
+
+		/// <inheritdoc/>
+		public override TValue Deserialize(BsonDeserializationContext context, BsonDeserializationArgs args) 
+			=> _converter.FromObjectId(context.Reader.ReadObjectId());
+
+		/// <inheritdoc/>
+		public override void Serialize(BsonSerializationContext context, BsonSerializationArgs args, TValue value) => context.Writer.WriteObjectId(_converter.ToObjectId(value));
+	}
+
+	/// <summary>
+	/// Class which serializes object id types to BSON
+	/// </summary>
+	public sealed class ObjectIdBsonSerializationProvider : BsonSerializationProviderBase
+	{
+		/// <inheritdoc/>
+		public override IBsonSerializer? GetSerializer(Type type, IBsonSerializerRegistry serializerRegistry)
+		{
+			ObjectIdConverterAttribute? attribute = type.GetCustomAttribute<ObjectIdConverterAttribute>();
+			if (attribute == null)
+			{
+				return null;
+			}
+			return (IBsonSerializer?)Activator.CreateInstance(typeof(ObjectIdBsonSerializer<,>).MakeGenericType(type, attribute.ConverterType));
+		}
+	}
+}
diff --git a/Engine/Source/Programs/Horde/HordeServer.Shared/Utilities/TimeSpanJsonConverter.cs b/Engine/Source/Programs/Horde/HordeServer.Shared/Utilities/TimeSpanJsonConverter.cs
new file mode 100644
--- /dev/null
+++ b/Engine/Source/Programs/Horde/HordeServer.Shared/Utilities/TimeSpanJsonConverter.cs
@@ -0,0 +1,36 @@
+// Copyright Epic Games, Inc. All Rights Reserved.
+
+using System;
+using System.Collections.Generic;
+using System.Globalization;
+using System.Linq;
+using System.Text;
+using System.Text.Json;
+using System.Text.Json.Serialization;
+using System.Threading.Tasks;
+
+namespace HordeServer.Utilities
+{
+	/// <summary>
+	/// Converter for <see cref="TimeSpan"/> types to Json
+	/// </summary>
+	public class TimeSpanJsonConverter : JsonConverter<TimeSpan>
+	{
+		/// <inheritdoc/>
+		public override TimeSpan Read(ref Utf8JsonReader reader, Type typeToConvert, JsonSerializerOptions options)
+		{
+			string? str = reader.GetString();
+			if (str == null)
+			{
+				throw new InvalidDataException("Unable to parse TimeSpan");
+			}
+			return TimeSpan.Parse(str, CultureInfo.CurrentCulture);
+		}
+
+		/// <inheritdoc/>
+		public override void Write(Utf8JsonWriter writer, TimeSpan timeSpan, JsonSerializerOptions options)
+		{
+			writer.WriteStringValue(timeSpan.ToString("c"));
+		}
+	}
+}
diff --git a/Engine/Source/Programs/Horde/HordeServer/Configuration/ConfigService.cs b/Engine/Source/Programs/Horde/HordeServer/Configuration/ConfigService.cs
--- a/Engine/Source/Programs/Horde/HordeServer/Configuration/ConfigService.cs
+++ b/Engine/Source/Programs/Horde/HordeServer/Configuration/ConfigService.cs
@@ -14,6 +14,7 @@ using EpicGames.Horde.Users;
 using EpicGames.Redis;
 using HordeServer.Server;
 using HordeServer.Users;
+using HordeServer.Utilities;
 using Microsoft.Extensions.Diagnostics.HealthChecks;
 using Microsoft.Extensions.Hosting;
 using Microsoft.Extensions.Logging;
@@ -145,7 +146,7 @@ namespace HordeServer.Configuration
 			_logger = logger;
 
 			_jsonOptions = new JsonSerializerOptions();
-			Startup.ConfigureJsonSerializer(_jsonOptions);
+			JsonUtils.ConfigureJsonSerializer(_jsonOptions);
 
 			_ticker = clock.AddSharedTicker<ConfigService>(_tickInterval, TickSharedAsync, logger);
 
diff --git a/Engine/Source/Programs/Horde/HordeServer/Notifications/Sinks/SlackNotificationSink.cs b/Engine/Source/Programs/Horde/HordeServer/Notifications/Sinks/SlackNotificationSink.cs
--- a/Engine/Source/Programs/Horde/HordeServer/Notifications/Sinks/SlackNotificationSink.cs
+++ b/Engine/Source/Programs/Horde/HordeServer/Notifications/Sinks/SlackNotificationSink.cs
@@ -41,6 +41,7 @@ using HordeServer.Perforce;
 using HordeServer.Server;
 using HordeServer.Streams;
 using HordeServer.Users;
+using HordeServer.Utilities;
 using Microsoft.AspNetCore.Hosting;
 using Microsoft.CodeAnalysis;
 using Microsoft.Extensions.Caching.Memory;
@@ -279,7 +280,7 @@ namespace HordeServer.Notifications.Sinks
 			}
 
 			_jsonSerializerOptions = new JsonSerializerOptions();
-			Startup.ConfigureJsonSerializer(_jsonSerializerOptions);
+			JsonUtils.ConfigureJsonSerializer(_jsonSerializerOptions);
 
 			if (!String.IsNullOrEmpty(settings.Value.SlackUsers))
 			{
diff --git a/Engine/Source/Programs/Horde/HordeServer/Startup.cs b/Engine/Source/Programs/Horde/HordeServer/Startup.cs
--- a/Engine/Source/Programs/Horde/HordeServer/Startup.cs
+++ b/Engine/Source/Programs/Horde/HordeServer/Startup.cs
@@ -3,14 +3,11 @@
 using System;
 using System.Collections.Generic;
 using System.Diagnostics;
-using System.Globalization;
 using System.IO;
 using System.Linq;
 using System.Net.Mime;
 using System.Reflection;
 using System.Runtime.InteropServices;
-using System.Text.Json;
-using System.Text.Json.Serialization;
 using System.Threading;
 using System.Threading.Tasks;
 using Amazon;
@@ -21,7 +18,6 @@ using Amazon.Extensions.NETCore.Setup;
 using Amazon.SQS;
 using EpicGames.AspNet;
 using EpicGames.Core;
-using EpicGames.Horde;
 using EpicGames.Horde.Acls;
 using EpicGames.Horde.Agents;
 using EpicGames.Horde.Agents.Leases;
@@ -38,6 +34,7 @@ using EpicGames.Horde.Users;
 using EpicGames.Redis;
 using Grpc.Core;
 using Grpc.Core.Interceptors;
+using HordeCommon;
 using HordeServer.Accounts;
 using HordeServer.Acls;
 using HordeServer.Agents;
@@ -91,7 +88,6 @@ using HordeServer.Tools;
 using HordeServer.Ugs;
 using HordeServer.Users;
 using HordeServer.Utilities;
-using HordeCommon;
 using Microsoft.AspNetCore.Authentication;
 using Microsoft.AspNetCore.Authentication.Cookies;
 using Microsoft.AspNetCore.Authentication.JwtBearer;
@@ -243,66 +239,6 @@ namespace HordeServer
 			}
 		}
 
-		class JsonObjectIdConverter : JsonConverter<ObjectId>
-		{
-			public override ObjectId Read(ref Utf8JsonReader reader, Type typeToConvert, JsonSerializerOptions options)
-			{
-				string? str = reader.GetString();
-				if (str == null)
-				{
-					throw new InvalidDataException("Unable to parse object id");
-				}
-				if (str.Length == 0)
-				{
-					return ObjectId.Empty;
-				}
-				return ObjectId.Parse(str);
-			}
-
-			public override void Write(Utf8JsonWriter writer, ObjectId objectId, JsonSerializerOptions options)
-			{
-				writer.WriteStringValue(objectId.ToString());
-			}
-		}
-
-		class JsonDateTimeConverter : JsonConverter<DateTime>
-		{
-			public override DateTime Read(ref Utf8JsonReader reader, Type typeToConvert, JsonSerializerOptions options)
-			{
-				Debug.Assert(typeToConvert == typeof(DateTime));
-
-				string? str = reader.GetString();
-				if (str == null)
-				{
-					throw new InvalidDataException("Unable to parse DateTime");
-				}
-				return DateTime.Parse(str, CultureInfo.CurrentCulture);
-			}
-
-			public override void Write(Utf8JsonWriter writer, DateTime dateTime, JsonSerializerOptions options)
-			{
-				writer.WriteStringValue(dateTime.ToUniversalTime().ToString("yyyy'-'MM'-'dd'T'HH':'mm':'ssZ", CultureInfo.CurrentCulture));
-			}
-		}
-
-		class JsonTimeSpanConverter : JsonConverter<TimeSpan>
-		{
-			public override TimeSpan Read(ref Utf8JsonReader reader, Type typeToConvert, JsonSerializerOptions options)
-			{
-				string? str = reader.GetString();
-				if (str == null)
-				{
-					throw new InvalidDataException("Unable to parse TimeSpan");
-				}
-				return TimeSpan.Parse(str, CultureInfo.CurrentCulture);
-			}
-
-			public override void Write(Utf8JsonWriter writer, TimeSpan timeSpan, JsonSerializerOptions options)
-			{
-				writer.WriteStringValue(timeSpan.ToString("c"));
-			}
-		}
-
 		class ObsoleteLoggingFilter : IActionFilter
 		{
 			public void OnActionExecuted(ActionExecutedContext context)
@@ -466,7 +402,7 @@ namespace HordeServer
 			services.AddSingleton<RedisService>(sp => redisService);
 			services.AddSingleton<IRedisService>(sp => sp.GetRequiredService<RedisService>());
 			services.AddDataProtection().PersistKeysToStackExchangeRedis(() => redisService.GetDatabase(), "aspnet-data-protection");
-			
+
 			services.AddResponseCompression(options =>
 			{
 				options.EnableForHttps = true;
@@ -474,7 +410,7 @@ namespace HordeServer
 				options.Providers.Add<BrotliCompressionProvider>();
 				options.Providers.Add<ZstdCompressionProvider>();
 			});
-			
+
 			if (settings.CorsEnabled)
 			{
 				services.AddCors(options =>
@@ -960,7 +896,7 @@ namespace HordeServer
 				options.KnownNetworks.Clear();
 			});
 
-			services.AddMvc().AddJsonOptions(options => ConfigureJsonSerializer(options.JsonSerializerOptions));
+			services.AddMvc().AddJsonOptions(options => JsonUtils.ConfigureJsonSerializer(options.JsonSerializerOptions));
 			services.AddControllersWithViews(options => options.Filters.Add(new ObsoleteLoggingFilter()))
 				.AddRazorRuntimeCompilation();
 
@@ -1049,27 +985,6 @@ namespace HordeServer
 			LogValueFormatter.RegisterTypeAnnotation<UserId>("UserId");
 		}
 
-		public static JsonSerializerOptions JsonSerializerOptions { get; } = CreateJsonSerializerOptions();
-
-		static JsonSerializerOptions CreateJsonSerializerOptions()
-		{
-			JsonSerializerOptions options = new JsonSerializerOptions();
-			ConfigureJsonSerializer(options);
-			options.MakeReadOnly(true);
-			return options;
-		}
-
-		public static void ConfigureJsonSerializer(JsonSerializerOptions options)
-		{
-			HordeHttpClient.ConfigureJsonSerializer(options);
-			options.Converters.Add(new JsonObjectIdConverter());
-			options.Converters.Add(new JsonKnownTypesConverterFactory());
-			options.Converters.Add(new ObjectIdJsonConverterFactory());
-			options.Converters.Add(new SubResourceIdJsonConverterFactory());
-			options.Converters.Add(new JsonDateTimeConverter());
-			options.Converters.Add(new JsonTimeSpanConverter());
-		}
-
 		private static void ConfigureLogStorage(IServiceCollection services)
 		{
 			services.AddSingleton<PersistentLogStorage>();
diff --git a/Engine/Source/Programs/Horde/HordeServer/Storage/StorageConfig.cs b/Engine/Source/Programs/Horde/HordeServer/Storage/StorageConfig.cs
--- a/Engine/Source/Programs/Horde/HordeServer/Storage/StorageConfig.cs
+++ b/Engine/Source/Programs/Horde/HordeServer/Storage/StorageConfig.cs
@@ -16,6 +16,7 @@ using EpicGames.Horde.Storage;
 using HordeServer.Acls;
 using HordeServer.Server;
 using HordeServer.Storage.ObjectStores;
+using HordeServer.Utilities;
 
 namespace HordeServer.Storage
 {
@@ -102,7 +103,7 @@ namespace HordeServer.Storage
 				using (MemoryStream stream = new MemoryStream())
 				{
 					JsonSerializerOptions options = new JsonSerializerOptions();
-					Startup.ConfigureJsonSerializer(options);
+					JsonUtils.ConfigureJsonSerializer(options);
 
 					JsonSerializer.Serialize(stream, backendConfig, options: options);
 
diff --git a/Engine/Source/Programs/Horde/HordeServer/Streams/StreamConfig.cs b/Engine/Source/Programs/Horde/HordeServer/Streams/StreamConfig.cs
--- a/Engine/Source/Programs/Horde/HordeServer/Streams/StreamConfig.cs
+++ b/Engine/Source/Programs/Horde/HordeServer/Streams/StreamConfig.cs
@@ -321,7 +321,7 @@ namespace HordeServer.Streams
 			}
 
 			// Compute a hash of this stream revision to make it easier to detect changes
-			byte[] streamData = JsonSerializer.SerializeToUtf8Bytes(this, Startup.JsonSerializerOptions);
+			byte[] streamData = JsonSerializer.SerializeToUtf8Bytes(this, JsonUtils.DefaultSerializerOptions);
 			Revision = IoHash.Compute(streamData).ToString();
 		}
 
diff --git a/Engine/Source/Programs/Horde/HordeServer/Telemetry/Sinks/EpicTelemetrySink.cs b/Engine/Source/Programs/Horde/HordeServer/Telemetry/Sinks/EpicTelemetrySink.cs
--- a/Engine/Source/Programs/Horde/HordeServer/Telemetry/Sinks/EpicTelemetrySink.cs
+++ b/Engine/Source/Programs/Horde/HordeServer/Telemetry/Sinks/EpicTelemetrySink.cs
@@ -11,6 +11,7 @@ using System.Threading;
 using System.Threading.Tasks;
 using EpicGames.Core;
 using EpicGames.Horde.Telemetry;
+using HordeServer.Utilities;
 using Microsoft.AspNetCore.WebUtilities;
 using Microsoft.Extensions.Http;
 using Microsoft.Extensions.Logging;
@@ -140,7 +141,7 @@ namespace HordeServer.Telemetry.Sinks
 			_httpClient = CreateHttpClient();
 
 			_jsonOptions = new JsonSerializerOptions();
-			Startup.ConfigureJsonSerializer(_jsonOptions);
+			JsonUtils.ConfigureJsonSerializer(_jsonOptions);
 			_jsonOptions.PropertyNamingPolicy = null;
 			_jsonOptions.Converters.Insert(0, new TableauDateTimeConverter());
 
diff --git a/Engine/Source/Programs/Horde/HordeServer/Telemetry/Sinks/MetricTelemetrySink.cs b/Engine/Source/Programs/Horde/HordeServer/Telemetry/Sinks/MetricTelemetrySink.cs
--- a/Engine/Source/Programs/Horde/HordeServer/Telemetry/Sinks/MetricTelemetrySink.cs
+++ b/Engine/Source/Programs/Horde/HordeServer/Telemetry/Sinks/MetricTelemetrySink.cs
@@ -6,6 +6,7 @@ using System.Threading;
 using System.Threading.Tasks;
 using EpicGames.Horde.Telemetry;
 using HordeServer.Telemetry.Metrics;
+using HordeServer.Utilities;
 
 namespace HordeServer.Telemetry.Sinks
 {
@@ -22,7 +23,7 @@ namespace HordeServer.Telemetry.Sinks
 			_metricCollection = metricCollection;
 
 			_jsonOptions = new JsonSerializerOptions();
-			Startup.ConfigureJsonSerializer(_jsonOptions);
+			JsonUtils.ConfigureJsonSerializer(_jsonOptions);
 		}
 
 		/// <inheritdoc/>
diff --git a/Engine/Source/Programs/Horde/HordeServer/Telemetry/Sinks/MongoTelemetrySink.cs b/Engine/Source/Programs/Horde/HordeServer/Telemetry/Sinks/MongoTelemetrySink.cs
--- a/Engine/Source/Programs/Horde/HordeServer/Telemetry/Sinks/MongoTelemetrySink.cs
+++ b/Engine/Source/Programs/Horde/HordeServer/Telemetry/Sinks/MongoTelemetrySink.cs
@@ -71,7 +71,7 @@ namespace HordeServer.Telemetry.Sinks
 			_logger = logger;
 
 			_jsonOptions = new JsonSerializerOptions();
-			Startup.ConfigureJsonSerializer(_jsonOptions);
+			JsonUtils.ConfigureJsonSerializer(_jsonOptions);
 		}
 
 		/// <inheritdoc/>
diff --git a/Engine/Source/Programs/Horde/HordeServer/Telemetry/TelemetryManager.cs b/Engine/Source/Programs/Horde/HordeServer/Telemetry/TelemetryManager.cs
--- a/Engine/Source/Programs/Horde/HordeServer/Telemetry/TelemetryManager.cs
+++ b/Engine/Source/Programs/Horde/HordeServer/Telemetry/TelemetryManager.cs
@@ -9,6 +9,7 @@ using System.Threading.Tasks;
 using EpicGames.Core;
 using EpicGames.Horde.Telemetry;
 using HordeServer.Telemetry.Sinks;
+using HordeServer.Utilities;
 using Microsoft.Extensions.DependencyInjection;
 using Microsoft.Extensions.Hosting;
 using Microsoft.Extensions.Logging;
@@ -87,7 +88,7 @@ namespace HordeServer.Telemetry
 			try
 			{
 				JsonSerializerOptions options = new JsonSerializerOptions();
-				Startup.ConfigureJsonSerializer(options);
+				JsonUtils.ConfigureJsonSerializer(options);
 
 				return JsonSerializer.Serialize(telemetryEvent, options);
 			}
diff --git a/Engine/Source/Programs/Horde/HordeServer/Utilities/ObjectIdConverter.cs b/Engine/Source/Programs/Horde/HordeServer/Utilities/ObjectIdConverter.cs
deleted file mode 100644
