Horde: Missing files.

#rnx

diff --git a/Engine/Source/Programs/Horde/HordeServer.Tests/Ddc/Functional/Utils.cs b/Engine/Source/Programs/Horde/HordeServer.Tests/Ddc/Functional/Utils.cs
--- a/Engine/Source/Programs/Horde/HordeServer.Tests/Ddc/Functional/Utils.cs
+++ b/Engine/Source/Programs/Horde/HordeServer.Tests/Ddc/Functional/Utils.cs
@@ -1,6 +1,7 @@
 // Copyright Epic Games, Inc. All Rights Reserved.
 
 using System.Text.Json;
+using HordeServer.Utilities;
 
 namespace HordeServer.Tests.Ddc.FunctionalTests
 {
@@ -11,7 +12,7 @@ namespace HordeServer.Tests.Ddc.FunctionalTests
 		private static JsonSerializerOptions ConfigureJsonOptions()
 		{
 			JsonSerializerOptions options = new JsonSerializerOptions();
-			Startup.ConfigureJsonSerializer(options);
+			JsonUtils.ConfigureJsonSerializer(options);
 			return options;
 		}
 	}
diff --git a/Engine/Source/Programs/Horde/HordeServer.Tests/Telemetry/TelemetryTests.cs b/Engine/Source/Programs/Horde/HordeServer.Tests/Telemetry/TelemetryTests.cs
--- a/Engine/Source/Programs/Horde/HordeServer.Tests/Telemetry/TelemetryTests.cs
+++ b/Engine/Source/Programs/Horde/HordeServer.Tests/Telemetry/TelemetryTests.cs
@@ -12,6 +12,7 @@ using HordeServer.Server;
 using HordeServer.Telemetry;
 using HordeServer.Telemetry.Metrics;
 using HordeServer.Telemetry.Sinks;
+using HordeServer.Utilities;
 using Json.Path;
 using Microsoft.Extensions.DependencyInjection;
 using Microsoft.VisualStudio.TestTools.UnitTesting;
@@ -334,7 +335,7 @@ namespace HordeServer.Tests.Telemetry
 		public void GroupConverterTest()
 		{
 			JsonSerializerOptions serializerOptions = new JsonSerializerOptions();
-			Startup.ConfigureJsonSerializer(serializerOptions);
+			JsonUtils.ConfigureJsonSerializer(serializerOptions);
 
 			MetricConfig? config = JsonSerializer.Deserialize<MetricConfig>("{ \"groupBy\": \"$.foo,$.bar\" }", serializerOptions);
 			Assert.IsNotNull(config);
