Horde: Add an ITelemetryWriter interface, distinct from ITelemetrySink, as the target for writing telemetry events. Allows abstracting Horde server metadata when sending new events.

#rnx

diff --git a/Engine/Source/Programs/Horde/HordeServer.Shared/Telemetry/ITelemetryWriter.cs b/Engine/Source/Programs/Horde/HordeServer.Shared/Telemetry/ITelemetryWriter.cs
new file mode 100644
--- /dev/null
+++ b/Engine/Source/Programs/Horde/HordeServer.Shared/Telemetry/ITelemetryWriter.cs
@@ -0,0 +1,40 @@
+// Copyright Epic Games, Inc. All Rights Reserved.
+
+using EpicGames.Horde.Telemetry;
+
+namespace HordeServer.Telemetry
+{
+	/// <summary>
+	/// Interface for writing telemetry events from the server
+	/// </summary>
+	public interface ITelemetryWriter
+	{
+		/// <summary>
+		/// Whether the writer is enabled. Can be used to shortcut construction of event data.
+		/// </summary>
+		bool Enabled { get; }
+
+		/// <summary>
+		/// Writes a telemetry event using the server's metadata
+		/// </summary>
+		/// <param name="telemetryStoreId">The telemetry store to write to</param>
+		/// <param name="payload">Event data to write</param>
+		void WriteEvent(TelemetryStoreId telemetryStoreId, object payload);
+
+		/// <summary>
+		/// Writes a telemetry event
+		/// </summary>
+		/// <param name="telemetryStoreId">The telemetry store to write to</param>
+		/// <param name="payload">Event data to write</param>
+		void WriteEvent(TelemetryStoreId telemetryStoreId, TelemetryRecordMeta metadata, object payload);
+	}
+
+	/// <summary>
+	/// Additional metadata associated with a telemetry event
+	/// </summary>
+	/// <param name="AppId">Identifier of the application sending the event</param>
+	/// <param name="AppVersion">Version number of the application</param>
+	/// <param name="AppEnvironment">Name of the environment that the sending application is running in</param>
+	/// <param name="SessionId">Unique identifier for the current session</param>
+	public record class TelemetryRecordMeta(string? AppId = null, string? AppVersion = null, string? AppEnvironment = null, string? SessionId = null);
+}
diff --git a/Engine/Source/Programs/Horde/HordeServer.Tests/Telemetry/TelemetryTests.cs b/Engine/Source/Programs/Horde/HordeServer.Tests/Telemetry/TelemetryTests.cs
--- a/Engine/Source/Programs/Horde/HordeServer.Tests/Telemetry/TelemetryTests.cs
+++ b/Engine/Source/Programs/Horde/HordeServer.Tests/Telemetry/TelemetryTests.cs
@@ -22,6 +22,8 @@ namespace HordeServer.Tests.Telemetry
 	[TestClass]
 	public class TelemetryTests : TestSetup
 	{
+		static readonly TelemetryRecordMeta s_metadata = new TelemetryRecordMeta();
+
 		protected override void ConfigureServices(IServiceCollection services)
 		{
 			base.ConfigureServices(services);
@@ -55,9 +57,9 @@ namespace HordeServer.Tests.Telemetry
 
 			// Test 1
 			{
-				sink.SendEvent(telemetryStoreConfig.Id, TelemetryRecordMeta.CurrentHordeInstance, new { EventName = "Included", foo = 1 });
-				sink.SendEvent(telemetryStoreConfig.Id, TelemetryRecordMeta.CurrentHordeInstance, new { EventName = "Included", foo = 2 });
-				sink.SendEvent(telemetryStoreConfig.Id, TelemetryRecordMeta.CurrentHordeInstance, new { EventName = "Excluded", foo = 3 });
+				sink.SendEvent(telemetryStoreConfig.Id, new TelemetryEvent(s_metadata, new { EventName = "Included", foo = 1 }));
+				sink.SendEvent(telemetryStoreConfig.Id, new TelemetryEvent(s_metadata, new { EventName = "Included", foo = 2 }));
+				sink.SendEvent(telemetryStoreConfig.Id, new TelemetryEvent(s_metadata, new { EventName = "Excluded", foo = 3 }));
 				await sink.FlushAsync(CancellationToken.None);
 				await collection.FlushAsync(CancellationToken.None);
 
@@ -93,9 +95,9 @@ namespace HordeServer.Tests.Telemetry
 
 			// Test 1
 			{
-				sink.SendEvent(telemetryStoreConfig.Id, TelemetryRecordMeta.CurrentHordeInstance, new { foo = 1 });
-				sink.SendEvent(telemetryStoreConfig.Id, TelemetryRecordMeta.CurrentHordeInstance, new { foo = 2 });
-				sink.SendEvent(telemetryStoreConfig.Id, TelemetryRecordMeta.CurrentHordeInstance, new { foo = 3 });
+				sink.SendEvent(telemetryStoreConfig.Id, new TelemetryEvent(s_metadata, new { foo = 1 }));
+				sink.SendEvent(telemetryStoreConfig.Id, new TelemetryEvent(s_metadata, new { foo = 2 }));
+				sink.SendEvent(telemetryStoreConfig.Id, new TelemetryEvent(s_metadata, new { foo = 3 }));
 				await sink.FlushAsync(CancellationToken.None);
 				await collection.FlushAsync(CancellationToken.None);
 
@@ -108,7 +110,7 @@ namespace HordeServer.Tests.Telemetry
 
 			// Test 2
 			{
-				sink.SendEvent(telemetryStoreConfig.Id, TelemetryRecordMeta.CurrentHordeInstance, new { foo = 3 });
+				sink.SendEvent(telemetryStoreConfig.Id, new TelemetryEvent(s_metadata, new { foo = 3 }));
 				await sink.FlushAsync(CancellationToken.None);
 				await collection.FlushAsync(CancellationToken.None);
 
@@ -123,7 +125,7 @@ namespace HordeServer.Tests.Telemetry
 			{
 				await Clock.AdvanceAsync(TimeSpan.FromHours(1.0));
 
-				sink.SendEvent(telemetryStoreConfig.Id, TelemetryRecordMeta.CurrentHordeInstance, new { foo = 4 });
+				sink.SendEvent(telemetryStoreConfig.Id, new TelemetryEvent(s_metadata, new { foo = 4 }));
 				await sink.FlushAsync(CancellationToken.None);
 				await collection.FlushAsync(CancellationToken.None);
 
@@ -171,9 +173,9 @@ namespace HordeServer.Tests.Telemetry
 
 			// Test 1
 			{
-				sink.SendEvent(telemetryStoreConfig.Id, TelemetryRecordMeta.CurrentHordeInstance, new { foo = 1 });
-				sink.SendEvent(telemetryStoreConfig.Id, TelemetryRecordMeta.CurrentHordeInstance, new { foo = 2, bar = 201 });
-				sink.SendEvent(telemetryStoreConfig.Id, TelemetryRecordMeta.CurrentHordeInstance, new { foo = 3, bar = new { baz = 101 } });
+				sink.SendEvent(telemetryStoreConfig.Id, new TelemetryEvent(s_metadata, new { foo = 1 }));
+				sink.SendEvent(telemetryStoreConfig.Id, new TelemetryEvent(s_metadata, new { foo = 2, bar = 201 }));
+				sink.SendEvent(telemetryStoreConfig.Id, new TelemetryEvent(s_metadata, new { foo = 3, bar = new { baz = 101 } }));
 				await sink.FlushAsync(CancellationToken.None);
 				await collection.FlushAsync(CancellationToken.None);
 
@@ -235,7 +237,7 @@ namespace HordeServer.Tests.Telemetry
 
 			foreach (double value in values)
 			{
-				sink.SendEvent(telemetryStoreConfig.Id, TelemetryRecordMeta.CurrentHordeInstance, new { foo = value });
+				sink.SendEvent(telemetryStoreConfig.Id, new TelemetryEvent(s_metadata, new { foo = value }));
 			}
 
 			await sink.FlushAsync(CancellationToken.None);
@@ -270,11 +272,11 @@ namespace HordeServer.Tests.Telemetry
 			MetricTelemetrySink sink = ServiceProvider.GetRequiredService<MetricTelemetrySink>();
 			IMetricCollection collection = ServiceProvider.GetRequiredService<IMetricCollection>();
 
-			sink.SendEvent(TelemetryStoreId.Default, TelemetryRecordMeta.CurrentHordeInstance, new { foo = 1, group = "first" });
-			sink.SendEvent(TelemetryStoreId.Default, TelemetryRecordMeta.CurrentHordeInstance, new { foo = 2, group = "first" });
-			sink.SendEvent(TelemetryStoreId.Default, TelemetryRecordMeta.CurrentHordeInstance, new { foo = 3, group = "first" });
-			sink.SendEvent(TelemetryStoreId.Default, TelemetryRecordMeta.CurrentHordeInstance, new { foo = 4, group = "second" });
-			sink.SendEvent(TelemetryStoreId.Default, TelemetryRecordMeta.CurrentHordeInstance, new { foo = 5 });
+			sink.SendEvent(TelemetryStoreId.Default, new TelemetryEvent(s_metadata, new { foo = 1, group = "first" }));
+			sink.SendEvent(TelemetryStoreId.Default, new TelemetryEvent(s_metadata, new { foo = 2, group = "first" }));
+			sink.SendEvent(TelemetryStoreId.Default, new TelemetryEvent(s_metadata, new { foo = 3, group = "first" }));
+			sink.SendEvent(TelemetryStoreId.Default, new TelemetryEvent(s_metadata, new { foo = 4, group = "second" }));
+			sink.SendEvent(TelemetryStoreId.Default, new TelemetryEvent(s_metadata, new { foo = 5 }));
 			await sink.FlushAsync(CancellationToken.None);
 			await collection.FlushAsync(CancellationToken.None);
 
@@ -318,7 +320,7 @@ namespace HordeServer.Tests.Telemetry
 			MetricTelemetrySink sink = ServiceProvider.GetRequiredService<MetricTelemetrySink>();
 			IMetricCollection collection = ServiceProvider.GetRequiredService<IMetricCollection>();
 
-			sink.SendEvent(telemetryStoreConfig.Id, TelemetryRecordMeta.CurrentHordeInstance, new { foo = 123, group = "first, second & third" });
+			sink.SendEvent(telemetryStoreConfig.Id, new TelemetryEvent(s_metadata, new { foo = 123, group = "first, second & third" }));
 			await sink.FlushAsync(CancellationToken.None);
 			await collection.FlushAsync(CancellationToken.None);
 
@@ -378,13 +380,13 @@ namespace HordeServer.Tests.Telemetry
 
 			// Test 1
 			{
-				sink.SendEvent(telemetryStoreConfig.Id, TelemetryRecordMeta.CurrentHordeInstance, new { EventName = "Included", foo = 1, groupFacetA = "groupA" });
-				sink.SendEvent(telemetryStoreConfig.Id, TelemetryRecordMeta.CurrentHordeInstance, new { EventName = "Included", foo = 2, groupFacetA = "groupA" });
-				sink.SendEvent(telemetryStoreConfig.Id, TelemetryRecordMeta.CurrentHordeInstance, new { EventName = "Included", foo = 3, groupFacetA = "groupA", groupFacetB = "groupB" });
-				sink.SendEvent(telemetryStoreConfig.Id, TelemetryRecordMeta.CurrentHordeInstance, new { EventName = "Included", foo = 4, groupFacetB = "groupB" });
-				sink.SendEvent(telemetryStoreConfig.Id, TelemetryRecordMeta.CurrentHordeInstance, new { EventName = "Included", foo = 5, groupFacetB = "groupA,groupB" });
-				sink.SendEvent(telemetryStoreConfig.Id, TelemetryRecordMeta.CurrentHordeInstance, new { EventName = "Included", foo = 6, groupFacetA = "groupA", groupFacetB = "groupB", groupFacetC = "groupC" });
-				sink.SendEvent(telemetryStoreConfig.Id, TelemetryRecordMeta.CurrentHordeInstance, new { EventName = "Excluded", foo = 6 });
+				sink.SendEvent(telemetryStoreConfig.Id, new TelemetryEvent(s_metadata, new { EventName = "Included", foo = 1, groupFacetA = "groupA" }));
+				sink.SendEvent(telemetryStoreConfig.Id, new TelemetryEvent(s_metadata, new { EventName = "Included", foo = 2, groupFacetA = "groupA" }));
+				sink.SendEvent(telemetryStoreConfig.Id, new TelemetryEvent(s_metadata, new { EventName = "Included", foo = 3, groupFacetA = "groupA", groupFacetB = "groupB" }));
+				sink.SendEvent(telemetryStoreConfig.Id, new TelemetryEvent(s_metadata, new { EventName = "Included", foo = 4, groupFacetB = "groupB" }));
+				sink.SendEvent(telemetryStoreConfig.Id, new TelemetryEvent(s_metadata, new { EventName = "Included", foo = 5, groupFacetB = "groupA,groupB" }));
+				sink.SendEvent(telemetryStoreConfig.Id, new TelemetryEvent(s_metadata, new { EventName = "Included", foo = 6, groupFacetA = "groupA", groupFacetB = "groupB", groupFacetC = "groupC" }));
+				sink.SendEvent(telemetryStoreConfig.Id, new TelemetryEvent(s_metadata, new { EventName = "Excluded", foo = 6 }));
 				await sink.FlushAsync(CancellationToken.None);
 				await collection.FlushAsync(CancellationToken.None);
 
diff --git a/Engine/Source/Programs/Horde/HordeServer.Tests/TestSetup.cs b/Engine/Source/Programs/Horde/HordeServer.Tests/TestSetup.cs
--- a/Engine/Source/Programs/Horde/HordeServer.Tests/TestSetup.cs
+++ b/Engine/Source/Programs/Horde/HordeServer.Tests/TestSetup.cs
@@ -201,7 +201,7 @@ namespace HordeServer.Tests
 			services.AddSingleton(typeof(IAuditLogFactory<>), typeof(AuditLogFactory<>));
 			services.AddSingleton<IAuditLog<AgentId>>(sp => sp.GetRequiredService<IAuditLogFactory<AgentId>>().Create("Agents.Log", "AgentId"));
 			services.AddSingleton<ITelemetrySink, NullTelemetrySink>();
-			services.AddSingleton<ITelemetrySinkInternal, MetricTelemetrySink>();
+			services.AddSingleton<ITelemetrySink, MetricTelemetrySink>();
 			services.AddSingleton<OpenTelemetry.Trace.Tracer>(sp => TracerProvider.Default.GetTracer("TestTracer"));
 			services.AddSingleton(sp => new Meter("TestMeter"));
 
diff --git a/Engine/Source/Programs/Horde/HordeServer/Issues/IssueCollection.cs b/Engine/Source/Programs/Horde/HordeServer/Issues/IssueCollection.cs
--- a/Engine/Source/Programs/Horde/HordeServer/Issues/IssueCollection.cs
+++ b/Engine/Source/Programs/Horde/HordeServer/Issues/IssueCollection.cs
@@ -596,7 +596,7 @@ namespace HordeServer.Issues
 		readonly IMongoCollection<IssueStep> _issueSteps;
 		readonly IMongoCollection<IssueSuspect> _issueSuspects;
 		readonly IAuditLog<int> _auditLog;
-		readonly ITelemetrySink _telemetrySink;
+		readonly ITelemetryWriter _telemetryWriter;
 		readonly IOptionsMonitor<GlobalConfig> _globalConfig;
 		readonly Tracer _tracer;
 		readonly ILogger _logger;
@@ -619,11 +619,11 @@ namespace HordeServer.Issues
 			});
 		}
 
-		public IssueCollection(IMongoService mongoService, IRedisService redisService, IUserCollection userCollection, IAuditLogFactory<int> auditLogFactory, ITelemetrySink telemetrySink, IOptionsMonitor<GlobalConfig> globalConfig, Tracer tracer, ILogger<IssueCollection> logger)
+		public IssueCollection(IMongoService mongoService, IRedisService redisService, IUserCollection userCollection, IAuditLogFactory<int> auditLogFactory, ITelemetryWriter telemetryWriter, IOptionsMonitor<GlobalConfig> globalConfig, Tracer tracer, ILogger<IssueCollection> logger)
 		{
 			_redisService = redisService;
 			_userCollection = userCollection;
-			_telemetrySink = telemetrySink;
+			_telemetryWriter = telemetryWriter;
 			_globalConfig = globalConfig;
 			_tracer = tracer;
 			_logger = logger;
@@ -741,7 +741,7 @@ namespace HordeServer.Issues
 
 			foreach (TelemetryStoreId telemetryStoreId in telemetryStoreIds)
 			{
-				_telemetrySink.SendEvent(telemetryStoreId, TelemetryRecordMeta.CurrentHordeInstance, new
+				_telemetryWriter.WriteEvent(telemetryStoreId, new
 				{
 					EventName = "State.Issue",
 					Id = issue.Id,
@@ -1512,7 +1512,7 @@ namespace HordeServer.Issues
 		{
 			if (_globalConfig.CurrentValue.TryGetStream(issueSpan.StreamId, out StreamConfig? streamConfig) && !streamConfig.TelemetryStoreId.IsEmpty)
 			{
-				_telemetrySink.SendEvent(streamConfig.TelemetryStoreId, TelemetryRecordMeta.CurrentHordeInstance, new
+				_telemetryWriter.WriteEvent(streamConfig.TelemetryStoreId, new
 				{
 					EventName = "State.IssueSpan",
 					Id = issueSpan.Id,
diff --git a/Engine/Source/Programs/Horde/HordeServer/Jobs/JobCollection.cs b/Engine/Source/Programs/Horde/HordeServer/Jobs/JobCollection.cs
--- a/Engine/Source/Programs/Horde/HordeServer/Jobs/JobCollection.cs
+++ b/Engine/Source/Programs/Horde/HordeServer/Jobs/JobCollection.cs
@@ -727,7 +727,7 @@ namespace HordeServer.Jobs
 		readonly MongoIndex<JobDocument> _createTimeIndex;
 		readonly MongoIndex<JobDocument> _streamThenTemplateThenCreationTimeIndex;
 		readonly MongoIndex<JobDocument> _startedByBisectTaskIdIndex;
-		readonly ITelemetrySink _telemetrySink;
+		readonly ITelemetryWriter _telemetryWriter;
 		readonly IClock _clock;
 		readonly IGraphCollection _graphCollection;
 		readonly ILogCollection _logCollection;
@@ -738,12 +738,12 @@ namespace HordeServer.Jobs
 		/// <summary>
 		/// Constructor
 		/// </summary>
-		public JobCollection(IMongoService mongoService, IClock clock, IGraphCollection graphCollection, ILogCollection logCollection, ITelemetrySink telemetrySink, IOptionsMonitor<GlobalConfig> globalConfig, Tracer tracer, ILogger<JobCollection> logger)
+		public JobCollection(IMongoService mongoService, IClock clock, IGraphCollection graphCollection, ILogCollection logCollection, ITelemetryWriter telemetryWriter, IOptionsMonitor<GlobalConfig> globalConfig, Tracer tracer, ILogger<JobCollection> logger)
 		{
 			_clock = clock;
 			_graphCollection = graphCollection;
 			_logCollection = logCollection;
-			_telemetrySink = telemetrySink;
+			_telemetryWriter = telemetryWriter;
 			_globalConfig = globalConfig;
 			_tracer = tracer;
 			_logger = logger;
@@ -792,7 +792,7 @@ namespace HordeServer.Jobs
 
 			if (_globalConfig.CurrentValue.TryGetStream(streamId, out StreamConfig? streamConfig) && !streamConfig.TelemetryStoreId.IsEmpty)
 			{
-				_telemetrySink.SendEvent(streamConfig.TelemetryStoreId, TelemetryRecordMeta.CurrentHordeInstance, new
+				_telemetryWriter.WriteEvent(streamConfig.TelemetryStoreId, new
 				{
 					EventName = "State.Job",
 					Id = newJob.Id,
diff --git a/Engine/Source/Programs/Horde/HordeServer/Jobs/JobStepRefCollection.cs b/Engine/Source/Programs/Horde/HordeServer/Jobs/JobStepRefCollection.cs
--- a/Engine/Source/Programs/Horde/HordeServer/Jobs/JobStepRefCollection.cs
+++ b/Engine/Source/Programs/Horde/HordeServer/Jobs/JobStepRefCollection.cs
@@ -104,19 +104,19 @@ namespace HordeServer.Jobs
 		}
 
 		readonly IMongoCollection<JobStepRef> _jobStepRefs;
-		readonly ITelemetrySink _telemetrySink;
+		readonly ITelemetryWriter _telemetryWriter;
 		readonly IOptionsMonitor<GlobalConfig> _globalConfig;
 
 		/// <summary>
 		/// Constructor
 		/// </summary>
-		public JobStepRefCollection(IMongoService mongoService, ITelemetrySink telemetrySink, IOptionsMonitor<GlobalConfig> globalConfig)
+		public JobStepRefCollection(IMongoService mongoService, ITelemetryWriter telemetryWriter, IOptionsMonitor<GlobalConfig> globalConfig)
 		{
 			List<MongoIndex<JobStepRef>> indexes = new List<MongoIndex<JobStepRef>>();
 			indexes.Add(keys => keys.Ascending(x => x.StreamId).Ascending(x => x.TemplateId).Ascending(x => x.Name).Descending(x => x.Change));
 
 			_jobStepRefs = mongoService.GetCollection<JobStepRef>("JobStepRefs", indexes);
-			_telemetrySink = telemetrySink;
+			_telemetryWriter = telemetryWriter;
 			_globalConfig = globalConfig;
 		}
 
@@ -128,7 +128,7 @@ namespace HordeServer.Jobs
 
 			if (_globalConfig.CurrentValue.TryGetStream(streamId, out StreamConfig? streamConfig) && !streamConfig.TelemetryStoreId.IsEmpty)
 			{
-				_telemetrySink.SendEvent(streamConfig.TelemetryStoreId, TelemetryRecordMeta.CurrentHordeInstance, new
+				_telemetryWriter.WriteEvent(streamConfig.TelemetryStoreId, new
 				{
 					EventName = "State.JobStepRef",
 					Id = id.ToString(),
diff --git a/Engine/Source/Programs/Horde/HordeServer/Server/RpcService.cs b/Engine/Source/Programs/Horde/HordeServer/Server/RpcService.cs
--- a/Engine/Source/Programs/Horde/HordeServer/Server/RpcService.cs
+++ b/Engine/Source/Programs/Horde/HordeServer/Server/RpcService.cs
@@ -46,7 +46,7 @@ namespace HordeServer.Server
 
 		readonly AgentService _agentService;
 		readonly LifetimeService _lifetimeService;
-		readonly ITelemetrySink _telemetrySink;
+		readonly ITelemetryWriter _telemetryWriter;
 		readonly IToolCollection _toolCollection;
 		readonly IAgentTelemetryCollection _agentTelemetryCollection;
 		readonly AclService _aclService;
@@ -56,11 +56,11 @@ namespace HordeServer.Server
 		/// <summary>
 		/// Constructor
 		/// </summary>
-		public RpcService(AgentService agentService, LifetimeService lifetimeService, ITelemetrySink telemetrySink, IToolCollection toolCollection, IAgentTelemetryCollection agentTelemetryCollection, AclService aclService, IOptionsSnapshot<GlobalConfig> globalConfig, ILogger<RpcService> logger)
+		public RpcService(AgentService agentService, LifetimeService lifetimeService, ITelemetryWriter telemetryWriter, IToolCollection toolCollection, IAgentTelemetryCollection agentTelemetryCollection, AclService aclService, IOptionsSnapshot<GlobalConfig> globalConfig, ILogger<RpcService> logger)
 		{
 			_agentService = agentService;
 			_lifetimeService = lifetimeService;
-			_telemetrySink = telemetrySink;
+			_telemetryWriter = telemetryWriter;
 			_toolCollection = toolCollection;
 			_agentTelemetryCollection = agentTelemetryCollection;
 			_aclService = aclService;
@@ -480,7 +480,7 @@ namespace HordeServer.Server
 				FieldDescriptor caseDescriptor = oneofDescriptor.Accessor.GetCaseFieldDescriptor(wrappedEvent);
 
 				object wrappedValue = caseDescriptor.Accessor.GetValue(wrappedEvent);
-				_telemetrySink.SendEvent(TelemetryStoreId.Default, agentMeta, wrappedValue);
+				_telemetryWriter.WriteEvent(TelemetryStoreId.Default, agentMeta, wrappedValue);
 			}
 
 			return new Empty();
diff --git a/Engine/Source/Programs/Horde/HordeServer/Startup.cs b/Engine/Source/Programs/Horde/HordeServer/Startup.cs
--- a/Engine/Source/Programs/Horde/HordeServer/Startup.cs
+++ b/Engine/Source/Programs/Horde/HordeServer/Startup.cs
@@ -787,7 +787,7 @@ namespace HordeServer
 				});
 
 			services.AddSingleton<TelemetryManager>();
-			services.AddSingleton<ITelemetrySink>(sp => sp.GetRequiredService<TelemetryManager>());
+			services.AddSingleton<ITelemetryWriter>(sp => sp.GetRequiredService<TelemetryManager>());
 			services.AddHostedService(sp => sp.GetRequiredService<TelemetryManager>());
 			services.AddSingleton<MongoTelemetrySink>();
 			services.AddHostedService(sp => sp.GetRequiredService<MongoTelemetrySink>());
diff --git a/Engine/Source/Programs/Horde/HordeServer/Telemetry/ITelemetrySink.cs b/Engine/Source/Programs/Horde/HordeServer/Telemetry/ITelemetrySink.cs
--- a/Engine/Source/Programs/Horde/HordeServer/Telemetry/ITelemetrySink.cs
+++ b/Engine/Source/Programs/Horde/HordeServer/Telemetry/ITelemetrySink.cs
@@ -10,46 +10,26 @@ namespace HordeServer.Telemetry
 	/// <summary>
 	/// Interface for a telemetry sink
 	/// </summary>
-	public interface ITelemetrySink
+	public interface ITelemetrySink : IAsyncDisposable
 	{
 		/// <summary>
 		/// Whether the sink is enabled
 		/// </summary>
 		bool Enabled { get; }
 
-		/// <summary>
-		/// Sends a telemetry event with the given information
-		/// </summary>
-		/// <param name="storeId">Identifier for the telemetry store</param>
-		/// <param name="telemetryEvent">The telemetry event that was received</param>
-		void SendEvent(TelemetryStoreId storeId, TelemetryEvent telemetryEvent);
-	}
-
-	/// <summary>
-	/// Internal interface for sinks
-	/// </summary>
-	public interface ITelemetrySinkInternal : ITelemetrySink, IAsyncDisposable
-	{
 		/// <summary>
 		/// Flush any queued metrics to underlying service
 		/// </summary>
-		/// <param name="cancellationToken"></param>
+		/// <param name="cancellationToken">Cancellation token for the operation</param>
 		/// <returns>Completion task</returns>
-		public ValueTask FlushAsync(CancellationToken cancellationToken);
-	}
+		ValueTask FlushAsync(CancellationToken cancellationToken);
 
-	/// <summary>
-	/// Extension methods for telemetry sinks
-	/// </summary>
-	public static class TelemetrySinkExtensions
-	{
 		/// <summary>
 		/// Sends a telemetry event with the given information
 		/// </summary>
-		public static void SendEvent(this ITelemetrySink sink, TelemetryStoreId storeId, TelemetryRecordMeta recordMeta, object payload)
-		{
-			sink.SendEvent(storeId, new TelemetryEvent(recordMeta, payload));
-		}
+		/// <param name="storeId">Identifier for the telemetry store</param>
+		/// <param name="telemetryEvent">The telemetry event that was received</param>
+		void SendEvent(TelemetryStoreId storeId, TelemetryEvent telemetryEvent);
 	}
 }
 
diff --git a/Engine/Source/Programs/Horde/HordeServer/Telemetry/Sinks/ClickHouseTelemetrySink.cs b/Engine/Source/Programs/Horde/HordeServer/Telemetry/Sinks/ClickHouseTelemetrySink.cs
--- a/Engine/Source/Programs/Horde/HordeServer/Telemetry/Sinks/ClickHouseTelemetrySink.cs
+++ b/Engine/Source/Programs/Horde/HordeServer/Telemetry/Sinks/ClickHouseTelemetrySink.cs
@@ -17,7 +17,7 @@ namespace HordeServer.Telemetry.Sinks
 	/// <summary>
 	/// ClickHouse telemetry sink
 	/// </summary>
-	public sealed class ClickHouseTelemetrySink : ITelemetrySinkInternal
+	public sealed class ClickHouseTelemetrySink : ITelemetrySink
 	{
 		/// <summary>
 		/// Name of the HTTP client for writing telemetry data
diff --git a/Engine/Source/Programs/Horde/HordeServer/Telemetry/Sinks/EpicTelemetrySink.cs b/Engine/Source/Programs/Horde/HordeServer/Telemetry/Sinks/EpicTelemetrySink.cs
--- a/Engine/Source/Programs/Horde/HordeServer/Telemetry/Sinks/EpicTelemetrySink.cs
+++ b/Engine/Source/Programs/Horde/HordeServer/Telemetry/Sinks/EpicTelemetrySink.cs
@@ -39,7 +39,7 @@ namespace HordeServer.Telemetry.Sinks
 	/// <summary>
 	/// Epic internal telemetry sink using the data router
 	/// </summary>
-	public sealed class EpicTelemetrySink : ITelemetrySinkInternal
+	public sealed class EpicTelemetrySink : ITelemetrySink
 	{
 		// Converter for datetime formats that Tableau can ingest
 		class TableauDateTimeConverter : JsonConverter<DateTime>
diff --git a/Engine/Source/Programs/Horde/HordeServer/Telemetry/Sinks/MetricTelemetrySink.cs b/Engine/Source/Programs/Horde/HordeServer/Telemetry/Sinks/MetricTelemetrySink.cs
--- a/Engine/Source/Programs/Horde/HordeServer/Telemetry/Sinks/MetricTelemetrySink.cs
+++ b/Engine/Source/Programs/Horde/HordeServer/Telemetry/Sinks/MetricTelemetrySink.cs
@@ -13,7 +13,7 @@ namespace HordeServer.Telemetry.Sinks
 	/// <summary>
 	/// Consumes telemetry events and generates metrics
 	/// </summary>
-	class MetricTelemetrySink : ITelemetrySinkInternal
+	class MetricTelemetrySink : ITelemetrySink
 	{
 		readonly IMetricCollection _metricCollection;
 		readonly JsonSerializerOptions _jsonOptions;
diff --git a/Engine/Source/Programs/Horde/HordeServer/Telemetry/Sinks/MongoTelemetrySink.cs b/Engine/Source/Programs/Horde/HordeServer/Telemetry/Sinks/MongoTelemetrySink.cs
--- a/Engine/Source/Programs/Horde/HordeServer/Telemetry/Sinks/MongoTelemetrySink.cs
+++ b/Engine/Source/Programs/Horde/HordeServer/Telemetry/Sinks/MongoTelemetrySink.cs
@@ -21,7 +21,7 @@ namespace HordeServer.Telemetry.Sinks
 	/// <summary>
 	/// Telemetry sink which writes data to a MongoDB collection
 	/// </summary>
-	public sealed class MongoTelemetrySink : ITelemetrySinkInternal, IHostedService, IAsyncDisposable
+	public sealed class MongoTelemetrySink : ITelemetrySink, IHostedService, IAsyncDisposable
 	{
 		class EventDocument
 		{
diff --git a/Engine/Source/Programs/Horde/HordeServer/Telemetry/Sinks/NullTelemetrySink.cs b/Engine/Source/Programs/Horde/HordeServer/Telemetry/Sinks/NullTelemetrySink.cs
--- a/Engine/Source/Programs/Horde/HordeServer/Telemetry/Sinks/NullTelemetrySink.cs
+++ b/Engine/Source/Programs/Horde/HordeServer/Telemetry/Sinks/NullTelemetrySink.cs
@@ -1,5 +1,7 @@
 // Copyright Epic Games, Inc. All Rights Reserved.
 
+using System.Threading;
+using System.Threading.Tasks;
 using EpicGames.Horde.Telemetry;
 
 namespace HordeServer.Telemetry.Sinks
@@ -12,9 +14,16 @@ namespace HordeServer.Telemetry.Sinks
 		/// <inheritdoc/>
 		public bool Enabled => false;
 
+		/// <inheritdoc/>
+		public ValueTask DisposeAsync()
+			=> default;
+
+		/// <inheritdoc/>
+		public ValueTask FlushAsync(CancellationToken cancellationToken)
+			=> default;
+
 		/// <inheritdoc/>
 		public void SendEvent(TelemetryStoreId telemetryStoreId, TelemetryEvent telemetryEvent)
-		{
-		}
+		{ }
 	}
 }
diff --git a/Engine/Source/Programs/Horde/HordeServer/Telemetry/TelemetryController.cs b/Engine/Source/Programs/Horde/HordeServer/Telemetry/TelemetryController.cs
--- a/Engine/Source/Programs/Horde/HordeServer/Telemetry/TelemetryController.cs
+++ b/Engine/Source/Programs/Horde/HordeServer/Telemetry/TelemetryController.cs
@@ -157,7 +157,7 @@ namespace HordeServer.Telemetry
 			TelemetryRecordMeta recordMeta = new TelemetryRecordMeta(AppId: appId, AppVersion: appVersion, AppEnvironment: appEnvironment);
 			foreach (JsonObject eventPayload in request.Events)
 			{
-				_telemetryManager.SendEvent(storeId, recordMeta, eventPayload);
+				_telemetryManager.WriteEvent(storeId, recordMeta, eventPayload);
 			}
 
 			return NoContent();
diff --git a/Engine/Source/Programs/Horde/HordeServer/Telemetry/TelemetryEvent.cs b/Engine/Source/Programs/Horde/HordeServer/Telemetry/TelemetryEvent.cs
--- a/Engine/Source/Programs/Horde/HordeServer/Telemetry/TelemetryEvent.cs
+++ b/Engine/Source/Programs/Horde/HordeServer/Telemetry/TelemetryEvent.cs
@@ -26,25 +26,4 @@ namespace HordeServer.Telemetry
 			Payload = payload;
 		}
 	}
-
-	/// <summary>
-	/// Additional metadata associated with a telemetry event
-	/// </summary>
-	/// <param name="AppId">Identifier of the application sending the event</param>
-	/// <param name="AppVersion">Version number of the application</param>
-	/// <param name="AppEnvironment">Name of the environment that the sending application is running in</param>
-	/// <param name="SessionId">Unique identifier for the current session</param>
-	public record class TelemetryRecordMeta(string? AppId = null, string? AppVersion = null, string? AppEnvironment = null, string? SessionId = null)
-	{
-		/// <summary>
-		/// Record metadata for the running Horde instance
-		/// </summary>
-		public static TelemetryRecordMeta CurrentHordeInstance { get; } = new TelemetryRecordMeta
-		{
-			AppId = "Horde",
-			AppVersion = ServerApp.Version.ToString(),
-			AppEnvironment = ServerApp.DeploymentEnvironment,
-			SessionId = ServerApp.SessionId
-		};
-	}
 }
diff --git a/Engine/Source/Programs/Horde/HordeServer/Telemetry/TelemetryManager.cs b/Engine/Source/Programs/Horde/HordeServer/Telemetry/TelemetryManager.cs
--- a/Engine/Source/Programs/Horde/HordeServer/Telemetry/TelemetryManager.cs
+++ b/Engine/Source/Programs/Horde/HordeServer/Telemetry/TelemetryManager.cs
@@ -14,6 +14,7 @@ using Microsoft.Extensions.DependencyInjection;
 using Microsoft.Extensions.Hosting;
 using Microsoft.Extensions.Logging;
 using Microsoft.Extensions.Options;
+using MongoDB.Driver.Core.WireProtocol.Messages;
 using OpenTelemetry.Trace;
 
 namespace HordeServer.Telemetry
@@ -21,16 +22,24 @@ namespace HordeServer.Telemetry
 	/// <summary>
 	/// Telemetry sink dispatching incoming events to all registered sinks
 	/// </summary>
-	public sealed class TelemetryManager : ITelemetrySinkInternal, IHostedService
+	public sealed class TelemetryManager : ITelemetryWriter, IHostedService
 	{
 		/// <inheritdoc/>
 		public bool Enabled => _telemetrySinks.Count > 0;
 
-		private readonly List<ITelemetrySinkInternal> _telemetrySinks = new();
+		private readonly List<ITelemetrySink> _telemetrySinks = new();
 		private readonly ITicker _ticker;
 		private readonly Tracer _tracer;
 		private readonly ILogger<TelemetryManager> _logger;
 
+		private readonly static TelemetryRecordMeta s_serverEventMetadata = new TelemetryRecordMeta
+		{
+			AppId = "Horde",
+			AppVersion = ServerApp.Version.ToString(),
+			AppEnvironment = ServerApp.DeploymentEnvironment,
+			SessionId = ServerApp.SessionId
+		};
+
 		/// <summary>
 		/// Constructor
 		/// </summary>
@@ -60,10 +69,18 @@ namespace HordeServer.Telemetry
 		}
 
 		/// <inheritdoc/>
-		public void SendEvent(TelemetryStoreId telemetryStoreId, TelemetryEvent telemetryEvent)
+		public void WriteEvent(TelemetryStoreId telemetryStoreId, object payload)
 		{
-			using TelemetrySpan span = _tracer.StartActiveSpan($"{nameof(TelemetryManager)}.{nameof(SendEvent)}");
-			foreach (ITelemetrySinkInternal sink in _telemetrySinks)
+			WriteEvent(telemetryStoreId, s_serverEventMetadata, payload);
+		}
+
+		/// <inheritdoc/>
+		public void WriteEvent(TelemetryStoreId telemetryStoreId, TelemetryRecordMeta metadata, object payload)
+		{
+			using TelemetrySpan span = _tracer.StartActiveSpan($"{nameof(TelemetryManager)}.{nameof(WriteEvent)}");
+
+			TelemetryEvent telemetryEvent = new TelemetryEvent(metadata, payload);
+			foreach (ITelemetrySink sink in _telemetrySinks)
 			{
 				using TelemetrySpan sinkSpan = _tracer.StartActiveSpan($"SendEvent");
 				string fullName = sink.GetType().FullName ?? "Unknown";
@@ -114,7 +131,7 @@ namespace HordeServer.Telemetry
 		public async ValueTask DisposeAsync()
 		{
 			await _ticker.DisposeAsync();
-			foreach (ITelemetrySinkInternal sink in _telemetrySinks)
+			foreach (ITelemetrySink sink in _telemetrySinks)
 			{
 				await sink.DisposeAsync();
 			}
@@ -123,7 +140,7 @@ namespace HordeServer.Telemetry
 		/// <inheritdoc />
 		public async ValueTask FlushAsync(CancellationToken cancellationToken)
 		{
-			foreach (ITelemetrySinkInternal sink in _telemetrySinks)
+			foreach (ITelemetrySink sink in _telemetrySinks)
 			{
 				await sink.FlushAsync(cancellationToken);
 			}
