Horde: Updating schema documentation

#rnx

diff --git a/Engine/Source/Programs/Horde/Docs/Config/Schema/Dashboard.md b/Engine/Source/Programs/Horde/Docs/Config/Schema/Dashboard.md
--- a/Engine/Source/Programs/Horde/Docs/Config/Schema/Dashboard.md
+++ b/Engine/Source/Programs/Horde/Docs/Config/Schema/Dashboard.md
@@ -112,17 +112,13 @@ Name | Description
 
 ## ConfigInclude
 
-Directive to merge config data from another source
-
 Name | Description
 ---- | -----------
-`path` | `string`<br>Path to the config data to be included. May be relative to the including file's location.
+`path` | `string`<br>
 
 ## ConfigMacro
 
-Declares a config macro
-
 Name | Description
 ---- | -----------
-`name` | `string`<br>Name of the macro property
-`value` | `string`<br>Value for the macro property
+`name` | `string`<br>
+`value` | `string`<br>
diff --git a/Engine/Source/Programs/Horde/Docs/Config/Schema/Globals.md b/Engine/Source/Programs/Horde/Docs/Config/Schema/Globals.md
--- a/Engine/Source/Programs/Horde/Docs/Config/Schema/Globals.md
+++ b/Engine/Source/Programs/Horde/Docs/Config/Schema/Globals.md
@@ -32,20 +32,16 @@ Name | Description
 
 ## ConfigInclude
 
-Directive to merge config data from another source
-
 Name | Description
 ---- | -----------
-`path` | `string`<br>Path to the config data to be included. May be relative to the including file's location.
+`path` | `string`<br>
 
 ## ConfigMacro
 
-Declares a config macro
-
 Name | Description
 ---- | -----------
-`name` | `string`<br>Name of the macro property
-`value` | `string`<br>Value for the macro property
+`name` | `string`<br>
+`value` | `string`<br>
 
 ## PoolConfig
 
@@ -275,44 +271,36 @@ Name | Description
 
 ## AclConfig
 
-Parameters to update an ACL
-
 Name | Description
 ---- | -----------
-`entries` | [`AclEntryConfig`](#aclentryconfig)`[]`<br>Entries to replace the existing ACL
-`profiles` | [`AclProfileConfig`](#aclprofileconfig)`[]`<br>Defines profiles which allow grouping sets of actions into named collections
-`inherit` | `boolean`<br>Whether to inherit permissions from the parent ACL
-`exceptions` | `string[]`<br>List of exceptions to the inherited setting
+`entries` | [`AclEntryConfig`](#aclentryconfig)`[]`<br>
+`profiles` | [`AclProfileConfig`](#aclprofileconfig)`[]`<br>
+`inherit` | `boolean`<br>
+`exceptions` | `string[]`<br>
 
 ## AclEntryConfig
 
-Individual entry in an ACL
-
 Name | Description
 ---- | -----------
-`claim` | [`AclClaimConfig`](#aclclaimconfig)<br>Name of the user or group
-`actions` | `string[]`<br>Array of actions to allow
-`profiles` | `string[]`<br>List of profiles to grant
+`claim` | [`AclClaimConfig`](#aclclaimconfig)<br>
+`actions` | `string[]`<br>
+`profiles` | `string[]`<br>
 
 ## AclClaimConfig
 
-New claim to create
-
 Name | Description
 ---- | -----------
-`type` | `string`<br>The claim type
-`value` | `string`<br>The claim value
+`type` | `string`<br>
+`value` | `string`<br>
 
 ## AclProfileConfig
 
-Configuration for an ACL profile. This defines a preset group of actions which can be given to a user via an ACL entry.
-
 Name | Description
 ---- | -----------
-`id` | `string`<br>Identifier for this profile
-`actions` | `string[]`<br>Actions to include
-`excludeActions` | `string[]`<br>Actions to exclude from the inherited actions
-`extends` | `string[]`<br>Other profiles to extend from
+`id` | `string`<br>
+`actions` | `string[]`<br>
+`excludeActions` | `string[]`<br>
+`extends` | `string[]`<br>
 
 ## SecretConfig
 
diff --git a/Engine/Source/Programs/Horde/Docs/Config/Schema/Projects.md b/Engine/Source/Programs/Horde/Docs/Config/Schema/Projects.md
--- a/Engine/Source/Programs/Horde/Docs/Config/Schema/Projects.md
+++ b/Engine/Source/Programs/Horde/Docs/Config/Schema/Projects.md
@@ -24,20 +24,16 @@ Name | Description
 
 ## ConfigInclude
 
-Directive to merge config data from another source
-
 Name | Description
 ---- | -----------
-`path` | `string`<br>Path to the config data to be included. May be relative to the including file's location.
+`path` | `string`<br>
 
 ## ConfigMacro
 
-Declares a config macro
-
 Name | Description
 ---- | -----------
-`name` | `string`<br>Name of the macro property
-`value` | `string`<br>Value for the macro property
+`name` | `string`<br>
+`value` | `string`<br>
 
 ## PoolConfig
 
@@ -215,41 +211,33 @@ Name | Description
 
 ## AclConfig
 
-Parameters to update an ACL
-
 Name | Description
 ---- | -----------
-`entries` | [`AclEntryConfig`](#aclentryconfig)`[]`<br>Entries to replace the existing ACL
-`profiles` | [`AclProfileConfig`](#aclprofileconfig)`[]`<br>Defines profiles which allow grouping sets of actions into named collections
-`inherit` | `boolean`<br>Whether to inherit permissions from the parent ACL
-`exceptions` | `string[]`<br>List of exceptions to the inherited setting
+`entries` | [`AclEntryConfig`](#aclentryconfig)`[]`<br>
+`profiles` | [`AclProfileConfig`](#aclprofileconfig)`[]`<br>
+`inherit` | `boolean`<br>
+`exceptions` | `string[]`<br>
 
 ## AclEntryConfig
 
-Individual entry in an ACL
-
 Name | Description
 ---- | -----------
-`claim` | [`AclClaimConfig`](#aclclaimconfig)<br>Name of the user or group
-`actions` | `string[]`<br>Array of actions to allow
-`profiles` | `string[]`<br>List of profiles to grant
+`claim` | [`AclClaimConfig`](#aclclaimconfig)<br>
+`actions` | `string[]`<br>
+`profiles` | `string[]`<br>
 
 ## AclClaimConfig
 
-New claim to create
-
 Name | Description
 ---- | -----------
-`type` | `string`<br>The claim type
-`value` | `string`<br>The claim value
+`type` | `string`<br>
+`value` | `string`<br>
 
 ## AclProfileConfig
 
-Configuration for an ACL profile. This defines a preset group of actions which can be given to a user via an ACL entry.
-
 Name | Description
 ---- | -----------
-`id` | `string`<br>Identifier for this profile
-`actions` | `string[]`<br>Actions to include
-`excludeActions` | `string[]`<br>Actions to exclude from the inherited actions
-`extends` | `string[]`<br>Other profiles to extend from
+`id` | `string`<br>
+`actions` | `string[]`<br>
+`excludeActions` | `string[]`<br>
+`extends` | `string[]`<br>
diff --git a/Engine/Source/Programs/Horde/Docs/Config/Schema/Streams.md b/Engine/Source/Programs/Horde/Docs/Config/Schema/Streams.md
--- a/Engine/Source/Programs/Horde/Docs/Config/Schema/Streams.md
+++ b/Engine/Source/Programs/Horde/Docs/Config/Schema/Streams.md
@@ -37,20 +37,16 @@ Name | Description
 
 ## ConfigInclude
 
-Directive to merge config data from another source
-
 Name | Description
 ---- | -----------
-`path` | `string`<br>Path to the config data to be included. May be relative to the including file's location.
+`path` | `string`<br>
 
 ## ConfigMacro
 
-Declares a config macro
-
 Name | Description
 ---- | -----------
-`name` | `string`<br>Name of the macro property
-`value` | `string`<br>Value for the macro property
+`name` | `string`<br>
+`value` | `string`<br>
 
 ## JobOptions
 
@@ -262,12 +258,10 @@ Name | Description
 
 ## AclClaimConfig
 
-New claim to create
-
 Name | Description
 ---- | -----------
-`type` | `string`<br>The claim type
-`value` | `string`<br>The claim value
+`type` | `string`<br>
+`value` | `string`<br>
 
 ## ScheduleGateConfig
 
@@ -322,35 +316,29 @@ Name | Description
 
 ## AclConfig
 
-Parameters to update an ACL
-
 Name | Description
 ---- | -----------
-`entries` | [`AclEntryConfig`](#aclentryconfig)`[]`<br>Entries to replace the existing ACL
-`profiles` | [`AclProfileConfig`](#aclprofileconfig)`[]`<br>Defines profiles which allow grouping sets of actions into named collections
-`inherit` | `boolean`<br>Whether to inherit permissions from the parent ACL
-`exceptions` | `string[]`<br>List of exceptions to the inherited setting
+`entries` | [`AclEntryConfig`](#aclentryconfig)`[]`<br>
+`profiles` | [`AclProfileConfig`](#aclprofileconfig)`[]`<br>
+`inherit` | `boolean`<br>
+`exceptions` | `string[]`<br>
 
 ## AclEntryConfig
 
-Individual entry in an ACL
-
 Name | Description
 ---- | -----------
-`claim` | [`AclClaimConfig`](#aclclaimconfig)<br>Name of the user or group
-`actions` | `string[]`<br>Array of actions to allow
-`profiles` | `string[]`<br>List of profiles to grant
+`claim` | [`AclClaimConfig`](#aclclaimconfig)<br>
+`actions` | `string[]`<br>
+`profiles` | `string[]`<br>
 
 ## AclProfileConfig
 
-Configuration for an ACL profile. This defines a preset group of actions which can be given to a user via an ACL entry.
-
 Name | Description
 ---- | -----------
-`id` | `string`<br>Identifier for this profile
-`actions` | `string[]`<br>Actions to include
-`excludeActions` | `string[]`<br>Actions to exclude from the inherited actions
-`extends` | `string[]`<br>Other profiles to extend from
+`id` | `string`<br>
+`actions` | `string[]`<br>
+`excludeActions` | `string[]`<br>
+`extends` | `string[]`<br>
 
 ## Priority (Enum)
 
diff --git a/Engine/Source/Programs/Horde/Docs/Config/Schema/Telemetry.md b/Engine/Source/Programs/Horde/Docs/Config/Schema/Telemetry.md
--- a/Engine/Source/Programs/Horde/Docs/Config/Schema/Telemetry.md
+++ b/Engine/Source/Programs/Horde/Docs/Config/Schema/Telemetry.md
@@ -14,44 +14,36 @@ Name | Description
 
 ## AclConfig
 
-Parameters to update an ACL
-
 Name | Description
 ---- | -----------
-`entries` | [`AclEntryConfig`](#aclentryconfig)`[]`<br>Entries to replace the existing ACL
-`profiles` | [`AclProfileConfig`](#aclprofileconfig)`[]`<br>Defines profiles which allow grouping sets of actions into named collections
-`inherit` | `boolean`<br>Whether to inherit permissions from the parent ACL
-`exceptions` | `string[]`<br>List of exceptions to the inherited setting
+`entries` | [`AclEntryConfig`](#aclentryconfig)`[]`<br>
+`profiles` | [`AclProfileConfig`](#aclprofileconfig)`[]`<br>
+`inherit` | `boolean`<br>
+`exceptions` | `string[]`<br>
 
 ## AclEntryConfig
 
-Individual entry in an ACL
-
 Name | Description
 ---- | -----------
-`claim` | [`AclClaimConfig`](#aclclaimconfig)<br>Name of the user or group
-`actions` | `string[]`<br>Array of actions to allow
-`profiles` | `string[]`<br>List of profiles to grant
+`claim` | [`AclClaimConfig`](#aclclaimconfig)<br>
+`actions` | `string[]`<br>
+`profiles` | `string[]`<br>
 
 ## AclClaimConfig
 
-New claim to create
-
 Name | Description
 ---- | -----------
-`type` | `string`<br>The claim type
-`value` | `string`<br>The claim value
+`type` | `string`<br>
+`value` | `string`<br>
 
 ## AclProfileConfig
 
-Configuration for an ACL profile. This defines a preset group of actions which can be given to a user via an ACL entry.
-
 Name | Description
 ---- | -----------
-`id` | `string`<br>Identifier for this profile
-`actions` | `string[]`<br>Actions to include
-`excludeActions` | `string[]`<br>Actions to exclude from the inherited actions
-`extends` | `string[]`<br>Other profiles to extend from
+`id` | `string`<br>
+`actions` | `string[]`<br>
+`excludeActions` | `string[]`<br>
+`extends` | `string[]`<br>
 
 ## MetricConfig
 
@@ -82,17 +74,13 @@ Name | Description
 
 ## ConfigInclude
 
-Directive to merge config data from another source
-
 Name | Description
 ---- | -----------
-`path` | `string`<br>Path to the config data to be included. May be relative to the including file's location.
+`path` | `string`<br>
 
 ## ConfigMacro
 
-Declares a config macro
-
 Name | Description
 ---- | -----------
-`name` | `string`<br>Name of the macro property
-`value` | `string`<br>Value for the macro property
+`name` | `string`<br>
+`value` | `string`<br>
diff --git a/Engine/Source/Programs/Horde/Docs/Deployment/ServerSettings.md b/Engine/Source/Programs/Horde/Docs/Deployment/ServerSettings.md
--- a/Engine/Source/Programs/Horde/Docs/Deployment/ServerSettings.md
+++ b/Engine/Source/Programs/Horde/Docs/Deployment/ServerSettings.md
@@ -390,44 +390,36 @@ Name | Description
 
 ## AclConfig
 
-Parameters to update an ACL
-
 Name | Description
 ---- | -----------
-`entries` | [`AclEntryConfig`](#aclentryconfig)`[]`<br>Entries to replace the existing ACL
-`profiles` | [`AclProfileConfig`](#aclprofileconfig)`[]`<br>Defines profiles which allow grouping sets of actions into named collections
-`inherit` | `boolean`<br>Whether to inherit permissions from the parent ACL
-`exceptions` | `string[]`<br>List of exceptions to the inherited setting
+`entries` | [`AclEntryConfig`](#aclentryconfig)`[]`<br>
+`profiles` | [`AclProfileConfig`](#aclprofileconfig)`[]`<br>
+`inherit` | `boolean`<br>
+`exceptions` | `string[]`<br>
 
 ## AclEntryConfig
 
-Individual entry in an ACL
-
 Name | Description
 ---- | -----------
-`claim` | [`AclClaimConfig`](#aclclaimconfig)<br>Name of the user or group
-`actions` | `string[]`<br>Array of actions to allow
-`profiles` | `string[]`<br>List of profiles to grant
+`claim` | [`AclClaimConfig`](#aclclaimconfig)<br>
+`actions` | `string[]`<br>
+`profiles` | `string[]`<br>
 
 ## AclClaimConfig
 
-New claim to create
-
 Name | Description
 ---- | -----------
-`type` | `string`<br>The claim type
-`value` | `string`<br>The claim value
+`type` | `string`<br>
+`value` | `string`<br>
 
 ## AclProfileConfig
 
-Configuration for an ACL profile. This defines a preset group of actions which can be given to a user via an ACL entry.
-
 Name | Description
 ---- | -----------
-`id` | `string`<br>Identifier for this profile
-`actions` | `string[]`<br>Actions to include
-`excludeActions` | `string[]`<br>Actions to exclude from the inherited actions
-`extends` | `string[]`<br>Other profiles to extend from
+`id` | `string`<br>
+`actions` | `string[]`<br>
+`excludeActions` | `string[]`<br>
+`extends` | `string[]`<br>
 
 ## OpenTelemetrySettings
 
