Horde: Fix compiler warnings

diff --git a/Engine/Source/Programs/Horde/HordeServer/Telemetry/TelemetryManager.cs b/Engine/Source/Programs/Horde/HordeServer/Telemetry/TelemetryManager.cs
--- a/Engine/Source/Programs/Horde/HordeServer/Telemetry/TelemetryManager.cs
+++ b/Engine/Source/Programs/Horde/HordeServer/Telemetry/TelemetryManager.cs
@@ -14,7 +14,6 @@ using Microsoft.Extensions.DependencyInjection;
 using Microsoft.Extensions.Hosting;
 using Microsoft.Extensions.Logging;
 using Microsoft.Extensions.Options;
-using MongoDB.Driver.Core.WireProtocol.Messages;
 using OpenTelemetry.Trace;
 
 namespace HordeServer.Telemetry
@@ -32,7 +31,7 @@ namespace HordeServer.Telemetry
 		private readonly Tracer _tracer;
 		private readonly ILogger<TelemetryManager> _logger;
 
-		private readonly static TelemetryRecordMeta s_serverEventMetadata = new TelemetryRecordMeta
+		private static readonly TelemetryRecordMeta s_serverEventMetadata = new TelemetryRecordMeta
 		{
 			AppId = "Horde",
 			AppVersion = ServerApp.Version.ToString(),
