Horde: Log where agent registration file is saved to/loaded from

When the file exists, the agent provides no hinting it reads this file which can be confusing when trying to reset the agent.

diff --git a/Engine/Source/Programs/Horde/HordeAgent/Services/SessionFactory.cs b/Engine/Source/Programs/Horde/HordeAgent/Services/SessionFactory.cs
--- a/Engine/Source/Programs/Horde/HordeAgent/Services/SessionFactory.cs
+++ b/Engine/Source/Programs/Horde/HordeAgent/Services/SessionFactory.cs
@@ -161,7 +161,11 @@ namespace HordeAgent.Services
 				registrationList.Entries.Add(registrationInfo);
 
 				await WriteRegistrationListAsync(registrationFile, registrationList, cancellationToken);
-				logger.LogInformation("Created agent (Id={AgentId}). Settings saved to {File}.", registrationInfo.Id, registrationFile);
+				logger.LogInformation("Created agent (Id={AgentId}). Server <-> agent registration info saved to {File}", registrationInfo.Id, registrationFile);
+			}
+			else
+			{
+				logger.LogInformation("Cached server <-> agent registration info loaded from {File}. AgentId={AgentId}", registrationFile, registrationInfo.Id);
 			}
 
 			// Create the session
