Horde: Fix failing tests referencing telemetry sinks

diff --git a/Engine/Source/Programs/Horde/HordeServer.Tests/TestSetup.cs b/Engine/Source/Programs/Horde/HordeServer.Tests/TestSetup.cs
--- a/Engine/Source/Programs/Horde/HordeServer.Tests/TestSetup.cs
+++ b/Engine/Source/Programs/Horde/HordeServer.Tests/TestSetup.cs
@@ -61,6 +61,7 @@ using HordeServer.Ugs;
 using HordeServer.Users;
 using HordeServer.Utilities;
 using HordeCommon;
+using HordeServer.Telemetry.Metrics;
 using Microsoft.AspNetCore.Http;
 using Microsoft.AspNetCore.Mvc;
 using Microsoft.Extensions.Caching.Memory;
@@ -201,9 +202,16 @@ namespace HordeServer.Tests
 			services.AddSingleton(typeof(IAuditLogFactory<>), typeof(AuditLogFactory<>));
 			services.AddSingleton<IAuditLog<AgentId>>(sp => sp.GetRequiredService<IAuditLogFactory<AgentId>>().Create("Agents.Log", "AgentId"));
 			services.AddSingleton<ITelemetrySink, NullTelemetrySink>();
+			services.AddSingleton<NullTelemetrySink>();
 			services.AddSingleton<ITelemetrySink, MetricTelemetrySink>();
+			services.AddSingleton<MetricTelemetrySink>();
+
+			services.AddSingleton<TelemetryManager>();
+			services.AddSingleton<ITelemetryWriter>(sp => sp.GetRequiredService<TelemetryManager>());
 			services.AddSingleton<OpenTelemetry.Trace.Tracer>(sp => TracerProvider.Default.GetTracer("TestTracer"));
 			services.AddSingleton(sp => new Meter("TestMeter"));
+			services.AddSingleton<MetricCollection>();
+			services.AddSingleton<IMetricCollection, MetricCollection>(sp => sp.GetRequiredService<MetricCollection>());
 
 			services.AddSingleton<ConfigService>();
 			services.AddSingleton<IOptionsFactory<GlobalConfig>>(sp => sp.GetRequiredService<ConfigService>());
