Horde: Log each configuration source during agent start

For better visibility where the agent is reading config from. Also log the actual logs dir.

diff --git a/Engine/Source/Programs/Horde/HordeAgent/AgentApp.cs b/Engine/Source/Programs/Horde/HordeAgent/AgentApp.cs
--- a/Engine/Source/Programs/Horde/HordeAgent/AgentApp.cs
+++ b/Engine/Source/Programs/Horde/HordeAgent/AgentApp.cs
@@ -21,6 +21,8 @@ using Microsoft.Extensions.Logging;
 using Microsoft.Win32;
 using OpenTracing.Util;
 using Polly;
+using IConfigurationSource = Microsoft.Extensions.Configuration.IConfigurationSource;
+using JsonConfigurationSource = Microsoft.Extensions.Configuration.Json.JsonConfigurationSource;
 
 namespace HordeAgent
 {
@@ -113,10 +115,13 @@ namespace HordeAgent
 			{
 				configOverrides.Add($"{AgentSettings.SectionName}:{nameof(AgentSettings.WorkingDir)}", workingDirOverride);
 			}
-
+			
+			List<string> configSources;
+			List<string> configSourcesInstalled = new();
+			
 			// Create the base configuration data by just reading from the application directory. We need to check some settings before
 			// being able to read user configuration files.
-			IConfiguration configuration = CreateConfig(false, null, configOverrides);
+			IConfiguration configuration = CreateConfig(false, null, configOverrides, out configSources);
 			AgentSettings settings = BindSettings(configuration);
 
 			if (settings.Installed)
@@ -132,12 +137,22 @@ namespace HordeAgent
 				}
 
 				FileReference agentConfigFile = FileReference.Combine(DataDir, "agent.json");
-				configuration = CreateConfig(true, agentConfigFile, configOverrides);
+				configuration = CreateConfig(true, agentConfigFile, configOverrides, out configSourcesInstalled);
 				settings = BindSettings(configuration);
 			}
 
 			using ILoggerFactory loggerFactory = Logging.CreateLoggerFactory(configuration);
-
+			
+			{
+				configSources.AddRange(configSourcesInstalled);
+				ILogger logger = loggerFactory.CreateLogger(typeof(AgentApp));
+				foreach (string configSource in configSources)
+				{
+					logger.LogInformation("Config source: {ConfigSource}", configSource);
+				}
+				logger.LogInformation("Logs dir: {LogsDir}", settings.LogsDir);
+			}
+			
 			IServiceCollection services = new ServiceCollection();
 			services.AddCommandsFromAssembly(Assembly.GetExecutingAssembly());
 			services.AddSingleton(loggerFactory);
@@ -241,8 +256,9 @@ namespace HordeAgent
 			return new DefaultSystemMetrics();
 		}
 
-		static IConfiguration CreateConfig(bool readInstalledConfig, FileReference? agentConfigFile, Dictionary<string, string?> configOverrides)
+		static IConfiguration CreateConfig(bool readInstalledConfig, FileReference? agentConfigFile, Dictionary<string, string?> configOverrides, out List<string> configSources)
 		{
+			configSources = new List<string>();
 			string? environment = Environment.GetEnvironmentVariable("DOTNET_ENVIRONMENT");
 			if (String.IsNullOrEmpty(environment))
 			{
@@ -254,17 +270,36 @@ namespace HordeAgent
 			{
 				builder = builder.Add(new RegistryConfigurationSource(Registry.LocalMachine, "SOFTWARE\\Epic Games\\Horde\\Agent", AgentSettings.SectionName));
 			}
-
-			builder.SetBasePath(AppDir.FullName)
+			
+			string basePath = AppDir.FullName;
+			builder.SetBasePath(basePath)
 				.AddJsonFile("appsettings.json", optional: false)
 				.AddJsonFile("appsettings.Build.json", optional: true) // specific settings for builds (installer/dockerfile)
 				.AddJsonFile($"appsettings.{environment}.json", optional: true) // environment variable overrides, also used in k8s setups with Helm
 				.AddJsonFile("appsettings.User.json", optional: true);
-
+			
 			if (agentConfigFile != null)
 			{
 				builder = builder.AddJsonFile(agentConfigFile.FullName, optional: true, reloadOnChange: true);
 			}
+			
+			foreach (IConfigurationSource source in builder.Sources)
+			{
+				switch (source)
+				{
+					case JsonConfigurationSource jcs:
+						configSources.Add("JSON file: " + Path.Join(basePath, jcs.Path));
+						break;
+					case RegistryConfigurationSource rcs:
+						if (OperatingSystem.IsWindows())
+						{
+							configSources.Add("Windows Registry key: " + rcs.GetRegistryKey());	
+						}
+						break;
+					case null:
+						throw new ArgumentException(nameof(source));
+				}
+			}
 
 			return builder
 				.AddInMemoryCollection(configOverrides)
diff --git a/Engine/Source/Programs/Shared/EpicGames.Core/RegistryConfigurationSource.cs b/Engine/Source/Programs/Shared/EpicGames.Core/RegistryConfigurationSource.cs
--- a/Engine/Source/Programs/Shared/EpicGames.Core/RegistryConfigurationSource.cs
+++ b/Engine/Source/Programs/Shared/EpicGames.Core/RegistryConfigurationSource.cs
@@ -31,6 +31,12 @@ namespace EpicGames.Core
 			_baseConfigName = baseConfigName;
 		}
 		
+		/// <summary>
+		/// Get the registry key used
+		/// </summary>
+		/// <returns></returns>
+		public string GetRegistryKey() => _baseKey + "\\" + _keyPath;
+		
 		/// <inheritdoc/>
 		public IConfigurationProvider Build(IConfigurationBuilder builder)
 			=> new RegistryConfigProvider(_baseKey, _keyPath, _baseConfigName);
