Horde: Move telemetry ACLs into shared library.

#rnx

diff --git a/Engine/Source/Programs/Horde/HordeServer.Shared/Telemetry/TelemetryAclAction.cs b/Engine/Source/Programs/Horde/HordeServer.Shared/Telemetry/TelemetryAclAction.cs
new file mode 100644
--- /dev/null
+++ b/Engine/Source/Programs/Horde/HordeServer.Shared/Telemetry/TelemetryAclAction.cs
@@ -0,0 +1,17 @@
+// Copyright Epic Games, Inc. All Rights Reserved.
+
+using EpicGames.Horde.Acls;
+
+namespace HordeServer.Telemetry
+{
+	/// <summary>
+	/// ACL actions for telemetry operations
+	/// </summary>
+	public static class TelemetryAclAction
+	{
+		/// <summary>
+		/// Ability to search for various metrics
+		/// </summary>
+		public static AclAction QueryMetrics { get; } = new AclAction("QueryMetrics");
+	}
+}
diff --git a/Engine/Source/Programs/Horde/HordeServer/Telemetry/TelemetryAclAction.cs b/Engine/Source/Programs/Horde/HordeServer/Telemetry/TelemetryAclAction.cs
deleted file mode 100644
