Horde: Artifact Search
#skipci

diff --git a/Engine/Source/Programs/Horde/HordeDashboard/src/backend/Api.ts b/Engine/Source/Programs/Horde/HordeDashboard/src/backend/Api.ts
--- a/Engine/Source/Programs/Horde/HordeDashboard/src/backend/Api.ts
+++ b/Engine/Source/Programs/Horde/HordeDashboard/src/backend/Api.ts
@@ -1691,12 +1691,14 @@ export type CreateZipRequest = {
 
 /** Describes an artifact */
 export type GetArtifactResponseV2 = {
-
 	id: string;
 	type: ArtifactContextType;
 	keys: string[]
 	name: string;
 	description?: string;
+	streamId?: string;
+	change?: number;
+	metadata?: string[];
 }
 
 /** Result of an artifact search */
diff --git a/Engine/Source/Programs/Horde/HordeDashboard/src/backend/Backend.ts b/Engine/Source/Programs/Horde/HordeDashboard/src/backend/Backend.ts
--- a/Engine/Source/Programs/Horde/HordeDashboard/src/backend/Backend.ts
+++ b/Engine/Source/Programs/Horde/HordeDashboard/src/backend/Backend.ts
@@ -502,6 +502,17 @@ export class Backend {
         });
     }
 
+    getArtifacts(streamId?: string, minChange?: number, maxChange?: number, name?: string, type?: string): Promise<FindArtifactsResponse> {
+
+        return new Promise<FindArtifactsResponse>((resolve, reject) => {
+
+            this.backend.get(`/api/v2/artifacts`, { params: { streamId: streamId, minChange: minChange, maxChange: maxChange, name: name, type: type } })
+                .then(response => { resolve(response.data); })
+                .catch(reason => reject(reason));
+        })
+    }
+
+
     getJobArtifactsV2(ids?: string[], keys?: string[]): Promise<FindArtifactsResponse> {
 
         const uniqueIds = Array.from(new Set(ids ?? []));
diff --git a/Engine/Source/Programs/Horde/HordeDashboard/src/components/TopNav.tsx b/Engine/Source/Programs/Horde/HordeDashboard/src/components/TopNav.tsx
--- a/Engine/Source/Programs/Horde/HordeDashboard/src/components/TopNav.tsx
+++ b/Engine/Source/Programs/Horde/HordeDashboard/src/components/TopNav.tsx
@@ -12,6 +12,7 @@ import { PreviewChangesModal } from './PreviewChanges';
 import { VersionModal } from './VersionModal';
 import { getHordeTheme } from '../styles/theme';
 import { getHordeStyling } from '../styles/Styles';
+import { FindArtifactsModal } from './artifacts/ArtifactsSearch';
 
 
 const getStyles = () => {
@@ -419,6 +420,9 @@ export const TopNav: React.FC<{ suppressServer?: boolean }> = observer(({ suppre
 
    const divRef = useRef(null);
 
+   const query = new URLSearchParams(window.location.search);
+   const showArtifacts = query.has("showArtifacts");
+
    // subscribe
    if (dashboard.updated) { }
 
@@ -435,7 +439,7 @@ export const TopNav: React.FC<{ suppressServer?: boolean }> = observer(({ suppre
       const features = dashboard.user?.dashboardFeatures;
 
       const serviceItems: IContextualMenuItem[] = [];
-      
+
       serviceItems.push({
          key: "admin_analytics",
          text: "Analytics",
@@ -456,6 +460,21 @@ export const TopNav: React.FC<{ suppressServer?: boolean }> = observer(({ suppre
          link: `/tools`
       });
 
+      serviceItems.push({
+         key: "job_artifacts",
+         text: "Artifacts",
+         onClick(ev, item) {
+
+            const search = new URLSearchParams(window.location.search);
+            search.append("showArtifacts", "true")
+
+            const url = `${window.location.pathname}?` + search.toString();
+            navigate(url, { replace: true })
+
+         },
+      });
+
+
       const style = { ...menuStyles } as Partial<IContextualMenuStyles>;
 
       const cbItem: ICommandBarItemProps = {
@@ -672,7 +691,7 @@ export const TopNav: React.FC<{ suppressServer?: boolean }> = observer(({ suppre
 
    }
 
-   function getInitials(name: string) {      
+   function getInitials(name: string) {
       const nameArray = name.indexOf(".") === -1 ? name.split(" ") : name.split(".");
       if (nameArray.length === 1) {
          return name.toUpperCase().slice(0, 2);
@@ -680,7 +699,7 @@ export const TopNav: React.FC<{ suppressServer?: boolean }> = observer(({ suppre
       const firstInitial = nameArray[0].charAt(0).toUpperCase();
       const lastInitial = nameArray[nameArray.length - 1].charAt(0).toUpperCase();
       return firstInitial + lastInitial;
-    }
+   }
 
    let initials = "??";
    try {
@@ -738,6 +757,12 @@ export const TopNav: React.FC<{ suppressServer?: boolean }> = observer(({ suppre
 
    return (
       <div style={{ backgroundColor: hordeTheme.horde.topNavBackground }}>
+         {showArtifacts && <FindArtifactsModal onClose={() => {
+            query.delete("showArtifacts");
+            const url = `${window.location.pathname}?` + query.toString();
+            navigate(url, { replace: true })
+
+         }} />}
          {showVersion && <VersionModal show={true} onClose={() => { setShowVersion(false) }} />}
          {showPreviewChanges && <PreviewChangesModal onClose={() => { setShowPreviewChanges(false) }} />}
          <Stack tokens={{ maxWidth: 1440, childrenGap: 0 }} disableShrink={true} styles={{ root: { backgroundColor: hordeTheme.horde.topNavBackground, margin: "auto", width: "100%" } }}>
@@ -789,7 +814,7 @@ export const TopNav: React.FC<{ suppressServer?: boolean }> = observer(({ suppre
                      />}
                   </Stack>
 
-                  <Stack style={{width: "32px"}} onMouseEnter={() => setShowMenu(true)}
+                  <Stack style={{ width: "32px" }} onMouseEnter={() => setShowMenu(true)}
                      onMouseLeave={() => { setShowMenu(false) }} >
                      <div ref={divRef}>
                         <Persona styles={{ root: { selectors: { ".ms-Persona-initials": { fontWeight: "unset", fontFamily: "Horde Open Sans SemiBold", cursor: "pointer" } } } }} imageShouldFadeIn={false} imageInitials={initials} imageUrl={dashboard.userImage32} size={PersonaSize.size32}
diff --git a/Engine/Source/Programs/Horde/HordeDashboard/src/components/artifacts/ArtifactsModal.tsx b/Engine/Source/Programs/Horde/HordeDashboard/src/components/artifacts/ArtifactsModal.tsx
--- a/Engine/Source/Programs/Horde/HordeDashboard/src/components/artifacts/ArtifactsModal.tsx
+++ b/Engine/Source/Programs/Horde/HordeDashboard/src/components/artifacts/ArtifactsModal.tsx
@@ -1,6 +1,6 @@
 // Copyright Epic Games, Inc. All Rights Reserved.
 
-import { Callout, DetailsList, DetailsListLayoutMode, DirectionalHint, FontIcon, IColumn, IconButton, Modal, PrimaryButton, ScrollablePane, Selection, SelectionMode, SelectionZone, Spinner, SpinnerSize, Stack, Text, mergeStyleSets } from "@fluentui/react";
+import { Callout, DetailsList, DetailsListLayoutMode, DirectionalHint, FontIcon, IColumn, IContextualMenuProps, IconButton, Modal, NavBase, PrimaryButton, ScrollablePane, Selection, SelectionMode, SelectionZone, Spinner, SpinnerSize, Stack, Text, mergeStyleSets } from "@fluentui/react";
 import { action, makeObservable, observable } from "mobx";
 import { observer } from "mobx-react-lite";
 import { useEffect, useState } from "react";
@@ -90,7 +90,7 @@ class ArtifactsHandler {
          return;
       }
 
-      let a: GetArtifactResponseV2 | undefined; 
+      let a: GetArtifactResponseV2 | undefined;
 
       if (this.artifactId) {
          a = artifacts.find(a => a.id === this.artifactId)
@@ -99,7 +99,7 @@ class ArtifactsHandler {
       if (!a) {
          a = artifacts.find(a => a.type === this.context)!;
       }
-      
+
 
       if (!a) {
          console.error("Unable to find artifact for context", this.context, artifacts);
@@ -433,6 +433,7 @@ function formatBytes(bytes: number, decimals = 2) {
 const DownloadButton: React.FC<{ handler: ArtifactsHandler }> = observer(({ handler }) => {
 
    const [selectKey, setSelectionKey] = useState(0);
+   const navigate = useNavigate();
 
 
    // subscribe
@@ -465,8 +466,30 @@ const DownloadButton: React.FC<{ handler: ArtifactsHandler }> = observer(({ hand
       buttonText = `Download (${sizeText})`;
    }
 
+   let jobUrl = "";
+   if (handler.jobId) {
+      jobUrl = `/job/${handler.jobId}`;
+      if (handler.stepId) {
+         jobUrl += `?step=${handler.stepId}`;
+      }
+   }
+
+   const downloadProps: IContextualMenuProps = {
+      items: [
+         {
+            key: 'navigate_to_job',
+            text: 'Navigate to Job',
+            onClick: () => {
+               navigate(`/job/${handler.jobId}`)
+            }
+         }
+      ],
+      directionalHint: DirectionalHint.bottomRightEdge
+   };
+
+
    return <Stack horizontal verticalAlign="center" tokens={{ childrenGap: 8 }}>
-      <PrimaryButton styles={{ root: { fontFamily: 'Horde Open Sans SemiBold !important' } }} disabled={!selection.filesSelected && !selection.directoriesSelected} onClick={async () => {
+      <PrimaryButton split menuProps={downloadProps} styles={{ root: { fontFamily: 'Horde Open Sans SemiBold !important' } }} disabled={!selection.filesSelected && !selection.directoriesSelected} onClick={async () => {
 
          const selection = handler.currentSelection.items;
 
@@ -521,7 +544,7 @@ const DownloadButton: React.FC<{ handler: ArtifactsHandler }> = observer(({ hand
 
 let idcounter = 0;
 
-const JobDetailArtifactsInner: React.FC<{ jobId: string; stepId: string, artifacts?: GetArtifactResponseV2[], contextType: ArtifactContextType, artifactPath?: string, artifactId? :string }> = observer(({ jobId, stepId, artifacts, contextType, artifactPath, artifactId }) => {
+const JobDetailArtifactsInner: React.FC<{ jobId: string; stepId: string, artifacts?: GetArtifactResponseV2[], contextType: ArtifactContextType, artifactPath?: string, artifactId?: string }> = observer(({ jobId, stepId, artifacts, contextType, artifactPath, artifactId }) => {
 
    // eslint-disable-next-line
    const handler = ArtifactsHandler.current ?? new ArtifactsHandler(jobId, stepId, contextType, artifactPath, artifacts, artifactId);
@@ -774,7 +797,7 @@ const JobDetailArtifactsInner: React.FC<{ jobId: string; stepId: string, artifac
 
 
 
-export const JobArtifactsModal: React.FC<{ jobId: string; stepId: string, artifacts?: GetArtifactResponseV2[], contextType: ArtifactContextType, artifactPath?: string, artifactId? :string, onClose: () => void }> = ({ jobId, stepId, artifacts, contextType, artifactPath, artifactId, onClose }) => {
+export const JobArtifactsModal: React.FC<{ jobId: string; stepId: string, artifacts?: GetArtifactResponseV2[], contextType: ArtifactContextType, artifactPath?: string, artifactId?: string, onClose: () => void }> = ({ jobId, stepId, artifacts, contextType, artifactPath, artifactId, onClose }) => {
 
    const { hordeClasses } = getHordeStyling();
 
diff --git a/Engine/Source/Programs/Horde/HordeDashboard/src/components/artifacts/ArtifactsSearch.tsx b/Engine/Source/Programs/Horde/HordeDashboard/src/components/artifacts/ArtifactsSearch.tsx
new file mode 100644
--- /dev/null
+++ b/Engine/Source/Programs/Horde/HordeDashboard/src/components/artifacts/ArtifactsSearch.tsx
@@ -0,0 +1,244 @@
+import { ComboBox, DefaultButton, DetailsList, DetailsListLayoutMode, IComboBox, IComboBoxOption, IContextualMenuItem, IContextualMenuProps, IDetailsListProps, ITextField, IconButton, Label, Modal, PrimaryButton, ScrollablePane, ScrollbarVisibility, SelectionMode, Spinner, SpinnerSize, Stack, Text, TextField } from "@fluentui/react";
+import { getHordeStyling } from "../../styles/Styles";
+import { IStreamChooser, StreamChooser } from "../projects/StreamChooser";
+import React, { useState } from "react";
+import backend from "../../backend";
+import { GetArtifactResponseV2 } from "../../backend/Api";
+import dashboard from "../../backend/Dashboard";
+import { JobArtifactsModal } from "./ArtifactsModal";
+import { projectStore } from "../../backend/ProjectStore";
+
+
+const ArtifactsList: React.FC<{ artifacts?: GetArtifactResponseV2[] }> = ({ artifacts }) => {
+
+   const [browse, setBrowse] = useState<{ jobId?: string, stepId?: string, artifactId?: string, artifact?: GetArtifactResponseV2 }>({});
+
+   if (!artifacts?.length) {
+      return null;
+   }
+
+   const sorted = artifacts.filter(a => {
+      return !!a.keys.find(k => k.startsWith("job:") && k.indexOf("/step:") !== -1)
+   }).sort((a, b) => {
+
+      if (a.name !== b.name) {
+         return a.name.localeCompare(b.name);
+      }
+
+      if (a.streamId !== b.streamId) {
+         return a.streamId!.localeCompare(b.streamId!);
+      }
+
+      return (a.change ?? 0) - (b.change ?? 0)
+   })
+
+   const renderRow: IDetailsListProps['onRenderRow'] = (props) => {
+
+      if (props) {
+
+         const item = props!.item as GetArtifactResponseV2;
+
+         let background: string | undefined;
+         if (props.itemIndex % 2 === 0) {
+            background = dashboard.darktheme ? "#1D2021" : "#EAE9E9";
+         }
+
+         const stream = projectStore.streamById(item.streamId);
+
+         return <Stack horizontal verticalAlign="center" verticalFill tokens={{ childrenGap: 24 }} styles={{ root: { backgroundColor: background, paddingLeft: 12, paddingRight: 12, paddingTop: 8, paddingBottom: 8 } }}>
+            <Stack style={{ width: 420, paddingLeft: 8 }}>
+               <Text style={{ fontWeight: 600 }}>{item.name}</Text>
+            </Stack>
+            <Stack style={{ width: 140 }}>
+               <Text>{stream ? (stream.fullname ?? stream.name) : item.streamId}</Text>
+            </Stack>
+            <Stack style={{ width: 72 }}>
+               <Text>{item.change}</Text>
+            </Stack>
+            <Stack style={{ width: 80 }}>
+               <Text>{item.type}</Text>
+            </Stack>
+            <Stack>
+               <DefaultButton style={{ width: 90, paddingRight: 8 }} text="Browse" onClick={() => {
+
+                  const key = item.keys.find(k => k.startsWith("job:") && k.indexOf("/step:") !== -1);
+                  if (!key) {
+                     return;
+                  }
+
+                  setBrowse({ jobId: key.slice(4, 28), stepId: key.slice(-4), artifactId: item.id, artifact: item })
+
+               }} />
+            </Stack>
+         </Stack>
+      }
+      return null;
+   };
+
+   return <Stack styles={{ root: { position: "relative", height: 590 } }}>
+      {!!browse.jobId && <JobArtifactsModal jobId={browse.jobId} stepId={browse.stepId!} artifactId={browse.artifactId} contextType={browse.artifact!.type} onClose={() => setBrowse({})} />}
+      <ScrollablePane scrollbarVisibility={ScrollbarVisibility.always}>
+         <DetailsList
+            isHeaderVisible={false}
+            styles={{
+               root: {
+                  paddingRight: 12
+               },
+               headerWrapper: {
+                  paddingTop: 0
+               }
+            }}
+            items={sorted}
+            selectionMode={SelectionMode.single}
+            layoutMode={DetailsListLayoutMode.justified}
+            compact={false}
+            onRenderRow={renderRow}
+         />
+      </ScrollablePane>
+   </Stack>
+
+}
+
+export const FindArtifactsModal: React.FC<{ onClose: () => void }> = ({ onClose }) => {
+
+   const [state, setState] = useState<{ searching?: boolean, artifacts?: GetArtifactResponseV2[] }>({});
+   const streamRef = React.useRef<IStreamChooser>(null);
+   const minChangeRef = React.useRef<ITextField>(null);
+   const maxChangeRef = React.useRef<ITextField>(null);
+   const nameRef = React.useRef<ITextField>(null);
+   const typeRef = React.useRef<IComboBox>(null);
+
+   const { hordeClasses } = getHordeStyling();
+
+   const queryArtifacts = async () => {
+
+      setState({ ...state, searching: true });
+
+      try {
+
+         let minChange: number | undefined = parseInt(minChangeRef.current?.value?.trim() ?? "0");
+         if (!minChange || isNaN(minChange)) {
+            minChange = undefined;
+         }
+
+         let maxChange: number | undefined = parseInt(maxChangeRef.current?.value?.trim() ?? "0");
+         if (!maxChange || isNaN(maxChange)) {
+            maxChange = undefined;
+         }
+
+         let name: string | undefined = nameRef.current?.value?.trim();
+         if (!name) {
+            name = undefined;
+         }
+
+         let type: string | undefined;
+
+         if (typeRef.current?.selectedOptions?.length) {
+            type = typeRef.current.selectedOptions[0].key as string;
+            if (type === "step-all") {
+               type = undefined;
+            }
+         }
+
+         let streamId = streamRef?.current?.streamId?.trim();
+         if (!streamId) {
+            streamId = undefined;
+         }
+
+         const artifacts = await backend.getArtifacts(streamId, minChange, maxChange, name, type);
+
+         setState({
+            searching: false, artifacts: artifacts.artifacts
+         });
+
+      } catch (reason) {
+         console.error(reason);
+         setState({ searching: false });
+      }
+   }
+
+   const Searching = () => {
+      return <Modal isOpen={true} isBlocking={true} topOffsetFixed={true} styles={{ main: { padding: 8, width: 320, height: 128, hasBeenOpened: false, top: "128px", position: "absolute" } }} className={hordeClasses.modal}>
+         <Stack horizontalAlign="center">
+            <Stack styles={{ root: { padding: 8, paddingBottom: 32 } }}>
+               <Text variant="mediumPlus">Searching for Artifacts</Text>
+            </Stack>
+            <Stack horizontalAlign="center">
+               <Spinner size={SpinnerSize.large} />
+            </Stack>
+         </Stack>
+      </Modal>
+   }
+
+   const artifactTypes: IComboBoxOption[] = [
+      {
+         key: `step-all`,
+         text: `All`
+      }, {
+         key: `step-saved`,
+         text: `step-saved`
+      }, {
+         key: `step-output`,
+         text: `step-output`
+      }, {
+         key: `step-trace`,
+         text: `step-trace`
+      },
+      {
+         key: `step-testdata`,
+         text: `step-testdata`
+      }
+   ]
+
+   return <Stack>
+      <Modal isOpen={true} isBlocking={true} topOffsetFixed={true} styles={{ main: { padding: 8, width: 1024, height: 820, hasBeenOpened: false, top: "80px", position: "absolute" } }} onDismiss={() => onClose()} className={hordeClasses.modal}>
+         {!!state.searching && <Searching />}
+         <Stack className="horde-no-darktheme" styles={{ root: { paddingTop: 10, paddingRight: 12 } }}>
+            <Stack style={{ paddingLeft: 24, paddingRight: 24 }}>
+               <Stack tokens={{ childrenGap: 12 }} style={{ height: 800 }}>
+                  <Stack horizontal verticalAlign="start">
+                     <Stack style={{ paddingTop: 3 }}>
+                        <Text variant="mediumPlus" styles={{ root: { fontFamily: "Horde Open Sans SemiBold" } }}>Find Artifacts</Text>
+                     </Stack>
+                     <Stack grow />
+                     <Stack horizontalAlign="end">
+                        <IconButton
+                           iconProps={{ iconName: 'Cancel' }}
+                           onClick={() => { onClose() }}
+                        />
+                     </Stack>
+                  </Stack>
+                  <Stack horizontal verticalAlign="center" tokens={{ childrenGap: 20 }}>
+                     <Stack >
+                        <Label>Stream</Label>
+                        <StreamChooser ref={streamRef} />
+                     </Stack>
+                     <Stack >
+                        <TextField componentRef={minChangeRef} style={{ width: 100 }} label="Min Changelist" />
+                     </Stack>
+                     <Stack >
+                        <TextField componentRef={maxChangeRef} style={{ width: 100 }} label="Max Changelist" />
+                     </Stack>
+                     <Stack >
+                        <TextField componentRef={nameRef} style={{ width: 224 }} label="Name" spellCheck={false} autoComplete="off" />
+                     </Stack>
+                     <Stack>
+                        <Label>Artifact Type</Label>
+                        <ComboBox componentRef={typeRef} style={{ width: 120, textAlign: "left" }} defaultSelectedKey="step-all" options={artifactTypes} calloutProps={{ doNotLayer: true }} />
+                     </Stack>
+                  </Stack>
+                  <Stack horizontal>
+                     <Stack grow />
+                     <Stack>
+                        <PrimaryButton disabled={!!state.searching} text="Find" onClick={() => (queryArtifacts())} />
+                     </Stack>
+                  </Stack>
+                  <Stack styles={{ root: { paddingTop: 12 } }}>
+                     <ArtifactsList artifacts={state.artifacts} />
+                  </Stack>
+               </Stack>
+            </Stack>
+         </Stack>
+      </Modal>
+   </Stack>;
+};
diff --git a/Engine/Source/Programs/Horde/HordeDashboard/src/components/projects/StreamChooser.tsx b/Engine/Source/Programs/Horde/HordeDashboard/src/components/projects/StreamChooser.tsx
new file mode 100644
--- /dev/null
+++ b/Engine/Source/Programs/Horde/HordeDashboard/src/components/projects/StreamChooser.tsx
@@ -0,0 +1,90 @@
+import { DefaultButton, IContextualMenuItem, IContextualMenuProps, IRefObject, Stack } from "@fluentui/react"
+import { projectStore } from "../../backend/ProjectStore"
+import { useImperativeHandle, useState } from "react";
+import React from "react";
+
+
+export interface IStreamChooser {
+   streamId?: string;
+}
+
+export interface IStreamChooserProps {
+
+}
+
+export const StreamChooser = React.forwardRef<IStreamChooser, IStreamChooserProps>((props, ref) => {
+
+   const [chosen, setChosen] = useState("");
+
+   useImperativeHandle(ref, () => ({
+      streamId: chosen,
+    }));   
+
+   const projectIds = projectStore.projects.sort((a, b) => a.name.localeCompare(b.name)).map(p => p.id);
+
+   let projectOptions: IContextualMenuItem[] = [];
+
+   projectIds.forEach(pid => {
+
+      const project = projectStore.projects.find(p => p.id === pid)!
+
+      const streams = project.streams?.sort((a, b) => {
+         return a.name.localeCompare(b.name);
+      })
+
+      const catItems: IContextualMenuItem[] = project.categories?.map(c => {
+         return {
+            key: c.name,
+            text: c.name,
+            data: c,
+            subMenuProps: { items: [] }
+         }
+      }) ?? [];
+
+      const streamItems: IContextualMenuItem[] = [];
+
+      streams?.forEach(s => {
+
+         const streamItem = {
+            key: s.id,
+            text: s.fullname ?? s.name,
+            onClick: () => {
+               setChosen(s.id)
+            }
+         };
+
+         const c = catItems.find(c => c.data.streams.find(sid => sid === s.id));
+         if (c) {
+            c.subMenuProps!.items.push(streamItem);
+         } else {
+            streamItems.push(streamItem)
+         }
+      });
+
+      projectOptions.push({
+         key: `project_key_${pid}`,
+         text: project.name,
+         subMenuProps: { items: catItems.length ? catItems : streamItems }
+      })
+
+   });
+
+   const templateMenuProps: IContextualMenuProps = {
+      shouldFocusOnMount: true,
+      subMenuHoverDelay: 0,
+      items: projectOptions,
+   };
+
+   let text = "Choose Stream";
+
+   if (chosen) {
+      const stream = projectStore.streamById(chosen);
+      if (stream) {
+         text = stream.fullname ?? stream.name;
+      }
+   }
+
+   return <div>
+      <DefaultButton style={{ width: 320, textAlign: "left" }} menuProps={templateMenuProps} text={text} />
+   </div>
+})
\ No newline at end of file
