Horde: Add No Pool filter to view agents which have no pools assigned
#rnx
#skipci

diff --git a/Engine/Source/Programs/Horde/HordeDashboard/src/components/AgentView.tsx b/Engine/Source/Programs/Horde/HordeDashboard/src/components/AgentView.tsx
--- a/Engine/Source/Programs/Horde/HordeDashboard/src/components/AgentView.tsx
+++ b/Engine/Source/Programs/Horde/HordeDashboard/src/components/AgentView.tsx
@@ -1403,7 +1403,7 @@ const agentContextMenuProps: IContextualMenuItem[] = [
    },
 ];
 
-const agentStatus = ["Active", "Ready", "Disabled", "Ephemeral", "Pending Conform", "Pending Shutdown", "Offline", "Offline (Autoscaler)", "Offline (Manual)", "Offline (Unexpected)"];
+const agentStatus = ["Active", "Ready", "Disabled", "Ephemeral", "No Pools", "Pending Conform", "Pending Shutdown", "Offline", "Offline (Autoscaler)", "Offline (Manual)", "Offline (Unexpected)"];
 
 
 export const AgentMenuBar: React.FC<{ agentView?: boolean }> = observer(({ agentView }) => {
@@ -2284,6 +2284,12 @@ export const AgentViewInner: React.FC<{ agentId?: string, poolId?: string, searc
          }
       }
 
+      if (filter.has("No Pools")) {
+         if (!item.pools?.length) {
+            filtered = false;
+         }
+      }
+
       if (filter.has("Ephemeral")) {
          if (item.ephemeral) {
             filtered = false;
