Horde: Remove the ClickHouse telemetry sink for now. We don't currently rely on this for anything, and agent CPU/RAM/disk telemetry is sent via a different mechanism into a custom collection.

#rnx

diff --git a/Engine/Source/Programs/Horde/HordeAgent/HordeAgent.csproj b/Engine/Source/Programs/Horde/HordeAgent/HordeAgent.csproj
--- a/Engine/Source/Programs/Horde/HordeAgent/HordeAgent.csproj
+++ b/Engine/Source/Programs/Horde/HordeAgent/HordeAgent.csproj
@@ -150,7 +150,6 @@
     <Protobuf Include="..\HordeServer\Protos\horde\messages\agent_requirements.proto" ProtoRoot="..\HordeServer\Protos" Access="Public" />
     <Protobuf Include="..\HordeServer\Protos\horde\messages\lease.proto" ProtoRoot="..\HordeServer\Protos" Access="Public" />
     <Protobuf Include="..\HordeServer\Protos\horde\messages\aws_lambda.proto" ProtoRoot="..\HordeServer\Protos" Access="Public" />
-    <Protobuf Include="..\HordeServer\Protos\horde\messages\telemetry.proto" ProtoRoot="..\HordeServer\Protos" Access="Public" />
     <Protobuf Include="..\HordeServer\Protos\horde\tasks\compute_task.proto" ProtoRoot="..\HordeServer\Protos" Access="Public" />
     <Protobuf Include="..\HordeServer\Protos\horde\tasks\restart_task.proto" ProtoRoot="..\HordeServer\Protos" Access="Public" />
     <Protobuf Include="..\HordeServer\Protos\horde\tasks\shutdown_task.proto" ProtoRoot="..\HordeServer\Protos" Access="Public" />
diff --git a/Engine/Source/Programs/Horde/HordeAgent/Services/TelemetryService.cs b/Engine/Source/Programs/Horde/HordeAgent/Services/TelemetryService.cs
--- a/Engine/Source/Programs/Horde/HordeAgent/Services/TelemetryService.cs
+++ b/Engine/Source/Programs/Horde/HordeAgent/Services/TelemetryService.cs
@@ -11,7 +11,6 @@ using Grpc.Core;
 using Grpc.Net.Client;
 using HordeCommon.Rpc;
 using HordeCommon.Rpc.Messages;
-using HordeCommon.Rpc.Messages.Telemetry;
 using HordeCommon.Rpc.Tasks;
 using Microsoft.Extensions.Hosting;
 using Microsoft.Extensions.Logging;
@@ -45,15 +44,6 @@ public class CpuMetrics
 	{
 		return $"User={User,5:F1}% System={System,5:F1}% Idle={Idle,5:F1}%";
 	}
-
-	/// <summary>
-	/// Convert to Protobuf-based event
-	/// </summary>
-	/// <returns></returns>
-	public RpcAgentCpuMetricsEvent ToEvent()
-	{
-		return new RpcAgentCpuMetricsEvent { User = User, System = System, Idle = Idle };
-	}
 }
 
 /// <summary>
@@ -86,15 +76,6 @@ public class MemoryMetrics
 	{
 		return $"Total={Total} kB, Available={Available} kB, Used={Used} kB, Used={UsedPercentage * 100.0:F1} %";
 	}
-
-	/// <summary>
-	/// Convert to Protobuf-based event
-	/// </summary>
-	/// <returns></returns>
-	public RpcAgentMemoryMetricsEvent ToEvent()
-	{
-		return new RpcAgentMemoryMetricsEvent { Total = Total, Free = Available, Used = Used, UsedPercentage = UsedPercentage };
-	}
 }
 
 /// <summary>
@@ -276,7 +257,7 @@ public sealed class WindowsSystemMetrics : ISystemMetrics, IDisposable
 /// <summary>
 /// Send telemetry events back to server at regular intervals
 /// </summary>
-class TelemetryService : BackgroundService
+class TelemetryService : IHostedService, IDisposable
 {
 	private readonly TimeSpan _heartbeatInterval = TimeSpan.FromSeconds(60);
 	private readonly TimeSpan _heartbeatMaxAllowedDiff = TimeSpan.FromSeconds(5);
@@ -286,10 +267,7 @@ class TelemetryService : BackgroundService
 	private readonly AgentSettings _agentSettings;
 	private readonly ILogger<TelemetryService> _logger;
 	private readonly TimeSpan _reportInterval;
-	private readonly RpcAgentMetadataEvent _agentMetadataEvent;
-	private readonly TimeSpan _agentMetadataReportInterval = TimeSpan.FromMinutes(2);
 	private readonly ISystemMetrics? _systemMetrics;
-	private DateTime _lastTimeAgentMetadataSent = DateTime.UnixEpoch;
 
 	private CancellationTokenSource? _eventLoopHeartbeatCts;
 	private Task? _eventLoopTask;
@@ -306,28 +284,23 @@ class TelemetryService : BackgroundService
 		_logger = logger;
 		_reportInterval = TimeSpan.FromMilliseconds(_agentSettings.TelemetryReportInterval);
 		_systemMetrics = systemMetrics;
-
-		// Calculate this once at startup as it should not change during lifetime of process
-		_agentMetadataEvent = GetAgentMetadataEvent();
 	}
 
 	/// <inheritdoc/>
-	public override void Dispose()
+	public void Dispose()
 	{
-		base.Dispose();
 		_eventLoopHeartbeatCts?.Dispose();
 	}
 
 	/// <inheritdoc />
-	public override Task StartAsync(CancellationToken cancellationToken)
+	public Task StartAsync(CancellationToken cancellationToken)
 	{
 		_eventLoopHeartbeatCts = CancellationTokenSource.CreateLinkedTokenSource(cancellationToken);
 		_eventLoopTask = EventLoopHeartbeatAsync(_eventLoopHeartbeatCts.Token);
-
-		return base.StartAsync(cancellationToken);
+		return Task.CompletedTask;
 	}
 
-	public override async Task StopAsync(CancellationToken cancellationToken)
+	public async Task StopAsync(CancellationToken cancellationToken)
 	{
 		_eventLoopHeartbeatCts?.Cancel();
 
@@ -342,8 +315,6 @@ class TelemetryService : BackgroundService
 				// Ignore cancellation exceptions
 			}
 		}
-
-		await base.StopAsync(cancellationToken);
 	}
 
 	/// <summary>
@@ -510,178 +481,4 @@ class TelemetryService : BackgroundService
 
 		return filters;
 	}
-
-	/// <inheritdoc />
-	protected override async Task ExecuteAsync(CancellationToken stoppingToken)
-	{
-		_logger.LogDebug("Starting telemetry background task");
-		while (!stoppingToken.IsCancellationRequested)
-		{
-			try
-			{
-				if (!await ExecuteBackgroundInternalAsync(stoppingToken))
-				{
-					break;
-				}
-			}
-			catch (OperationCanceledException) when (stoppingToken.IsCancellationRequested)
-			{
-				break;
-			}
-			catch (Exception ex)
-			{
-				_logger.LogWarning(ex, "Exception in TelemetryService: {Message}", ex.Message);
-			}
-
-			// Wait a moment before attempting to restart the background work
-			try
-			{
-				await Task.Delay(TimeSpan.FromSeconds(20), stoppingToken);
-			}
-			catch (OperationCanceledException) when (stoppingToken.IsCancellationRequested)
-			{
-				break;
-			}
-		}
-		_logger.LogDebug("Stopping telemetry background task");
-	}
-
-	private async Task<bool> ExecuteBackgroundInternalAsync(CancellationToken stoppingToken)
-	{
-		if (_systemMetrics == null || !_agentSettings.EnableTelemetry)
-		{
-			return false;
-		}
-
-		using GrpcChannel channel = await _grpcService.CreateGrpcChannelAsync(stoppingToken);
-		CallInvoker invoker = _grpcService.GetInvoker(channel);
-		HordeRpc.HordeRpcClient client = new(invoker);
-
-		while (!stoppingToken.IsCancellationRequested)
-		{
-			_logger.LogDebug("Sending telemetry events to server...");
-
-			CpuMetrics? cpuMetrics = _systemMetrics.GetCpu();
-			MemoryMetrics? memMetrics = _systemMetrics.GetMemory();
-
-			RpcSendTelemetryEventsRequest request = new();
-			Timestamp utcNow = Timestamp.FromDateTime(DateTime.UtcNow);
-
-			RpcExecutionMetadata em = new();
-			foreach (RpcLease lease in _workerService.GetActiveLeases())
-			{
-				if (lease.Payload.TryUnpack<ExecuteJobTask>(out ExecuteJobTask task))
-				{
-					em.LeaseId = lease.Id;
-					em.JobId = task.JobId;
-					em.JobBatchId = task.BatchId;
-					break;
-				}
-			}
-
-			if (cpuMetrics != null)
-			{
-				RpcAgentCpuMetricsEvent cpuMetricsEvent = cpuMetrics.ToEvent();
-				cpuMetricsEvent.AgentId = _agentMetadataEvent.AgentId;
-				cpuMetricsEvent.Timestamp = utcNow;
-				cpuMetricsEvent.ExecutionMetadata = em;
-				request.Events.Add(new RpcWrappedTelemetryEvent { Cpu = cpuMetricsEvent });
-			}
-
-			if (memMetrics != null)
-			{
-				RpcAgentMemoryMetricsEvent memMetricsEvent = memMetrics.ToEvent();
-				memMetricsEvent.AgentId = _agentMetadataEvent.AgentId;
-				memMetricsEvent.Timestamp = utcNow;
-				memMetricsEvent.ExecutionMetadata = em;
-				request.Events.Add(new RpcWrappedTelemetryEvent { Mem = memMetricsEvent });
-			}
-
-			if (DateTime.UtcNow > _lastTimeAgentMetadataSent + _agentMetadataReportInterval)
-			{
-				// Report agent metadata every now and then as events are not guaranteed to be delivered.
-				// Re-sending ensures the metadata will eventually make it to the server.
-				request.Events.Add(new RpcWrappedTelemetryEvent { AgentMetadata = _agentMetadataEvent });
-				_lastTimeAgentMetadataSent = DateTime.UtcNow;
-			}
-
-			await client.SendTelemetryEventsAsync(request, new CallOptions(cancellationToken: stoppingToken));
-
-			await Task.Delay(_reportInterval, stoppingToken);
-		}
-
-		return true;
-	}
-
-	RpcAgentMetadataEvent GetAgentMetadataEvent()
-	{
-		RpcAgentMetadataEvent e = new()
-		{
-			Ip = null,
-			Hostname = _agentSettings.GetAgentName(),
-			Region = null,
-			AvailabilityZone = null,
-			Environment = null,
-			AgentVersion = AgentApp.Version,
-			Os = GetOs(),
-			OsVersion = Environment.OSVersion.Version.ToString(),
-			Architecture = RuntimeInformation.OSArchitecture.ToString(),
-		};
-		e.PoolIds.AddRange(_workerService.PoolIds);
-		e.AgentId = CalculateAgentId(e);
-		return e;
-	}
-
-	/// <summary>
-	/// Calculate an agent ID
-	/// </summary>
-	/// <returns>A unique hash for all fields</returns>
-	private static long CalculateAgentId(RpcAgentMetadataEvent e)
-	{
-		using SHA256 sha256 = SHA256.Create();
-		using MemoryStream ms = new(200);
-		using BinaryWriter bw = new(ms);
-
-		bw.Write(e.Ip ?? "<empty ip>");
-		bw.Write(e.Hostname ?? "<empty hostname>");
-		bw.Write(e.Region ?? "<empty region>");
-		bw.Write(e.AvailabilityZone ?? "<empty az>");
-		bw.Write(e.Environment ?? "<empty env>");
-		bw.Write(e.AgentVersion ?? "<empty version>");
-		bw.Write(e.Os ?? "<empty os>");
-		bw.Write(e.OsVersion ?? "<empty os version>");
-		bw.Write(e.Architecture ?? "<empty os architecture>");
-
-		foreach (KeyValuePair<string, string> pair in e.Properties)
-		{
-			bw.Write(pair.Key ?? "<empty key>");
-			bw.Write(pair.Value ?? "<empty value>");
-		}
-
-		foreach (string poolId in e.PoolIds)
-		{
-			bw.Write(poolId);
-		}
-
-		ms.Position = 0;
-		byte[] hash = sha256.ComputeHash(ms);
-		return BitConverter.ToInt64(hash, 0);
-	}
-
-	private static string GetOs()
-	{
-		if (RuntimeInformation.IsOSPlatform(OSPlatform.Windows))
-		{
-			return "Windows";
-		}
-		if (RuntimeInformation.IsOSPlatform(OSPlatform.Linux))
-		{
-			return "Linux";
-		}
-		if (RuntimeInformation.IsOSPlatform(OSPlatform.OSX))
-		{
-			return "macOS";
-		}
-		return "Unknown";
-	}
 }
\ No newline at end of file
diff --git a/Engine/Source/Programs/Horde/HordeServer.Shared/Telemetry/ITelemetryWriter.cs b/Engine/Source/Programs/Horde/HordeServer.Shared/Telemetry/ITelemetryWriter.cs
--- a/Engine/Source/Programs/Horde/HordeServer.Shared/Telemetry/ITelemetryWriter.cs
+++ b/Engine/Source/Programs/Horde/HordeServer.Shared/Telemetry/ITelemetryWriter.cs
@@ -37,4 +37,21 @@ namespace HordeServer.Telemetry
 	/// <param name="AppEnvironment">Name of the environment that the sending application is running in</param>
 	/// <param name="SessionId">Unique identifier for the current session</param>
 	public record class TelemetryRecordMeta(string? AppId = null, string? AppVersion = null, string? AppEnvironment = null, string? SessionId = null);
+
+	/// <summary>
+	/// Inactive implementation of <see cref="ITelemetryWriter"/>
+	/// </summary>
+	public class NullTelemetryWriter : ITelemetryWriter
+	{
+		/// <inheritdoc/>
+		public bool Enabled { get; } = false;
+
+		/// <inheritdoc/>
+		public void WriteEvent(TelemetryStoreId telemetryStoreId, object payload)
+		{ }
+
+		/// <inheritdoc/>
+		public void WriteEvent(TelemetryStoreId telemetryStoreId, TelemetryRecordMeta metadata, object payload)
+		{ }
+	}
 }
diff --git a/Engine/Source/Programs/Horde/HordeServer/HordeServer.csproj b/Engine/Source/Programs/Horde/HordeServer/HordeServer.csproj
--- a/Engine/Source/Programs/Horde/HordeServer/HordeServer.csproj
+++ b/Engine/Source/Programs/Horde/HordeServer/HordeServer.csproj
@@ -112,7 +112,6 @@
     <Protobuf Include="Protos\horde\messages\agent_requirements.proto" ProtoRoot="Protos" />
     <Protobuf Include="Protos\horde\messages\agent_workspace.proto" ProtoRoot="Protos" />
     <Protobuf Include="Protos\horde\messages\lease.proto" ProtoRoot="Protos" />
-    <Protobuf Include="Protos\horde\messages\telemetry.proto" ProtoRoot="Protos" />
     <Protobuf Include="Protos\horde\tasks\compute_task.proto" ProtoRoot="Protos" />
     <Protobuf Include="Protos\horde\tasks\conform_task.proto" ProtoRoot="Protos" />
     <Protobuf Include="Protos\horde\tasks\job_task.proto" ProtoRoot="Protos" />
diff --git a/Engine/Source/Programs/Horde/HordeServer/Protos/horde/messages/telemetry.cs b/Engine/Source/Programs/Horde/HordeServer/Protos/horde/messages/telemetry.cs
deleted file mode 100644
diff --git a/Engine/Source/Programs/Horde/HordeServer/Server/RpcService.cs b/Engine/Source/Programs/Horde/HordeServer/Server/RpcService.cs
--- a/Engine/Source/Programs/Horde/HordeServer/Server/RpcService.cs
+++ b/Engine/Source/Programs/Horde/HordeServer/Server/RpcService.cs
@@ -26,7 +26,6 @@ using HordeServer.Tools;
 using HordeServer.Utilities;
 using HordeCommon.Rpc;
 using HordeCommon.Rpc.Messages;
-using HordeCommon.Rpc.Messages.Telemetry;
 using Microsoft.AspNetCore.Authorization;
 using Microsoft.Extensions.Logging;
 using Microsoft.Extensions.Options;
@@ -456,34 +455,5 @@ namespace HordeServer.Server
 
 			return Task.FromResult(new Empty());
 		}
-
-		/// <summary>
-		/// Receives telemetry events from agents
-		/// </summary>
-		/// <param name="request">Request arguments</param>
-		/// <param name="context">Context for the RPC call</param>
-		/// <returns>An empty response</returns>
-		public override async Task<Empty> SendTelemetryEvents(RpcSendTelemetryEventsRequest request, ServerCallContext context)
-		{
-			ISession? session = null;
-
-			SessionId? sessionId = context.GetHttpContext().User.GetSessionClaim();
-			if (sessionId != null)
-			{
-				session = await _agentService.GetSessionAsync(sessionId.Value);
-			}
-
-			TelemetryRecordMeta agentMeta = new TelemetryRecordMeta("HordeAgent", session?.Version ?? "(Unknown)", ServerApp.DeploymentEnvironment, sessionId?.ToString() ?? "(Unknown)");
-			foreach (RpcWrappedTelemetryEvent wrappedEvent in request.Events)
-			{
-				OneofDescriptor oneofDescriptor = RpcWrappedTelemetryEvent.Descriptor.Oneofs[0];
-				FieldDescriptor caseDescriptor = oneofDescriptor.Accessor.GetCaseFieldDescriptor(wrappedEvent);
-
-				object wrappedValue = caseDescriptor.Accessor.GetValue(wrappedEvent);
-				_telemetryWriter.WriteEvent(TelemetryStoreId.Default, agentMeta, wrappedValue);
-			}
-
-			return new Empty();
-		}
 	}
 }
diff --git a/Engine/Source/Programs/Horde/HordeServer/Startup.cs b/Engine/Source/Programs/Horde/HordeServer/Startup.cs
--- a/Engine/Source/Programs/Horde/HordeServer/Startup.cs
+++ b/Engine/Source/Programs/Horde/HordeServer/Startup.cs
@@ -798,7 +798,6 @@ namespace HordeServer
 			services.AddSingleton<IMetricCollection, MetricCollection>(sp => sp.GetRequiredService<MetricCollection>());
 
 			services.AddHttpClient(EpicTelemetrySink.HttpClientName, client => { });
-			services.AddHttpClient(ClickHouseTelemetrySink.HttpClientName, client => { });
 
 			// Create the agent telemetry collection, and register the hosted service so we can flush from any server.
 			services.AddSingleton<AgentTelemetryCollection>();
diff --git a/Engine/Source/Programs/Horde/HordeServer/Telemetry/Sinks/ClickHouseTelemetrySink.cs b/Engine/Source/Programs/Horde/HordeServer/Telemetry/Sinks/ClickHouseTelemetrySink.cs
deleted file mode 100644
diff --git a/Engine/Source/Programs/Horde/HordeServer/Telemetry/TelemetryManager.cs b/Engine/Source/Programs/Horde/HordeServer/Telemetry/TelemetryManager.cs
--- a/Engine/Source/Programs/Horde/HordeServer/Telemetry/TelemetryManager.cs
+++ b/Engine/Source/Programs/Horde/HordeServer/Telemetry/TelemetryManager.cs
@@ -55,9 +55,6 @@ namespace HordeServer.Telemetry
 					case EpicTelemetryConfig epicConfig:
 						_telemetrySinks.Add(new EpicTelemetrySink(epicConfig, loggerFactory.CreateLogger<EpicTelemetrySink>()));
 						break;
-					case ClickHouseTelemetryConfig chConfig:
-						_telemetrySinks.Add(new ClickHouseTelemetrySink(chConfig, httpClientFactory, loggerFactory.CreateLogger<ClickHouseTelemetrySink>()));
-						break;
 					case MongoTelemetryConfig mongoConfig:
 						_telemetrySinks.Add(serviceProvider.GetRequiredService<MongoTelemetrySink>());
 						break;
