Horde: Fix 404 error when deleting a deployment
#rnx

diff --git a/Engine/Source/Programs/Horde/HordeServer/Tools/ToolsController.cs b/Engine/Source/Programs/Horde/HordeServer/Tools/ToolsController.cs
--- a/Engine/Source/Programs/Horde/HordeServer/Tools/ToolsController.cs
+++ b/Engine/Source/Programs/Horde/HordeServer/Tools/ToolsController.cs
@@ -156,7 +156,7 @@ namespace HordeServer.Tools
 			if (request.State != null)
 			{
 				deployment = await deployment.UpdateAsync(request.State.Value, HttpContext.RequestAborted);
-				if (deployment == null)
+				if (request.State != ToolDeploymentState.Cancelled && deployment == null)
 				{
 					return NotFound(id, deploymentId);
 				}
