Horde: Fix issues with tool management
#rnx
#skipci

diff --git a/Engine/Source/Programs/Horde/HordeDashboard/src/backend/Api.ts b/Engine/Source/Programs/Horde/HordeDashboard/src/backend/Api.ts
--- a/Engine/Source/Programs/Horde/HordeDashboard/src/backend/Api.ts
+++ b/Engine/Source/Programs/Horde/HordeDashboard/src/backend/Api.ts
@@ -5061,7 +5061,7 @@ export enum ToolDeploymentState {
 	Complete = "Complete",
 
 	/** The deployment has been cancelled. */
-	Cancelled = "Canceled"
+	Cancelled = "Cancelled"
 }
 
 /** Summary for a particular tool */
diff --git a/Engine/Source/Programs/Horde/HordeDashboard/src/backend/Backend.ts b/Engine/Source/Programs/Horde/HordeDashboard/src/backend/Backend.ts
--- a/Engine/Source/Programs/Horde/HordeDashboard/src/backend/Backend.ts
+++ b/Engine/Source/Programs/Horde/HordeDashboard/src/backend/Backend.ts
@@ -1636,7 +1636,7 @@ export class Backend {
     /// Updates the state of an active deployment.
     updateDeployment(toolId: string, deploymentId: string, update: UpdateDeploymentRequest): Promise<void> {
         return new Promise<void>((resolve, reject) => {
-            this.backend.patch(`/api/v1/tools/${toolId}/deployments/${deploymentId}`, { request: update }).then(() => {
+            this.backend.patch(`/api/v1/tools/${toolId}/deployments/${deploymentId}`, update).then(() => {
                 resolve();
             }).catch(reason => {
                 reject(reason);
diff --git a/Engine/Source/Programs/Horde/HordeDashboard/src/components/tools/ToolManagement.tsx b/Engine/Source/Programs/Horde/HordeDashboard/src/components/tools/ToolManagement.tsx
--- a/Engine/Source/Programs/Horde/HordeDashboard/src/components/tools/ToolManagement.tsx
+++ b/Engine/Source/Programs/Horde/HordeDashboard/src/components/tools/ToolManagement.tsx
@@ -36,7 +36,6 @@ class ToolHandler extends PollBase {
 
       if (tool) {
          this.currentDeployment = this.data!.deployments?.find(d => d.id === tool.deploymentId);
-
       }
 
       this.setUpdated();
@@ -82,13 +81,13 @@ const ActionConfirmation: React.FC<{ handler: ToolHandler, request: UpdateDeploy
 
          await backend.updateDeployment(handler.data!.id, deployment.id, request);
 
-         handler.poll();         
 
       } catch (reason) {
 
          errorReason = reason ? reason : "Unknown Error";
 
       } finally {
+
          if (errorReason) {
 
             ErrorHandler.set({
@@ -97,8 +96,10 @@ const ActionConfirmation: React.FC<{ handler: ToolHandler, request: UpdateDeploy
                title: `Error Updating Deployment`,
                message: `There was an issue updating the tool deployment.\n\nReason: ${errorReason}\n\nTime: ${moment.utc().format("MMM Do, HH:mm z")}`
 
-            }, true);            
-
+            }, true);
+         }
+         else {
+            handler.poll();
          }
 
          setSubmitting(false);
@@ -230,6 +231,11 @@ const ToolManagementInner: React.FC<{ handler: ToolHandler }> = observer(({ hand
       if (item.deploy.id === handler.currentDeployment?.id) {
          state = "Current";
       }
+
+      if (item.deploy.progress < 1.0) {
+         state = "Deploying"
+      }
+
       return <Stack style={{ width: 96 }} horizontalAlign="center"><Text style={{ fontWeight: state === "Current" ? 600 : undefined }}>{state}</Text></Stack>
    }
 
@@ -241,7 +247,7 @@ const ToolManagementInner: React.FC<{ handler: ToolHandler }> = observer(({ hand
 
          let background: string | undefined;
          if (props.itemIndex % 2 === 0) {
-            background = dashboard.darktheme ? "#1D2021" : "#EAE9E9";
+            //background = dashboard.darktheme ? "#1D2021" : "#FAF9F9";
          }
 
          return <Stack horizontal verticalAlign="center" verticalFill tokens={{ childrenGap: 24 }} styles={{ root: { backgroundColor: background, paddingLeft: 12, paddingRight: 12, paddingTop: 8, paddingBottom: 8 } }}>
@@ -291,7 +297,7 @@ export const ToolManagementModal: React.FC<{ toolId: string; toolName?: string,
             <Stack>
                <Text variant="mediumPlus">{headerText}</Text>
             </Stack>
-            <Stack grow />
+            <Stack grow />;
             <Stack>
                <IconButton
                   iconProps={{ iconName: 'Cancel' }}
