Horde: Fix static analysis warnings

#rnx

diff --git a/Engine/Source/Programs/Horde/HordeAgent.Tests/Services/TelemetryServiceTest.cs b/Engine/Source/Programs/Horde/HordeAgent.Tests/Services/TelemetryServiceTest.cs
--- a/Engine/Source/Programs/Horde/HordeAgent.Tests/Services/TelemetryServiceTest.cs
+++ b/Engine/Source/Programs/Horde/HordeAgent.Tests/Services/TelemetryServiceTest.cs
@@ -25,7 +25,7 @@ public sealed class TelemetryServiceTest : IDisposable
 		using WorkerService workerService = new WorkerService(null!, null!, null!, null!);
 		GrpcService grpcService = new(settingsOpt, NullLogger<GrpcService>.Instance, _loggerFactory);
 
-		_telemetryService = new TelemetryService(workerService, grpcService, null!, settingsOpt, NullLogger<TelemetryService>.Instance);
+		_telemetryService = new TelemetryService(NullLogger<TelemetryService>.Instance);
 	}
 
 	[TestMethod]
diff --git a/Engine/Source/Programs/Horde/HordeAgent/Services/TelemetryService.cs b/Engine/Source/Programs/Horde/HordeAgent/Services/TelemetryService.cs
--- a/Engine/Source/Programs/Horde/HordeAgent/Services/TelemetryService.cs
+++ b/Engine/Source/Programs/Horde/HordeAgent/Services/TelemetryService.cs
@@ -3,18 +3,10 @@
 using System.Diagnostics;
 using System.Runtime.InteropServices;
 using System.Runtime.Versioning;
-using System.Security.Cryptography;
 using System.Text;
 using EpicGames.Core;
-using Google.Protobuf.WellKnownTypes;
-using Grpc.Core;
-using Grpc.Net.Client;
-using HordeCommon.Rpc;
-using HordeCommon.Rpc.Messages;
-using HordeCommon.Rpc.Tasks;
 using Microsoft.Extensions.Hosting;
 using Microsoft.Extensions.Logging;
-using Microsoft.Extensions.Options;
 using Microsoft.Management.Infrastructure;
 
 namespace HordeAgent.Services;
@@ -262,12 +254,7 @@ class TelemetryService : IHostedService, IDisposable
 	private readonly TimeSpan _heartbeatInterval = TimeSpan.FromSeconds(60);
 	private readonly TimeSpan _heartbeatMaxAllowedDiff = TimeSpan.FromSeconds(5);
 
-	private readonly WorkerService _workerService;
-	private readonly GrpcService _grpcService;
-	private readonly AgentSettings _agentSettings;
 	private readonly ILogger<TelemetryService> _logger;
-	private readonly TimeSpan _reportInterval;
-	private readonly ISystemMetrics? _systemMetrics;
 
 	private CancellationTokenSource? _eventLoopHeartbeatCts;
 	private Task? _eventLoopTask;
@@ -276,14 +263,9 @@ class TelemetryService : IHostedService, IDisposable
 	/// <summary>
 	/// Constructor
 	/// </summary>
-	public TelemetryService(WorkerService workerService, GrpcService grpcService, ISystemMetrics systemMetrics, IOptions<AgentSettings> settings, ILogger<TelemetryService> logger)
+	public TelemetryService(ILogger<TelemetryService> logger)
 	{
-		_workerService = workerService;
-		_grpcService = grpcService;
-		_agentSettings = settings.Value;
 		_logger = logger;
-		_reportInterval = TimeSpan.FromMilliseconds(_agentSettings.TelemetryReportInterval);
-		_systemMetrics = systemMetrics;
 	}
 
 	/// <inheritdoc/>
diff --git a/Engine/Source/Programs/Horde/HordeServer/Server/RpcService.cs b/Engine/Source/Programs/Horde/HordeServer/Server/RpcService.cs
--- a/Engine/Source/Programs/Horde/HordeServer/Server/RpcService.cs
+++ b/Engine/Source/Programs/Horde/HordeServer/Server/RpcService.cs
@@ -11,12 +11,12 @@ using System.Threading;
 using System.Threading.Tasks;
 using EpicGames.Horde.Agents;
 using EpicGames.Horde.Agents.Sessions;
-using EpicGames.Horde.Telemetry;
 using EpicGames.Horde.Tools;
 using Google.Protobuf;
-using Google.Protobuf.Reflection;
 using Google.Protobuf.WellKnownTypes;
 using Grpc.Core;
+using HordeCommon.Rpc;
+using HordeCommon.Rpc.Messages;
 using HordeServer.Acls;
 using HordeServer.Agents;
 using HordeServer.Agents.Sessions;
@@ -24,8 +24,6 @@ using HordeServer.Agents.Telemetry;
 using HordeServer.Telemetry;
 using HordeServer.Tools;
 using HordeServer.Utilities;
-using HordeCommon.Rpc;
-using HordeCommon.Rpc.Messages;
 using Microsoft.AspNetCore.Authorization;
 using Microsoft.Extensions.Logging;
 using Microsoft.Extensions.Options;
@@ -45,7 +43,6 @@ namespace HordeServer.Server
 
 		readonly AgentService _agentService;
 		readonly LifetimeService _lifetimeService;
-		readonly ITelemetryWriter _telemetryWriter;
 		readonly IToolCollection _toolCollection;
 		readonly IAgentTelemetryCollection _agentTelemetryCollection;
 		readonly AclService _aclService;
@@ -55,11 +52,10 @@ namespace HordeServer.Server
 		/// <summary>
 		/// Constructor
 		/// </summary>
-		public RpcService(AgentService agentService, LifetimeService lifetimeService, ITelemetryWriter telemetryWriter, IToolCollection toolCollection, IAgentTelemetryCollection agentTelemetryCollection, AclService aclService, IOptionsSnapshot<GlobalConfig> globalConfig, ILogger<RpcService> logger)
+		public RpcService(AgentService agentService, LifetimeService lifetimeService, IToolCollection toolCollection, IAgentTelemetryCollection agentTelemetryCollection, AclService aclService, IOptionsSnapshot<GlobalConfig> globalConfig, ILogger<RpcService> logger)
 		{
 			_agentService = agentService;
 			_lifetimeService = lifetimeService;
-			_telemetryWriter = telemetryWriter;
 			_toolCollection = toolCollection;
 			_agentTelemetryCollection = agentTelemetryCollection;
 			_aclService = aclService;
