Horde: Couple more static analysis fixes.

#rnx

diff --git a/Engine/Source/Programs/Horde/HordeServer/Server/RpcService.cs b/Engine/Source/Programs/Horde/HordeServer/Server/RpcService.cs
--- a/Engine/Source/Programs/Horde/HordeServer/Server/RpcService.cs
+++ b/Engine/Source/Programs/Horde/HordeServer/Server/RpcService.cs
@@ -21,7 +21,6 @@ using HordeServer.Acls;
 using HordeServer.Agents;
 using HordeServer.Agents.Sessions;
 using HordeServer.Agents.Telemetry;
-using HordeServer.Telemetry;
 using HordeServer.Tools;
 using HordeServer.Utilities;
 using Microsoft.AspNetCore.Authorization;
diff --git a/Engine/Source/Programs/Horde/HordeServer/Telemetry/TelemetryManager.cs b/Engine/Source/Programs/Horde/HordeServer/Telemetry/TelemetryManager.cs
--- a/Engine/Source/Programs/Horde/HordeServer/Telemetry/TelemetryManager.cs
+++ b/Engine/Source/Programs/Horde/HordeServer/Telemetry/TelemetryManager.cs
@@ -2,7 +2,6 @@
 
 using System;
 using System.Collections.Generic;
-using System.Net.Http;
 using System.Text.Json;
 using System.Threading;
 using System.Threading.Tasks;
@@ -42,7 +41,7 @@ namespace HordeServer.Telemetry
 		/// <summary>
 		/// Constructor
 		/// </summary>
-		public TelemetryManager(IServiceProvider serviceProvider, IHttpClientFactory httpClientFactory, IClock clock, IOptions<ServerSettings> serverSettings, Tracer tracer, ILoggerFactory loggerFactory)
+		public TelemetryManager(IServiceProvider serviceProvider, IClock clock, IOptions<ServerSettings> serverSettings, Tracer tracer, ILoggerFactory loggerFactory)
 		{
 			_tracer = tracer;
 			_logger = loggerFactory.CreateLogger<TelemetryManager>();
