Horde: Allow aribtraty artifsct types to be entered for search, show more data with the info box
#skipci
#rnx

diff --git a/Engine/Source/Programs/Horde/HordeDashboard/src/components/artifacts/ArtifactsModal.tsx b/Engine/Source/Programs/Horde/HordeDashboard/src/components/artifacts/ArtifactsModal.tsx
--- a/Engine/Source/Programs/Horde/HordeDashboard/src/components/artifacts/ArtifactsModal.tsx
+++ b/Engine/Source/Programs/Horde/HordeDashboard/src/components/artifacts/ArtifactsModal.tsx
@@ -692,7 +692,7 @@ const JobDetailArtifactsInner: React.FC<{ jobId: string; stepId: string, artifac
                      </Stack>
                   </Stack>
                   <Stack style={{ paddingLeft: 0 }}>
-                     <Text style={{ fontSize: 11, whiteSpace: "pre-wrap", fontFamily: "Horde Cousine Regular" }}>{JSON.stringify(item.dirResponse ?? item.fileResponse ?? {}, undefined, 2).replaceAll("\\r", "").replaceAll("\\n", "\n")}</Text></Stack>
+                     <Text style={{ fontSize: 11, whiteSpace: "pre-wrap", fontFamily: "Horde Cousine Regular" }}>{ (handler.artifact ? JSON.stringify(handler.artifact, undefined, 2) + ",\n" : "") + JSON.stringify(item.dirResponse ?? item.fileResponse ?? {}, undefined, 2).replaceAll("\\r", "").replaceAll("\\n", "\n")}</Text></Stack>
                </Stack>
             </Callout>}
          </Stack>
diff --git a/Engine/Source/Programs/Horde/HordeDashboard/src/components/artifacts/ArtifactsSearch.tsx b/Engine/Source/Programs/Horde/HordeDashboard/src/components/artifacts/ArtifactsSearch.tsx
--- a/Engine/Source/Programs/Horde/HordeDashboard/src/components/artifacts/ArtifactsSearch.tsx
+++ b/Engine/Source/Programs/Horde/HordeDashboard/src/components/artifacts/ArtifactsSearch.tsx
@@ -99,6 +99,26 @@ const ArtifactsList: React.FC<{ artifacts?: GetArtifactResponseV2[] }> = ({ arti
 
 }
 
+const artifactTypes: IComboBoxOption[] = [
+   {
+      key: `step-all`,
+      text: `All`
+   }, {
+      key: `step-saved`,
+      text: `step-saved`
+   }, {
+      key: `step-output`,
+      text: `step-output`
+   }, {
+      key: `step-trace`,
+      text: `step-trace`
+   },
+   {
+      key: `step-testdata`,
+      text: `step-testdata`
+   }
+]
+
 export const FindArtifactsModal: React.FC<{ onClose: () => void }> = ({ onClose }) => {
 
    const [state, setState] = useState<{ searching?: boolean, artifacts?: GetArtifactResponseV2[] }>({});
@@ -134,12 +154,20 @@ export const FindArtifactsModal: React.FC<{ onClose: () => void }> = ({ onClose
          let type: string | undefined;
 
          if (typeRef.current?.selectedOptions?.length) {
-            type = typeRef.current.selectedOptions[0].key as string;
-            if (type === "step-all") {
+
+            type = (typeRef.current.selectedOptions[0].key as string)?.trim();
+            const text = (typeRef.current.selectedOptions[0].text as string)?.trim();
+            if (type === "step-all" || !text) {
                type = undefined;
+            } else {
+               const existing = artifactTypes.find(t => t.key === type);
+               if (!existing) {
+                  artifactTypes.push({ key: text, text: type });
+               }
             }
          }
 
+
          let streamId = streamRef?.current?.streamId?.trim();
          if (!streamId) {
             streamId = undefined;
@@ -170,25 +198,7 @@ export const FindArtifactsModal: React.FC<{ onClose: () => void }> = ({ onClose
       </Modal>
    }
 
-   const artifactTypes: IComboBoxOption[] = [
-      {
-         key: `step-all`,
-         text: `All`
-      }, {
-         key: `step-saved`,
-         text: `step-saved`
-      }, {
-         key: `step-output`,
-         text: `step-output`
-      }, {
-         key: `step-trace`,
-         text: `step-trace`
-      },
-      {
-         key: `step-testdata`,
-         text: `step-testdata`
-      }
-   ]
+   const filterTypes = [...artifactTypes].filter(t => !!t.text?.trim());
 
    return <Stack>
       <Modal isOpen={true} isBlocking={true} topOffsetFixed={true} styles={{ main: { padding: 8, width: 1024, height: 820, hasBeenOpened: false, top: "80px", position: "absolute" } }} onDismiss={() => onClose()} className={hordeClasses.modal}>
@@ -214,17 +224,17 @@ export const FindArtifactsModal: React.FC<{ onClose: () => void }> = ({ onClose
                         <StreamChooser ref={streamRef} />
                      </Stack>
                      <Stack >
-                        <TextField componentRef={minChangeRef} style={{ width: 100 }} label="Min Changelist" />
+                        <TextField componentRef={minChangeRef} style={{ width: 92 }} label="Min Changelist" />
                      </Stack>
                      <Stack >
-                        <TextField componentRef={maxChangeRef} style={{ width: 100 }} label="Max Changelist" />
+                        <TextField componentRef={maxChangeRef} style={{ width: 92 }} label="Max Changelist" />
                      </Stack>
                      <Stack >
-                        <TextField componentRef={nameRef} style={{ width: 224 }} label="Name" spellCheck={false} autoComplete="off" />
+                        <TextField componentRef={nameRef} style={{ width: 220 }} label="Name" spellCheck={false} autoComplete="off" />
                      </Stack>
                      <Stack>
                         <Label>Artifact Type</Label>
-                        <ComboBox componentRef={typeRef} style={{ width: 120, textAlign: "left" }} defaultSelectedKey="step-all" options={artifactTypes} calloutProps={{ doNotLayer: true }} />
+                        <ComboBox componentRef={typeRef} allowFreeform={true} autoComplete="off" spellCheck={false} style={{ width: 144, textAlign: "left" }} defaultSelectedKey="step-all" options={filterTypes} calloutProps={{ doNotLayer: true }} />
                      </Stack>
                   </Stack>
                   <Stack horizontal>
