Horde: Tweaks to artifact info
#rnx
#skipci

diff --git a/Engine/Source/Programs/Horde/HordeDashboard/src/components/artifacts/ArtifactsModal.tsx b/Engine/Source/Programs/Horde/HordeDashboard/src/components/artifacts/ArtifactsModal.tsx
--- a/Engine/Source/Programs/Horde/HordeDashboard/src/components/artifacts/ArtifactsModal.tsx
+++ b/Engine/Source/Programs/Horde/HordeDashboard/src/components/artifacts/ArtifactsModal.tsx
@@ -430,7 +430,7 @@ function formatBytes(bytes: number, decimals = 2) {
 }
 
 
-const DownloadButton: React.FC<{ handler: ArtifactsHandler }> = observer(({ handler }) => {
+const DownloadButton: React.FC<{ handler: ArtifactsHandler, openArtifactInfo: () => void }> = observer(({ handler, openArtifactInfo }) => {
 
    const [selectKey, setSelectionKey] = useState(0);
    const navigate = useNavigate();
@@ -476,11 +476,19 @@ const DownloadButton: React.FC<{ handler: ArtifactsHandler }> = observer(({ hand
 
    const downloadProps: IContextualMenuProps = {
       items: [
+         {
+            key: 'view_artifact_info',
+            text: 'View Artifact Info',
+            disabled: !handler.artifact,
+            onClick: () => {
+               openArtifactInfo();
+            }
+         },
          {
             key: 'navigate_to_job',
             text: 'Navigate to Job',
             onClick: () => {
-               navigate(`/job/${handler.jobId}`)
+               navigate(jobUrl);
             }
          }
       ],
@@ -488,7 +496,7 @@ const DownloadButton: React.FC<{ handler: ArtifactsHandler }> = observer(({ hand
    };
 
 
-   return <Stack horizontal verticalAlign="center" tokens={{ childrenGap: 8 }}>
+   return <Stack id="callout_target_artifactinfo" horizontal verticalAlign="center" tokens={{ childrenGap: 8 }}>
       <PrimaryButton split menuProps={downloadProps} styles={{ root: { fontFamily: 'Horde Open Sans SemiBold !important' } }} disabled={!selection.filesSelected && !selection.directoriesSelected} onClick={async () => {
 
          const selection = handler.currentSelection.items;
@@ -549,7 +557,7 @@ const JobDetailArtifactsInner: React.FC<{ jobId: string; stepId: string, artifac
    // eslint-disable-next-line
    const handler = ArtifactsHandler.current ?? new ArtifactsHandler(jobId, stepId, contextType, artifactPath, artifacts, artifactId);
 
-   const [infoKey, setInfoKey] = useState("");
+   const [viewArtifactInfo, setViewArtifactInfo] = useState(false);
 
    useEffect(() => {
       return () => {
@@ -659,42 +667,10 @@ const JobDetailArtifactsInner: React.FC<{ jobId: string; stepId: string, artifac
          const href = getFileHRef(item);
 
          return <Stack>
-            <Stack id={`callout_target_${item.key}`} data-selection-disabled verticalAlign="center" verticalFill horizontal horizontalAlign="end" style={{ paddingTop: 0, paddingBottom: 0 }}>
+            <Stack data-selection-disabled verticalAlign="center" verticalFill horizontal horizontalAlign="end" style={{ paddingTop: 0, paddingBottom: 0 }}>
                {isFile && <IconButton id="artifactview" href={`${href}&inline=true`} target="_blank" style={{ paddingTop: 1, color: "#106EBE" }} iconProps={{ iconName: "Eye", styles: { root: { fontSize: "14px" } } }} />}
                {isFile && <IconButton id="artifactview" href={href} target="_blank" style={{ paddingTop: 1, color: "#106EBE" }} iconProps={{ iconName: "CloudDownload", styles: { root: { fontSize: "14px" } } }} />}
-               <IconButton id="artifactview" onClick={() => setInfoKey(item.key)} style={{ paddingTop: 1, color: "#106EBE", height: "18px" }} iconProps={{ iconName: "Info", styles: { root: { fontSize: "14px" } } }} />
             </Stack>
-            {infoKey === item.key && <Callout
-               styles={{ root: { padding: "32px 24px", maxWidth: 1300 } }}
-               role="dialog"
-               gapSpace={0}
-               target={`#callout_target_${item.key}`}
-               isBeakVisible={true}
-               beakWidth={12}
-               onDismiss={() => {
-                  setInfoKey("")
-               }}
-               directionalHint={DirectionalHint.leftCenter}
-               setInitialFocus>
-               <Stack style={{ maxWidth: 1140 }}>
-                  <Stack style={{ paddingBottom: 12 }}>
-                     <Stack horizontal verticalAlign="center" verticalFill={true}>
-                        <Stack>
-                           <Text style={{ fontSize: 14, fontFamily: "Horde Open Sans SemiBold" }}>{`Artifact ID ${handler.artifact?.id}`}</Text>
-                        </Stack>
-                        <Stack grow />
-                        <Stack>
-                           <IconButton
-                              iconProps={{ iconName: 'Cancel', styles: { root: { fontSize: "14px" } } }}
-                              onClick={() => { setInfoKey("") }}
-                           />
-                        </Stack>
-                     </Stack>
-                  </Stack>
-                  <Stack style={{ paddingLeft: 0 }}>
-                     <Text style={{ fontSize: 11, whiteSpace: "pre-wrap", fontFamily: "Horde Cousine Regular" }}>{ (handler.artifact ? JSON.stringify(handler.artifact, undefined, 2) + ",\n" : "") + JSON.stringify(item.dirResponse ?? item.fileResponse ?? {}, undefined, 2).replaceAll("\\r", "").replaceAll("\\n", "\n")}</Text></Stack>
-               </Stack>
-            </Callout>}
          </Stack>
       }
 
@@ -748,51 +724,84 @@ const JobDetailArtifactsInner: React.FC<{ jobId: string; stepId: string, artifac
 
    }
 
-   return <Stack key={`jobdetailartifacts_${idcounter++}`} tokens={{ childrenGap: 12 }}>
-      <Stack horizontal verticalAlign="center" tokens={{ childrenGap: 18 }} style={{ paddingBottom: 12 }}>
-         <Stack>
-            <BrowseBreadCrumbs handler={handler} />
+   return <Stack key={`jobdetailartifacts_${idcounter++}`}>
+      {viewArtifactInfo && <Callout
+         styles={{ root: { padding: "32px 24px", maxWidth: 1300 } }}
+         role="dialog"
+         gapSpace={12}
+         target={`#callout_target_artifactinfo`}
+         isBeakVisible={true}
+         beakWidth={12}
+         onDismiss={() => {
+            setViewArtifactInfo(false);
+         }}
+         directionalHint={DirectionalHint.leftTopEdge}
+         setInitialFocus>
+         <Stack style={{ maxWidth: 1140 }}>
+            <Stack style={{ paddingBottom: 12 }}>
+               <Stack horizontal verticalAlign="center" verticalFill={true}>
+                  <Stack>
+                     <Text style={{ fontSize: 14, fontFamily: "Horde Open Sans SemiBold" }}>{`Artifact ID ${handler.artifact?.id}`}</Text>
+                  </Stack>
+                  <Stack grow />
+                  <Stack>
+                     <IconButton
+                        iconProps={{ iconName: 'Cancel', styles: { root: { fontSize: "14px" } } }}
+                        onClick={() => { setViewArtifactInfo(false); }}
+                     />
+                  </Stack>
+               </Stack>
+            </Stack>
+            <Stack style={{ paddingLeft: 0 }}>
+               <Text style={{ fontSize: 11, whiteSpace: "pre-wrap", fontFamily: "Horde Cousine Regular" }}>{handler.artifact ? JSON.stringify(handler.artifact, undefined, 2).replaceAll("\\r", "").replaceAll("\\n", "\n") : ""}</Text></Stack>
          </Stack>
-         <Stack grow />
-         <DownloadButton handler={handler} />
+      </Callout>}
+
+      <Stack tokens={{childrenGap: 12}}>
+         <Stack horizontal verticalAlign="center" tokens={{ childrenGap: 18 }} style={{ paddingBottom: 12 }}>
+            <Stack>
+               <BrowseBreadCrumbs handler={handler} />
+            </Stack>
+            <Stack grow />
+            <DownloadButton handler={handler} openArtifactInfo={() => setViewArtifactInfo(true)} />
+         </Stack >
+         {handler.loading && <Stack>
+            <Spinner styles={{ root: { opacity: 0, animation: "hordeFadeIn 1s ease-in-out 2s forwards" } }} size={SpinnerSize.large} />
+         </Stack>}
+         {!handler.loading && <Stack style={{ height: 492 + 160, position: "relative" }}>
+            <ScrollablePane style={{ height: 492 + 160 }}>
+               <SelectionZone selection={handler.selection!}>
+                  <DetailsList
+                     styles={{ root: { overflowX: "hidden" } }}
+                     className={styles.list}
+                     isHeaderVisible={false}
+                     compact={true}
+                     items={items}
+                     columns={columns}
+                     layoutMode={DetailsListLayoutMode.fixedColumns}
+                     selectionMode={SelectionMode.multiple}
+                     enableUpdateAnimations={false}
+                     selection={handler.selection}
+                     selectionPreservedOnEmptyClick={true}
+                     onShouldVirtualize={() => false}
+                     onItemInvoked={(item: BrowserItem) => {
+                        if (item?.type !== BrowserType.File) {
+                           return;
+                        }
+                        const href = getFileHRef(item);
+                        if (href) {
+                           window.open(href + "&inline=true", "_blank");
+                        }
+                     }}
+                     onRenderItemColumn={renderItem}
+                  />
+               </SelectionZone>
+
+            </ScrollablePane>
+         </Stack>}
       </Stack >
-      {handler.loading && <Stack>
-         <Spinner styles={{ root: { opacity: 0, animation: "hordeFadeIn 1s ease-in-out 2s forwards" } }} size={SpinnerSize.large} />
-      </Stack>}
-      {!handler.loading && <Stack style={{ height: 492 + 160, position: "relative" }}>
-         <ScrollablePane style={{ height: 492 + 160 }}>
-            <SelectionZone selection={handler.selection!}>
-               <DetailsList
-                  styles={{ root: { overflowX: "hidden" } }}
-                  className={styles.list}
-                  isHeaderVisible={false}
-                  compact={true}
-                  items={items}
-                  columns={columns}
-                  layoutMode={DetailsListLayoutMode.fixedColumns}
-                  selectionMode={SelectionMode.multiple}
-                  enableUpdateAnimations={false}
-                  selection={handler.selection}
-                  selectionPreservedOnEmptyClick={true}
-                  onShouldVirtualize={() => false}
-                  onItemInvoked={(item: BrowserItem) => {
-                     if (item?.type !== BrowserType.File) {
-                        return;
-                     }
-                     const href = getFileHRef(item);
-                     if (href) {
-                        window.open(href + "&inline=true", "_blank");
-                     }
-                  }}
-                  onRenderItemColumn={renderItem}
-               />
-            </SelectionZone>
-
-         </ScrollablePane>
-      </Stack>}
-   </Stack >
-
-   //
+   </Stack>
+
 })
 
 
