Horde: Improve Datadog enricher for OpenTelemetry traces

- Avoids using the full path as "resource.name" because of high cardinality
- Sets username and user ID as tags
- Sets dashboard's last user activity for tracking of inactive browser tabs

diff --git a/Engine/Source/Programs/Horde/HordeServer.Shared/Utilities/HordeClaimTypes.cs b/Engine/Source/Programs/Horde/HordeServer.Shared/Utilities/HordeClaimTypes.cs
--- a/Engine/Source/Programs/Horde/HordeServer.Shared/Utilities/HordeClaimTypes.cs
+++ b/Engine/Source/Programs/Horde/HordeServer.Shared/Utilities/HordeClaimTypes.cs
@@ -111,32 +111,28 @@ namespace HordeServer.Utilities
 		public static AccountId? GetAccountId(this ClaimsPrincipal principal)
 		{
 			string? idValue = principal.FindFirstValue(HordeClaimTypes.AccountId);
-			if (idValue == null)
-			{
-				return null;
-			}
-			else
-			{
-				return AccountId.Parse(idValue);
-			}
+			return idValue == null ? null : AccountId.Parse(idValue);
 		}
 
 		/// <summary>
-		/// Gets the Horde user id from a principal
+		/// Gets the Horde user ID from a principal
 		/// </summary>
 		/// <param name="principal"></param>
 		/// <returns></returns>
 		public static UserId? GetUserId(this ClaimsPrincipal principal)
 		{
 			string? idValue = principal.FindFirstValue(HordeClaimTypes.UserId);
-			if (idValue == null)
-			{
-				return null;
-			}
-			else
-			{
-				return UserId.Parse(idValue);
-			}
+			return idValue == null ? null : UserId.Parse(idValue);
+		}
+		
+		/// <summary>
+		/// Gets the Horde username from a principal
+		/// </summary>
+		/// <param name="principal"></param>
+		/// <returns></returns>
+		public static string? GetUser(this ClaimsPrincipal principal)
+		{
+			return principal.FindFirstValue(HordeClaimTypes.User);
 		}
 
 		/// <summary>
@@ -147,11 +143,7 @@ namespace HordeServer.Utilities
 		public static string? GetExternalIssueUser(this ClaimsPrincipal user)
 		{
 			Claim? claim = user.FindFirst(HordeClaimTypes.ExternalIssueUser);
-			if (claim == null)
-			{
-				return null;
-			}
-			return claim.Value;
+			return claim?.Value;
 		}
 
 		/// <summary>
diff --git a/Engine/Source/Programs/Horde/HordeServer/Utilities/OpenTelemetry.cs b/Engine/Source/Programs/Horde/HordeServer/Utilities/OpenTelemetry.cs
--- a/Engine/Source/Programs/Horde/HordeServer/Utilities/OpenTelemetry.cs
+++ b/Engine/Source/Programs/Horde/HordeServer/Utilities/OpenTelemetry.cs
@@ -14,6 +14,7 @@ using EpicGames.Horde.Streams;
 using Microsoft.AspNetCore.Http;
 using Microsoft.Extensions.DependencyInjection;
 using Microsoft.Extensions.Logging;
+using Microsoft.Extensions.Primitives;
 using MongoDB.Bson;
 using MongoDB.Driver;
 using OpenTelemetry.Exporter;
@@ -87,11 +88,24 @@ public static class OpenTelemetryHelper
 
 		void DatadogAspNetRequestEnricher(Activity activity, HttpRequest request)
 		{
+			activity.DisplayName = $"{request.Method} {request.Headers.Host}{request.Path}";
+			
+			// From https://docs.datadoghq.com/standard-attributes/
 			activity.SetTag("service.name", settings.ServiceName);
 			activity.SetTag("operation.name", "http.request");
-			string url = $"{request.Method} {request.Headers.Host}{request.Path}";
-			activity.DisplayName = url;
-			activity.SetTag("resource.name", url);
+			activity.SetTag("http.client_ip", request.HttpContext.Connection.RemoteIpAddress);
+			activity.SetTag("usr.id", request.HttpContext.User.GetUserId());
+			activity.SetTag("usr.name", request.HttpContext.User.GetUser());
+			
+			// Header sent by the dashboard to indicate how long a user has been inactive for a particular browser page (in seconds)
+			if (request.Headers.TryGetValue("X-Horde-LastUserActivity", out StringValues values))
+			{
+				string? value = values.FirstOrDefault();
+				if (!String.IsNullOrEmpty(value) && Int32.TryParse(value, out int lastUserActivity))
+				{
+					activity.SetTag("horde.lastUserActivity", lastUserActivity);
+				}
+			}
 		}
 
 		bool FilterHttpRequests(HttpContext context)
