Horde: Initial support for plugins.

* Each plugin implements the IPluginStartup interface, which allows it to register services into the standard DI container.
* The IPluginStartup interface also exposes a global config type, which is used to mount extension points under the "plugins" section of the global config object.
* The server automatically registers IOptions<>, IOptionsMonitor<> etc... interfaces for the plugins configuration object, and registers controllers in the plugin assembly.

diff --git a/Engine/Source/Programs/Horde/HordeServer.Shared/Plugins/IPluginConfig.cs b/Engine/Source/Programs/Horde/HordeServer.Shared/Plugins/IPluginConfig.cs
new file mode 100644
--- /dev/null
+++ b/Engine/Source/Programs/Horde/HordeServer.Shared/Plugins/IPluginConfig.cs
@@ -0,0 +1,25 @@
+// Copyright Epic Games, Inc. All Rights Reserved.
+
+using System;
+using System.Collections.Generic;
+using System.Linq;
+using System.Text;
+using System.Text.Json;
+using System.Text.Json.Serialization;
+using System.Threading.Tasks;
+using HordeServer.Acls;
+
+namespace HordeServer.Plugins
+{
+	/// <summary>
+	/// Interface for plugin extensions to the global config object
+	/// </summary>
+	public interface IPluginConfig
+	{
+		/// <summary>
+		/// Called to fixup a plugin's configuration after deserialization
+		/// </summary>
+		/// <param name="parentAcl">The parent ACL scope</param>
+		void PostLoad(AclConfig parentAcl);
+	}
+}
diff --git a/Engine/Source/Programs/Horde/HordeServer.Shared/Plugins/IPluginStartup.cs b/Engine/Source/Programs/Horde/HordeServer.Shared/Plugins/IPluginStartup.cs
new file mode 100644
--- /dev/null
+++ b/Engine/Source/Programs/Horde/HordeServer.Shared/Plugins/IPluginStartup.cs
@@ -0,0 +1,46 @@
+// Copyright Epic Games, Inc. All Rights Reserved.
+
+using System;
+using Microsoft.Extensions.DependencyInjection;
+
+namespace HordeServer.Plugins
+{
+	/// <summary>
+	/// Attribute used to identify types that should be constructed for the given plugin name
+	/// </summary>
+	[AttributeUsage(AttributeTargets.Class)]
+	public sealed class PluginStartupAttribute : Attribute
+	{
+		/// <summary>
+		/// Name of the plugin
+		/// </summary>
+		public string Name { get; }
+
+		/// <summary>
+		/// Constructor
+		/// </summary>
+		public PluginStartupAttribute(string name)
+			=> Name = name;
+	}
+
+	/// <summary>
+	/// Interface for modules that can extend the services known by Horde
+	/// </summary>
+	public interface IPluginStartup
+	{
+		/// <summary>
+		/// Type containing global settings for this plugin. Must implement the <see cref="IPluginConfig"/> interface.
+		/// 
+		/// Standard <see cref="Microsoft.Extensions.Options.IOptions{TOptions}"/> and 
+		/// <see cref="Microsoft.Extensions.Options.IOptionsMonitor{TOptions}"/> objects will be registered for this type 
+		/// in the service container.
+		/// </summary>
+		Type? GlobalConfigType { get; }
+
+		/// <summary>
+		/// Configure the services provided by the plugin
+		/// </summary>
+		/// <param name="serviceCollection">Collection of services to add to</param>
+		void ConfigureServices(IServiceCollection serviceCollection);
+	}
+}
diff --git a/Engine/Source/Programs/Horde/HordeServer/Configuration/ConfigService.cs b/Engine/Source/Programs/Horde/HordeServer/Configuration/ConfigService.cs
--- a/Engine/Source/Programs/Horde/HordeServer/Configuration/ConfigService.cs
+++ b/Engine/Source/Programs/Horde/HordeServer/Configuration/ConfigService.cs
@@ -12,6 +12,7 @@ using System.Threading.Tasks;
 using EpicGames.Core;
 using EpicGames.Horde.Users;
 using EpicGames.Redis;
+using HordeServer.Plugins;
 using HordeServer.Server;
 using HordeServer.Users;
 using HordeServer.Utilities;
@@ -136,7 +137,7 @@ namespace HordeServer.Configuration
 		/// <summary>
 		/// Constructor
 		/// </summary>
-		public ConfigService(IRedisService redisService, IOptions<ServerSettings> serverSettings, IEnumerable<IConfigSource> sources, IClock clock, IHealthMonitor<ConfigService> health, ILogger<ConfigService> logger)
+		public ConfigService(IRedisService redisService, IOptions<ServerSettings> serverSettings, IEnumerable<IConfigSource> sources, IPluginCollection pluginCollection, IClock clock, IHealthMonitor<ConfigService> health, ILogger<ConfigService> logger)
 		{
 			_redisService = redisService;
 			_serverSettings = serverSettings.Value;
@@ -148,6 +149,16 @@ namespace HordeServer.Configuration
 			_jsonOptions = new JsonSerializerOptions();
 			JsonUtils.ConfigureJsonSerializer(_jsonOptions);
 
+			PluginConfigCollectionConverter pluginConfigConverter = new PluginConfigCollectionConverter();
+			foreach ((string name, IPluginStartup plugin) in pluginCollection.EnabledPlugins)
+			{
+				if (plugin.GlobalConfigType != null)
+				{
+					pluginConfigConverter.NameToType[name] = plugin.GlobalConfigType;
+				}
+			}
+			_jsonOptions.Converters.Add(pluginConfigConverter);
+
 			_ticker = clock.AddSharedTicker<ConfigService>(_tickInterval, TickSharedAsync, logger);
 
 			_updateTask = new BackgroundTask(WaitForUpdatesAsync);
diff --git a/Engine/Source/Programs/Horde/HordeServer/Plugins/IPluginCollection.cs b/Engine/Source/Programs/Horde/HordeServer/Plugins/IPluginCollection.cs
new file mode 100644
--- /dev/null
+++ b/Engine/Source/Programs/Horde/HordeServer/Plugins/IPluginCollection.cs
@@ -0,0 +1,17 @@
+// Copyright Epic Games, Inc. All Rights Reserved.
+
+using System.Collections.Generic;
+
+namespace HordeServer.Plugins
+{
+	/// <summary>
+	/// Interface for querying the state of plugins on the server
+	/// </summary>
+	public interface IPluginCollection
+	{
+		/// <summary>
+		/// The enabled plugins
+		/// </summary>
+		IReadOnlyDictionary<string, IPluginStartup> EnabledPlugins { get; }
+	}
+}
diff --git a/Engine/Source/Programs/Horde/HordeServer/Plugins/PluginCollection.cs b/Engine/Source/Programs/Horde/HordeServer/Plugins/PluginCollection.cs
new file mode 100644
--- /dev/null
+++ b/Engine/Source/Programs/Horde/HordeServer/Plugins/PluginCollection.cs
@@ -0,0 +1,21 @@
+// Copyright Epic Games, Inc. All Rights Reserved.
+
+using System.Collections.Generic;
+
+namespace HordeServer.Plugins
+{
+	/// <summary>
+	/// Concrete implementation of <see cref="IPluginCollection"/>
+	/// </summary>
+	class PluginCollection : IPluginCollection
+	{
+		/// <inheritdoc/>
+		public IReadOnlyDictionary<string, IPluginStartup> EnabledPlugins { get; }
+
+		/// <summary>
+		/// Constructor
+		/// </summary>
+		public PluginCollection(IReadOnlyDictionary<string, IPluginStartup> enabled)
+			=> EnabledPlugins = enabled;
+	}
+}
diff --git a/Engine/Source/Programs/Horde/HordeServer/Plugins/PluginConfigCollection.cs b/Engine/Source/Programs/Horde/HordeServer/Plugins/PluginConfigCollection.cs
new file mode 100644
--- /dev/null
+++ b/Engine/Source/Programs/Horde/HordeServer/Plugins/PluginConfigCollection.cs
@@ -0,0 +1,102 @@
+﻿// Copyright Epic Games, Inc. All Rights Reserved.
+
+using System;
+using System.Collections.Generic;
+using System.Diagnostics.CodeAnalysis;
+using System.Text.Json;
+using System.Text.Json.Serialization;
+
+namespace HordeServer.Plugins
+{
+	/// <summary>
+	/// Collection of <see cref="IPluginConfig"/> objects.
+	/// </summary>
+	[JsonConverter(typeof(PluginConfigCollectionConverter))]
+	public class PluginConfigCollection : Dictionary<string, IPluginConfig>
+	{
+		/// <summary>
+		/// Constructor
+		/// </summary>
+		public PluginConfigCollection()
+			: base(StringComparer.OrdinalIgnoreCase)
+		{ }
+
+		/// <summary>
+		/// Attempts to get a plugin config of a specific type
+		/// </summary>
+		public bool TryGetValue<T>(string name, [NotNullWhen(true)] out T? value) where T : class, IPluginConfig
+		{
+			IPluginConfig? config;
+			if (base.TryGetValue(name, out config) && config is T typedConfig)
+			{
+				value = typedConfig;
+				return true;
+			}
+			else
+			{
+				value = null;
+				return false;
+			}
+		}
+	}
+
+	/// <summary>
+	/// Serializer to and from strongly typed plugin config classes
+	/// </summary>
+	class PluginConfigCollectionConverter : JsonConverter<PluginConfigCollection>
+	{
+		/// <summary>
+		/// Map of plugin names to the expected config type
+		/// </summary>
+		public Dictionary<string, Type> NameToType { get; } = new Dictionary<string, Type>(StringComparer.OrdinalIgnoreCase);
+
+		/// <inheritdoc/>
+		public override PluginConfigCollection? Read(ref Utf8JsonReader reader, Type typeToConvert, JsonSerializerOptions options)
+		{
+			if (reader.TokenType != JsonTokenType.StartObject)
+			{
+				throw new JsonException("Expected an object for plugin config type");
+			}
+
+			PluginConfigCollection result = new PluginConfigCollection();
+			while (reader.Read() && reader.TokenType == JsonTokenType.PropertyName)
+			{
+				string? name = reader.GetString();
+				if (name == null || !reader.Read())
+				{
+					throw new JsonException("Invalid plugin name");
+				}
+				if (NameToType.TryGetValue(name, out Type? type))
+				{
+					object? obj = JsonSerializer.Deserialize(ref reader, type, options);
+					if (obj != null)
+					{
+						result.Add(name, (IPluginConfig)obj);
+					}
+				}
+				else
+				{
+					reader.Skip();
+				}
+			}
+
+			if (reader.TokenType != JsonTokenType.EndObject)
+			{
+				throw new JsonException("Invalid token while parsing plugin config");
+			}
+			return result;
+		}
+
+		/// <inheritdoc/>
+		public override void Write(Utf8JsonWriter writer, PluginConfigCollection value, JsonSerializerOptions options)
+		{
+			writer.WriteStartObject();
+			foreach ((string name, IPluginConfig config) in value)
+			{
+				writer.WritePropertyName(name);
+				JsonSerializer.Serialize(writer, config, config.GetType(), options);
+			}
+			writer.WriteEndObject();
+		}
+	}
+}
diff --git a/Engine/Source/Programs/Horde/HordeServer/Plugins/PluginOptionsFactory.cs b/Engine/Source/Programs/Horde/HordeServer/Plugins/PluginOptionsFactory.cs
new file mode 100644
--- /dev/null
+++ b/Engine/Source/Programs/Horde/HordeServer/Plugins/PluginOptionsFactory.cs
@@ -0,0 +1,89 @@
+// Copyright Epic Games, Inc. All Rights Reserved.
+
+using System;
+using HordeServer.Server;
+using Microsoft.Extensions.DependencyInjection;
+using Microsoft.Extensions.Options;
+using Microsoft.Extensions.Primitives;
+
+namespace HordeServer.Plugins
+{
+	/// <summary>
+	/// Factory for accessing plugin config instances
+	/// </summary>
+	/// <typeparam name="T">Type of the options object to return</typeparam>
+	sealed class PluginOptionsFactory<T> : IOptionsFactory<T>, IOptionsChangeTokenSource<T>
+		where T : class, IPluginConfig, new()
+	{
+		readonly string _pluginName;
+		readonly IOptionsMonitor<GlobalConfig> _globalConfig;
+		readonly IOptionsChangeTokenSource<GlobalConfig> _globalConfigChangeTokenSource;
+
+		/// <inheritdoc/>
+		string IOptionsChangeTokenSource<T>.Name => String.Empty;
+
+		/// <summary>
+		/// Constructor
+		/// </summary>
+		public PluginOptionsFactory(string pluginName, IServiceProvider serviceProvider)
+			: this(pluginName, serviceProvider.GetRequiredService<IOptionsMonitor<GlobalConfig>>(), serviceProvider.GetRequiredService<IOptionsChangeTokenSource<GlobalConfig>>())
+		{
+		}
+
+		/// <summary>
+		/// Constructor
+		/// </summary>
+		public PluginOptionsFactory(string pluginName, IOptionsMonitor<GlobalConfig> globalConfig, IOptionsChangeTokenSource<GlobalConfig> globalConfigChangeTokenSource)
+		{
+			_pluginName = pluginName;
+			_globalConfig = globalConfig;
+			_globalConfigChangeTokenSource = globalConfigChangeTokenSource;
+		}
+
+		/// <inheritdoc/>
+		public T Create(string name)
+		{
+			T? config;
+			if (!_globalConfig.CurrentValue.Plugins.TryGetValue<T>(_pluginName, out config))
+			{
+				config = new T();
+			}
+			return config;
+		}
+
+		/// <inheritdoc/>
+		public IChangeToken GetChangeToken()
+			=> _globalConfigChangeTokenSource.GetChangeToken();
+	}
+
+	/// <summary>
+	/// Helper methods for registering 
+	/// </summary>
+	public static class PluginOptionsFactoryExtensions
+	{
+		/// <summary>
+		/// Register an options factory for the given plugin options type
+		/// </summary>
+		public static void AddPluginConfig(this IServiceCollection services, string name, Type configType)
+		{
+			Type helperType = typeof(RegistrationHelper<>).MakeGenericType(configType);
+			RegistrationHelper registration = (RegistrationHelper)Activator.CreateInstance(helperType)!;
+			registration.Register(services, name);
+		}
+
+		abstract class RegistrationHelper
+		{
+			public abstract void Register(IServiceCollection services, string name);
+		}
+
+		class RegistrationHelper<T> : RegistrationHelper where T : class, IPluginConfig, new()
+		{
+			public override void Register(IServiceCollection services, string name)
+			{
+				services.AddSingleton<PluginOptionsFactory<T>>(sp => new PluginOptionsFactory<T>(name, sp));
+				services.AddSingleton<IOptionsFactory<T>>(sp => sp.GetRequiredService<PluginOptionsFactory<T>>());
+				services.AddSingleton<IOptionsChangeTokenSource<T>>(sp => sp.GetRequiredService<PluginOptionsFactory<T>>());
+			}
+		}
+	}
+}
diff --git a/Engine/Source/Programs/Horde/HordeServer/Server/GlobalConfig.cs b/Engine/Source/Programs/Horde/HordeServer/Server/GlobalConfig.cs
--- a/Engine/Source/Programs/Horde/HordeServer/Server/GlobalConfig.cs
+++ b/Engine/Source/Programs/Horde/HordeServer/Server/GlobalConfig.cs
@@ -38,6 +38,7 @@ using HordeServer.Jobs.Bisect;
 using HordeServer.Logs;
 using HordeServer.Notifications;
 using HordeServer.Perforce;
+using HordeServer.Plugins;
 using HordeServer.Projects;
 using HordeServer.Secrets;
 using HordeServer.Storage;
@@ -245,6 +246,11 @@ namespace HordeServer.Server
 		/// </summary>
 		public List<TelemetryStoreConfig> TelemetryStores { get; set; } = new List<TelemetryStoreConfig>();
 
+		/// <summary>
+		/// Plugin config objects
+		/// </summary>
+		public PluginConfigCollection Plugins { get; set; } = new PluginConfigCollection();
+
 		/// <summary>
 		/// General parameters for other tools. Can be queried through the api/v1/parameters endpoint.
 		/// </summary>
@@ -355,6 +361,11 @@ namespace HordeServer.Server
 
 			Storage.PostLoad(this);
 
+			foreach (IPluginConfig pluginConfig in Plugins.Values)
+			{
+				pluginConfig.PostLoad(Acl);
+			}
+
 			_aclLookup.Clear();
 			BuildAclScopeLookup(Acl, _aclLookup);
 
diff --git a/Engine/Source/Programs/Horde/HordeServer/Startup.cs b/Engine/Source/Programs/Horde/HordeServer/Startup.cs
--- a/Engine/Source/Programs/Horde/HordeServer/Startup.cs
+++ b/Engine/Source/Programs/Horde/HordeServer/Startup.cs
@@ -71,6 +71,7 @@ using HordeServer.Logs.Storage;
 using HordeServer.Notifications;
 using HordeServer.Notifications.Sinks;
 using HordeServer.Perforce;
+using HordeServer.Plugins;
 using HordeServer.Replicators;
 using HordeServer.Secrets;
 using HordeServer.Secrets.Providers;
@@ -366,6 +367,19 @@ namespace HordeServer
 		// This method gets called *multiple times* by the runtime. Use this method to add services to the container.
 		public void ConfigureServices(IServiceCollection services)
 		{
+			// Register the plugin collection
+			Dictionary<string, IPluginStartup> plugins = new Dictionary<string, IPluginStartup>();
+			services.AddSingleton<IPluginCollection>(new PluginCollection(plugins));
+
+			// Register all the plugin config types
+			foreach ((string name, IPluginStartup plugin) in plugins)
+			{
+				if (plugin.GlobalConfigType != null)
+				{
+					services.AddPluginConfig(name, plugin.GlobalConfigType);
+				}
+			}
+
 			// IOptionsMonitor pattern for live updating of configuration settings
 			services.Configure<ServerSettings>(x => BindServerSettings(Configuration, x));
 
@@ -895,7 +909,14 @@ namespace HordeServer
 				options.KnownNetworks.Clear();
 			});
 
-			services.AddMvc().AddJsonOptions(options => JsonUtils.ConfigureJsonSerializer(options.JsonSerializerOptions));
+			IMvcBuilder mvcBuilder = services.AddMvc()
+				.AddJsonOptions(options => JsonUtils.ConfigureJsonSerializer(options.JsonSerializerOptions));
+
+			foreach (Assembly pluginAssembly in plugins.Values.Select(x => x.GetType().Assembly).Distinct())
+			{
+				mvcBuilder.AddApplicationPart(pluginAssembly);
+			}
+
 			services.AddControllersWithViews(options => options.Filters.Add(new ObsoleteLoggingFilter()))
 				.AddRazorRuntimeCompilation();
 
