Horde: Missing file.

#rnx

diff --git a/Engine/Source/Programs/Horde/HordeServer.Tests/TestSetup.cs b/Engine/Source/Programs/Horde/HordeServer.Tests/TestSetup.cs
--- a/Engine/Source/Programs/Horde/HordeServer.Tests/TestSetup.cs
+++ b/Engine/Source/Programs/Horde/HordeServer.Tests/TestSetup.cs
@@ -73,6 +73,7 @@ using Microsoft.Extensions.Options;
 using Microsoft.VisualStudio.TestTools.UnitTesting;
 using Moq;
 using OpenTelemetry.Trace;
+using HordeServer.Plugins;
 
 namespace HordeServer.Tests
 {
@@ -196,6 +197,8 @@ namespace HordeServer.Tests
 
 			services.AddHttpClient<RpcService>();
 
+			services.AddSingleton<IPluginCollection>(new PluginCollection(new Dictionary<string, IPluginStartup>()));
+
 			services.AddLogging(builder => { builder.AddConsole().SetMinimumLevel(LogLevel.Debug); });
 			services.AddSingleton<IMemoryCache>(sp => new MemoryCache(new MemoryCacheOptions { }));
 
