Horde: Expose global server settings through an IServerInfo interface.

#rnx

diff --git a/Engine/Source/Programs/Horde/HordeServer.Shared/IServerInfo.cs b/Engine/Source/Programs/Horde/HordeServer.Shared/IServerInfo.cs
new file mode 100644
--- /dev/null
+++ b/Engine/Source/Programs/Horde/HordeServer.Shared/IServerInfo.cs
@@ -0,0 +1,38 @@
+// Copyright Epic Games, Inc. All Rights Reserved.
+
+using EpicGames.Core;
+using Microsoft.Extensions.Configuration;
+
+namespace HordeServer
+{
+	/// <summary>
+	/// Provides access to server deployment information
+	/// </summary>
+	public interface IServerInfo
+	{
+		/// <summary>
+		/// Current version of the server
+		/// </summary>
+		SemVer Version { get; }
+
+		/// <summary>
+		/// Environment that the server is deployed to
+		/// </summary>
+		string Environment { get; }
+
+		/// <summary>
+		/// Unique session id
+		/// </summary>
+		string SessionId { get; }
+
+		/// <summary>
+		/// Directory containing the server executable
+		/// </summary>
+		DirectoryReference AppDir { get; }
+
+		/// <summary>
+		/// Directory to store server data files
+		/// </summary>
+		DirectoryReference DataDir { get; }
+	}
+}
diff --git a/Engine/Source/Programs/Horde/HordeServer.Tests/TestSetup.cs b/Engine/Source/Programs/Horde/HordeServer.Tests/TestSetup.cs
--- a/Engine/Source/Programs/Horde/HordeServer.Tests/TestSetup.cs
+++ b/Engine/Source/Programs/Horde/HordeServer.Tests/TestSetup.cs
@@ -197,6 +197,7 @@ namespace HordeServer.Tests
 
 			services.AddHttpClient<RpcService>();
 
+			services.AddSingleton<IServerInfo>(new ServerInfo());
 			services.AddSingleton<IPluginCollection>(new PluginCollection(new Dictionary<string, IPluginStartup>()));
 
 			services.AddLogging(builder => { builder.AddConsole().SetMinimumLevel(LogLevel.Debug); });
diff --git a/Engine/Source/Programs/Horde/HordeServer/Server/GlobalConfig.cs b/Engine/Source/Programs/Horde/HordeServer/Server/GlobalConfig.cs
--- a/Engine/Source/Programs/Horde/HordeServer/Server/GlobalConfig.cs
+++ b/Engine/Source/Programs/Horde/HordeServer/Server/GlobalConfig.cs
@@ -356,7 +356,7 @@ namespace HordeServer.Server
 			foreach (TelemetryStoreConfig telemetryStore in TelemetryStores)
 			{
 				_telemetryStoreLookup.Add(telemetryStore.Id, telemetryStore);
-				telemetryStore.PostLoad(this);
+				telemetryStore.PostLoad(Acl);
 			}
 
 			Storage.PostLoad(this);
diff --git a/Engine/Source/Programs/Horde/HordeServer/Startup.cs b/Engine/Source/Programs/Horde/HordeServer/Startup.cs
--- a/Engine/Source/Programs/Horde/HordeServer/Startup.cs
+++ b/Engine/Source/Programs/Horde/HordeServer/Startup.cs
@@ -367,6 +367,8 @@ namespace HordeServer
 		// This method gets called *multiple times* by the runtime. Use this method to add services to the container.
 		public void ConfigureServices(IServiceCollection services)
 		{
+			services.AddSingleton<IServerInfo>(new ServerInfo());
+
 			// Register the plugin collection
 			Dictionary<string, IPluginStartup> plugins = new Dictionary<string, IPluginStartup>();
 			services.AddSingleton<IPluginCollection>(new PluginCollection(plugins));
diff --git a/Engine/Source/Programs/Horde/HordeServer/Telemetry/Metrics/MetricConfig.cs b/Engine/Source/Programs/Horde/HordeServer/Telemetry/Metrics/MetricConfig.cs
--- a/Engine/Source/Programs/Horde/HordeServer/Telemetry/Metrics/MetricConfig.cs
+++ b/Engine/Source/Programs/Horde/HordeServer/Telemetry/Metrics/MetricConfig.cs
@@ -10,7 +10,6 @@ using EpicGames.Horde.Telemetry;
 using EpicGames.Horde.Telemetry.Metrics;
 using HordeServer.Acls;
 using HordeServer.Configuration;
-using HordeServer.Server;
 using HordeServer.Utilities;
 using Json.Path;
 
@@ -161,9 +160,9 @@ namespace HordeServer.Telemetry.Metrics
 		/// <summary>
 		/// Called after the store has been deserialized to compute cached values
 		/// </summary>
-		public void PostLoad(GlobalConfig globalConfig)
+		public void PostLoad(AclConfig parentAcl)
 		{
-			Acl.PostLoad(globalConfig.Acl, $"telemetry-store:{Id}");
+			Acl.PostLoad(parentAcl, $"telemetry-store:{Id}");
 
 			_metricLookup.Clear();
 			foreach (MetricConfig metric in Metrics)
diff --git a/Engine/Source/Programs/Horde/HordeServer/Telemetry/TelemetryManager.cs b/Engine/Source/Programs/Horde/HordeServer/Telemetry/TelemetryManager.cs
--- a/Engine/Source/Programs/Horde/HordeServer/Telemetry/TelemetryManager.cs
+++ b/Engine/Source/Programs/Horde/HordeServer/Telemetry/TelemetryManager.cs
@@ -29,23 +29,17 @@ namespace HordeServer.Telemetry
 		private readonly ITicker _ticker;
 		private readonly Tracer _tracer;
 		private readonly ILogger<TelemetryManager> _logger;
-
-		private static readonly TelemetryRecordMeta s_serverEventMetadata = new TelemetryRecordMeta
-		{
-			AppId = "Horde",
-			AppVersion = ServerApp.Version.ToString(),
-			AppEnvironment = ServerApp.DeploymentEnvironment,
-			SessionId = ServerApp.SessionId
-		};
+		private readonly TelemetryRecordMeta _serverEventMetadata;
 
 		/// <summary>
 		/// Constructor
 		/// </summary>
-		public TelemetryManager(IServiceProvider serviceProvider, IClock clock, IOptions<ServerSettings> serverSettings, Tracer tracer, ILoggerFactory loggerFactory)
+		public TelemetryManager(IServiceProvider serviceProvider, IServerInfo serverInfo, IClock clock, IOptions<ServerSettings> serverSettings, Tracer tracer, ILoggerFactory loggerFactory)
 		{
 			_tracer = tracer;
 			_logger = loggerFactory.CreateLogger<TelemetryManager>();
 			_ticker = clock.AddTicker<EpicTelemetrySink>(TimeSpan.FromSeconds(30.0), FlushAsync, _logger);
+			_serverEventMetadata = new TelemetryRecordMeta("Horde", serverInfo.Version.ToString(), serverInfo.Environment, serverInfo.SessionId);
 
 			foreach (BaseTelemetryConfig config in serverSettings.Value.Telemetry)
 			{
@@ -66,7 +60,7 @@ namespace HordeServer.Telemetry
 		/// <inheritdoc/>
 		public void WriteEvent(TelemetryStoreId telemetryStoreId, object payload)
 		{
-			WriteEvent(telemetryStoreId, s_serverEventMetadata, payload);
+			WriteEvent(telemetryStoreId, _serverEventMetadata, payload);
 		}
 
 		/// <inheritdoc/>
