Horde: Fix user ID/username not getting set in Datadog enricher for OpenTelemetry traces

Must be accessed once response has been sent.

diff --git a/Engine/Source/Programs/Horde/HordeServer/Utilities/OpenTelemetry.cs b/Engine/Source/Programs/Horde/HordeServer/Utilities/OpenTelemetry.cs
--- a/Engine/Source/Programs/Horde/HordeServer/Utilities/OpenTelemetry.cs
+++ b/Engine/Source/Programs/Horde/HordeServer/Utilities/OpenTelemetry.cs
@@ -94,8 +94,6 @@ public static class OpenTelemetryHelper
 			activity.SetTag("service.name", settings.ServiceName);
 			activity.SetTag("operation.name", "http.request");
 			activity.SetTag("http.client_ip", request.HttpContext.Connection.RemoteIpAddress);
-			activity.SetTag("usr.id", request.HttpContext.User.GetUserId());
-			activity.SetTag("usr.name", request.HttpContext.User.GetUser());
 			
 			// Header sent by the dashboard to indicate how long a user has been inactive for a particular browser page (in seconds)
 			if (request.Headers.TryGetValue("X-Horde-LastUserActivity", out StringValues values))
@@ -107,6 +105,14 @@ public static class OpenTelemetryHelper
 				}
 			}
 		}
+		
+		void DatadogAspNetResponseEnricher(Activity activity, HttpResponse response)
+		{
+			// The request hook above is executed too early in the middleware chain so user related information must be read here
+			// HttpContext can only be accessed after the response has been sent. Both of these tags are standard Datadog attributes.
+			activity.SetTag("usr.id", response.HttpContext.User.GetUserId());
+			activity.SetTag("usr.name", response.HttpContext.User.GetUser());
+		}
 
 		bool FilterHttpRequests(HttpContext context)
 		{
@@ -134,6 +140,7 @@ public static class OpenTelemetryHelper
 				if (settings.EnableDatadogCompatibility)
 				{
 					options.EnrichWithHttpRequest = DatadogAspNetRequestEnricher;
+					options.EnrichWithHttpResponse = DatadogAspNetResponseEnricher;
 				}
 			})
 			.AddGrpcClientInstrumentation(options =>
