Horde: Missing file.

diff --git a/Engine/Source/Programs/Horde/HordeServer/Utilities/JsonSchemaCache.cs b/Engine/Source/Programs/Horde/HordeServer/Utilities/JsonSchemaCache.cs
new file mode 100644
--- /dev/null
+++ b/Engine/Source/Programs/Horde/HordeServer/Utilities/JsonSchemaCache.cs
@@ -0,0 +1,53 @@
+// Copyright Epic Games, Inc. All Rights Reserved.
+
+using System;
+using System.Collections.Concurrent;
+using EpicGames.Core;
+using HordeServer.Plugins;
+
+namespace HordeServer.Utilities
+{
+	/// <summary>
+	/// Cache for Json schemas
+	/// </summary>
+	public class JsonSchemaCache
+	{
+		readonly JsonSchemaFactory _factory;
+		readonly ConcurrentDictionary<Type, JsonSchema> _cachedSchemas = new ConcurrentDictionary<Type, JsonSchema>();
+
+		/// <summary>
+		/// Constructor
+		/// </summary>
+		public JsonSchemaCache(IPluginCollection pluginCollection)
+		{
+			_factory = new JsonSchemaFactory(new XmlDocReader());
+
+			JsonSchemaObject pluginObj = new JsonSchemaObject();
+			foreach (ILoadedPlugin plugin in pluginCollection.LoadedPlugins)
+			{
+				if (plugin.GlobalConfigType != null)
+				{
+					JsonSchemaType pluginSchemaType = _factory.CreateSchemaType(plugin.GlobalConfigType);
+					pluginObj.Properties.Add(new JsonSchemaProperty(plugin.Name.ToString(), $"Configuration for the {plugin.Name} plugin", pluginSchemaType));
+				}
+			}
+			_factory.TypeCache.Add(typeof(PluginConfigCollection), pluginObj);
+		}
+
+		/// <summary>
+		/// Create a Json schema (or retrieve a cached schema)
+		/// </summary>
+		public JsonSchema CreateSchema(Type type)
+		{
+			JsonSchema? schema;
+			if (!_cachedSchemas.TryGetValue(type, out schema))
+			{
+				lock (_factory)
+				{
+					schema = _cachedSchemas.GetOrAdd(type, x => _factory.CreateSchema(x));
+				}
+			}
+			return schema;
+		}
+	}
+}
