Horde: Various plugin improvements.

* Plugin names now use the PluginName type
* [Plugin] attribute now specifies global config and server config types, allowing them to be injected into constructors of startup objects.
* Added PluginMetadata class similar to PluginDescriptors for UE plugins.

diff --git a/Engine/Source/Programs/Horde/HordeServer.Shared/Plugins/IPluginStartup.cs b/Engine/Source/Programs/Horde/HordeServer.Shared/Plugins/IPluginStartup.cs
--- a/Engine/Source/Programs/Horde/HordeServer.Shared/Plugins/IPluginStartup.cs
+++ b/Engine/Source/Programs/Horde/HordeServer.Shared/Plugins/IPluginStartup.cs
@@ -1,42 +1,17 @@
 // Copyright Epic Games, Inc. All Rights Reserved.
 
 using System;
+using Microsoft.CodeAnalysis.CSharp.Syntax;
+using Microsoft.Extensions.Configuration;
 using Microsoft.Extensions.DependencyInjection;
 
 namespace HordeServer.Plugins
 {
-	/// <summary>
-	/// Attribute used to identify types that should be constructed for the given plugin name
-	/// </summary>
-	[AttributeUsage(AttributeTargets.Class)]
-	public sealed class PluginStartupAttribute : Attribute
-	{
-		/// <summary>
-		/// Name of the plugin
-		/// </summary>
-		public string Name { get; }
-
-		/// <summary>
-		/// Constructor
-		/// </summary>
-		public PluginStartupAttribute(string name)
-			=> Name = name;
-	}
-
 	/// <summary>
 	/// Interface for modules that can extend the services known by Horde
 	/// </summary>
 	public interface IPluginStartup
 	{
-		/// <summary>
-		/// Type containing global settings for this plugin. Must implement the <see cref="IPluginConfig"/> interface.
-		/// 
-		/// Standard <see cref="Microsoft.Extensions.Options.IOptions{TOptions}"/> and 
-		/// <see cref="Microsoft.Extensions.Options.IOptionsMonitor{TOptions}"/> objects will be registered for this type 
-		/// in the service container.
-		/// </summary>
-		Type? GlobalConfigType { get; }
-
 		/// <summary>
 		/// Configure the services provided by the plugin
 		/// </summary>
diff --git a/Engine/Source/Programs/Horde/HordeServer.Shared/Plugins/PluginAttribute.cs b/Engine/Source/Programs/Horde/HordeServer.Shared/Plugins/PluginAttribute.cs
new file mode 100644
--- /dev/null
+++ b/Engine/Source/Programs/Horde/HordeServer.Shared/Plugins/PluginAttribute.cs
@@ -0,0 +1,44 @@
+// Copyright Epic Games, Inc. All Rights Reserved.
+
+using System;
+using System.Collections.Generic;
+using System.Linq;
+using System.Text;
+using System.Threading.Tasks;
+
+namespace HordeServer.Plugins
+{
+	/// <summary>
+	/// Attribute used to identify types that should be constructed for the given plugin name. Should be attached to a public class implementing
+	/// <see cref="IPluginStartup"/>.
+	/// </summary>
+	[AttributeUsage(AttributeTargets.Class)]
+	public sealed class PluginAttribute : Attribute
+	{
+		/// <summary>
+		/// Name of the plugin
+		/// </summary>
+		public string Name { get; }
+
+		/// <summary>
+		/// Type containing server settings for this plugin. Will be injected into a <see cref="Microsoft.Extensions.Options.IOptions{TOptions}"/>
+		/// object in the service container, as well as (optionally) the constructor for the the concrete plugin startup object.
+		/// </summary>
+		public Type? ServerConfigType { get; set; }
+
+		/// <summary>
+		/// Type containing global settings for this plugin. Must implement the <see cref="IPluginConfig"/> interface.
+		/// 
+		/// Standard <see cref="Microsoft.Extensions.Options.IOptions{TOptions}"/> and 
+		/// <see cref="Microsoft.Extensions.Options.IOptionsMonitor{TOptions}"/> objects will be registered for this type 
+		/// in the service container.
+		/// </summary>
+		public Type? GlobalConfigType { get; set; }
+
+		/// <summary>
+		/// Constructor
+		/// </summary>
+		public PluginAttribute(string name)
+			=> Name = name;
+	}
+}
diff --git a/Engine/Source/Programs/Horde/HordeServer.Tests/Server/SchemaTests.cs b/Engine/Source/Programs/Horde/HordeServer.Tests/Server/SchemaTests.cs
--- a/Engine/Source/Programs/Horde/HordeServer.Tests/Server/SchemaTests.cs
+++ b/Engine/Source/Programs/Horde/HordeServer.Tests/Server/SchemaTests.cs
@@ -14,7 +14,7 @@ namespace HordeServer.Tests.Server
 		[TestMethod]
 		public void StreamSchema()
 		{
-			JsonSchemaCache cache = new JsonSchemaCache(PluginCollection.Empty);
+			JsonSchemaCache cache = new JsonSchemaCache(new PluginCollection());
 			JsonSchema schema = cache.CreateSchema(typeof(StreamConfig));
 			_ = schema;
 		}
@@ -22,7 +22,7 @@ namespace HordeServer.Tests.Server
 		[TestMethod]
 		public void ProjectSchema()
 		{
-			JsonSchemaCache cache = new JsonSchemaCache(PluginCollection.Empty);
+			JsonSchemaCache cache = new JsonSchemaCache(new PluginCollection());
 			JsonSchema schema = cache.CreateSchema(typeof(StreamConfig));
 			_ = schema;
 		}
diff --git a/Engine/Source/Programs/Horde/HordeServer.Tests/TestSetup.cs b/Engine/Source/Programs/Horde/HordeServer.Tests/TestSetup.cs
--- a/Engine/Source/Programs/Horde/HordeServer.Tests/TestSetup.cs
+++ b/Engine/Source/Programs/Horde/HordeServer.Tests/TestSetup.cs
@@ -198,7 +198,7 @@ namespace HordeServer.Tests
 			services.AddHttpClient<RpcService>();
 
 			services.AddSingleton<IServerInfo>(new ServerInfo());
-			services.AddSingleton<IPluginCollection>(new PluginCollection(new Dictionary<string, IPluginStartup>()));
+			services.AddSingleton<IPluginCollection>(new PluginCollection());
 
 			services.AddLogging(builder => { builder.AddConsole().SetMinimumLevel(LogLevel.Debug); });
 			services.AddSingleton<IMemoryCache>(sp => new MemoryCache(new MemoryCacheOptions { }));
diff --git a/Engine/Source/Programs/Horde/HordeServer/Commands/Generate/ConfigSchemasCommand.cs b/Engine/Source/Programs/Horde/HordeServer/Commands/Generate/ConfigSchemasCommand.cs
--- a/Engine/Source/Programs/Horde/HordeServer/Commands/Generate/ConfigSchemasCommand.cs
+++ b/Engine/Source/Programs/Horde/HordeServer/Commands/Generate/ConfigSchemasCommand.cs
@@ -1,7 +1,6 @@
 // Copyright Epic Games, Inc. All Rights Reserved.
 
 using System;
-using System.Collections.Generic;
 using System.ComponentModel;
 using System.Threading.Tasks;
 using EpicGames.Core;
@@ -23,7 +22,7 @@ namespace HordeServer.Commands.Generate
 		{
 			_outputDir ??= DirectoryReference.Combine(ServerApp.AppDir, "Schemas");
 
-			JsonSchemaCache schemaCache = new JsonSchemaCache(new PluginCollection(new Dictionary<string, IPluginStartup>()));
+			JsonSchemaCache schemaCache = new JsonSchemaCache(new PluginCollection());
 
 			DirectoryReference.CreateDirectory(_outputDir);
 			foreach (Type schemaType in SchemaController.ConfigSchemas)
diff --git a/Engine/Source/Programs/Horde/HordeServer/Configuration/ConfigService.cs b/Engine/Source/Programs/Horde/HordeServer/Configuration/ConfigService.cs
--- a/Engine/Source/Programs/Horde/HordeServer/Configuration/ConfigService.cs
+++ b/Engine/Source/Programs/Horde/HordeServer/Configuration/ConfigService.cs
@@ -150,11 +150,11 @@ namespace HordeServer.Configuration
 			JsonUtils.ConfigureJsonSerializer(_jsonOptions);
 
 			PluginConfigCollectionConverter pluginConfigConverter = new PluginConfigCollectionConverter();
-			foreach ((string name, IPluginStartup plugin) in pluginCollection.EnabledPlugins)
+			foreach (ILoadedPlugin plugin in pluginCollection.LoadedPlugins)
 			{
 				if (plugin.GlobalConfigType != null)
 				{
-					pluginConfigConverter.NameToType[name] = plugin.GlobalConfigType;
+					pluginConfigConverter.NameToType[plugin.Name] = plugin.GlobalConfigType;
 				}
 			}
 			_jsonOptions.Converters.Add(pluginConfigConverter);
diff --git a/Engine/Source/Programs/Horde/HordeServer/Plugins/IPlugin.cs b/Engine/Source/Programs/Horde/HordeServer/Plugins/IPlugin.cs
new file mode 100644
--- /dev/null
+++ b/Engine/Source/Programs/Horde/HordeServer/Plugins/IPlugin.cs
@@ -0,0 +1,68 @@
+// Copyright Epic Games, Inc. All Rights Reserved.
+
+using System;
+using System.Reflection;
+using Microsoft.Extensions.DependencyInjection;
+
+namespace HordeServer.Plugins
+{
+	/// <summary>
+	/// Interface for manipulating a plugin
+	/// </summary>
+	public interface IPlugin
+	{
+		/// <summary>
+		/// Name of the plugin
+		/// </summary>
+		PluginName Name { get; }
+
+		/// <summary>
+		/// Metadata for the plugin
+		/// </summary>
+		IPluginMetadata Metadata { get; }
+
+		/// <summary>
+		/// Load this plugin
+		/// </summary>
+		ILoadedPlugin Load();
+	}
+
+	/// <summary>
+	/// Plugin which has been loaded
+	/// </summary>
+	public interface ILoadedPlugin : IPlugin
+	{
+		/// <summary>
+		/// Assembly containing the plugin's implementation
+		/// </summary>
+		Assembly Assembly { get; }
+
+		/// <summary>
+		/// Type used for the server config.
+		/// </summary>
+		Type ServerConfigType { get; }
+
+		/// <summary>
+		/// Type used for the global config.
+		/// </summary>
+		Type GlobalConfigType { get; }
+
+		/// <summary>
+		/// Configure services for the plugin
+		/// </summary>
+		/// <param name="serviceCollection">Service collection instance</param>
+		void ConfigureServices(IServiceCollection serviceCollection);
+	}
+
+	/// <summary>
+	/// Extension methods for plugins
+	/// </summary>
+	public static class PluginExtensions
+	{
+		/// <summary>
+		/// Adds a plugin's services to a service collection
+		/// </summary>
+		public static void AddPlugin(this IServiceCollection serviceCollection, ILoadedPlugin plugin)
+			=> plugin.ConfigureServices(serviceCollection);
+	}
+}
diff --git a/Engine/Source/Programs/Horde/HordeServer/Plugins/IPluginCollection.cs b/Engine/Source/Programs/Horde/HordeServer/Plugins/IPluginCollection.cs
--- a/Engine/Source/Programs/Horde/HordeServer/Plugins/IPluginCollection.cs
+++ b/Engine/Source/Programs/Horde/HordeServer/Plugins/IPluginCollection.cs
@@ -10,8 +10,13 @@ namespace HordeServer.Plugins
 	public interface IPluginCollection
 	{
 		/// <summary>
-		/// The enabled plugins
+		/// List of available plugins
 		/// </summary>
-		IReadOnlyDictionary<string, IPluginStartup> EnabledPlugins { get; }
+		IReadOnlyList<IPlugin> Plugins { get; }
+
+		/// <summary>
+		/// List of the enabled plugins
+		/// </summary>
+		IReadOnlyList<ILoadedPlugin> LoadedPlugins { get; }
 	}
 }
diff --git a/Engine/Source/Programs/Horde/HordeServer/Plugins/IPluginMetadata.cs b/Engine/Source/Programs/Horde/HordeServer/Plugins/IPluginMetadata.cs
new file mode 100644
--- /dev/null
+++ b/Engine/Source/Programs/Horde/HordeServer/Plugins/IPluginMetadata.cs
@@ -0,0 +1,53 @@
+// Copyright Epic Games, Inc. All Rights Reserved.
+
+using System;
+using System.Collections.Generic;
+
+namespace HordeServer.Plugins
+{
+	/// <summary>
+	/// Information about a Horde plugin
+	/// </summary>
+	public interface IPluginMetadata
+	{
+		/// <summary>
+		/// Name of the plugin
+		/// </summary>
+		PluginName Name { get; }
+
+		/// <summary>
+		/// Description of the plugin
+		/// </summary>
+		string Description { get; }
+
+		/// <summary>
+		/// Other plugins that this plugin depends on
+		/// </summary>
+		IReadOnlyList<string> DependsOn { get; }
+
+		/// <summary>
+		/// Unique implementations of singleton features which this plugin provides
+		/// </summary>
+		IReadOnlyList<string>? ImplementsSingletons { get; }
+	}
+
+	/// <summary>
+	/// Information about a Horde plugin
+	/// </summary>
+	public class PluginMetadata : IPluginMetadata
+	{
+		/// <inheritdoc/>
+		public PluginName Name { get; set; }
+
+		/// <inheritdoc/>
+		public string Description { get; set; } = String.Empty;
+
+		/// <inheritdoc/>
+		public List<string> DependsOn { get; set; } = new List<string>();
+		IReadOnlyList<string> IPluginMetadata.DependsOn => DependsOn;
+
+		/// <inheritdoc/>
+		public List<string>? ImplementsSingletons { get; set; }
+		IReadOnlyList<string>? IPluginMetadata.ImplementsSingletons => ImplementsSingletons;
+	}
+}
diff --git a/Engine/Source/Programs/Horde/HordeServer/Plugins/PluginCollection.cs b/Engine/Source/Programs/Horde/HordeServer/Plugins/PluginCollection.cs
--- a/Engine/Source/Programs/Horde/HordeServer/Plugins/PluginCollection.cs
+++ b/Engine/Source/Programs/Horde/HordeServer/Plugins/PluginCollection.cs
@@ -1,6 +1,9 @@
 // Copyright Epic Games, Inc. All Rights Reserved.
 
+using System;
 using System.Collections.Generic;
+using System.Reflection;
+using Microsoft.Extensions.DependencyInjection;
 
 namespace HordeServer.Plugins
 {
@@ -9,18 +12,91 @@ namespace HordeServer.Plugins
 	/// </summary>
 	class PluginCollection : IPluginCollection
 	{
+		// Plugin which is compiled into the application or already loaded
+		class StaticPlugin : IPlugin
+		{
+			public PluginName Name => _metadata.Name;
+			public IPluginMetadata Metadata => _metadata;
+
+			readonly IPluginMetadata _metadata;
+			readonly ILoadedPlugin _loadedPlugin;
+
+			public StaticPlugin(IPluginMetadata metadata, ILoadedPlugin loadedPlugin)
+			{
+				_metadata = metadata;
+				_loadedPlugin = loadedPlugin;
+			}
+
+			public ILoadedPlugin Load()
+				=> _loadedPlugin;
+		}
+
+		class LoadedPlugin<TServerConfig, TGlobalConfig, TStartup> : ILoadedPlugin
+			where TGlobalConfig : class, IPluginConfig, new()
+			where TStartup : class, IPluginStartup, new()
+		{
+			readonly IPluginMetadata _metadata;
+
+			public PluginName Name => _metadata.Name;
+			public IPluginMetadata Metadata => _metadata;
+			public Assembly Assembly => typeof(TStartup).Assembly;
+			public Type ServerConfigType => typeof(TServerConfig);
+			public Type GlobalConfigType => typeof(TGlobalConfig);
+
+			public LoadedPlugin(IPluginMetadata metadata)
+				=> _metadata = metadata;
+
+			public void ConfigureServices(IServiceCollection serviceCollection)
+			{
+				serviceCollection.AddPluginConfig<TGlobalConfig>(Name);
+
+				TStartup startup = new TStartup();
+				startup.ConfigureServices(serviceCollection);
+			}
+
+			public ILoadedPlugin Load()
+				=> this;
+		}
+
+		/// <inheritdoc/>
+		public IReadOnlyList<IPlugin> Plugins => _plugins;
+
+		/// <inheritdoc/>
+		public IReadOnlyList<ILoadedPlugin> LoadedPlugins => _loadedPlugins;
+
+		readonly List<IPlugin> _plugins = new List<IPlugin>();
+		readonly List<ILoadedPlugin> _loadedPlugins = new List<ILoadedPlugin>();
+		readonly HashSet<PluginName> _loadedPluginNames = new HashSet<PluginName>();
+
 		/// <summary>
-		/// Static empty plugin collection
+		/// Adds a plugin with the given startup class
 		/// </summary>
-		public static IPluginCollection Empty { get; } = new PluginCollection(new Dictionary<string, IPluginStartup>());
+		ILoadedPlugin AddLoadedPlugin(IPluginMetadata metadata, Type startupType)
+		{
+			if (!_loadedPluginNames.Add(metadata.Name))
+			{
+				throw new NotImplementedException();
+			}
 
-		/// <inheritdoc/>
-		public IReadOnlyDictionary<string, IPluginStartup> EnabledPlugins { get; }
+			PluginAttribute attr = startupType.GetCustomAttribute<PluginAttribute>()
+				?? throw new InvalidOperationException($"Cannot add {startupType.Name} as a plugin. No {nameof(PluginAttribute)} was found.");
+
+			Type pluginType = typeof(LoadedPlugin<,,>).MakeGenericType(attr.ServerConfigType ?? typeof(object), attr.GlobalConfigType ?? typeof(object), startupType);
+			ILoadedPlugin loadedPlugin = (ILoadedPlugin)Activator.CreateInstance(pluginType, metadata)!;
+			_loadedPlugins.Add(loadedPlugin);
+
+			return loadedPlugin;
+		}
 
 		/// <summary>
-		/// Constructor
+		/// Adds a plugin with the given startup class
 		/// </summary>
-		public PluginCollection(IReadOnlyDictionary<string, IPluginStartup> enabled)
-			=> EnabledPlugins = enabled;
+		/// <typeparam name="T">Type of the startup class</typeparam>
+		public ILoadedPlugin Add<T>(IPluginMetadata metadata) where T : class, IPluginStartup
+		{
+			ILoadedPlugin plugin = AddLoadedPlugin(metadata, typeof(T));
+			_plugins.Add(plugin);
+			return plugin;
+		}
 	}
 }
diff --git a/Engine/Source/Programs/Horde/HordeServer/Plugins/PluginConfigCollection.cs b/Engine/Source/Programs/Horde/HordeServer/Plugins/PluginConfigCollection.cs
--- a/Engine/Source/Programs/Horde/HordeServer/Plugins/PluginConfigCollection.cs
+++ b/Engine/Source/Programs/Horde/HordeServer/Plugins/PluginConfigCollection.cs
@@ -5,6 +5,7 @@ using System.Collections.Generic;
 using System.Diagnostics.CodeAnalysis;
 using System.Text.Json;
 using System.Text.Json.Serialization;
+using EpicGames.Horde;
 
 namespace HordeServer.Plugins
 {
@@ -12,19 +13,12 @@ namespace HordeServer.Plugins
 	/// Collection of <see cref="IPluginConfig"/> objects.
 	/// </summary>
 	[JsonConverter(typeof(PluginConfigCollectionConverter))]
-	public class PluginConfigCollection : Dictionary<string, IPluginConfig>
+	public class PluginConfigCollection : Dictionary<PluginName, IPluginConfig>
 	{
-		/// <summary>
-		/// Constructor
-		/// </summary>
-		public PluginConfigCollection()
-			: base(StringComparer.OrdinalIgnoreCase)
-		{ }
-
 		/// <summary>
 		/// Attempts to get a plugin config of a specific type
 		/// </summary>
-		public bool TryGetValue<T>(string name, [NotNullWhen(true)] out T? value) where T : class, IPluginConfig
+		public bool TryGetValue<T>(PluginName name, [NotNullWhen(true)] out T? value) where T : class, IPluginConfig
 		{
 			IPluginConfig? config;
 			if (base.TryGetValue(name, out config) && config is T typedConfig)
@@ -48,7 +42,7 @@ namespace HordeServer.Plugins
 		/// <summary>
 		/// Map of plugin names to the expected config type
 		/// </summary>
-		public Dictionary<string, Type> NameToType { get; } = new Dictionary<string, Type>(StringComparer.OrdinalIgnoreCase);
+		public Dictionary<PluginName, Type> NameToType { get; } = new Dictionary<PluginName, Type>();
 
 		/// <inheritdoc/>
 		public override PluginConfigCollection? Read(ref Utf8JsonReader reader, Type typeToConvert, JsonSerializerOptions options)
@@ -61,11 +55,13 @@ namespace HordeServer.Plugins
 			PluginConfigCollection result = new PluginConfigCollection();
 			while (reader.Read() && reader.TokenType == JsonTokenType.PropertyName)
 			{
-				string? name = reader.GetString();
-				if (name == null || !reader.Read())
+				string? nameStr = reader.GetString();
+				if (nameStr == null || !reader.Read())
 				{
 					throw new JsonException("Invalid plugin name");
 				}
+
+				PluginName name = new PluginName(new StringId(nameStr));
 				if (NameToType.TryGetValue(name, out Type? type))
 				{
 					object? obj = JsonSerializer.Deserialize(ref reader, type, options);
@@ -91,9 +87,9 @@ namespace HordeServer.Plugins
 		public override void Write(Utf8JsonWriter writer, PluginConfigCollection value, JsonSerializerOptions options)
 		{
 			writer.WriteStartObject();
-			foreach ((string name, IPluginConfig config) in value)
+			foreach ((PluginName name, IPluginConfig config) in value)
 			{
-				writer.WritePropertyName(name);
+				writer.WritePropertyName(name.ToString());
 				JsonSerializer.Serialize(writer, config, config.GetType(), options);
 			}
 			writer.WriteEndObject();
diff --git a/Engine/Source/Programs/Horde/HordeServer/Plugins/PluginName.cs b/Engine/Source/Programs/Horde/HordeServer/Plugins/PluginName.cs
new file mode 100644
--- /dev/null
+++ b/Engine/Source/Programs/Horde/HordeServer/Plugins/PluginName.cs
@@ -0,0 +1,44 @@
+// Copyright Epic Games, Inc. All Rights Reserved.
+
+using System.ComponentModel;
+using EpicGames.Core;
+using EpicGames.Horde;
+
+namespace HordeServer.Plugins
+{
+	/// <summary>
+	/// Name of a plugin
+	/// </summary>
+	/// <param name="Id"></param>
+	[LogValueType]
+	[JsonSchemaString]
+	[TypeConverter(typeof(StringIdTypeConverter<PluginName, PluginNameStringIdConverter>))]
+	[StringIdConverter(typeof(PluginNameStringIdConverter))]
+	public record struct PluginName(StringId Id)
+	{
+		/// <summary>
+		/// Constructor
+		/// </summary>
+		/// <param name="name"></param>
+		public PluginName(string name)
+			: this(new StringId(name))
+		{ }
+
+		/// <inheritdoc/>
+		public override string ToString() => Id.ToString();
+	}
+
+	/// <summary>
+	/// Converter from PluginName objects to regular StringId types
+	/// </summary>
+	class PluginNameStringIdConverter : StringIdConverter<PluginName>
+	{
+		/// <inheritdoc/>
+		public override PluginName FromStringId(StringId id)
+			=> new PluginName(id);
+
+		/// <inheritdoc/>
+		public override StringId ToStringId(PluginName value)
+			=> value.Id;
+	}
+}
diff --git a/Engine/Source/Programs/Horde/HordeServer/Plugins/PluginOptionsFactory.cs b/Engine/Source/Programs/Horde/HordeServer/Plugins/PluginOptionsFactory.cs
--- a/Engine/Source/Programs/Horde/HordeServer/Plugins/PluginOptionsFactory.cs
+++ b/Engine/Source/Programs/Horde/HordeServer/Plugins/PluginOptionsFactory.cs
@@ -15,7 +15,7 @@ namespace HordeServer.Plugins
 	sealed class PluginOptionsFactory<T> : IOptionsFactory<T>, IOptionsChangeTokenSource<T>
 		where T : class, IPluginConfig, new()
 	{
-		readonly string _pluginName;
+		readonly PluginName _pluginName;
 		readonly IOptionsMonitor<GlobalConfig> _globalConfig;
 		readonly IOptionsChangeTokenSource<GlobalConfig> _globalConfigChangeTokenSource;
 
@@ -25,7 +25,7 @@ namespace HordeServer.Plugins
 		/// <summary>
 		/// Constructor
 		/// </summary>
-		public PluginOptionsFactory(string pluginName, IServiceProvider serviceProvider)
+		public PluginOptionsFactory(PluginName pluginName, IServiceProvider serviceProvider)
 			: this(pluginName, serviceProvider.GetRequiredService<IOptionsMonitor<GlobalConfig>>(), serviceProvider.GetRequiredService<IOptionsChangeTokenSource<GlobalConfig>>())
 		{
 		}
@@ -33,7 +33,7 @@ namespace HordeServer.Plugins
 		/// <summary>
 		/// Constructor
 		/// </summary>
-		public PluginOptionsFactory(string pluginName, IOptionsMonitor<GlobalConfig> globalConfig, IOptionsChangeTokenSource<GlobalConfig> globalConfigChangeTokenSource)
+		public PluginOptionsFactory(PluginName pluginName, IOptionsMonitor<GlobalConfig> globalConfig, IOptionsChangeTokenSource<GlobalConfig> globalConfigChangeTokenSource)
 		{
 			_pluginName = pluginName;
 			_globalConfig = globalConfig;
@@ -64,26 +64,11 @@ namespace HordeServer.Plugins
 		/// <summary>
 		/// Register an options factory for the given plugin options type
 		/// </summary>
-		public static void AddPluginConfig(this IServiceCollection services, string name, Type configType)
+		public static void AddPluginConfig<T>(this IServiceCollection services, PluginName name) where T : class, IPluginConfig, new()
 		{
-			Type helperType = typeof(RegistrationHelper<>).MakeGenericType(configType);
-			RegistrationHelper registration = (RegistrationHelper)Activator.CreateInstance(helperType)!;
-			registration.Register(services, name);
-		}
-
-		abstract class RegistrationHelper
-		{
-			public abstract void Register(IServiceCollection services, string name);
-		}
-
-		class RegistrationHelper<T> : RegistrationHelper where T : class, IPluginConfig, new()
-		{
-			public override void Register(IServiceCollection services, string name)
-			{
-				services.AddSingleton<PluginOptionsFactory<T>>(sp => new PluginOptionsFactory<T>(name, sp));
-				services.AddSingleton<IOptionsFactory<T>>(sp => sp.GetRequiredService<PluginOptionsFactory<T>>());
-				services.AddSingleton<IOptionsChangeTokenSource<T>>(sp => sp.GetRequiredService<PluginOptionsFactory<T>>());
-			}
+			services.AddSingleton<PluginOptionsFactory<T>>(sp => new PluginOptionsFactory<T>(name, sp));
+			services.AddSingleton<IOptionsFactory<T>>(sp => sp.GetRequiredService<PluginOptionsFactory<T>>());
+			services.AddSingleton<IOptionsChangeTokenSource<T>>(sp => sp.GetRequiredService<PluginOptionsFactory<T>>());
 		}
 	}
 }
diff --git a/Engine/Source/Programs/Horde/HordeServer/Startup.cs b/Engine/Source/Programs/Horde/HordeServer/Startup.cs
--- a/Engine/Source/Programs/Horde/HordeServer/Startup.cs
+++ b/Engine/Source/Programs/Horde/HordeServer/Startup.cs
@@ -371,16 +371,13 @@ namespace HordeServer
 			services.AddSingleton<JsonSchemaCache>();
 
 			// Register the plugin collection
-			Dictionary<string, IPluginStartup> plugins = new Dictionary<string, IPluginStartup>();
-			services.AddSingleton<IPluginCollection>(new PluginCollection(plugins));
+			PluginCollection pluginCollection = new PluginCollection();
+			services.AddSingleton<IPluginCollection>(pluginCollection);
 
 			// Register all the plugin config types
-			foreach ((string name, IPluginStartup plugin) in plugins)
+			foreach (ILoadedPlugin plugin in pluginCollection.LoadedPlugins)
 			{
-				if (plugin.GlobalConfigType != null)
-				{
-					services.AddPluginConfig(name, plugin.GlobalConfigType);
-				}
+				plugin.ConfigureServices(services);
 			}
 
 			// IOptionsMonitor pattern for live updating of configuration settings
@@ -915,7 +912,7 @@ namespace HordeServer
 			IMvcBuilder mvcBuilder = services.AddMvc()
 				.AddJsonOptions(options => JsonUtils.ConfigureJsonSerializer(options.JsonSerializerOptions));
 
-			foreach (Assembly pluginAssembly in plugins.Values.Select(x => x.GetType().Assembly).Distinct())
+			foreach (Assembly pluginAssembly in pluginCollection.LoadedPlugins.Select(x => x.Assembly).Distinct())
 			{
 				mvcBuilder.AddApplicationPart(pluginAssembly);
 			}
