Horde: Missing file.

#rnx

diff --git a/Engine/Source/Programs/Horde/HordeServer/ServerInfo.cs b/Engine/Source/Programs/Horde/HordeServer/ServerInfo.cs
new file mode 100644
--- /dev/null
+++ b/Engine/Source/Programs/Horde/HordeServer/ServerInfo.cs
@@ -0,0 +1,27 @@
+// Copyright Epic Games, Inc. All Rights Reserved.
+
+using EpicGames.Core;
+
+namespace HordeServer
+{
+	/// <summary>
+	/// Concrete implementation of <see cref="IServerInfo"/>
+	/// </summary>
+	class ServerInfo : IServerInfo
+	{
+		/// <inheritdoc/>
+		public SemVer Version => ServerApp.Version;
+
+		/// <inheritdoc/>
+		public string Environment => ServerApp.DeploymentEnvironment;
+
+		/// <inheritdoc/>
+		public string SessionId => ServerApp.SessionId;
+
+		/// <inheritdoc/>
+		public DirectoryReference AppDir => ServerApp.AppDir;
+
+		/// <inheritdoc/>
+		public DirectoryReference DataDir => ServerApp.DataDir;
+	}
+}
