Horde: Fix preflights not navigating to existing jobs
#rnx
#skipci

diff --git a/Engine/Source/Programs/Horde/HordeDashboard/src/components/Preflight.tsx b/Engine/Source/Programs/Horde/HordeDashboard/src/components/Preflight.tsx
--- a/Engine/Source/Programs/Horde/HordeDashboard/src/components/Preflight.tsx
+++ b/Engine/Source/Programs/Horde/HordeDashboard/src/components/Preflight.tsx
@@ -105,7 +105,7 @@ export const PreflightRedirector: React.FC = () => {
 
       console.log(`Redirecting preflight: ${window.location.href}`);
 
-      backend.getJobs({ filter: "id", count: 1, preflightChange: cl }).then(result => {
+      backend.getJobs({ filter: "id,streamId", count: 1, preflightChange: cl }).then(result => {
 
          if (result && result.length === 1) {
 
