Horde: Support injecting server config types into plugin startup constructors.

#rnx

diff --git a/Engine/Source/Programs/Horde/HordeServer.Shared/Plugins/IPluginConfig.cs b/Engine/Source/Programs/Horde/HordeServer.Shared/Plugins/IPluginConfig.cs
--- a/Engine/Source/Programs/Horde/HordeServer.Shared/Plugins/IPluginConfig.cs
+++ b/Engine/Source/Programs/Horde/HordeServer.Shared/Plugins/IPluginConfig.cs
@@ -22,4 +22,14 @@ namespace HordeServer.Plugins
 		/// <param name="parentAcl">The parent ACL scope</param>
 		void PostLoad(AclConfig parentAcl);
 	}
+
+	/// <summary>
+	/// Empty implementation of <see cref="IPluginConfig"/>
+	/// </summary>
+	public sealed class EmptyPluginConfig : IPluginConfig
+	{
+		/// <inheritdoc/>
+		public void PostLoad(AclConfig parentAcl)
+		{ }
+	}
 }
diff --git a/Engine/Source/Programs/Horde/HordeServer/Plugins/IPlugin.cs b/Engine/Source/Programs/Horde/HordeServer/Plugins/IPlugin.cs
--- a/Engine/Source/Programs/Horde/HordeServer/Plugins/IPlugin.cs
+++ b/Engine/Source/Programs/Horde/HordeServer/Plugins/IPlugin.cs
@@ -2,6 +2,7 @@
 
 using System;
 using System.Reflection;
+using Microsoft.Extensions.Configuration;
 using Microsoft.Extensions.DependencyInjection;
 
 namespace HordeServer.Plugins
@@ -50,8 +51,9 @@ namespace HordeServer.Plugins
 		/// <summary>
 		/// Configure services for the plugin
 		/// </summary>
+		/// <param name="configuration">Static configuration data</param>
 		/// <param name="serviceCollection">Service collection instance</param>
-		void ConfigureServices(IServiceCollection serviceCollection);
+		void ConfigureServices(IConfiguration configuration, IServiceCollection serviceCollection);
 	}
 
 	/// <summary>
@@ -62,7 +64,7 @@ namespace HordeServer.Plugins
 		/// <summary>
 		/// Adds a plugin's services to a service collection
 		/// </summary>
-		public static void AddPlugin(this IServiceCollection serviceCollection, ILoadedPlugin plugin)
-			=> plugin.ConfigureServices(serviceCollection);
+		public static void AddPlugin(this IServiceCollection serviceCollection, ILoadedPlugin plugin, IConfiguration configuration)
+			=> plugin.ConfigureServices(configuration, serviceCollection);
 	}
 }
diff --git a/Engine/Source/Programs/Horde/HordeServer/Plugins/PluginCollection.cs b/Engine/Source/Programs/Horde/HordeServer/Plugins/PluginCollection.cs
--- a/Engine/Source/Programs/Horde/HordeServer/Plugins/PluginCollection.cs
+++ b/Engine/Source/Programs/Horde/HordeServer/Plugins/PluginCollection.cs
@@ -2,7 +2,10 @@
 
 using System;
 using System.Collections.Generic;
+using System.Linq;
 using System.Reflection;
+using HordeServer.Acls;
+using Microsoft.Extensions.Configuration;
 using Microsoft.Extensions.DependencyInjection;
 
 namespace HordeServer.Plugins
@@ -32,8 +35,9 @@ namespace HordeServer.Plugins
 		}
 
 		class LoadedPlugin<TServerConfig, TGlobalConfig, TStartup> : ILoadedPlugin
+			where TServerConfig : class, new()
 			where TGlobalConfig : class, IPluginConfig, new()
-			where TStartup : class, IPluginStartup, new()
+			where TStartup : class, IPluginStartup
 		{
 			readonly IPluginMetadata _metadata;
 
@@ -46,14 +50,73 @@ namespace HordeServer.Plugins
 			public LoadedPlugin(IPluginMetadata metadata)
 				=> _metadata = metadata;
 
-			public void ConfigureServices(IServiceCollection serviceCollection)
+			public void ConfigureServices(IConfiguration config, IServiceCollection serviceCollection)
 			{
 				serviceCollection.AddPluginConfig<TGlobalConfig>(Name);
 
-				TStartup startup = new TStartup();
+				TStartup startup = CreateStartup(config);
 				startup.ConfigureServices(serviceCollection);
 			}
 
+			static TStartup CreateStartup(IConfiguration configuration)
+			{
+				ConstructorInfo? chosenConstructor = null;
+				ParameterInfo[]? chosenConstructorParams = null;
+
+				ConstructorInfo[] constructors = typeof(TStartup).GetConstructors(BindingFlags.Public | BindingFlags.Instance);
+				foreach (ConstructorInfo constructor in constructors)
+				{
+					ParameterInfo[] parameters = constructor.GetParameters();
+					if (IsValidConstructor(parameters))
+					{
+						if (chosenConstructorParams == null || parameters.Length > chosenConstructorParams.Length)
+						{
+							chosenConstructor = constructor;
+							chosenConstructorParams = parameters;
+						}
+					}
+				}
+
+				TStartup startup;
+				if (chosenConstructor == null || chosenConstructorParams == null)
+				{
+					startup = Activator.CreateInstance<TStartup>();
+				}
+				else
+				{
+					object[] arguments = new object[chosenConstructorParams.Length];
+					for (int idx = 0; idx < chosenConstructorParams.Length; idx++)
+					{
+						ParameterInfo parameter = chosenConstructorParams[idx];
+						if (parameter.ParameterType == typeof(TServerConfig))
+						{
+							arguments[idx] = new TServerConfig();
+							configuration.Bind(arguments[idx]);
+						}
+						else
+						{
+							throw new NotImplementedException();
+						}
+					}
+
+					startup = (TStartup)chosenConstructor.Invoke(arguments);
+				}
+
+				return startup;
+			}
+
+			static bool IsValidConstructor(ParameterInfo[] parameters)
+			{
+				foreach (ParameterInfo parameter in parameters)
+				{
+					if (parameter.ParameterType != typeof(TServerConfig))
+					{
+						return false;
+					}
+				}
+				return true;
+			}
+
 			public ILoadedPlugin Load()
 				=> this;
 		}
@@ -81,7 +144,7 @@ namespace HordeServer.Plugins
 			PluginAttribute attr = startupType.GetCustomAttribute<PluginAttribute>()
 				?? throw new InvalidOperationException($"Cannot add {startupType.Name} as a plugin. No {nameof(PluginAttribute)} was found.");
 
-			Type pluginType = typeof(LoadedPlugin<,,>).MakeGenericType(attr.ServerConfigType ?? typeof(object), attr.GlobalConfigType ?? typeof(object), startupType);
+			Type pluginType = typeof(LoadedPlugin<,,>).MakeGenericType(attr.ServerConfigType ?? typeof(object), attr.GlobalConfigType ?? typeof(EmptyPluginConfig), startupType);
 			ILoadedPlugin loadedPlugin = (ILoadedPlugin)Activator.CreateInstance(pluginType, metadata)!;
 			_loadedPlugins.Add(loadedPlugin);
 
diff --git a/Engine/Source/Programs/Horde/HordeServer/Startup.cs b/Engine/Source/Programs/Horde/HordeServer/Startup.cs
--- a/Engine/Source/Programs/Horde/HordeServer/Startup.cs
+++ b/Engine/Source/Programs/Horde/HordeServer/Startup.cs
@@ -370,16 +370,6 @@ namespace HordeServer
 			services.AddSingleton<IServerInfo>(new ServerInfo());
 			services.AddSingleton<JsonSchemaCache>();
 
-			// Register the plugin collection
-			PluginCollection pluginCollection = new PluginCollection();
-			services.AddSingleton<IPluginCollection>(pluginCollection);
-
-			// Register all the plugin config types
-			foreach (ILoadedPlugin plugin in pluginCollection.LoadedPlugins)
-			{
-				plugin.ConfigureServices(services);
-			}
-
 			// IOptionsMonitor pattern for live updating of configuration settings
 			services.Configure<ServerSettings>(x => BindServerSettings(Configuration, x));
 
@@ -393,6 +383,18 @@ namespace HordeServer
 			ServerSettings settings = new();
 			BindServerSettings(Configuration, settings);
 
+			// Register the plugin collection
+			PluginCollection pluginCollection = new PluginCollection();
+			services.AddSingleton<IPluginCollection>(pluginCollection);
+
+			// Register all the plugin services
+			IConfigurationSection pluginsConfig = Configuration.GetSection("Horde").GetSection("Plugins");
+			foreach (ILoadedPlugin plugin in pluginCollection.LoadedPlugins)
+			{
+				IConfigurationSection pluginConfig = pluginsConfig.GetSection(plugin.Name.ToString());
+				plugin.ConfigureServices(pluginConfig, services);
+			}
+
 			OpenTelemetryHelper.Configure(services, settings.OpenTelemetry);
 
 			if (settings.GlobalThreadPoolMinSize != null)
