Horde: Resolve HTTP client IP for Datadog trace enricher

diff --git a/Engine/Source/Programs/Horde/HordeServer/Utilities/OpenTelemetry.cs b/Engine/Source/Programs/Horde/HordeServer/Utilities/OpenTelemetry.cs
--- a/Engine/Source/Programs/Horde/HordeServer/Utilities/OpenTelemetry.cs
+++ b/Engine/Source/Programs/Horde/HordeServer/Utilities/OpenTelemetry.cs
@@ -93,7 +93,13 @@ public static class OpenTelemetryHelper
 			// From https://docs.datadoghq.com/standard-attributes/
 			activity.SetTag("service.name", settings.ServiceName);
 			activity.SetTag("operation.name", "http.request");
-			activity.SetTag("http.client_ip", request.HttpContext.Connection.RemoteIpAddress);
+			
+			// Resolve client's IP via headers or actual TCP/IP remote IP
+			string? forwardedForHeader = request.HttpContext.Request.Headers["X-Forwarded-For"].FirstOrDefault();
+			string? clientIp = String.IsNullOrEmpty(forwardedForHeader)
+				? request.HttpContext.Connection.RemoteIpAddress?.ToString()
+				: forwardedForHeader;
+			activity.SetTag("http.client_ip", clientIp);
 			
 			// Header sent by the dashboard to indicate how long a user has been inactive for a particular browser page (in seconds)
 			if (request.Headers.TryGetValue("X-Horde-LastUserActivity", out StringValues values))
