Horde: Fix compile warnings in plugins

diff --git a/Engine/Source/Programs/Horde/HordeServer/Plugins/PluginCollection.cs b/Engine/Source/Programs/Horde/HordeServer/Plugins/PluginCollection.cs
--- a/Engine/Source/Programs/Horde/HordeServer/Plugins/PluginCollection.cs
+++ b/Engine/Source/Programs/Horde/HordeServer/Plugins/PluginCollection.cs
@@ -2,9 +2,7 @@
 
 using System;
 using System.Collections.Generic;
-using System.Linq;
 using System.Reflection;
-using HordeServer.Acls;
 using Microsoft.Extensions.Configuration;
 using Microsoft.Extensions.DependencyInjection;
 
