Horde: UBA analytics chart for cache rate

#rnx

diff --git a/Engine/Source/Programs/Horde/HordeServer/Defaults/default-analytics.dashboard.json b/Engine/Source/Programs/Horde/HordeServer/Defaults/default-analytics.dashboard.json
--- a/Engine/Source/Programs/Horde/HordeServer/Defaults/default-analytics.dashboard.json
+++ b/Engine/Source/Programs/Horde/HordeServer/Defaults/default-analytics.dashboard.json
@@ -569,6 +569,15 @@
 								}
 							]
 						},
+						{
+							"name": "Cache Rate All (Average)",
+							"display": "Ratio",
+							"metrics": [
+								{
+									"id": "ubt-uba-cache-rate-all"
+								}
+							]
+						},
 						{
 							"name": "Actions All (Sum)",
 							"display": "Value",
@@ -609,6 +618,15 @@
 								}
 							]
 						},
+						{
+							"name": "CacheRate (Average)",
+							"display": "Ratio",
+							"metrics": [
+								{
+									"id": "ubt-uba-cache-rate"
+								}
+							]
+						},
 						{
 							"name": "Actions (Average)",
 							"display": "Value",
diff --git a/Engine/Source/Programs/Horde/HordeServer/Defaults/default-metrics.telemetry.json b/Engine/Source/Programs/Horde/HordeServer/Defaults/default-metrics.telemetry.json
--- a/Engine/Source/Programs/Horde/HordeServer/Defaults/default-metrics.telemetry.json
+++ b/Engine/Source/Programs/Horde/HordeServer/Defaults/default-metrics.telemetry.json
@@ -626,6 +626,14 @@
 			"function": "Average",
 			"interval": "00:30:00"
 		},
+		{
+			"id": "ubt-uba-cache-rate",
+			"filter": "$[?(@.Payload.EventName == 'Core.UBT.Executor' && search(@.Payload.Executor, \"Unreal Build Accelerator\") )]",
+			"groupBy": "$.Payload.Build_BranchName,$.Payload.Target,$.Payload.Mode,$.Payload.Configuration,$.Payload.Platform",
+			"property": "$.Payload.CacheRate",
+			"function": "Average",
+			"interval": "00:30:00"
+		},
 		{
 			"id": "ubt-uba-remote-usage-all",
 			"filter": "$[?(@.Payload.EventName == 'Core.UBT.Executor' && search(@.Payload.Executor, \"Unreal Build Accelerator\") )]",
@@ -658,6 +666,14 @@
 			"function": "Average",
 			"interval": "00:30:00"
 		},
+		{
+			"id": "ubt-uba-cache-rate-all",
+			"filter": "$[?(@.Payload.EventName == 'Core.UBT.Executor' && search(@.Payload.Executor, \"Unreal Build Accelerator\") )]",
+			"groupBy": "$.Payload.Build_BranchName",
+			"property": "$.Payload.CacheRate",
+			"function": "Average",
+			"interval": "00:30:00"
+		},
 		{
 			"id": "ubt-completed-duration",
 			"filter": "$[?(@.Payload.EventName == 'Core.UBT.Completed' )]",
