Horde: Ensure that we create a config object for each plugin in the global config rather than constructing a new one each time, to make sure the PostLoad method is called and it has the opportunity to fixup the ACL chain.

#rnx

diff --git a/Engine/Source/Programs/Horde/HordeServer/Configuration/ConfigService.cs b/Engine/Source/Programs/Horde/HordeServer/Configuration/ConfigService.cs
--- a/Engine/Source/Programs/Horde/HordeServer/Configuration/ConfigService.cs
+++ b/Engine/Source/Programs/Horde/HordeServer/Configuration/ConfigService.cs
@@ -112,6 +112,7 @@ namespace HordeServer.Configuration
 		readonly IRedisService _redisService;
 		readonly ServerSettings _serverSettings;
 		readonly Dictionary<string, IConfigSource> _sources;
+		readonly IPluginCollection _pluginCollection;
 		readonly JsonSerializerOptions _jsonOptions;
 		readonly RedisKey _snapshotKey = "config";
 		readonly IHealthMonitor _health;
@@ -142,6 +143,7 @@ namespace HordeServer.Configuration
 			_redisService = redisService;
 			_serverSettings = serverSettings.Value;
 			_sources = sources.ToDictionary(x => x.Scheme, x => x, StringComparer.OrdinalIgnoreCase);
+			_pluginCollection = pluginCollection;
 			_health = health;
 			_health.SetName("GlobalConfig");
 			_logger = logger;
@@ -744,6 +746,16 @@ namespace HordeServer.Configuration
 			GlobalConfig globalConfig = JsonSerializer.Deserialize<GlobalConfig>(snapshot.Data, _jsonOptions)!;
 			globalConfig.Revision = IoHash.Compute(data.Span).ToString();
 
+			// Ensure that all plugins have an entry in the global config so they can register their ACLs
+			foreach (ILoadedPlugin loadedPlugin in _pluginCollection.LoadedPlugins)
+			{
+				if (!globalConfig.Plugins.TryGetValue(loadedPlugin.Name, out _))
+				{
+					IPluginConfig pluginConfig = (IPluginConfig)Activator.CreateInstance(loadedPlugin.GlobalConfigType)!;
+					globalConfig.Plugins.Add(loadedPlugin.Name, pluginConfig);
+				}
+			}
+
 			// Run the postload callbacks on all the config objects
 			globalConfig.PostLoad(_serverSettings);
 			return globalConfig;
diff --git a/Engine/Source/Programs/Horde/HordeServer/Plugins/PluginOptionsFactory.cs b/Engine/Source/Programs/Horde/HordeServer/Plugins/PluginOptionsFactory.cs
--- a/Engine/Source/Programs/Horde/HordeServer/Plugins/PluginOptionsFactory.cs
+++ b/Engine/Source/Programs/Horde/HordeServer/Plugins/PluginOptionsFactory.cs
@@ -42,14 +42,7 @@ namespace HordeServer.Plugins
 
 		/// <inheritdoc/>
 		public T Create(string name)
-		{
-			T? config;
-			if (!_globalConfig.CurrentValue.Plugins.TryGetValue<T>(_pluginName, out config))
-			{
-				config = new T();
-			}
-			return config;
-		}
+			=> (T)_globalConfig.CurrentValue.Plugins[_pluginName];
 
 		/// <inheritdoc/>
 		public IChangeToken GetChangeToken()
