Horde: Prevent throwing an exception from the Swagger middleware if the server documentation does not exist on disk.

#rnx

diff --git a/Engine/Source/Programs/Horde/HordeServer/Startup.cs b/Engine/Source/Programs/Horde/HordeServer/Startup.cs
--- a/Engine/Source/Programs/Horde/HordeServer/Startup.cs
+++ b/Engine/Source/Programs/Horde/HordeServer/Startup.cs
@@ -948,7 +948,12 @@ namespace HordeServer
 			{
 				config.SchemaGeneratorOptions.SchemaFilters.Add(new SwaggerSchemaFilter());
 				config.SwaggerDoc("v1", new OpenApiInfo { Title = "Horde Server API", Version = "v1" });
-				config.IncludeXmlComments(Path.Combine(AppContext.BaseDirectory, $"{Assembly.GetExecutingAssembly().GetName().Name}.xml"));
+
+				string xmlDocFile = Path.Combine(AppContext.BaseDirectory, $"{Assembly.GetExecutingAssembly().GetName().Name}.xml");
+				if (File.Exists(xmlDocFile))
+				{
+					config.IncludeXmlComments(xmlDocFile);
+				}
 			});
 
 			services.Configure<ApiBehaviorOptions>(options =>
