Horde: Move analytics into a plugin (currently statically linked, and always enabled).

Configuration for the telemetry sinks is moved from Horde:Telemetry to Horde:Plugins:Analytics:Sinks.

diff --git a/Engine/Source/Programs/Horde/HordeServer/HordeServer.csproj b/Engine/Source/Programs/Horde/HordeServer/HordeServer.csproj
--- a/Engine/Source/Programs/Horde/HordeServer/HordeServer.csproj
+++ b/Engine/Source/Programs/Horde/HordeServer/HordeServer.csproj
@@ -134,6 +134,7 @@
     <ProjectReference Include="..\..\Shared\EpicGames.Serialization\EpicGames.Serialization.csproj" />
     <ProjectReference Include="..\..\Shared\EpicGames.Slack\EpicGames.Slack.csproj" />
     <ProjectReference Include="..\HordeServer.Shared\HordeServer.Shared.csproj" />
+    <ProjectReference Include="..\Plugins\Analytics\HordeServer.Analytics\HordeServer.Analytics.csproj" />
   </ItemGroup>
 
   <ItemGroup>
diff --git a/Engine/Source/Programs/Horde/HordeServer/ServerSettings.cs b/Engine/Source/Programs/Horde/HordeServer/ServerSettings.cs
--- a/Engine/Source/Programs/Horde/HordeServer/ServerSettings.cs
+++ b/Engine/Source/Programs/Horde/HordeServer/ServerSettings.cs
@@ -164,93 +164,6 @@ namespace HordeServer
 		Worker
 	}
 
-	/// <summary>
-	/// Type of telemetry provider to use
-	/// </summary>
-	public enum TelemetrySinkType
-	{
-		/// <summary>
-		/// No telemetry sink (default)
-		/// </summary>
-		None,
-
-		/// <summary>
-		/// Use the Epic telemetry sink
-		/// </summary>
-		Epic,
-
-		/// <summary>
-		/// Use the ClickHouse telemetry sink
-		/// </summary>
-		ClickHouse,
-
-		/// <summary>
-		/// Mongo telemetry
-		/// </summary>
-		Mongo,
-	}
-
-	/// <summary>
-	/// Configuration for the telemetry sink
-	/// </summary>
-	public class BaseTelemetryConfig
-	{
-		/// <summary>
-		/// Type of telemetry sink
-		/// </summary>
-		public TelemetrySinkType Type { get; set; } = TelemetrySinkType.None;
-	}
-
-	/// <summary>
-	/// Configuration for the telemetry sink
-	/// </summary>
-	public class EpicTelemetryConfig : BaseTelemetryConfig, IEpicTelemetrySinkConfig
-	{
-		/// <summary>
-		/// Base URL for the telemetry server
-		/// </summary>
-		public Uri? Url { get; set; }
-
-		/// <summary>
-		/// Application name to send in the event messages
-		/// </summary>
-		public string AppId { get; set; } = "Horde";
-
-		/// <inheritdoc />
-		public override string ToString()
-		{
-			return $"{nameof(Url)}={Url} {nameof(AppId)}={AppId}";
-		}
-	}
-
-	/// <summary>
-	/// Configuration for the telemetry sink
-	/// </summary>
-	public class ClickHouseTelemetryConfig : BaseTelemetryConfig
-	{
-		/// <summary>
-		/// Base URL for ClickHouse server
-		/// </summary>
-		public Uri? Url { get; set; }
-
-		/// <inheritdoc />
-		public override string ToString()
-		{
-			return $"{nameof(Url)}={Url}";
-		}
-	}
-
-	/// <summary>
-	/// Configuration for the telemetry sink
-	/// </summary>
-	public class MongoTelemetryConfig : BaseTelemetryConfig
-	{
-		/// <summary>
-		/// Number of days worth of telmetry events to keep
-		/// </summary>
-		public double RetainDays { get; set; } = 1;
-	}
-
 	/// <summary>
 	/// Feature flags to aid rollout of new features.
 	///
@@ -981,11 +894,6 @@ namespace HordeServer
 		/// </summary>
 		public CommitSettings Commits { get; set; } = new CommitSettings();
 
-		/// <summary>
-		/// Settings for sending telemetry events to external services (for example Snowflake, ClickHouse etc)
-		/// </summary>
-		public List<BaseTelemetryConfig> Telemetry { get; set; } = new();
-
 		/// <summary>
 		/// Tools bundled along with the server. Data for each tool can be produced using the 'bundle create' command, and should be stored in the Tools directory.
 		/// </summary>
diff --git a/Engine/Source/Programs/Horde/HordeServer/Startup.cs b/Engine/Source/Programs/Horde/HordeServer/Startup.cs
--- a/Engine/Source/Programs/Horde/HordeServer/Startup.cs
+++ b/Engine/Source/Programs/Horde/HordeServer/Startup.cs
@@ -280,57 +280,10 @@ namespace HordeServer
 
 		public IConfiguration Configuration { get; }
 
-		/// <summary>
-		/// Bind config to concrete subclasses of BaseTelemetryConfig
-		/// </summary>
-		/// <param name="settings">Settings object being updated</param>
-		/// <param name="telemetrySection">Telemetry config section</param>
-		/// <exception cref="ConfigurationException"></exception>
-		private static void BindTelemetrySettings(ServerSettings settings, IConfiguration telemetrySection)
-		{
-			List<BaseTelemetryConfig> telemetryConfigs = new();
-			foreach (IConfigurationSection child in telemetrySection.GetChildren())
-			{
-				string? typeStr = child.GetValue<string>("Type");
-				if (typeStr == null || !Enum.TryParse(typeStr, true, out TelemetrySinkType sinkType))
-				{
-					throw new ConfigurationException($"Unable to parse sink type '{typeStr}'");
-				}
-
-				switch (sinkType)
-				{
-					case TelemetrySinkType.Epic:
-						EpicTelemetryConfig epic = new();
-						child.Bind(epic);
-						telemetryConfigs.Add(epic);
-						break;
-
-					case TelemetrySinkType.ClickHouse:
-						ClickHouseTelemetryConfig clickHouse = new();
-						child.Bind(clickHouse);
-						telemetryConfigs.Add(clickHouse);
-						break;
-
-					case TelemetrySinkType.Mongo:
-						MongoTelemetryConfig mongo = new();
-						child.Bind(mongo);
-						telemetryConfigs.Add(mongo);
-						break;
-
-					case TelemetrySinkType.None:
-					default:
-						break;
-				}
-			}
-
-			settings.Telemetry = telemetryConfigs;
-		}
-
 		public static void BindServerSettings(IConfiguration configuration, ServerSettings settings)
 		{
 			IConfigurationSection hordeSection = configuration.GetSection("Horde");
 			hordeSection.Bind(settings);
-			BindTelemetrySettings(settings, hordeSection.GetSection("Telemetry"));
 			settings.Validate();
 		}
 
@@ -385,6 +338,7 @@ namespace HordeServer
 
 			// Register the plugin collection
 			PluginCollection pluginCollection = new PluginCollection();
+			pluginCollection.Add<AnalyticsPlugin>(new PluginMetadata { Name = new PluginName("analytics") });
 			services.AddSingleton<IPluginCollection>(pluginCollection);
 
 			// Register all the plugin services
@@ -802,19 +756,6 @@ namespace HordeServer
 						.Build();
 				});
 
-			services.AddSingleton<TelemetryManager>();
-			services.AddSingleton<ITelemetryWriter>(sp => sp.GetRequiredService<TelemetryManager>());
-			services.AddHostedService(sp => sp.GetRequiredService<TelemetryManager>());
-			services.AddSingleton<MongoTelemetrySink>();
-			services.AddHostedService(sp => sp.GetRequiredService<MongoTelemetrySink>());
-			services.AddSingleton<MetricTelemetrySink>();
-
-			services.AddSingleton<MetricCollection>();
-			services.AddHostedService(sp => sp.GetRequiredService<MetricCollection>());
-			services.AddSingleton<IMetricCollection, MetricCollection>(sp => sp.GetRequiredService<MetricCollection>());
-
-			services.AddHttpClient(EpicTelemetrySink.HttpClientName, client => { });
-
 			// Create the agent telemetry collection, and register the hosted service so we can flush from any server.
 			services.AddSingleton<AgentTelemetryCollection>();
 			services.AddSingleton<IAgentTelemetryCollection>(sp => sp.GetRequiredService<AgentTelemetryCollection>());
diff --git a/Engine/Source/Programs/Horde/HordeServer/Telemetry/ITelemetrySink.cs b/Engine/Source/Programs/Horde/HordeServer/Telemetry/ITelemetrySink.cs
deleted file mode 100644
diff --git a/Engine/Source/Programs/Horde/HordeServer/Telemetry/Metrics/IMetric.cs b/Engine/Source/Programs/Horde/HordeServer/Telemetry/Metrics/IMetric.cs
deleted file mode 100644
diff --git a/Engine/Source/Programs/Horde/HordeServer/Telemetry/Metrics/IMetricCollection.cs b/Engine/Source/Programs/Horde/HordeServer/Telemetry/Metrics/IMetricCollection.cs
deleted file mode 100644
diff --git a/Engine/Source/Programs/Horde/HordeServer/Telemetry/Metrics/MetricCollection.cs b/Engine/Source/Programs/Horde/HordeServer/Telemetry/Metrics/MetricCollection.cs
deleted file mode 100644
diff --git a/Engine/Source/Programs/Horde/HordeServer/Telemetry/Metrics/MetricConfig.cs b/Engine/Source/Programs/Horde/HordeServer/Telemetry/Metrics/MetricConfig.cs
deleted file mode 100644
diff --git a/Engine/Source/Programs/Horde/HordeServer/Telemetry/Sinks/EpicTelemetrySink.cs b/Engine/Source/Programs/Horde/HordeServer/Telemetry/Sinks/EpicTelemetrySink.cs
deleted file mode 100644
diff --git a/Engine/Source/Programs/Horde/HordeServer/Telemetry/Sinks/MetricTelemetrySink.cs b/Engine/Source/Programs/Horde/HordeServer/Telemetry/Sinks/MetricTelemetrySink.cs
deleted file mode 100644
diff --git a/Engine/Source/Programs/Horde/HordeServer/Telemetry/Sinks/MongoTelemetrySink.cs b/Engine/Source/Programs/Horde/HordeServer/Telemetry/Sinks/MongoTelemetrySink.cs
deleted file mode 100644
diff --git a/Engine/Source/Programs/Horde/HordeServer/Telemetry/Sinks/NullTelemetrySink.cs b/Engine/Source/Programs/Horde/HordeServer/Telemetry/Sinks/NullTelemetrySink.cs
deleted file mode 100644
diff --git a/Engine/Source/Programs/Horde/HordeServer/Telemetry/TelemetryController.cs b/Engine/Source/Programs/Horde/HordeServer/Telemetry/TelemetryController.cs
deleted file mode 100644
diff --git a/Engine/Source/Programs/Horde/HordeServer/Telemetry/TelemetryEvent.cs b/Engine/Source/Programs/Horde/HordeServer/Telemetry/TelemetryEvent.cs
deleted file mode 100644
diff --git a/Engine/Source/Programs/Horde/HordeServer/Telemetry/TelemetryManager.cs b/Engine/Source/Programs/Horde/HordeServer/Telemetry/TelemetryManager.cs
deleted file mode 100644
diff --git a/Engine/Source/Programs/Horde/Plugins/Analytics/HordeServer.Analytics/AnalyticsGlobalConfig.cs b/Engine/Source/Programs/Horde/Plugins/Analytics/HordeServer.Analytics/AnalyticsGlobalConfig.cs
new file mode 100644
--- /dev/null
+++ b/Engine/Source/Programs/Horde/Plugins/Analytics/HordeServer.Analytics/AnalyticsGlobalConfig.cs
@@ -0,0 +1,59 @@
+// Copyright Epic Games, Inc. All Rights Reserved.
+
+using System;
+using System.Collections.Generic;
+using System.Diagnostics.CodeAnalysis;
+using System.Linq;
+using System.Security.Claims;
+using System.Text;
+using System.Threading.Tasks;
+using EpicGames.Horde.Acls;
+using EpicGames.Horde.Telemetry;
+using HordeServer.Acls;
+using HordeServer.Plugins;
+using HordeServer.Telemetry.Metrics;
+
+#pragma warning disable CA2227 // Change X to be read-only by removing the property setter
+
+namespace HordeServer
+{
+	/// <summary>
+	/// Config settings for analytics
+	/// </summary>
+	public class AnalyticsGlobalConfig : IPluginConfig
+	{
+		/// <summary>
+		/// Metrics to aggregate on the Horde server
+		/// </summary>
+		public List<TelemetryStoreConfig> TelemetryStores { get; set; } = new List<TelemetryStoreConfig>();
+
+		private AclConfig _parentAcl = null!;
+		private readonly Dictionary<TelemetryStoreId, TelemetryStoreConfig> _telemetryStoreLookup = new Dictionary<TelemetryStoreId, TelemetryStoreConfig>();
+
+		public void PostLoad(AclConfig parentAcl)
+		{
+			_parentAcl = parentAcl;
+
+			_telemetryStoreLookup.Clear();
+			foreach (TelemetryStoreConfig telemetryStore in TelemetryStores)
+			{
+				_telemetryStoreLookup.Add(telemetryStore.Id, telemetryStore);
+				telemetryStore.PostLoad(parentAcl);
+			}
+		}
+
+		/// <inheritdoc cref="AclConfig.Authorize(AclAction, ClaimsPrincipal)"/>
+		public bool Authorize(AclAction action, ClaimsPrincipal user)
+			=> _parentAcl.Authorize(action, user);
+
+
+		/// <summary>
+		/// Attempts to get configuration for a pool from this object
+		/// </summary>
+		/// <param name="telemetryStoreId">The pool identifier</param>
+		/// <param name="config">Configuration for the telemetry store</param>
+		/// <returns>True if the telemetry configuration was found</returns>
+		public bool TryGetTelemetryStore(TelemetryStoreId telemetryStoreId, [NotNullWhen(true)] out TelemetryStoreConfig? config) 
+			=> _telemetryStoreLookup.TryGetValue(telemetryStoreId, out config);
+	}
+}
diff --git a/Engine/Source/Programs/Horde/Plugins/Analytics/HordeServer.Analytics/AnalyticsPlugin.cs b/Engine/Source/Programs/Horde/Plugins/Analytics/HordeServer.Analytics/AnalyticsPlugin.cs
new file mode 100644
--- /dev/null
+++ b/Engine/Source/Programs/Horde/Plugins/Analytics/HordeServer.Analytics/AnalyticsPlugin.cs
@@ -0,0 +1,40 @@
+﻿// Copyright Epic Games, Inc. All Rights Reserved.
+
+using System;
+using System.Collections.Generic;
+using System.Linq;
+using System.Text;
+using System.Threading.Tasks;
+using HordeServer.Plugins;
+using HordeServer.Telemetry;
+using HordeServer.Telemetry.Metrics;
+using HordeServer.Telemetry.Sinks;
+using Microsoft.Extensions.Configuration;
+using Microsoft.Extensions.DependencyInjection;
+
+namespace HordeServer
+{
+	/// <summary>
+	/// Main entry point for the analytics plugin
+	/// </summary>
+	[Plugin("Analytics", ServerConfigType = typeof(AnalyticsServerConfig), GlobalConfigType = typeof(AnalyticsGlobalConfig))]
+	public class AnalyticsPlugin : IPluginStartup
+	{
+		/// <inheritdoc/>
+		public void ConfigureServices(IServiceCollection services)
+		{
+			services.AddSingleton<TelemetryManager>();
+			services.AddSingleton<ITelemetryWriter>(sp => sp.GetRequiredService<TelemetryManager>());
+			services.AddHostedService(sp => sp.GetRequiredService<TelemetryManager>());
+			services.AddSingleton<MongoTelemetrySink>();
+			services.AddHostedService(sp => sp.GetRequiredService<MongoTelemetrySink>());
+			services.AddSingleton<MetricTelemetrySink>();
+
+			services.AddSingleton<MetricCollection>();
+			services.AddHostedService(sp => sp.GetRequiredService<MetricCollection>());
+			services.AddSingleton<IMetricCollection, MetricCollection>(sp => sp.GetRequiredService<MetricCollection>());
+
+			services.AddHttpClient(EpicTelemetrySink.HttpClientName, client => { });
+		}
+	}
+}
diff --git a/Engine/Source/Programs/Horde/Plugins/Analytics/HordeServer.Analytics/AnalyticsServerConfig.cs b/Engine/Source/Programs/Horde/Plugins/Analytics/HordeServer.Analytics/AnalyticsServerConfig.cs
new file mode 100644
--- /dev/null
+++ b/Engine/Source/Programs/Horde/Plugins/Analytics/HordeServer.Analytics/AnalyticsServerConfig.cs
@@ -0,0 +1,84 @@
+// Copyright Epic Games, Inc. All Rights Reserved.
+
+using System;
+using System.Collections.Generic;
+using System.Linq;
+using System.Text;
+using System.Threading.Tasks;
+using HordeServer.Telemetry.Sinks;
+
+#pragma warning disable CA2227 // Change x to be read-only by removing the property setter
+
+namespace HordeServer
+{
+	/// <summary>
+	/// Server configuration for the analytics system
+	/// </summary>
+	public class AnalyticsServerConfig
+	{
+		/// <summary>
+		/// Settings for the various telemetry sinks
+		/// </summary>
+		public TelemetrySinkConfig Sinks { get; } = new TelemetrySinkConfig();
+	}
+
+	/// <summary>
+	/// Telemetry sinks
+	/// </summary>
+	public class TelemetrySinkConfig
+	{
+		/// <summary>
+		/// Settings for the Epic telemetry sink
+		/// </summary>
+		public EpicTelemetryConfig? Epic { get; set; }
+
+		/// <summary>
+		/// Settings for the MongoDB telemetry sink
+		/// </summary>
+		public MongoTelemetryConfig? Mongo { get; set; }
+	}
+
+	/// <summary>
+	/// Configuration for the telemetry sink
+	/// </summary>
+	public class BaseTelemetryConfig
+	{
+		/// <summary>
+		/// Whether to enable this sink
+		/// </summary>
+		public bool Enabled { get; set; } = false;
+	}
+
+	/// <summary>
+	/// Configuration for the telemetry sink
+	/// </summary>
+	public class EpicTelemetryConfig : BaseTelemetryConfig, IEpicTelemetrySinkConfig
+	{
+		/// <summary>
+		/// Base URL for the telemetry server
+		/// </summary>
+		public Uri? Url { get; set; }
+
+		/// <summary>
+		/// Application name to send in the event messages
+		/// </summary>
+		public string AppId { get; set; } = "Horde";
+
+		/// <inheritdoc />
+		public override string ToString()
+		{
+			return $"{nameof(Url)}={Url} {nameof(AppId)}={AppId}";
+		}
+	}
+
+	/// <summary>
+	/// Configuration for the telemetry sink
+	/// </summary>
+	public class MongoTelemetryConfig : BaseTelemetryConfig
+	{
+		/// <summary>
+		/// Number of days worth of telmetry events to keep
+		/// </summary>
+		public double RetainDays { get; set; } = 1;
+	}
+}
diff --git a/Engine/Source/Programs/Horde/Plugins/Analytics/HordeServer.Analytics/HordeServer.Analytics.csproj b/Engine/Source/Programs/Horde/Plugins/Analytics/HordeServer.Analytics/HordeServer.Analytics.csproj
new file mode 100644
--- /dev/null
+++ b/Engine/Source/Programs/Horde/Plugins/Analytics/HordeServer.Analytics/HordeServer.Analytics.csproj
@@ -0,0 +1,25 @@
+﻿<Project Sdk="Microsoft.NET.Sdk">
+
+  <PropertyGroup>
+    <TargetFramework>net8.0</TargetFramework>
+    <ImplicitUsings>enable</ImplicitUsings>
+    <Nullable>enable</Nullable>
+    <RootNamespace>HordeServer</RootNamespace>
+  </PropertyGroup>
+
+  <ItemGroup>
+    <ProjectReference Include="..\..\..\..\Shared\EpicGames.Core\EpicGames.Core.csproj" />
+    <ProjectReference Include="..\..\..\..\Shared\EpicGames.Horde\EpicGames.Horde.csproj" />
+    <ProjectReference Include="..\..\..\HordeServer.Shared\HordeServer.Shared.csproj" />
+  </ItemGroup>
+
+  <ItemGroup>
+    <PackageReference Include="JsonPath.Net" Version="0.6.7" />
+    <PackageReference Include="Meziantou.Analyzer" Version="2.0.139">
+      <PrivateAssets>all</PrivateAssets>
+      <IncludeAssets>runtime; build; native; contentfiles; analyzers; buildtransitive</IncludeAssets>
+    </PackageReference>
+    <PackageReference Include="T-Digest.NET" Version="1.0.3" />
+  </ItemGroup>
+
+</Project>
diff --git a/Engine/Source/Programs/Horde/Plugins/Analytics/HordeServer.Analytics/Telemetry/ITelemetrySink.cs b/Engine/Source/Programs/Horde/Plugins/Analytics/HordeServer.Analytics/Telemetry/ITelemetrySink.cs
new file mode 100644
--- /dev/null
+++ b/Engine/Source/Programs/Horde/Plugins/Analytics/HordeServer.Analytics/Telemetry/ITelemetrySink.cs
@@ -0,0 +1,35 @@
+// Copyright Epic Games, Inc. All Rights Reserved.
+
+using System;
+using System.Threading;
+using System.Threading.Tasks;
+using EpicGames.Horde.Telemetry;
+
+namespace HordeServer.Telemetry
+{
+	/// <summary>
+	/// Interface for a telemetry sink
+	/// </summary>
+	public interface ITelemetrySink : IAsyncDisposable
+	{
+		/// <summary>
+		/// Whether the sink is enabled
+		/// </summary>
+		bool Enabled { get; }
+
+		/// <summary>
+		/// Flush any queued metrics to underlying service
+		/// </summary>
+		/// <param name="cancellationToken">Cancellation token for the operation</param>
+		/// <returns>Completion task</returns>
+		ValueTask FlushAsync(CancellationToken cancellationToken);
+
+		/// <summary>
+		/// Sends a telemetry event with the given information
+		/// </summary>
+		/// <param name="storeId">Identifier for the telemetry store</param>
+		/// <param name="telemetryEvent">The telemetry event that was received</param>
+		void SendEvent(TelemetryStoreId storeId, TelemetryEvent telemetryEvent);
+	}
+}
+
diff --git a/Engine/Source/Programs/Horde/Plugins/Analytics/HordeServer.Analytics/Telemetry/Metrics/IMetric.cs b/Engine/Source/Programs/Horde/Plugins/Analytics/HordeServer.Analytics/Telemetry/Metrics/IMetric.cs
new file mode 100644
--- /dev/null
+++ b/Engine/Source/Programs/Horde/Plugins/Analytics/HordeServer.Analytics/Telemetry/Metrics/IMetric.cs
@@ -0,0 +1,38 @@
+// Copyright Epic Games, Inc. All Rights Reserved.
+
+using System;
+using EpicGames.Horde.Telemetry.Metrics;
+
+namespace HordeServer.Telemetry.Metrics
+{
+	/// <summary>
+	/// Interface for a metric event
+	/// </summary>
+	public interface IMetric
+	{
+		/// <summary>
+		/// Unique identifier for this event
+		/// </summary>
+		MetricId MetricId { get; }
+
+		/// <summary>
+		/// Start time for this aggregation interval
+		/// </summary>
+		DateTime Time { get; }
+
+		/// <summary>
+		/// Grouping key for this metric
+		/// </summary>
+		string? Group { get; }
+
+		/// <summary>
+		/// Aggregated value for the metric
+		/// </summary>
+		double Value { get; }
+
+		/// <summary>
+		/// Number of samples
+		/// </summary>
+		int Count { get; }
+	}
+}
diff --git a/Engine/Source/Programs/Horde/Plugins/Analytics/HordeServer.Analytics/Telemetry/Metrics/IMetricCollection.cs b/Engine/Source/Programs/Horde/Plugins/Analytics/HordeServer.Analytics/Telemetry/Metrics/IMetricCollection.cs
new file mode 100644
--- /dev/null
+++ b/Engine/Source/Programs/Horde/Plugins/Analytics/HordeServer.Analytics/Telemetry/Metrics/IMetricCollection.cs
@@ -0,0 +1,43 @@
+// Copyright Epic Games, Inc. All Rights Reserved.
+
+using System;
+using System.Collections.Generic;
+using System.Text.Json.Nodes;
+using System.Threading;
+using System.Threading.Tasks;
+using EpicGames.Horde.Telemetry;
+using EpicGames.Horde.Telemetry.Metrics;
+
+namespace HordeServer.Telemetry.Metrics
+{
+	/// <summary>
+	/// Collection of aggregated metrics
+	/// </summary>
+	public interface IMetricCollection
+	{
+		/// <summary>
+		/// Adds a new event to the collection
+		/// </summary>
+		/// <param name="storeId">Identifier for a telemetry store</param>
+		/// <param name="node">Arbitrary event node</param>
+		void AddEvent(TelemetryStoreId storeId, JsonNode node);
+
+		/// <summary>
+		/// Finds metrics over a given time period
+		/// </summary>
+		/// <param name="storeId">Identifier for a telemetry store</param>
+		/// <param name="metricIds">Metrics to search for</param>
+		/// <param name="minTime">Start of the time period to query</param>
+		/// <param name="maxTime">End of the time period to query</param>
+		/// <param name="group">Grouping key to filter results</param>
+		/// <param name="maxResults">Maximum number of results to return</param>
+		/// <param name="cancellationToken">Cancellation token for the operation</param>
+		Task<List<IMetric>> FindAsync(TelemetryStoreId storeId, MetricId[] metricIds, DateTime? minTime = null, DateTime? maxTime = null, string? group = null, int maxResults = 50, CancellationToken cancellationToken = default);
+
+		/// <summary>
+		/// Flush all the pending events to storage
+		/// </summary>
+		/// <param name="cancellationToken">Cancellation token for the operation</param>
+		Task FlushAsync(CancellationToken cancellationToken);
+	}
+}
diff --git a/Engine/Source/Programs/Horde/Plugins/Analytics/HordeServer.Analytics/Telemetry/Metrics/MetricCollection.cs b/Engine/Source/Programs/Horde/Plugins/Analytics/HordeServer.Analytics/Telemetry/Metrics/MetricCollection.cs
new file mode 100644
--- /dev/null
+++ b/Engine/Source/Programs/Horde/Plugins/Analytics/HordeServer.Analytics/Telemetry/Metrics/MetricCollection.cs
@@ -0,0 +1,362 @@
+// Copyright Epic Games, Inc. All Rights Reserved.
+
+using System;
+using System.Collections.Generic;
+using System.Linq;
+using System.Text.Json.Nodes;
+using System.Threading;
+using System.Threading.Tasks;
+using EpicGames.Core;
+using EpicGames.Horde.Telemetry;
+using EpicGames.Horde.Telemetry.Metrics;
+using HordeServer.Server;
+using HordeServer.Utilities;
+using Json.Path;
+using Microsoft.Extensions.Hosting;
+using Microsoft.Extensions.Logging;
+using Microsoft.Extensions.Options;
+using MongoDB.Bson;
+using MongoDB.Bson.Serialization.Attributes;
+using MongoDB.Driver;
+using TDigestNet;
+
+namespace HordeServer.Telemetry.Metrics
+{
+	class MetricCollection : IMetricCollection, IHostedService, IAsyncDisposable
+	{
+		class MetricDocument : IMetric
+		{
+			public ObjectId Id { get; set; }
+
+			[BsonElement("ts")]
+			public TelemetryStoreId TelemetryStoreId { get; set; }
+
+			[BsonElement("met")]
+			public MetricId MetricId { get; set; }
+
+			[BsonElement("grp")]
+			public string Group { get; set; } = String.Empty;
+
+			[BsonElement("time")]
+			public DateTime Time { get; set; }
+
+			[BsonElement("value")]
+			public double Value { get; set; }
+
+			[BsonElement("count")]
+			public int Count { get; set; }
+
+			[BsonElement("state")]
+			public byte[]? State { get; set; }
+		}
+
+		record class SampleKey(TelemetryStoreId Store, MetricId Metric, string Group, DateTime Time);
+
+		readonly object _lockObject = new object();
+		readonly IMongoCollection<MetricDocument> _metrics;
+		readonly AsyncEvent _newDataEvent = new AsyncEvent();
+		readonly AsyncEvent _flushEvent = new AsyncEvent();
+		readonly BackgroundTask _flushTask;
+		readonly IClock _clock;
+		readonly IOptionsMonitor<AnalyticsGlobalConfig> _globalConfig;
+		readonly ILogger _logger;
+
+		Dictionary<SampleKey, List<double>> _queuedSamples = new Dictionary<SampleKey, List<double>>();
+
+		public MetricCollection(IMongoService mongoService, IClock clock, IOptionsMonitor<AnalyticsGlobalConfig> globalConfig, ILogger<MetricCollection> logger)
+		{
+			List<MongoIndex<MetricDocument>> indexes = new List<MongoIndex<MetricDocument>>();
+			indexes.Add(MongoIndex.Create<MetricDocument>(keys => keys.Ascending(x => x.TelemetryStoreId).Descending(x => x.Time).Ascending(x => x.MetricId).Ascending(x => x.Group)));
+			_metrics = mongoService.GetCollection<MetricDocument>("Metrics", indexes);
+
+			_flushTask = new BackgroundTask(BackgroundTickAsync);
+			_clock = clock;
+			_globalConfig = globalConfig;
+			_logger = logger;
+		}
+
+		/// <inheritdoc/>
+		public Task StartAsync(CancellationToken cancellationToken)
+		{
+			_flushTask.Start();
+			return Task.CompletedTask;
+		}
+
+		/// <inheritdoc/>
+		public async Task StopAsync(CancellationToken cancellationToken)
+		{
+			await _flushTask.StopAsync(cancellationToken);
+		}
+
+		/// <inheritdoc/>
+		public async ValueTask DisposeAsync()
+		{
+			await _flushTask.DisposeAsync();
+		}
+
+		/// <inheritdoc/>
+		public void AddEvent(TelemetryStoreId storeId, JsonNode node)
+		{
+			JsonArray array = new JsonArray { node };
+
+			TelemetryStoreConfig? telemetryStoreConfig;
+			if (_globalConfig.CurrentValue.TryGetTelemetryStore(storeId, out telemetryStoreConfig))
+			{
+				foreach (MetricConfig metric in telemetryStoreConfig.Metrics)
+				{
+					AddEvent(storeId, metric, node, array);
+				}
+			}
+
+			array.Remove(node);
+		}
+
+		void AddEvent(TelemetryStoreId telemetryStoreId, MetricConfig metric, JsonNode node, JsonArray array)
+		{
+			if (metric.Filter != null)
+			{
+				PathResult filterResult = metric.Filter.Evaluate(array);
+				if (filterResult.Error != null)
+				{
+					_logger.LogWarning("Error evaluating filter for metric {MetricId}: {Message}", metric.Id, filterResult.Error);
+					return;
+				}
+				if (filterResult.Matches == null || filterResult.Matches.Count == 0)
+				{
+					return;
+				}
+			}
+
+			List<double> values = new List<double>();
+
+			if (metric.Property != null)
+			{
+				PathResult result = metric.Property.Evaluate(node);
+				if (result.Error != null)
+				{
+					_logger.LogWarning("Error evaluating filter for metric {MetricId}: {Message}", metric.Id, result.Error);
+					return;
+				}
+
+				if (result.Matches != null && result.Matches.Count > 0)
+				{
+
+					foreach (Json.Path.Node match in result.Matches)
+					{
+						JsonValue? value = match.Value as JsonValue;
+						if (value != null && value.TryGetValue(out double doubleValue))
+						{
+							values.Add(doubleValue);
+						}
+					}
+				}
+			}
+			else
+			{
+				if (metric.Function != AggregationFunction.Count)
+				{
+					_logger.LogWarning("Missing property parameter for metric {MetricId}", metric.Id);
+					return;
+				}
+
+				values.Add(1);
+			}
+
+			if (values.Count > 0)
+			{
+				List<string> groupKeys = new List<string>();
+				foreach (JsonPath groupByPath in metric.GroupByPaths)
+				{
+					PathResult groupResult = groupByPath.Evaluate(node);
+
+					string groupKey;
+					if (groupResult.Error != null || groupResult.Matches == null || groupResult.Matches.Count == 0)
+					{
+						groupKey = "";
+					}
+					else
+					{
+						groupKey = EscapeCsv(groupResult.Matches.Select(x => x.Value?.ToString() ?? String.Empty));
+					}
+
+					groupKeys.Add(groupKey);
+				}
+
+				string group = String.Join(",", groupKeys);
+
+				DateTime utcNow = _clock.UtcNow;
+				DateTime sampleTime = new DateTime(utcNow.Ticks - (utcNow.Ticks % metric.Interval.Ticks), DateTimeKind.Utc);
+
+				SampleKey key = new SampleKey(telemetryStoreId, metric.Id, group.ToString(), sampleTime);
+				lock (_lockObject)
+				{
+					QueueSampleValues(key, values);
+				}
+
+				_newDataEvent.Set();
+			}
+		}
+
+		static string EscapeCsv(IEnumerable<string> items)
+		{
+			return String.Join(",", items.Select(x => EscapeCsv(x)));
+		}
+
+		static readonly char[] s_csvEscapeChars = { ',', '\n', '\"' };
+
+		static string? EscapeCsv(string text)
+		{
+			if (text.IndexOfAny(s_csvEscapeChars) != -1)
+			{
+				text = text.Replace("\"", "\"\"", StringComparison.Ordinal);
+				text = $"\"{text}\"";
+			}
+			return text;
+		}
+
+		void QueueSampleValues(SampleKey key, List<double> values)
+		{
+			List<double>? queuedValues;
+			if (!_queuedSamples.TryGetValue(key, out queuedValues))
+			{
+				queuedValues = new List<double>();
+				_queuedSamples.Add(key, queuedValues);
+			}
+			queuedValues.AddRange(values);
+		}
+
+		async Task BackgroundTickAsync(CancellationToken cancellationToken)
+		{
+			Task newDataTask = _newDataEvent.Task;
+			Task flushTask = _flushEvent.Task;
+
+			while (!cancellationToken.IsCancellationRequested)
+			{
+				await newDataTask.WaitAsync(cancellationToken);
+				await Task.WhenAny(flushTask, Task.Delay(TimeSpan.FromSeconds(5.0), cancellationToken));
+
+				newDataTask = _newDataEvent.Task;
+				flushTask = _flushEvent.Task;
+
+				await FlushAsync(cancellationToken);
+			}
+		}
+
+		public async Task FlushAsync(CancellationToken cancellationToken)
+		{
+			// Copy the current sample buffer
+			Dictionary<SampleKey, List<double>> samples;
+			lock (_lockObject)
+			{
+				samples = _queuedSamples;
+				_queuedSamples = new Dictionary<SampleKey, List<double>>();
+			}
+
+			// Add the samples to the database
+			try
+			{
+				foreach ((SampleKey sampleKey, List<double> sampleValues) in samples)
+				{
+					MetricConfig? metricConfig;
+					if (_globalConfig.CurrentValue.TryGetTelemetryStore(sampleKey.Store, out TelemetryStoreConfig? telemetryStoreConfig) && telemetryStoreConfig.TryGetMetric(sampleKey.Metric, out metricConfig))
+					{
+						await CombineValuesAsync(sampleKey.Store, metricConfig, sampleKey.Group, sampleKey.Time, sampleValues, cancellationToken);
+					}
+				}
+			}
+			catch
+			{
+				lock (_lockObject)
+				{
+					foreach ((SampleKey key, List<double> values) in samples)
+					{
+						QueueSampleValues(key, values);
+					}
+				}
+				throw;
+			}
+		}
+
+		async Task CombineValuesAsync(TelemetryStoreId telemetryStoreId, MetricConfig metricConfig, string group, DateTime time, List<double> values, CancellationToken cancellationToken)
+		{
+			for (; ; )
+			{
+				// Find or add the current metric document
+				FilterDefinition<MetricDocument> filter = Builders<MetricDocument>.Filter.Expr(x => x.TelemetryStoreId == telemetryStoreId && x.MetricId == metricConfig.Id && x.Group == group && x.Time == time);
+				UpdateDefinition<MetricDocument> update = Builders<MetricDocument>.Update.SetOnInsert(x => x.Count, 0);
+				MetricDocument metric = await _metrics.FindOneAndUpdateAsync(filter, update, new FindOneAndUpdateOptions<MetricDocument, MetricDocument> { IsUpsert = true, ReturnDocument = ReturnDocument.After }, cancellationToken);
+
+				// Combine the samples
+				TDigest digest = (metric.Count == 0) ? new TDigest() : TDigest.Deserialize(metric.State);
+				foreach (double value in values)
+				{
+					digest.Add(value);
+				}
+				metric.State = digest.Serialize();
+
+				// Save the previous count of samples to sequence updates to the document
+				int prevCount = metric.Count;
+				metric.Count += values.Count;
+
+				// Update the current value
+				switch (metricConfig.Function)
+				{
+					case AggregationFunction.Count:
+						metric.Value = metric.Count;
+						break;
+					case AggregationFunction.Min:
+						metric.Value = digest.Min;
+						break;
+					case AggregationFunction.Max:
+						metric.Value = digest.Max;
+						break;
+					case AggregationFunction.Sum:
+						metric.Value += values.Sum();
+						break;
+					case AggregationFunction.Average:
+						metric.Value = digest.Average;
+						break;
+					case AggregationFunction.Percentile:
+						metric.Value = digest.Quantile(metricConfig.Percentile / 100.0);
+						break;
+					default:
+						_logger.LogWarning("Unhandled aggregation function '{Function}'", metricConfig.Function);
+						break;
+				}
+
+				// Update the document
+				ReplaceOneResult result = await _metrics.ReplaceOneAsync(x => x.TelemetryStoreId == telemetryStoreId && x.Id == metric.Id && x.Count == prevCount, metric, cancellationToken: cancellationToken);
+				if (result.MatchedCount > 0)
+				{
+					break;
+				}
+			}
+		}
+
+		/// <inheritdoc/>
+		public async Task<List<IMetric>> FindAsync(TelemetryStoreId telemetryStoreId, MetricId[] metricIds, DateTime? minTime = null, DateTime? maxTime = null, string? group = null, int maxResults = 50, CancellationToken cancellationToken = default)
+		{
+			FilterDefinition<MetricDocument> filter = FilterDefinition<MetricDocument>.Empty;
+			filter &= Builders<MetricDocument>.Filter.Eq(x => x.TelemetryStoreId, telemetryStoreId);
+
+			if (minTime != null)
+			{
+				filter &= Builders<MetricDocument>.Filter.Gte(x => x.Time, minTime.Value);
+			}
+			if (maxTime != null)
+			{
+				filter &= Builders<MetricDocument>.Filter.Lte(x => x.Time, maxTime.Value);
+			}
+
+			filter &= Builders<MetricDocument>.Filter.In(x => x.MetricId, metricIds);
+
+			if (group != null)
+			{
+				filter &= Builders<MetricDocument>.Filter.Eq(x => x.Group, group);
+			}
+
+			List<MetricDocument> results = await _metrics.Find(filter).SortByDescending(x => x.Time).Limit(maxResults).ToListAsync(cancellationToken);
+			return results.ConvertAll<MetricDocument, IMetric>(x => x);
+		}
+	}
+}
diff --git a/Engine/Source/Programs/Horde/Plugins/Analytics/HordeServer.Analytics/Telemetry/Metrics/MetricConfig.cs b/Engine/Source/Programs/Horde/Plugins/Analytics/HordeServer.Analytics/Telemetry/Metrics/MetricConfig.cs
new file mode 100644
--- /dev/null
+++ b/Engine/Source/Programs/Horde/Plugins/Analytics/HordeServer.Analytics/Telemetry/Metrics/MetricConfig.cs
@@ -0,0 +1,185 @@
+// Copyright Epic Games, Inc. All Rights Reserved.
+
+using System;
+using System.Collections.Generic;
+using System.Diagnostics.CodeAnalysis;
+using System.Linq;
+using System.Text.Json.Serialization;
+using EpicGames.Core;
+using EpicGames.Horde.Telemetry;
+using EpicGames.Horde.Telemetry.Metrics;
+using HordeServer.Acls;
+using HordeServer.Configuration;
+using HordeServer.Utilities;
+using Json.Path;
+
+#pragma warning disable CA2227 // Change x to be read-only by removing the property setter
+
+namespace HordeServer.Telemetry.Metrics
+{
+	/// <summary>
+	/// Method for aggregating samples into a metric
+	/// </summary>
+	public enum AggregationFunction
+	{
+		/// <summary>
+		/// Count the number of matching elements
+		/// </summary>
+		Count,
+
+		/// <summary>
+		/// Take the minimum value of all samples
+		/// </summary>
+		Min,
+
+		/// <summary>
+		/// Take the maximum value of all samples
+		/// </summary>
+		Max,
+
+		/// <summary>
+		/// Sum all the reported values
+		/// </summary>
+		Sum,
+
+		/// <summary>
+		/// Average all the samples
+		/// </summary>
+		Average,
+
+		/// <summary>
+		/// Estimates the value at a certain percentile
+		/// </summary>
+		Percentile,
+	}
+
+	/// <summary>
+	/// Configures a metric to aggregate on the server
+	/// </summary>
+	public class MetricConfig
+	{
+		/// <summary>
+		/// Identifier for this metric
+		/// </summary>
+		public MetricId Id { get; set; }
+
+		/// <summary>
+		/// Filter expression to evaluate to determine which events to include. This query is evaluated against an array.
+		/// </summary>
+		[JsonSchemaString]
+		public JsonPath? Filter { get; set; }
+
+		/// <summary>
+		/// Property to aggregate
+		/// </summary>
+		[JsonSchemaString]
+		public JsonPath? Property { get; set; }
+
+		/// <summary>
+		/// Property to group by. Specified as a comma-separated list of JSON path expressions.
+		/// </summary>
+		public string GroupBy
+		{
+			get => _groupBy;
+			set
+			{
+				_groupBy = value;
+
+				GroupByPaths.Clear();
+				if (!String.IsNullOrWhiteSpace(_groupBy))
+				{
+					List<string> fields = _groupBy.Split(',').Select(x => x.Trim()).Where(x => x.Length > 0).ToList();
+					if (fields.Count > 0)
+					{
+						GroupByPaths.AddRange(fields.Select(x => JsonPath.Parse(x)));
+					}
+				}
+			}
+		}
+
+		[JsonIgnore]
+		string _groupBy = String.Empty;
+
+		/// <summary>
+		/// Accessor for the <see cref="GroupBy"/> field, parsed as a list of json paths
+		/// </summary>
+		[JsonIgnore]
+		public List<JsonPath> GroupByPaths { get; } = new List<JsonPath>();
+
+		/// <summary>
+		/// How to aggregate samples for this metric
+		/// </summary>
+		public AggregationFunction Function { get; set; }
+
+		/// <summary>
+		/// For the percentile function, specifies the percentile to measure
+		/// </summary>
+		public int Percentile { get; set; } = 95;
+
+		/// <summary>
+		/// Interval for each metric. Supports times such as "2d", "1h", "1h30m", "20s".
+		/// </summary>
+		[JsonConverter(typeof(IntervalJsonConverter))]
+		public TimeSpan Interval { get; set; } = TimeSpan.FromHours(1.0);
+	}
+
+	/// <summary>
+	/// Config for metrics
+	/// </summary>
+	[JsonSchema("https://unrealengine.com/horde/telemetry")]
+	[JsonSchemaCatalog("Horde Telemetry", "Horde telemetry configuration file", new[] { "*.telemetry.json", "*.metrics.json", "Metrics/*.json" })]
+	[ConfigIncludeRoot]
+	[ConfigMacroScope]
+	public class TelemetryStoreConfig
+	{
+		/// <summary>
+		/// Identifier for this store
+		/// </summary>
+		public TelemetryStoreId Id { get; set; }
+
+		/// <summary>
+		/// Permissions for this store
+		/// </summary>
+		public AclConfig Acl { get; set; } = new AclConfig();
+
+		/// <summary>
+		/// Metrics to aggregate on the Horde server
+		/// </summary>
+		public List<MetricConfig> Metrics { get; set; } = new List<MetricConfig>();
+
+		/// <summary>
+		/// Includes for other configuration files
+		/// </summary>
+		public List<ConfigInclude> Include { get; set; } = new List<ConfigInclude>();
+
+		readonly Dictionary<MetricId, MetricConfig> _metricLookup = new Dictionary<MetricId, MetricConfig>();
+
+		/// <summary>
+		/// Macros within this configuration
+		/// </summary>
+		public List<ConfigMacro> Macros { get; set; } = new List<ConfigMacro>();
+
+		/// <summary>
+		/// Called after the store has been deserialized to compute cached values
+		/// </summary>
+		public void PostLoad(AclConfig parentAcl)
+		{
+			Acl.PostLoad(parentAcl, $"telemetry-store:{Id}");
+
+			_metricLookup.Clear();
+			foreach (MetricConfig metric in Metrics)
+			{
+				_metricLookup.Add(metric.Id, metric);
+			}
+		}
+
+		/// <summary>
+		/// Attempt to get config for a metric with the given id
+		/// </summary>
+		/// <param name="metricId">Metric id</param>
+		/// <param name="metricConfig">Receives the config object on success</param>
+		/// <returns>True if the metric was found</returns>
+		public bool TryGetMetric(MetricId metricId, [NotNullWhen(true)] out MetricConfig? metricConfig)
+			=> _metricLookup.TryGetValue(metricId, out metricConfig);
+	}
+}
diff --git a/Engine/Source/Programs/Horde/Plugins/Analytics/HordeServer.Analytics/Telemetry/Sinks/EpicTelemetrySink.cs b/Engine/Source/Programs/Horde/Plugins/Analytics/HordeServer.Analytics/Telemetry/Sinks/EpicTelemetrySink.cs
new file mode 100644
--- /dev/null
+++ b/Engine/Source/Programs/Horde/Plugins/Analytics/HordeServer.Analytics/Telemetry/Sinks/EpicTelemetrySink.cs
@@ -0,0 +1,278 @@
+// Copyright Epic Games, Inc. All Rights Reserved.
+
+using System;
+using System.Collections.Generic;
+using System.Diagnostics;
+using System.Net.Http;
+using System.Net.Http.Headers;
+using System.Text.Json;
+using System.Text.Json.Serialization;
+using System.Threading;
+using System.Threading.Tasks;
+using EpicGames.Core;
+using EpicGames.Horde.Telemetry;
+using HordeServer.Utilities;
+using Microsoft.AspNetCore.WebUtilities;
+using Microsoft.Extensions.Http;
+using Microsoft.Extensions.Logging;
+using Polly;
+using Polly.Extensions.Http;
+
+namespace HordeServer.Telemetry.Sinks
+{
+	/// <summary>
+	/// Options for the Epic telemetry sink
+	/// </summary>
+	public interface IEpicTelemetrySinkConfig
+	{
+		/// <summary>
+		/// Base URL for sending events
+		/// </summary>
+		public Uri? Url { get; }
+
+		/// <summary>
+		/// The application id
+		/// </summary>
+		public string AppId { get; }
+	}
+
+	/// <summary>
+	/// Epic internal telemetry sink using the data router
+	/// </summary>
+	public sealed class EpicTelemetrySink : ITelemetrySink
+	{
+		// Converter for datetime formats that Tableau can ingest
+		class TableauDateTimeConverter : JsonConverter<DateTime>
+		{
+			const string Format = "yyyy-MM-dd HH:mm:ss";
+
+			public override void Write(Utf8JsonWriter writer, DateTime date, JsonSerializerOptions options)
+			{
+				writer.WriteStringValue(date.ToString(Format));
+			}
+
+			public override DateTime Read(ref Utf8JsonReader reader, Type typeToConvert, JsonSerializerOptions options)
+			{
+				return DateTime.ParseExact(reader.GetString() ?? String.Empty, Format, null);
+			}
+		}
+
+		/// <summary>
+		/// Name of the HTTP client for writing telemetry data
+		/// </summary>
+		public const string HttpClientName = "EpicTelemetrySink";
+
+		static readonly Utf8String s_packetPrefix = new Utf8String("{\"Events\":[");
+		static readonly Utf8String s_packetSuffix = new Utf8String("]}");
+
+		readonly object _lockObject = new object();
+
+		readonly HttpClient _httpClient;
+		readonly Uri? _baseUrl;
+		readonly JsonSerializerOptions _jsonOptions;
+
+		class Writer : IDisposable
+		{
+			readonly Uri _uri;
+			readonly JsonSerializerOptions _jsonOptions;
+			readonly ArrayMemoryWriter _packetMemoryWriter;
+			readonly Utf8JsonWriter _packetWriter;
+
+			public Stopwatch Timer { get; } = Stopwatch.StartNew();
+			public Uri Uri => _uri;
+			public bool HasData => _packetMemoryWriter.WrittenMemory.Length > 0;
+
+			public Writer(Uri uri, JsonSerializerOptions jsonOptions)
+			{
+				_uri = uri;
+				_jsonOptions = jsonOptions;
+				_packetMemoryWriter = new ArrayMemoryWriter(65536);
+				_packetWriter = new Utf8JsonWriter(_packetMemoryWriter);
+			}
+
+			public void Dispose()
+			{
+				_packetWriter.Dispose();
+			}
+
+			public void AddEvent(object payload)
+			{
+				// Restart the timer
+				Timer.Restart();
+
+				// If this is the first event written, write the outer events array
+				if (_packetMemoryWriter.Length == 0)
+				{
+					_packetMemoryWriter.WriteFixedLengthBytes(s_packetPrefix.Span);
+				}
+				else
+				{
+					_packetMemoryWriter.WriteUInt8((byte)',');
+				}
+
+				// Serialize the event data
+				_packetWriter.Reset();
+				JsonSerializer.Serialize(_packetWriter, payload, _jsonOptions);
+				_packetWriter.Flush();
+			}
+
+			public byte[] Flush()
+			{
+				_packetMemoryWriter.WriteFixedLengthBytes(s_packetSuffix.Span);
+				byte[] packet = _packetMemoryWriter.WrittenMemory.ToArray();
+
+				_packetMemoryWriter.Clear();
+				_packetWriter.Reset(_packetMemoryWriter);
+				return packet;
+			}
+		}
+
+		readonly IServerInfo _serverInfo;
+		readonly Dictionary<TelemetryRecordMeta, Writer> _writers = new Dictionary<TelemetryRecordMeta, Writer>();
+		readonly ILogger _logger;
+
+		/// <inheritdoc/>
+		public bool Enabled => _baseUrl != null;
+
+		/// <summary>
+		/// Constructor
+		/// </summary>
+		public EpicTelemetrySink(IEpicTelemetrySinkConfig config, IServerInfo serverInfo, ILogger<EpicTelemetrySink> logger)
+		{
+			_serverInfo = serverInfo;
+			_httpClient = CreateHttpClient();
+
+			_jsonOptions = new JsonSerializerOptions();
+			JsonUtils.ConfigureJsonSerializer(_jsonOptions);
+			_jsonOptions.PropertyNamingPolicy = null;
+			_jsonOptions.Converters.Insert(0, new TableauDateTimeConverter());
+
+			_logger = logger;
+			_baseUrl = config.Url;
+		}
+
+		static HttpClient CreateHttpClient()
+		{
+			HttpMessageHandler handler = new SocketsHttpHandler();
+
+#pragma warning disable CA2000
+			TimeSpan[] retryTimes = new[] { TimeSpan.FromSeconds(1.0), TimeSpan.FromSeconds(5.0), TimeSpan.FromSeconds(10.0) };
+			IAsyncPolicy<HttpResponseMessage> handleTransientErrorPolicy = Policy<HttpResponseMessage>.Handle<HttpRequestException>().OrTransientHttpStatusCode().WaitAndRetryAsync(retryTimes);
+			handler = new PolicyHttpMessageHandler(handleTransientErrorPolicy) { InnerHandler = handler };
+#pragma warning restore CA2000
+
+			return new HttpClient(handler, true);
+		}
+
+		/// <inheritdoc/>
+		public ValueTask DisposeAsync()
+		{
+			foreach (Writer writer in _writers.Values)
+			{
+				writer.Dispose();
+			}
+			_httpClient.Dispose();
+			return new ValueTask();
+		}
+
+		/// <inheritdoc/>
+		public async ValueTask FlushAsync(CancellationToken cancellationToken)
+		{
+			// Make sure the settings are valid
+			if (_baseUrl == null)
+			{
+				return;
+			}
+
+			// Generate the content for this event
+			List<(Uri, byte[])> packets = new List<(Uri, byte[])>();
+			lock (_lockObject)
+			{
+				// Get all the data to write
+				List<TelemetryRecordMeta> removeKeys = new List<TelemetryRecordMeta>();
+				foreach ((TelemetryRecordMeta key, Writer writer) in _writers)
+				{
+					if (writer.HasData)
+					{
+						packets.Add((writer.Uri, writer.Flush()));
+					}
+					else if (writer.Timer.Elapsed > TimeSpan.FromSeconds(30.0))
+					{
+						removeKeys.Add(key);
+					}
+				}
+
+				// Remove any writers that haven't been written to in 30s
+				foreach (TelemetryRecordMeta removeKey in removeKeys)
+				{
+					if (_writers.Remove(removeKey, out Writer? writer))
+					{
+						writer.Dispose();
+					}
+				}
+			}
+
+			// Post the event data
+			foreach ((Uri uri, byte[] packet) in packets)
+			{
+				using (HttpRequestMessage request = new HttpRequestMessage())
+				{
+					request.RequestUri = uri;
+					request.Method = HttpMethod.Post;
+					request.Headers.UserAgent.Add(new ProductInfoHeaderValue("Horde", _serverInfo.Version.ToString()));
+					request.Content = new ByteArrayContent(packet);
+					request.Content.Headers.ContentType = new MediaTypeHeaderValue("application/json");
+
+					using (HttpResponseMessage response = await _httpClient.SendAsync(request, cancellationToken))
+					{
+						if (response.IsSuccessStatusCode)
+						{
+							_logger.LogDebug("Sending {Size} bytes of telemetry data to {Url}", packet.Length, uri);
+						}
+						else
+						{
+							string content = await response.Content.ReadAsStringAsync(cancellationToken);
+							_logger.LogError("Unable to send telemetry data to server ({Code}): {Message}", response.StatusCode, content);
+						}
+					}
+				}
+			}
+		}
+
+		/// <inheritdoc/>
+		public void SendEvent(TelemetryStoreId telemetryStoreId, TelemetryEvent telemetryEvent)
+		{
+			if (_baseUrl == null)
+			{
+				return;
+			}
+
+			TelemetryRecordMeta recordMeta = telemetryEvent.RecordMeta;
+			if (recordMeta.AppId == null || recordMeta.AppVersion == null || recordMeta.AppEnvironment == null)
+			{
+				_logger.LogDebug("Unable to send telemetry event to Epic data router; missing required fields in record metadata. {@Event}", telemetryEvent);
+				return;
+			}
+
+			lock (_lockObject)
+			{
+				Writer? writer;
+				if (!_writers.TryGetValue(recordMeta, out writer))
+				{
+					Dictionary<string, string?> queryParams = new Dictionary<string, string?>();
+					queryParams.Add("AppID", telemetryEvent.RecordMeta.AppId);
+					queryParams.Add("AppVersion", telemetryEvent.RecordMeta.AppVersion);
+					queryParams.Add("AppEnvironment", telemetryEvent.RecordMeta.AppEnvironment);
+					queryParams.Add("SessionID", telemetryEvent.RecordMeta.SessionId);
+					queryParams.Add("UploadType", "eteventstream");
+
+					Uri uri = new Uri(QueryHelpers.AddQueryString(_baseUrl.ToString(), queryParams));
+					writer = new Writer(uri, _jsonOptions);
+
+					_writers.Add(recordMeta, writer);
+				}
+				writer.AddEvent(telemetryEvent.Payload);
+			}
+		}
+	}
+}
diff --git a/Engine/Source/Programs/Horde/Plugins/Analytics/HordeServer.Analytics/Telemetry/Sinks/MetricTelemetrySink.cs b/Engine/Source/Programs/Horde/Plugins/Analytics/HordeServer.Analytics/Telemetry/Sinks/MetricTelemetrySink.cs
new file mode 100644
--- /dev/null
+++ b/Engine/Source/Programs/Horde/Plugins/Analytics/HordeServer.Analytics/Telemetry/Sinks/MetricTelemetrySink.cs
@@ -0,0 +1,48 @@
+// Copyright Epic Games, Inc. All Rights Reserved.
+
+using System.Text.Json;
+using System.Text.Json.Nodes;
+using System.Threading;
+using System.Threading.Tasks;
+using EpicGames.Horde.Telemetry;
+using HordeServer.Telemetry.Metrics;
+using HordeServer.Utilities;
+
+namespace HordeServer.Telemetry.Sinks
+{
+	/// <summary>
+	/// Consumes telemetry events and generates metrics
+	/// </summary>
+	class MetricTelemetrySink : ITelemetrySink
+	{
+		readonly IMetricCollection _metricCollection;
+		readonly JsonSerializerOptions _jsonOptions;
+
+		public MetricTelemetrySink(IMetricCollection metricCollection)
+		{
+			_metricCollection = metricCollection;
+
+			_jsonOptions = new JsonSerializerOptions();
+			JsonUtils.ConfigureJsonSerializer(_jsonOptions);
+		}
+
+		/// <inheritdoc/>
+		public bool Enabled => true;
+
+		/// <inheritdoc/>
+		public ValueTask DisposeAsync() => default;
+
+		/// <inheritdoc/>
+		public ValueTask FlushAsync(CancellationToken cancellationToken) => default;
+
+		/// <inheritdoc/>
+		public void SendEvent(TelemetryStoreId telemetryStoreId, TelemetryEvent telemetryEvent)
+		{
+			JsonNode? node = JsonSerializer.SerializeToNode(telemetryEvent, _jsonOptions);
+			if (node != null)
+			{
+				_metricCollection.AddEvent(telemetryStoreId, node);
+			}
+		}
+	}
+}
diff --git a/Engine/Source/Programs/Horde/Plugins/Analytics/HordeServer.Analytics/Telemetry/Sinks/MongoTelemetrySink.cs b/Engine/Source/Programs/Horde/Plugins/Analytics/HordeServer.Analytics/Telemetry/Sinks/MongoTelemetrySink.cs
new file mode 100644
--- /dev/null
+++ b/Engine/Source/Programs/Horde/Plugins/Analytics/HordeServer.Analytics/Telemetry/Sinks/MongoTelemetrySink.cs
@@ -0,0 +1,128 @@
+// Copyright Epic Games, Inc. All Rights Reserved.
+
+using System;
+using System.Linq;
+using System.Text.Json;
+using System.Threading;
+using System.Threading.Tasks;
+using EpicGames.Core;
+using EpicGames.Horde.Telemetry;
+using HordeServer.Server;
+using HordeServer.Utilities;
+using Microsoft.Extensions.Hosting;
+using Microsoft.Extensions.Logging;
+using Microsoft.Extensions.Options;
+using MongoDB.Bson;
+using MongoDB.Bson.Serialization.Attributes;
+using MongoDB.Driver;
+
+namespace HordeServer.Telemetry.Sinks
+{
+	/// <summary>
+	/// Telemetry sink which writes data to a MongoDB collection
+	/// </summary>
+	public sealed class MongoTelemetrySink : ITelemetrySink, IHostedService, IAsyncDisposable
+	{
+		class EventDocument
+		{
+			public ObjectId Id { get; set; }
+
+			[BsonElement("ts")]
+			public TelemetryStoreId TelemetryStoreId { get; set; }
+
+			[BsonElement("data")]
+			public BsonDocument Data { get; set; }
+
+			[BsonConstructor]
+			public EventDocument()
+			{
+				Data = new BsonDocument();
+			}
+
+			public EventDocument(TelemetryStoreId telemetryStoreId, BsonDocument data)
+			{
+				Id = ObjectId.GenerateNewId();
+				TelemetryStoreId = telemetryStoreId;
+				Data = data;
+			}
+		}
+
+		readonly MongoTelemetryConfig? _config;
+		readonly IMongoCollection<EventDocument> _collection;
+#pragma warning disable CA2213 // False positive? _writer is disposed in DisposeAsync.
+		readonly MongoBufferedWriter<EventDocument> _writer;
+#pragma warning restore CA2213
+		readonly ITicker _cleanupTicker;
+		readonly JsonSerializerOptions _jsonOptions;
+		readonly ILogger _logger;
+
+		/// <inheritdoc/>
+		public bool Enabled => _config != null;
+
+		/// <summary>
+		/// Constructor
+		/// </summary>
+		public MongoTelemetrySink(IMongoService mongoService, IClock clock, IOptions<AnalyticsServerConfig> serverSettings, ILogger<MongoTelemetrySink> logger)
+		{
+			_config = serverSettings.Value.Sinks.Mongo;
+			_collection = mongoService.GetCollection<EventDocument>("Telemetry", builder => builder.Ascending(x => x.TelemetryStoreId).Descending(x => x.Id));
+			_writer = new MongoBufferedWriter<EventDocument>(_collection, logger);
+			_cleanupTicker = clock.AddSharedTicker<MongoTelemetrySink>(TimeSpan.FromHours(4.0), CleanupAsync, logger);
+			_logger = logger;
+
+			_jsonOptions = new JsonSerializerOptions();
+			JsonUtils.ConfigureJsonSerializer(_jsonOptions);
+		}
+
+		/// <inheritdoc/>
+		public async Task StartAsync(CancellationToken cancellationToken)
+		{
+			if (Enabled)
+			{
+				await _writer.StartAsync();
+				await _cleanupTicker.StartAsync();
+			}
+		}
+
+		/// <inheritdoc/>
+		public async Task StopAsync(CancellationToken cancellationToken)
+		{
+			await _cleanupTicker.StopAsync();
+			await _writer.StopAsync(cancellationToken);
+		}
+
+		async ValueTask CleanupAsync(CancellationToken cancellationToken)
+		{
+			if (_config == null)
+			{
+				return;
+			}
+
+			DateTime baseTime = DateTime.UtcNow - TimeSpan.FromDays(_config.RetainDays);
+#pragma warning disable CS0618 // Type or member is obsolete
+			ObjectId minObjectId = new ObjectId(baseTime, 0, 0, 0);
+#pragma warning restore CS0618
+			DeleteResult result = await _collection.DeleteManyAsync(x => x.Id < minObjectId, cancellationToken);
+			_logger.LogInformation("Deleted {NumItems} telemetry records", result.DeletedCount);
+		}
+
+		/// <inheritdoc/>
+		public async ValueTask DisposeAsync()
+		{
+			await FlushAsync(default);
+			await _cleanupTicker.DisposeAsync();
+			await _writer.DisposeAsync();
+		}
+
+		/// <inheritdoc/>
+		public ValueTask FlushAsync(CancellationToken cancellationToken)
+			=> _writer.FlushAsync(cancellationToken);
+
+		/// <inheritdoc/>
+		public void SendEvent(TelemetryStoreId telemetryStoreId, TelemetryEvent telemetryEvent)
+		{
+			BsonDocument bson = BsonDocument.Parse(JsonSerializer.Serialize(telemetryEvent, _jsonOptions));
+			_writer.Write(new EventDocument(telemetryStoreId, bson));
+		}
+	}
+}
diff --git a/Engine/Source/Programs/Horde/Plugins/Analytics/HordeServer.Analytics/Telemetry/Sinks/NullTelemetrySink.cs b/Engine/Source/Programs/Horde/Plugins/Analytics/HordeServer.Analytics/Telemetry/Sinks/NullTelemetrySink.cs
new file mode 100644
--- /dev/null
+++ b/Engine/Source/Programs/Horde/Plugins/Analytics/HordeServer.Analytics/Telemetry/Sinks/NullTelemetrySink.cs
@@ -0,0 +1,29 @@
+// Copyright Epic Games, Inc. All Rights Reserved.
+
+using System.Threading;
+using System.Threading.Tasks;
+using EpicGames.Horde.Telemetry;
+
+namespace HordeServer.Telemetry.Sinks
+{
+	/// <summary>
+	/// Telemetry sink that discards all events
+	/// </summary>
+	public sealed class NullTelemetrySink : ITelemetrySink
+	{
+		/// <inheritdoc/>
+		public bool Enabled => false;
+
+		/// <inheritdoc/>
+		public ValueTask DisposeAsync()
+			=> default;
+
+		/// <inheritdoc/>
+		public ValueTask FlushAsync(CancellationToken cancellationToken)
+			=> default;
+
+		/// <inheritdoc/>
+		public void SendEvent(TelemetryStoreId telemetryStoreId, TelemetryEvent telemetryEvent)
+		{ }
+	}
+}
diff --git a/Engine/Source/Programs/Horde/Plugins/Analytics/HordeServer.Analytics/Telemetry/TelemetryController.cs b/Engine/Source/Programs/Horde/Plugins/Analytics/HordeServer.Analytics/Telemetry/TelemetryController.cs
new file mode 100644
--- /dev/null
+++ b/Engine/Source/Programs/Horde/Plugins/Analytics/HordeServer.Analytics/Telemetry/TelemetryController.cs
@@ -0,0 +1,166 @@
+// Copyright Epic Games, Inc. All Rights Reserved.
+
+using System;
+using System.Collections.Generic;
+using System.ComponentModel.DataAnnotations;
+using System.Linq;
+using System.Text.Json.Nodes;
+using System.Threading;
+using System.Threading.Tasks;
+using EpicGames.Horde.Telemetry;
+using EpicGames.Horde.Telemetry.Metrics;
+using HordeServer.Server;
+using HordeServer.Telemetry.Metrics;
+using HordeServer.Utilities;
+using Microsoft.AspNetCore.Authorization;
+using Microsoft.AspNetCore.Mvc;
+using Microsoft.Extensions.Options;
+
+namespace HordeServer.Telemetry
+{
+	/// <summary>
+	/// Controller for the /api/v1/telemetry endpoint
+	/// </summary>
+	[ApiController]
+	[Route("[controller]")]
+	public class TelemetryController : HordeControllerBase
+	{
+		readonly TelemetryManager _telemetryManager;
+		readonly IMetricCollection _metricCollection;
+		readonly IOptionsSnapshot<AnalyticsGlobalConfig> _globalConfig;
+
+		/// <summary>
+		/// Constructor
+		/// </summary>
+		public TelemetryController(TelemetryManager telemetryManager, IMetricCollection metricCollection, IOptionsSnapshot<AnalyticsGlobalConfig> globalConfig)
+		{
+			_telemetryManager = telemetryManager;
+			_metricCollection = metricCollection;
+			_globalConfig = globalConfig;
+		}
+
+		/// <summary>
+		/// Queries aggregated metrics from the telemetry system
+		/// </summary>
+		/// <param name="id">The metrics to query</param>
+		/// <param name="minTime">Minimum time interval to query</param>
+		/// <param name="maxTime">Maximum time interval to query</param>
+		/// <param name="group">Grouping key</param>
+		/// <param name="results">Number of results to return</param>
+		/// <param name="cancellationToken">Cancellation token for the operation</param>
+		[HttpGet]
+		[Authorize]
+		[Route("/api/v1/telemetry/metrics")]
+		[Obsolete("Pass a telemetry store id in the route instead")]
+		public Task<ActionResult<List<GetTelemetryMetricsResponse>>> GetMetricsAsync([FromQuery] MetricId[] id, [FromQuery] DateTime? minTime = null, [FromQuery] DateTime? maxTime = null, [FromQuery] string? group = null, [FromQuery] int results = 50, CancellationToken cancellationToken = default)
+		{
+			return GetMetricsAsync(TelemetryStoreId.Default, id, minTime, maxTime, group, results, cancellationToken);
+		}
+
+		/// <summary>
+		/// Queries aggregated metrics from the telemetry system
+		/// </summary>
+		/// <param name="storeId">The telemetry store id</param>
+		/// <param name="id">The metrics to query</param>
+		/// <param name="minTime">Minimum time interval to query</param>
+		/// <param name="maxTime">Maximum time interval to query</param>
+		/// <param name="group">Grouping key</param>
+		/// <param name="results">Number of results to return</param>
+		/// <param name="cancellationToken">Cancellation token for the operation</param>
+		[HttpGet]
+		[Authorize]
+		[Route("/api/v1/telemetry/{storeId}/metrics")]
+		public async Task<ActionResult<List<GetTelemetryMetricsResponse>>> GetMetricsAsync(TelemetryStoreId storeId, [FromQuery] MetricId[] id, [FromQuery] DateTime? minTime = null, [FromQuery] DateTime? maxTime = null, [FromQuery] string? group = null, [FromQuery] int results = 50, CancellationToken cancellationToken = default)
+		{
+			if (!_globalConfig.Value.Authorize(TelemetryAclAction.QueryMetrics, User))
+			{
+				return Forbid(TelemetryAclAction.QueryMetrics);
+			}
+
+			TelemetryStoreConfig? telemetryStoreConfig;
+			if (!_globalConfig.Value.TryGetTelemetryStore(storeId, out telemetryStoreConfig))
+			{
+				return NotFound(storeId);
+			}
+
+			List<GetTelemetryMetricsResponse> result = new List<GetTelemetryMetricsResponse>();
+
+			List<IMetric> metrics = await _metricCollection.FindAsync(storeId, id, minTime, maxTime, group, results, cancellationToken);
+
+			HashSet<MetricId> unique = new HashSet<MetricId>(metrics.Select(m => m.MetricId));
+
+			foreach (MetricId metricId in unique)
+			{
+				MetricConfig? metricConfig = telemetryStoreConfig.Metrics.Find(m => m.Id == metricId);
+
+				if (metricConfig == null)
+				{
+					continue;
+				}
+
+				GetTelemetryMetricsResponse response = new GetTelemetryMetricsResponse();
+				response.MetricId = metricId.ToString();
+				response.GroupBy = metricConfig.GroupBy;
+				result.Add(response);
+
+				foreach (IMetric metric in metrics.Where(m => m.MetricId == metricId))
+				{
+					GetTelemetryMetricResponse rmetric = new GetTelemetryMetricResponse();
+					rmetric.Time = metric.Time;
+					rmetric.Value = metric.Value;
+
+					if (group == null)
+					{
+						rmetric.Group = metric.Group;
+					}
+
+					response.Metrics.Add(rmetric);
+				}
+			}
+
+			return result;
+		}
+
+		/// <summary>
+		/// Posts a new telemetry event. This API is modeled after Epic's external data router, allowing events in the engine to be sent to Horde using the same mechanism.
+		/// </summary>
+		/// <param name="request">The event data</param>
+		/// <param name="appId">Identifier of the application sending the event</param>
+		/// <param name="appVersion">Version number of the application</param>
+		/// <param name="appEnvironment">Name of the environment that the sending application is running in</param>
+		/// <param name="uploadType">Type of data being uploaded</param>
+		[HttpPost]
+		[Route("/api/v1/telemetry")]
+		public ActionResult PostEvent([FromBody] PostTelemetryEventStreamRequest request, [FromQuery][Required] string appId, [FromQuery][Required] string appVersion, [FromQuery][Required] string appEnvironment, [FromQuery][Required] TelemetryUploadType uploadType)
+		{
+			return PostEvent(TelemetryStoreId.Default, request, appId, appVersion, appEnvironment, uploadType);
+		}
+
+		/// <summary>
+		/// Posts a new telemetry event. This API is modeled after Epic's external data router, allowing events in the engine to be sent to Horde using the same mechanism.
+		/// </summary>
+		/// <param name="storeId">The telemetry store</param>
+		/// <param name="request">The event data</param>
+		/// <param name="appId">Identifier of the application sending the event</param>
+		/// <param name="appVersion">Version number of the application</param>
+		/// <param name="appEnvironment">Name of the environment that the sending application is running in</param>
+		/// <param name="uploadType">Type of data being uploaded</param>
+		[HttpPost]
+		[Route("/api/v1/telemetry/{storeId}")]
+		public ActionResult PostEvent(TelemetryStoreId storeId, [FromBody] PostTelemetryEventStreamRequest request, [FromQuery][Required] string appId, [FromQuery][Required] string appVersion, [FromQuery][Required] string appEnvironment, [FromQuery][Required] TelemetryUploadType uploadType)
+		{
+			if (uploadType != TelemetryUploadType.EtEventStream)
+			{
+				return BadRequest("Invalid upload type");
+			}
+
+			TelemetryRecordMeta recordMeta = new TelemetryRecordMeta(AppId: appId, AppVersion: appVersion, AppEnvironment: appEnvironment);
+			foreach (JsonObject eventPayload in request.Events)
+			{
+				_telemetryManager.WriteEvent(storeId, recordMeta, eventPayload);
+			}
+
+			return NoContent();
+		}
+	}
+}
diff --git a/Engine/Source/Programs/Horde/Plugins/Analytics/HordeServer.Analytics/Telemetry/TelemetryEvent.cs b/Engine/Source/Programs/Horde/Plugins/Analytics/HordeServer.Analytics/Telemetry/TelemetryEvent.cs
new file mode 100644
--- /dev/null
+++ b/Engine/Source/Programs/Horde/Plugins/Analytics/HordeServer.Analytics/Telemetry/TelemetryEvent.cs
@@ -0,0 +1,29 @@
+// Copyright Epic Games, Inc. All Rights Reserved.
+
+namespace HordeServer.Telemetry
+{
+	/// <summary>
+	/// Wrapper around native telemetry objects
+	/// </summary>
+	public class TelemetryEvent
+	{
+		/// <summary>
+		/// Record metadata
+		/// </summary>
+		public TelemetryRecordMeta RecordMeta { get; set; }
+
+		/// <summary>
+		/// Accessor for the native object
+		/// </summary>
+		public object Payload { get; }
+
+		/// <summary>
+		/// Constructor
+		/// </summary>
+		public TelemetryEvent(TelemetryRecordMeta recordMeta, object payload)
+		{
+			RecordMeta = recordMeta;
+			Payload = payload;
+		}
+	}
+}
diff --git a/Engine/Source/Programs/Horde/Plugins/Analytics/HordeServer.Analytics/Telemetry/TelemetryManager.cs b/Engine/Source/Programs/Horde/Plugins/Analytics/HordeServer.Analytics/Telemetry/TelemetryManager.cs
new file mode 100644
--- /dev/null
+++ b/Engine/Source/Programs/Horde/Plugins/Analytics/HordeServer.Analytics/Telemetry/TelemetryManager.cs
@@ -0,0 +1,135 @@
+// Copyright Epic Games, Inc. All Rights Reserved.
+
+using System;
+using System.Collections.Generic;
+using System.Text.Json;
+using System.Threading;
+using System.Threading.Tasks;
+using EpicGames.Core;
+using EpicGames.Horde.Telemetry;
+using HordeServer.Telemetry.Sinks;
+using HordeServer.Utilities;
+using Microsoft.Extensions.DependencyInjection;
+using Microsoft.Extensions.Hosting;
+using Microsoft.Extensions.Logging;
+using Microsoft.Extensions.Options;
+using OpenTelemetry.Trace;
+
+namespace HordeServer.Telemetry
+{
+	/// <summary>
+	/// Telemetry sink dispatching incoming events to all registered sinks
+	/// </summary>
+	public sealed class TelemetryManager : ITelemetryWriter, IHostedService
+	{
+		/// <inheritdoc/>
+		public bool Enabled => _telemetrySinks.Count > 0;
+
+		private readonly List<ITelemetrySink> _telemetrySinks = new();
+		private readonly ITicker _ticker;
+		private readonly Tracer _tracer;
+		private readonly ILogger<TelemetryManager> _logger;
+		private readonly TelemetryRecordMeta _serverEventMetadata;
+
+		/// <summary>
+		/// Constructor
+		/// </summary>
+		public TelemetryManager(IServiceProvider serviceProvider, IServerInfo serverInfo, IClock clock, IOptions<AnalyticsServerConfig> serverSettings, Tracer tracer, ILoggerFactory loggerFactory)
+		{
+			_tracer = tracer;
+			_logger = loggerFactory.CreateLogger<TelemetryManager>();
+			_ticker = clock.AddTicker<EpicTelemetrySink>(TimeSpan.FromSeconds(30.0), FlushAsync, _logger);
+			_serverEventMetadata = new TelemetryRecordMeta("Horde", serverInfo.Version.ToString(), serverInfo.Environment, serverInfo.SessionId);
+
+			TelemetrySinkConfig telemetrySinks = serverSettings.Value.Sinks;
+			if (telemetrySinks.Epic?.Enabled ?? false)
+			{
+				_telemetrySinks.Add(new EpicTelemetrySink(telemetrySinks.Epic, serverInfo, loggerFactory.CreateLogger<EpicTelemetrySink>()));
+			}
+			if (telemetrySinks.Mongo?.Enabled ?? false)
+			{
+				_telemetrySinks.Add(serviceProvider.GetRequiredService<MongoTelemetrySink>());
+			}
+
+			_telemetrySinks.Add(serviceProvider.GetRequiredService<MetricTelemetrySink>());
+		}
+
+		/// <inheritdoc/>
+		public void WriteEvent(TelemetryStoreId telemetryStoreId, object payload)
+		{
+			WriteEvent(telemetryStoreId, _serverEventMetadata, payload);
+		}
+
+		/// <inheritdoc/>
+		public void WriteEvent(TelemetryStoreId telemetryStoreId, TelemetryRecordMeta metadata, object payload)
+		{
+			using TelemetrySpan span = _tracer.StartActiveSpan($"{nameof(TelemetryManager)}.{nameof(WriteEvent)}");
+
+			TelemetryEvent telemetryEvent = new TelemetryEvent(metadata, payload);
+			foreach (ITelemetrySink sink in _telemetrySinks)
+			{
+				using TelemetrySpan sinkSpan = _tracer.StartActiveSpan($"SendEvent");
+				string fullName = sink.GetType().FullName ?? "Unknown";
+				span.SetAttribute("sink", fullName);
+
+				if (sink.Enabled)
+				{
+					try
+					{
+						sink.SendEvent(telemetryStoreId, telemetryEvent);
+					}
+					catch (Exception e)
+					{
+						_logger.LogWarning(e, "Failed sending event to {Sink}. Message: {Message}. Event: {Event}", fullName, e.Message, GetEventText(telemetryEvent));
+					}
+				}
+			}
+		}
+
+		static string GetEventText(TelemetryEvent telemetryEvent)
+		{
+			try
+			{
+				JsonSerializerOptions options = new JsonSerializerOptions();
+				JsonUtils.ConfigureJsonSerializer(options);
+
+				return JsonSerializer.Serialize(telemetryEvent, options);
+			}
+			catch
+			{
+				return "(Unable to serialize)";
+			}
+		}
+
+		/// <inheritdoc />
+		public async Task StartAsync(CancellationToken cancellationToken)
+		{
+			await _ticker.StartAsync();
+		}
+
+		/// <inheritdoc />
+		public async Task StopAsync(CancellationToken cancellationToken)
+		{
+			await _ticker.StopAsync();
+		}
+
+		/// <inheritdoc />
+		public async ValueTask DisposeAsync()
+		{
+			await _ticker.DisposeAsync();
+			foreach (ITelemetrySink sink in _telemetrySinks)
+			{
+				await sink.DisposeAsync();
+			}
+		}
+
+		/// <inheritdoc />
+		public async ValueTask FlushAsync(CancellationToken cancellationToken)
+		{
+			foreach (ITelemetrySink sink in _telemetrySinks)
+			{
+				await sink.FlushAsync(cancellationToken);
+			}
+		}
+	}
+}
