Horde: Move OpenTelemetryTracers into HordeServer.Shared.

#rnx

diff --git a/Engine/Source/Programs/Horde/HordeServer.Shared/Utilities/OpenTelemetryTracers.cs b/Engine/Source/Programs/Horde/HordeServer.Shared/Utilities/OpenTelemetryTracers.cs
new file mode 100644
--- /dev/null
+++ b/Engine/Source/Programs/Horde/HordeServer.Shared/Utilities/OpenTelemetryTracers.cs
@@ -0,0 +1,46 @@
+// Copyright Epic Games, Inc. All Rights Reserved.
+
+using OpenTelemetry.Trace;
+
+namespace HordeServer.Utilities
+{
+	/// <summary>
+	/// Static initialization of all available OpenTelemetry tracers
+	/// </summary>
+	public static class OpenTelemetryTracers
+	{
+		/// <summary>
+		/// Name of resource name attribute in Datadog
+		/// Some traces use this for prettier display inside their UI
+		/// </summary>
+		public const string DatadogResourceAttribute = "resource.name";
+
+		/// <summary>
+		/// Name of default Horde tracer (aka activity source)
+		/// </summary>
+		public const string HordeName = "Horde";
+
+		/// <summary>
+		/// Name of MongoDB tracer (aka activity source)
+		/// </summary>
+		public const string MongoDbName = "MongoDB";
+
+		/// <summary>
+		/// List of all source names configured in this class.
+		/// They are needed at startup when initializing OpenTelemetry
+		/// </summary>
+		public static string[] SourceNames => new[] { HordeName, MongoDbName };
+
+		/// <summary>
+		/// Default tracer used in Horde
+		/// Prefer dependency-injected tracer over this static member.
+		/// </summary>
+		public static readonly Tracer Horde = TracerProvider.Default.GetTracer(HordeName);
+
+		/// <summary>
+		/// Tracer specific to MongoDB
+		/// Prefer StartMongoDbSpan static extension for Tracer.
+		/// </summary>
+		public static readonly Tracer MongoDb = TracerProvider.Default.GetTracer(MongoDbName);
+	}
+}
diff --git a/Engine/Source/Programs/Horde/HordeServer/Utilities/OpenTelemetry.cs b/Engine/Source/Programs/Horde/HordeServer/Utilities/OpenTelemetry.cs
--- a/Engine/Source/Programs/Horde/HordeServer/Utilities/OpenTelemetry.cs
+++ b/Engine/Source/Programs/Horde/HordeServer/Utilities/OpenTelemetry.cs
@@ -239,46 +239,6 @@ public static class OpenTelemetryHelper
 	}
 }
 
-/// <summary>
-/// Static initialization of all available OpenTelemetry tracers
-/// </summary>
-public static class OpenTelemetryTracers
-{
-	/// <summary>
-	/// Name of resource name attribute in Datadog
-	/// Some traces use this for prettier display inside their UI
-	/// </summary>
-	public const string DatadogResourceAttribute = "resource.name";
-
-	/// <summary>
-	/// Name of default Horde tracer (aka activity source)
-	/// </summary>
-	public const string HordeName = "Horde";
-
-	/// <summary>
-	/// Name of MongoDB tracer (aka activity source)
-	/// </summary>
-	public const string MongoDbName = "MongoDB";
-
-	/// <summary>
-	/// List of all source names configured in this class.
-	/// They are needed at startup when initializing OpenTelemetry
-	/// </summary>
-	public static string[] SourceNames => new[] { HordeName, MongoDbName };
-
-	/// <summary>
-	/// Default tracer used in Horde
-	/// Prefer dependency-injected tracer over this static member.
-	/// </summary>
-	public static readonly Tracer Horde = TracerProvider.Default.GetTracer(HordeName);
-
-	/// <summary>
-	/// Tracer specific to MongoDB
-	/// Prefer StartMongoDbSpan static extension for Tracer.
-	/// </summary>
-	public static readonly Tracer MongoDb = TracerProvider.Default.GetTracer(MongoDbName);
-}
-
 /// <summary>
 /// Extensions to handle Horde specific data types in the OpenTelemetry library
 /// </summary>
