Horde: Fix test compile errors.

#rnx

diff --git a/Engine/Source/Programs/Horde/HordeServer.Tests/HordeServer.Tests.csproj b/Engine/Source/Programs/Horde/HordeServer.Tests/HordeServer.Tests.csproj
--- a/Engine/Source/Programs/Horde/HordeServer.Tests/HordeServer.Tests.csproj
+++ b/Engine/Source/Programs/Horde/HordeServer.Tests/HordeServer.Tests.csproj
@@ -40,6 +40,7 @@
       <Aliases>JobDriver</Aliases>
     </ProjectReference>
     <ProjectReference Include="..\HordeServer\HordeServer.csproj" />
+    <ProjectReference Include="..\Plugins\Analytics\HordeServer.Analytics\HordeServer.Analytics.csproj" />
   </ItemGroup>
 
   <ItemGroup>
diff --git a/Engine/Source/Programs/Horde/Plugins/Analytics/HordeServer.Analytics/Properties/AssemblyInfo.cs b/Engine/Source/Programs/Horde/Plugins/Analytics/HordeServer.Analytics/Properties/AssemblyInfo.cs
new file mode 100644
--- /dev/null
+++ b/Engine/Source/Programs/Horde/Plugins/Analytics/HordeServer.Analytics/Properties/AssemblyInfo.cs
@@ -0,0 +1,7 @@
+// Copyright Epic Games, Inc. All Rights Reserved.
+
+using System.Resources;
+using System.Runtime.CompilerServices;
+
+[assembly: InternalsVisibleTo("HordeServer.Tests")]
+[assembly: NeutralResourcesLanguage("en-US")]
