Horde: Remove legacy artifacts and log system.

diff --git a/Engine/Source/Programs/Horde/HordeServer.Tests/ControllerIntegrationTest.cs b/Engine/Source/Programs/Horde/HordeServer.Tests/ControllerIntegrationTest.cs
--- a/Engine/Source/Programs/Horde/HordeServer.Tests/ControllerIntegrationTest.cs
+++ b/Engine/Source/Programs/Horde/HordeServer.Tests/ControllerIntegrationTest.cs
@@ -8,7 +8,6 @@ using Amazon.CloudWatch;
 using HordeServer.Agents;
 using HordeServer.Configuration;
 using HordeServer.Jobs;
-using HordeServer.Jobs.Artifacts;
 using HordeServer.Jobs.Graphs;
 using HordeServer.Jobs.Templates;
 using HordeServer.Perforce;
diff --git a/Engine/Source/Programs/Horde/HordeServer.Tests/Fixture.cs b/Engine/Source/Programs/Horde/HordeServer.Tests/Fixture.cs
--- a/Engine/Source/Programs/Horde/HordeServer.Tests/Fixture.cs
+++ b/Engine/Source/Programs/Horde/HordeServer.Tests/Fixture.cs
@@ -13,7 +13,6 @@ using EpicGames.Horde.Streams;
 using HordeServer.Agents;
 using HordeServer.Configuration;
 using HordeServer.Jobs;
-using HordeServer.Jobs.Artifacts;
 using HordeServer.Jobs.Graphs;
 using HordeServer.Jobs.Templates;
 using HordeServer.Projects;
@@ -32,16 +31,14 @@ namespace HordeServer.Tests
 		public StreamConfig? StreamConfig { get; private set; }
 		public TemplateId TemplateRefId1 { get; private set; }
 		public TemplateId TemplateRefId2 { get; private set; }
-		public IArtifactV1 Job1Artifact { get; private set; } = null!;
-		public string Job1ArtifactData { get; private set; } = null!;
 		public IAgent Agent1 { get; private set; } = null!;
 		public string Agent1Name { get; private set; } = null!;
 		public const string PoolName = "TestingPool";
 
-		public static async Task<Fixture> CreateAsync(ConfigService configService, IGraphCollection graphCollection, ITemplateCollection templateCollection, JobService jobService, IArtifactCollectionV1 artifactCollection, AgentService agentService, ServerSettings serverSettings)
+		public static async Task<Fixture> CreateAsync(ConfigService configService, IGraphCollection graphCollection, ITemplateCollection templateCollection, JobService jobService, AgentService agentService, ServerSettings serverSettings)
 		{
 			Fixture fixture = new Fixture();
-			await fixture.PopulateAsync(configService, graphCollection, templateCollection, jobService, artifactCollection, agentService, serverSettings);
+			await fixture.PopulateAsync(configService, graphCollection, templateCollection, jobService, agentService, serverSettings);
 
 			//			(PerforceService as PerforceServiceStub)?.AddChange("//UE5/Main", 112233, "leet.coder", "Did stuff", new []{"file.cpp"});
 			//			(PerforceService as PerforceServiceStub)?.AddChange("//UE5/Main", 1111, "swarm", "A shelved CL here", new []{"renderer.cpp"});
@@ -49,7 +46,7 @@ namespace HordeServer.Tests
 			return fixture;
 		}
 
-		private async Task PopulateAsync(ConfigService configService, IGraphCollection graphCollection, ITemplateCollection templateCollection, JobService jobService, IArtifactCollectionV1 artifactCollection, AgentService agentService, ServerSettings serverSettings)
+		private async Task PopulateAsync(ConfigService configService, IGraphCollection graphCollection, ITemplateCollection templateCollection, JobService jobService, AgentService agentService, ServerSettings serverSettings)
 		{
 			FixtureGraph fg = new FixtureGraph();
 			fg.Id = ContentHash.Empty;
@@ -115,11 +112,6 @@ namespace HordeServer.Tests
 			);
 			Job2 = (await jobService.GetJobAsync(Job2.Id))!;
 
-			Job1ArtifactData = "For The Horde!";
-			using MemoryStream job1ArtifactStream = new MemoryStream(Encoding.UTF8.GetBytes(Job1ArtifactData));
-			Job1Artifact = await artifactCollection.CreateArtifactAsync(Job1.Id, JobStepId.Parse("22"), "myfile.txt",
-				"text/plain", job1ArtifactStream);
-
 			Agent1Name = "testAgent1";
 			Agent1 = await agentService.CreateAgentAsync(Agent1Name, false, "");
 		}
diff --git a/Engine/Source/Programs/Horde/HordeServer.Tests/Jobs/JobsControllerTest.cs b/Engine/Source/Programs/Horde/HordeServer.Tests/Jobs/JobsControllerTest.cs
deleted file mode 100644
diff --git a/Engine/Source/Programs/Horde/HordeServer.Tests/TestSetup.cs b/Engine/Source/Programs/Horde/HordeServer.Tests/TestSetup.cs
--- a/Engine/Source/Programs/Horde/HordeServer.Tests/TestSetup.cs
+++ b/Engine/Source/Programs/Horde/HordeServer.Tests/TestSetup.cs
@@ -37,7 +37,6 @@ using HordeServer.Dashboard;
 using HordeServer.Devices;
 using HordeServer.Issues;
 using HordeServer.Jobs;
-using HordeServer.Jobs.Artifacts;
 using HordeServer.Jobs.Bisect;
 using HordeServer.Jobs.Graphs;
 using HordeServer.Jobs.Schedules;
@@ -45,7 +44,6 @@ using HordeServer.Jobs.Templates;
 using HordeServer.Jobs.TestData;
 using HordeServer.Jobs.Timing;
 using HordeServer.Logs;
-using HordeServer.Logs.Storage;
 using HordeServer.Notifications;
 using HordeServer.Perforce;
 using HordeServer.Server;
@@ -120,7 +118,6 @@ namespace HordeServer.Tests
 		public IssueService IssueService => ServiceProvider.GetRequiredService<IssueService>();
 		public JobTaskSource JobTaskSource => ServiceProvider.GetRequiredService<JobTaskSource>();
 		public JobService JobService => ServiceProvider.GetRequiredService<JobService>();
-		public IArtifactCollectionV1 ArtifactCollection => ServiceProvider.GetRequiredService<IArtifactCollectionV1>();
 		public IDowntimeService DowntimeService => ServiceProvider.GetRequiredService<IDowntimeService>();
 		public RpcService RpcService => ServiceProvider.GetRequiredService<RpcService>();
 		public PoolService PoolService => ServiceProvider.GetRequiredService<PoolService>();
diff --git a/Engine/Source/Programs/Horde/HordeServer/Jobs/Artifacts/ArtifactCollectionV1.cs b/Engine/Source/Programs/Horde/HordeServer/Jobs/Artifacts/ArtifactCollectionV1.cs
deleted file mode 100644
diff --git a/Engine/Source/Programs/Horde/HordeServer/Jobs/Artifacts/ArtifactMessagesV1.cs b/Engine/Source/Programs/Horde/HordeServer/Jobs/Artifacts/ArtifactMessagesV1.cs
deleted file mode 100644
diff --git a/Engine/Source/Programs/Horde/HordeServer/Jobs/Artifacts/ArtifactsControllerV1.cs b/Engine/Source/Programs/Horde/HordeServer/Jobs/Artifacts/ArtifactsControllerV1.cs
deleted file mode 100644
diff --git a/Engine/Source/Programs/Horde/HordeServer/Jobs/Artifacts/IArtifactCollectionV1.cs b/Engine/Source/Programs/Horde/HordeServer/Jobs/Artifacts/IArtifactCollectionV1.cs
deleted file mode 100644
diff --git a/Engine/Source/Programs/Horde/HordeServer/Jobs/Artifacts/IArtifactV1.cs b/Engine/Source/Programs/Horde/HordeServer/Jobs/Artifacts/IArtifactV1.cs
deleted file mode 100644
diff --git a/Engine/Source/Programs/Horde/HordeServer/Jobs/JobsController.cs b/Engine/Source/Programs/Horde/HordeServer/Jobs/JobsController.cs
--- a/Engine/Source/Programs/Horde/HordeServer/Jobs/JobsController.cs
+++ b/Engine/Source/Programs/Horde/HordeServer/Jobs/JobsController.cs
@@ -17,7 +17,6 @@ using HordeServer.Acls;
 using HordeServer.Agents;
 using HordeServer.Artifacts;
 using HordeServer.Commits;
-using HordeServer.Jobs.Artifacts;
 using HordeServer.Jobs.Graphs;
 using HordeServer.Jobs.Templates;
 using HordeServer.Jobs.Timing;
@@ -48,7 +47,6 @@ namespace HordeServer.Jobs
 		private readonly JobService _jobService;
 		private readonly ITemplateCollection _templateCollection;
 		private readonly IArtifactCollection _artifactCollection;
-		private readonly IArtifactCollectionV1 _artifactCollectionV1;
 		private readonly IUserCollection _userCollection;
 		private readonly INotificationService _notificationService;
 		private readonly AgentService _agentService;
@@ -57,7 +55,7 @@ namespace HordeServer.Jobs
 		/// <summary>
 		/// Constructor
 		/// </summary>
-		public JobsController(IGraphCollection graphs, ICommitService commitService, IPerforceService perforce, JobService jobService, ITemplateCollection templateCollection, IArtifactCollection artifactCollection, IArtifactCollectionV1 artifactCollectionV1, IUserCollection userCollection, INotificationService notificationService, AgentService agentService, IOptionsSnapshot<GlobalConfig> globalConfig)
+		public JobsController(IGraphCollection graphs, ICommitService commitService, IPerforceService perforce, JobService jobService, ITemplateCollection templateCollection, IArtifactCollection artifactCollection, IUserCollection userCollection, INotificationService notificationService, AgentService agentService, IOptionsSnapshot<GlobalConfig> globalConfig)
 		{
 			_graphs = graphs;
 			_commitService = commitService;
@@ -65,7 +63,6 @@ namespace HordeServer.Jobs
 			_jobService = jobService;
 			_templateCollection = templateCollection;
 			_artifactCollection = artifactCollection;
-			_artifactCollectionV1 = artifactCollectionV1;
 			_userCollection = userCollection;
 			_notificationService = notificationService;
 			_agentService = agentService;
@@ -1766,102 +1763,6 @@ namespace HordeServer.Jobs
 			return new GetNotificationResponse(subscription);
 		}
 
-		/// <summary>
-		/// Gets a particular step currently scheduled to be executed for a job
-		/// </summary>
-		/// <param name="jobId">Unique id for the job</param>
-		/// <param name="batchId">Unique id for the batch</param>
-		/// <param name="stepId">Unique id for the step</param>
-		/// <param name="name"></param>
-		/// <param name="cancellationToken">Cancellation token for the operation</param>
-		/// <returns>List of nodes to be executed</returns>
-		[HttpGet]
-		[Route("/api/v1/jobs/{jobId}/batches/{batchId}/steps/{stepId}/artifacts/{*name}")]
-		public async Task<ActionResult> GetArtifactAsync(JobId jobId, JobStepBatchId batchId, JobStepId stepId, string name, CancellationToken cancellationToken)
-		{
-			IJob? job = await _jobService.GetJobAsync(jobId, cancellationToken);
-			if (job == null)
-			{
-				return NotFound(jobId);
-			}
-
-			StreamConfig? streamConfig;
-			if (!_globalConfig.Value.TryGetStream(job.StreamId, out streamConfig))
-			{
-				return NotFound(job.StreamId);
-			}
-			if (!streamConfig.Authorize(JobAclAction.ViewJob, User))
-			{
-				return Forbid(JobAclAction.ViewJob, jobId);
-			}
-
-			if (!job.TryGetBatch(batchId, out IJobStepBatch? batch))
-			{
-				return NotFound(jobId, batchId);
-			}
-			if (!batch.TryGetStep(stepId, out _))
-			{
-				return NotFound(jobId, batchId, stepId);
-			}
-
-			IReadOnlyList<IArtifactV1> artifacts = await _artifactCollectionV1.GetArtifactsAsync(jobId, stepId, name, cancellationToken);
-			if (artifacts.Count == 0)
-			{
-				return NotFound();
-			}
-
-			IArtifactV1 artifact = artifacts[0];
-			return new FileStreamResult(await _artifactCollectionV1.OpenArtifactReadStreamAsync(artifact, cancellationToken), artifact.MimeType);
-		}
-
-		/// <summary>
-		/// Gets a particular step currently scheduled to be executed for a job
-		/// </summary>
-		/// <param name="jobId">Unique id for the job</param>
-		/// <param name="batchId">Unique id for the batch</param>
-		/// <param name="stepId">Unique id for the step</param>
-		/// <param name="cancellationToken">Cancellation token for the operation</param>
-		/// <returns>List of nodes to be executed</returns>
-		[HttpGet]
-		[Route("/api/v1/jobs/{jobId}/batches/{batchId}/steps/{stepId}/trace")]
-		public async Task<ActionResult> GetStepTraceAsync(JobId jobId, JobStepBatchId batchId, JobStepId stepId, CancellationToken cancellationToken)
-		{
-			IJob? job = await _jobService.GetJobAsync(jobId, cancellationToken);
-			if (job == null)
-			{
-				return NotFound(jobId);
-			}
-
-			StreamConfig? streamConfig;
-			if (!_globalConfig.Value.TryGetStream(job.StreamId, out streamConfig))
-			{
-				return NotFound(job.StreamId);
-			}
-			if (!streamConfig.Authorize(JobAclAction.ViewJob, User))
-			{
-				return Forbid(JobAclAction.ViewJob, jobId);
-			}
-
-			if (!job.TryGetBatch(batchId, out IJobStepBatch? batch))
-			{
-				return NotFound(jobId, batchId);
-			}
-			if (!batch.TryGetStep(stepId, out _))
-			{
-				return NotFound(jobId, batchId, stepId);
-			}
-
-			IReadOnlyList<IArtifactV1> artifacts = await _artifactCollectionV1.GetArtifactsAsync(jobId, stepId, null, cancellationToken);
-			foreach (IArtifactV1 artifact in artifacts)
-			{
-				if (artifact.Name.Equals("trace.json", StringComparison.OrdinalIgnoreCase))
-				{
-					return new FileStreamResult(await _artifactCollectionV1.OpenArtifactReadStreamAsync(artifact, cancellationToken), "text/json");
-				}
-			}
-			return NotFound();
-		}
-
 		/// <summary>
 		/// Updates notifications for a specific label.
 		/// </summary>
diff --git a/Engine/Source/Programs/Horde/HordeServer/Logs/Data/LogChunkData.cs b/Engine/Source/Programs/Horde/HordeServer/Logs/Data/LogChunkData.cs
deleted file mode 100644
diff --git a/Engine/Source/Programs/Horde/HordeServer/Logs/Data/LogIndexData.cs b/Engine/Source/Programs/Horde/HordeServer/Logs/Data/LogIndexData.cs
deleted file mode 100644
diff --git a/Engine/Source/Programs/Horde/HordeServer/Logs/Data/LogSubChunkData.cs b/Engine/Source/Programs/Horde/HordeServer/Logs/Data/LogSubChunkData.cs
deleted file mode 100644
diff --git a/Engine/Source/Programs/Horde/HordeServer/Logs/Data/LogToken.cs b/Engine/Source/Programs/Horde/HordeServer/Logs/Data/LogToken.cs
deleted file mode 100644
diff --git a/Engine/Source/Programs/Horde/HordeServer/Logs/ILogStorage.cs b/Engine/Source/Programs/Horde/HordeServer/Logs/ILogStorage.cs
deleted file mode 100644
diff --git a/Engine/Source/Programs/Horde/HordeServer/Logs/LogCollection.cs b/Engine/Source/Programs/Horde/HordeServer/Logs/LogCollection.cs
--- a/Engine/Source/Programs/Horde/HordeServer/Logs/LogCollection.cs
+++ b/Engine/Source/Programs/Horde/HordeServer/Logs/LogCollection.cs
@@ -16,7 +16,6 @@ using EpicGames.Horde.Agents.Sessions;
 using EpicGames.Horde.Jobs;
 using EpicGames.Horde.Logs;
 using EpicGames.Horde.Storage;
-using HordeServer.Logs.Data;
 using HordeServer.Server;
 using HordeServer.Storage;
 using HordeServer.Utilities;
@@ -260,7 +259,6 @@ namespace HordeServer.Logs
 		readonly IMongoCollection<LogDocument> _logCollection;
 		readonly IMongoCollection<LogEventDocument> _logEvents;
 		readonly IMongoCollection<LegacyLogEventDocument> _legacyEvents;
-		readonly ILogStorage _storage;
 		readonly LogTailService _logTailService;
 		readonly StorageService _storageService;
 		readonly Tracer _tracer;
@@ -270,10 +268,9 @@ namespace HordeServer.Logs
 		/// <summary>
 		/// Constructor
 		/// </summary>
-		public LogCollection(IMongoService mongoService, ILogStorage storage, LogTailService logTailService, StorageService storageService, Tracer tracer, ILogger<LogCollection> logger)
+		public LogCollection(IMongoService mongoService, LogTailService logTailService, StorageService storageService, Tracer tracer, ILogger<LogCollection> logger)
 		{
 			_logCollection = mongoService.GetCollection<LogDocument>("LogFiles");
-			_storage = storage;
 			_logTailService = logTailService;
 			_storageService = storageService;
 			_tracer = tracer;
@@ -320,19 +317,11 @@ namespace HordeServer.Logs
 		public async Task<ILog?> GetAsync(LogId logId, CancellationToken cancellationToken)
 		{
 			LogDocument? logDocument = await _logCollection.Find<LogDocument>(x => x.Id == logId).FirstOrDefaultAsync(cancellationToken);
-			if (logDocument == null)
+			if (logDocument == null || !logDocument.UseNewStorageBackend)
 			{
 				return null;
 			}
-
-			if (logDocument.UseNewStorageBackend)
-			{
-				return new Log(this, logDocument);
-			}
-			else
-			{
-				return new LogV1(this, logDocument);
-			}
+			return new Log(this, logDocument);
 		}
 
 		/// <inheritdoc/>
@@ -646,13 +635,13 @@ namespace HordeServer.Logs
 		}
 
 		/// <summary>
-		/// Helper method for catching exceptions in <see cref="LogText.ConvertToPlainText(ReadOnlySpan{Byte}, Byte[], Int32)"/>
+		/// Helper method for catching exceptions in <see cref="ConvertToPlainText(ReadOnlySpan{Byte}, Byte[], Int32)"/>
 		/// </summary>
 		public static int GuardedConvertToPlainText(ReadOnlySpan<byte> input, byte[] output, int outputOffset, ILogger logger)
 		{
 			try
 			{
-				return LogText.ConvertToPlainText(input, output, outputOffset);
+				return ConvertToPlainText(input, output, outputOffset);
 			}
 			catch (Exception ex)
 			{
@@ -662,6 +651,124 @@ namespace HordeServer.Logs
 			}
 		}
 
+		/// <summary>
+		/// Converts a JSON log line to plain text
+		/// </summary>
+		/// <param name="input">The JSON data</param>
+		/// <param name="output">Output buffer for the converted line</param>
+		/// <param name="outputOffset">Offset within the buffer to write the converted data</param>
+		/// <returns></returns>
+		public static int ConvertToPlainText(ReadOnlySpan<byte> input, byte[] output, int outputOffset)
+		{
+			if (IsEmptyOrWhitespace(input))
+			{
+				output[outputOffset] = (byte)'\n';
+				return outputOffset + 1;
+			}
+
+			Utf8JsonReader reader = new Utf8JsonReader(input);
+			if (reader.Read() && reader.TokenType == JsonTokenType.StartObject)
+			{
+				while (reader.Read() && reader.TokenType == JsonTokenType.PropertyName)
+				{
+					if (!reader.ValueTextEquals("message"))
+					{
+						reader.Skip();
+						continue;
+					}
+					if (!reader.Read() || reader.TokenType != JsonTokenType.String)
+					{
+						reader.Skip();
+						continue;
+					}
+
+					int unescapedLength = UnescapeUtf8(reader.ValueSpan, output.AsSpan(outputOffset));
+					outputOffset += unescapedLength;
+
+					output[outputOffset] = (byte)'\n';
+					outputOffset++;
+
+					break;
+				}
+			}
+			return outputOffset;
+		}
+
+		/// <summary>
+		/// Determines if the given line is empty
+		/// </summary>
+		/// <param name="input">The input data</param>
+		/// <returns>True if the given text is empty</returns>
+		static bool IsEmptyOrWhitespace(ReadOnlySpan<byte> input)
+		{
+			for (int idx = 0; idx < input.Length; idx++)
+			{
+				byte v = input[idx];
+				if (v != (byte)'\n' && v != '\r' && v != ' ')
+				{
+					return false;
+				}
+			}
+			return true;
+		}
+
+		/// <summary>
+		/// Unescape a json utf8 string
+		/// </summary>
+		/// <param name="source">Source span of bytes</param>
+		/// <param name="target">Target span of bytes</param>
+		/// <returns>Length of the converted data</returns>
+		static int UnescapeUtf8(ReadOnlySpan<byte> source, Span<byte> target)
+		{
+			int length = 0;
+			for (; ; )
+			{
+				// Copy up to the next backslash
+				int backslash = source.IndexOf((byte)'\\');
+				if (backslash == -1)
+				{
+					source.CopyTo(target);
+					length += source.Length;
+					break;
+				}
+				else if (backslash > 0)
+				{
+					source.Slice(0, backslash).CopyTo(target);
+					source = source.Slice(backslash);
+					target = target.Slice(backslash);
+					length += backslash;
+				}
+
+				// Check what the escape code is
+				if (source[1] == 'u')
+				{
+					char[] chars = { (char)((StringUtils.ParseHexByte(source, 2) << 8) | StringUtils.ParseHexByte(source, 4)) };
+					int encodedLength = Encoding.UTF8.GetBytes(chars.AsSpan(), target);
+					source = source.Slice(6);
+					target = target.Slice(encodedLength);
+					length += encodedLength;
+				}
+				else
+				{
+					target[0] = source[1] switch
+					{
+						(byte)'\"' => (byte)'\"',
+						(byte)'\\' => (byte)'\\',
+						(byte)'b' => (byte)'\b',
+						(byte)'f' => (byte)'\f',
+						(byte)'n' => (byte)'\n',
+						(byte)'r' => (byte)'\r',
+						(byte)'t' => (byte)'\t',
+						_ => source[1]
+					};
+					source = source.Slice(2);
+					target = target.Slice(1);
+					length++;
+				}
+			}
+			return length;
+		}
+
 		#region Log events
 
 		/// <inheritdoc/>
@@ -925,522 +1032,5 @@ namespace HordeServer.Logs
 			/// <inheritdoc/>
 			public override void Write(byte[] buffer, int offset, int count) => throw new NotImplementedException();
 		}
-
-		#region V1
-
-		class LogV1 : ILog
-		{
-			readonly LogCollection _collection;
-			readonly LogDocument _document;
-
-			LogId ILog.Id => _document.Id;
-			JobId ILog.JobId => _document.JobId;
-			LeaseId? ILog.LeaseId => _document.LeaseId;
-			SessionId? ILog.SessionId => _document.SessionId;
-			LogType ILog.Type => _document.Type;
-			NamespaceId ILog.NamespaceId => _document.NamespaceId;
-			RefName ILog.RefName => _document.RefName;
-
-			public LogV1(LogCollection collection, LogDocument document)
-			{
-				_collection = collection;
-				_document = document;
-			}
-
-			public async Task<ILog> UpdateLineCountAsync(int lineCount, bool complete, CancellationToken cancellationToken = default)
-			{
-				LogDocument newDocument = await _collection.UpdateLineCountAsync(_document, lineCount, complete, cancellationToken);
-				return new Log(_collection, newDocument);
-			}
-
-			public Task DeleteAsync(CancellationToken cancellationToken = default)
-				=> _collection.DeleteAsync(_document, cancellationToken);
-
-			public Task<List<Utf8String>> ReadLinesAsync(int index, int count, CancellationToken cancellationToken = default)
-				=> _collection.ReadLinesV1Async(_document, index, count, cancellationToken);
-
-			public Task<LogMetadata> GetMetadataAsync(CancellationToken cancellationToken)
-				=> _collection.GetMetadataV1Async(_document);
-
-			public Task<Stream> OpenRawStreamAsync(CancellationToken cancellationToken = default)
-				=> _collection.OpenRawStreamV1Async(_document, 0, Int64.MaxValue);
-
-			public Task<Stream> OpenRawStreamAsync(long offset, long length, CancellationToken cancellationToken)
-				=> _collection.OpenRawStreamV1Async(_document, offset, length);
-
-			public Task CopyPlainTextStreamAsync(Stream outputStream, CancellationToken cancellationToken = default)
-				=> _collection.CopyPlainTextStreamAsync(this, outputStream, cancellationToken);
-
-			public Task<List<int>> SearchLogDataAsync(string text, int firstLine, int count, SearchStats stats, CancellationToken cancellationToken)
-				=> _collection.SearchLogDataV1Async(_document, text, firstLine, count, stats, cancellationToken);
-
-			public Task AddEventsAsync(List<NewLogEventData> newEvents, CancellationToken cancellationToken = default)
-				=> _collection.AddEventsAsync(_document.Id, newEvents, cancellationToken);
-
-			public Task<List<ILogEvent>> GetEventsAsync(ObjectId? spanId = null, int? index = null, int? count = null, CancellationToken cancellationToken = default)
-				=> _collection.GetEventsAsync(this, spanId, index, count, cancellationToken);
-		}
-
-		async Task<List<Utf8String>> ReadLinesV1Async(LogDocument log, int index, int count, CancellationToken cancellationToken)
-		{
-			List<Utf8String> lines = new List<Utf8String>();
-
-			(_, long minOffset) = await GetLineOffsetAsync(log, index, cancellationToken);
-			(_, long maxOffset) = await GetLineOffsetAsync(log, index + Math.Min(count, Int32.MaxValue - index), cancellationToken);
-
-			byte[] result;
-			using (System.IO.Stream stream = await OpenRawStreamAsync(log, minOffset, maxOffset - minOffset, cancellationToken))
-			{
-				result = new byte[stream.Length];
-				await stream.ReadFixedSizeDataAsync(result, 0, result.Length, cancellationToken);
-			}
-
-			int offset = 0;
-			for (int idx = 0; idx < result.Length; idx++)
-			{
-				if (result[idx] == (byte)'\n')
-				{
-					lines.Add(new Utf8String(result.AsMemory(offset, idx - offset)));
-					offset = idx + 1;
-				}
-			}
-
-			return lines;
-		}
-
-		/// <inheritdoc/>
-		async Task<LogMetadata> GetMetadataV1Async(LogDocument log)
-		{
-			LogMetadata metadata = new LogMetadata();
-			if (log.Chunks.Count > 0)
-			{
-				LogChunkDocument chunk = log.Chunks[log.Chunks.Count - 1];
-				if (log.MaxLineIndex == null || chunk.Length == 0)
-				{
-					LogChunkData chunkData = await ReadChunkV1Async(log, log.Chunks.Count - 1);
-					metadata.Length = chunk.Offset + chunkData.Length;
-					metadata.MaxLineIndex = chunk.LineIndex + chunkData.LineCount;
-				}
-				else
-				{
-					metadata.Length = chunk.Offset + chunk.Length;
-					metadata.MaxLineIndex = log.MaxLineIndex.Value;
-				}
-			}
-			return metadata;
-		}
-
-		/// <inheritdoc/>
-		async Task<Stream> OpenRawStreamV1Async(LogDocument log, long offset, long length)
-		{
-			if (log.Chunks.Count == 0)
-			{
-				return new MemoryStream(Array.Empty<byte>(), false);
-			}
-			else
-			{
-				int lastChunkIdx = log.Chunks.Count - 1;
-
-				// Clamp the length of the request
-				LogChunkDocument lastChunk = log.Chunks[lastChunkIdx];
-				if (length > lastChunk.Offset)
-				{
-					long lastChunkLength = lastChunk.Length;
-					if (lastChunkLength <= 0)
-					{
-						LogChunkData lastChunkData = await ReadChunkV1Async(log, lastChunkIdx);
-						lastChunkLength = lastChunkData.Length;
-					}
-					length = Math.Min(length, (lastChunk.Offset + lastChunkLength) - offset);
-				}
-
-				// Create the new stream
-				return new ResponseStreamV1(this, log, offset, length);
-			}
-		}
-
-		/// <inheritdoc/>
-		async Task<List<int>> SearchLogDataV1Async(LogDocument log, string text, int firstLine, int count, SearchStats searchStats, CancellationToken cancellationToken)
-		{
-			Stopwatch timer = Stopwatch.StartNew();
-
-			using TelemetrySpan span = _tracer.StartActiveSpan($"{nameof(LogCollection)}.{nameof(SearchLogDataAsync)}");
-			span.SetAttribute("logId", log.Id.ToString());
-			span.SetAttribute("text", text);
-			span.SetAttribute("count", count);
-
-			List<int> results = new List<int>();
-			if (count > 0)
-			{
-				IAsyncEnumerable<int> enumerable = SearchLogDataInternalV1Async(log, text, firstLine, searchStats);
-
-				await using IAsyncEnumerator<int> enumerator = enumerable.GetAsyncEnumerator(cancellationToken);
-				while (await enumerator.MoveNextAsync() && results.Count < count)
-				{
-					results.Add(enumerator.Current);
-				}
-			}
-
-			_logger.LogDebug("Search for \"{SearchText}\" in log {LogId} found {NumResults}/{MaxResults} results, took {Time}ms ({@Stats})", text, log.Id, results.Count, count, timer.ElapsedMilliseconds, searchStats);
-			return results;
-		}
-
-		/// <summary>
-		/// Reads a chunk from storage
-		/// </summary>
-		/// <param name="log">Log file to read from</param>
-		/// <param name="chunkIdx">The chunk to read</param>
-		/// <returns>Chunk data</returns>
-		private async Task<LogChunkData> ReadChunkV1Async(LogDocument log, int chunkIdx)
-		{
-			LogChunkDocument chunk = log.Chunks[chunkIdx];
-
-			// Try to read the chunk data from storage
-			LogChunkData? chunkData = null;
-			try
-			{
-				chunkData = await _storage.ReadChunkAsync(log.Id, chunk.Offset, chunk.LineIndex);
-			}
-			catch (Exception ex)
-			{
-				_logger.LogError(ex, "Unable to read log {LogId} at offset {Offset}", log.Id, chunk.Offset);
-			}
-
-			// Get the minimum length and line count for the chunk
-			if (chunkIdx + 1 < log.Chunks.Count)
-			{
-				LogChunkDocument nextChunk = log.Chunks[chunkIdx + 1];
-				chunkData = RepairChunkDataV1(log, chunkIdx, chunkData, (int)(nextChunk.Offset - chunk.Offset), nextChunk.LineIndex - chunk.LineIndex, $"before next");
-			}
-			else
-			{
-				if (log.MaxLineIndex != null && chunk.Length != 0)
-				{
-					chunkData = RepairChunkDataV1(log, chunkIdx, chunkData, chunk.Length, log.MaxLineIndex.Value - chunk.LineIndex, $"last chunk (max line index = {log.MaxLineIndex})");
-				}
-				else
-				{
-					chunkData ??= RepairChunkDataV1(log, chunkIdx, chunkData, 1024, 1, "default");
-				}
-			}
-
-			return chunkData;
-		}
-
-		/// <summary>
-		/// Validates the given chunk data, and fix it up if necessary
-		/// </summary>
-		/// <param name="log">The log file instance</param>
-		/// <param name="chunkIdx">Index of the chunk within the log</param>
-		/// <param name="chunkData">The chunk data that was read</param>
-		/// <param name="length">Expected length of the data</param>
-		/// <param name="lineCount">Expected number of lines in the data</param>
-		/// <param name="context">Context string for diagnostic output</param>
-		/// <returns>Repaired chunk data</returns>
-		LogChunkData RepairChunkDataV1(LogDocument log, int chunkIdx, LogChunkData? chunkData, int length, int lineCount, string context)
-		{
-			int currentLength = 0;
-			int currentLineCount = 0;
-			if (chunkData != null)
-			{
-				currentLength = chunkData.Length;
-				currentLineCount = chunkData.LineCount;
-			}
-
-			if (chunkData == null || currentLength < length || currentLineCount < lineCount)
-			{
-				_logger.LogWarning("Creating placeholder subchunk for log {LogId} chunk {ChunkIdx} (length {Length} vs expected {ExpLength}, lines {LineCount} vs expected {ExpLineCount}, context {Context})", log.Id, chunkIdx, currentLength, length, currentLineCount, lineCount, context);
-
-				List<LogSubChunkData> subChunks = new List<LogSubChunkData>();
-				if (chunkData != null && chunkData.Length < length && chunkData.LineCount < lineCount)
-				{
-					subChunks.AddRange(chunkData.SubChunks);
-				}
-
-				LogText text = new LogText();
-				text.AppendMissingDataInfo(chunkIdx, log.Chunks[chunkIdx].Server, length - currentLength, lineCount - currentLineCount);
-				subChunks.Add(new LogSubChunkData(log.Type, currentLength, currentLineCount, text));
-
-				LogChunkDocument chunk = log.Chunks[chunkIdx];
-				chunkData = new LogChunkData(chunk.Offset, chunk.LineIndex, subChunks);
-			}
-			return chunkData;
-		}
-
-		async IAsyncEnumerable<int> SearchLogDataInternalV1Async(LogDocument log, string text, int firstLine, SearchStats searchStats)
-		{
-			SearchText searchText = new SearchText(text);
-
-			// Read the index for this log file
-			if (log.IndexLength != null)
-			{
-				LogIndexData? indexData = await ReadIndexV1Async(log, log.IndexLength.Value);
-				if (indexData != null && firstLine < indexData.LineCount)
-				{
-					using TelemetrySpan span = _tracer.StartActiveSpan($"{nameof(LogCollection)}.{nameof(SearchLogDataInternalV1Async)}.Indexed");
-					span.SetAttribute("lineCount", indexData.LineCount);
-
-					foreach (int lineIndex in indexData.Search(firstLine, searchText, searchStats))
-					{
-						yield return lineIndex;
-					}
-
-					firstLine = indexData.LineCount;
-				}
-			}
-
-			// Manually search through the rest of the log
-			int chunkIdx = GetChunkForLineV1(log.Chunks, firstLine);
-			for (; chunkIdx < log.Chunks.Count; chunkIdx++)
-			{
-				LogChunkDocument chunk = log.Chunks[chunkIdx];
-
-				// Read the chunk data
-				LogChunkData chunkData = await ReadChunkV1Async(log, chunkIdx);
-				if (firstLine < chunkData.LineIndex + chunkData.LineCount)
-				{
-					// Find the first sub-chunk we're looking for
-					int subChunkIdx = 0;
-					if (firstLine > chunk.LineIndex)
-					{
-						subChunkIdx = chunkData.GetSubChunkForLine(firstLine - chunk.LineIndex);
-					}
-
-					// Search through the sub-chunks
-					for (; subChunkIdx < chunkData.SubChunks.Count; subChunkIdx++)
-					{
-						LogSubChunkData subChunkData = chunkData.SubChunks[subChunkIdx];
-						if (firstLine < subChunkData.LineIndex + subChunkData.LineCount)
-						{
-							// Create an index containing just this sub-chunk
-							LogIndexData index = subChunkData.BuildIndex(_logger);
-							foreach (int lineIndex in index.Search(firstLine, searchText, searchStats))
-							{
-								yield return lineIndex;
-							}
-						}
-					}
-				}
-			}
-		}
-
-		/// <summary>
-		/// Reads a chunk from storage
-		/// </summary>
-		/// <param name="log">Log file to read from</param>
-		/// <param name="length">Length of the log covered by the index</param>
-		/// <returns>Chunk data</returns>
-		private async Task<LogIndexData?> ReadIndexV1Async(LogDocument log, long length)
-		{
-			try
-			{
-				LogIndexData? index = await _storage.ReadIndexAsync(log.Id, length);
-				return index;
-			}
-			catch (Exception ex)
-			{
-				_logger.LogError(ex, "Unable to read log {LogId} index at length {Length}", log.Id, length);
-				return null;
-			}
-		}
-
-		/// <summary>
-		/// Gets the chunk index containing the given offset.
-		/// </summary>
-		/// <param name="chunks">The chunks to search</param>
-		/// <param name="offset">The offset to search for</param>
-		/// <returns>The chunk index containing the given offset</returns>
-		static int GetChunkForOffsetV1(IReadOnlyList<LogChunkDocument> chunks, long offset)
-		{
-			int chunkIndex = chunks.BinarySearch(x => x.Offset, offset);
-			if (chunkIndex < 0)
-			{
-				chunkIndex = ~chunkIndex - 1;
-			}
-			return chunkIndex;
-		}
-
-		/// <summary>
-		/// Gets the starting chunk index for the given line
-		/// </summary>
-		/// <param name="chunks">The chunks to search</param>
-		/// <param name="lineIndex">Index of the line to query</param>
-		/// <returns>Index of the chunk to fetch</returns>
-		static int GetChunkForLineV1(IReadOnlyList<LogChunkDocument> chunks, int lineIndex)
-		{
-			int chunkIndex = chunks.BinarySearch(x => x.LineIndex, lineIndex);
-			if (chunkIndex < 0)
-			{
-				chunkIndex = ~chunkIndex - 1;
-			}
-			return chunkIndex;
-		}
-
-		/// <summary>
-		/// Streams log data to a caller
-		/// </summary>
-		class ResponseStreamV1 : Stream
-		{
-			/// <summary>
-			/// The log file service that created this stream
-			/// </summary>
-			readonly LogCollection _logCollection;
-
-			/// <summary>
-			/// The log file being read
-			/// </summary>
-			readonly LogDocument _log;
-
-			/// <summary>
-			/// Starting offset within the file of the data to return 
-			/// </summary>
-			readonly long _responseOffset;
-
-			/// <summary>
-			/// Length of data to return
-			/// </summary>
-			readonly long _responseLength;
-
-			/// <summary>
-			/// Current offset within the stream
-			/// </summary>
-			long _currentOffset;
-
-			/// <summary>
-			/// The current chunk index
-			/// </summary>
-			int _chunkIdx;
-
-			/// <summary>
-			/// Buffer containing a message for missing data
-			/// </summary>
-			ReadOnlyMemory<byte> _sourceBuffer;
-
-			/// <summary>
-			/// Offset within the source buffer
-			/// </summary>
-			int _sourcePos;
-
-			/// <summary>
-			/// Length of the source buffer being copied from
-			/// </summary>
-			int _sourceEnd;
-
-			/// <summary>
-			/// Constructor
-			/// </summary>
-			/// <param name="logService">The log file service, for q</param>
-			/// <param name="log"></param>
-			/// <param name="offset"></param>
-			/// <param name="length"></param>
-			public ResponseStreamV1(LogCollection logService, LogDocument log, long offset, long length)
-			{
-				_logCollection = logService;
-				_log = log;
-
-				_responseOffset = offset;
-				_responseLength = length;
-
-				_currentOffset = offset;
-
-				_chunkIdx = GetChunkForOffsetV1(log.Chunks, offset);
-				_sourceBuffer = null!;
-			}
-
-			/// <inheritdoc/>
-			public override bool CanRead => true;
-
-			/// <inheritdoc/>
-			public override bool CanSeek => false;
-
-			/// <inheritdoc/>
-			public override bool CanWrite => false;
-
-			/// <inheritdoc/>
-			public override long Length => _responseLength;
-
-			/// <inheritdoc/>
-			public override long Position
-			{
-				get => _currentOffset - _responseOffset;
-				set => throw new NotImplementedException();
-			}
-
-			/// <inheritdoc/>
-			public override void Flush()
-			{
-			}
-
-			/// <inheritdoc/>
-			public override int Read(byte[] buffer, int offset, int count)
-			{
-#pragma warning disable VSTHRD002
-				return ReadAsync(buffer, offset, count, CancellationToken.None).Result;
-#pragma warning restore VSTHRD002
-			}
-
-			/// <inheritdoc/>
-			public override async Task<int> ReadAsync(byte[] buffer, int offset, int length, CancellationToken cancellationToken)
-			{
-				return await ReadAsync(buffer.AsMemory(offset, length), cancellationToken);
-			}
-
-			/// <inheritdoc/>
-			public override async ValueTask<int> ReadAsync(Memory<byte> buffer, CancellationToken cancellationToken)
-			{
-				int readBytes = 0;
-				while (readBytes < buffer.Length)
-				{
-					if (_sourcePos < _sourceEnd)
-					{
-						// Try to copy from the current buffer
-						int blockSize = Math.Min(_sourceEnd - _sourcePos, buffer.Length - readBytes);
-						_sourceBuffer.Slice(_sourcePos, blockSize).Span.CopyTo(buffer.Slice(readBytes).Span);
-						_currentOffset += blockSize;
-						readBytes += blockSize;
-						_sourcePos += blockSize;
-					}
-					else if (_currentOffset < _responseOffset + _responseLength)
-					{
-						// Move to the right chunk
-						while (_chunkIdx + 1 < _log.Chunks.Count && _currentOffset >= _log.Chunks[_chunkIdx + 1].Offset)
-						{
-							_chunkIdx++;
-						}
-
-						// Get the chunk data
-						LogChunkDocument chunk = _log.Chunks[_chunkIdx];
-						LogChunkData chunkData = await _logCollection.ReadChunkV1Async(_log, _chunkIdx);
-
-						// Figure out which sub-chunk to use
-						int subChunkIdx = chunkData.GetSubChunkForOffsetWithinChunk((int)(_currentOffset - chunk.Offset));
-						LogSubChunkData subChunkData = chunkData.SubChunks[subChunkIdx];
-
-						// Get the source data
-						long subChunkOffset = chunk.Offset + chunkData.SubChunkOffset[subChunkIdx];
-						_sourceBuffer = subChunkData.InflateText().Data;
-						_sourcePos = (int)(_currentOffset - subChunkOffset);
-						_sourceEnd = (int)Math.Min(_sourceBuffer.Length, (_responseOffset + _responseLength) - subChunkOffset);
-					}
-					else
-					{
-						// End of the log
-						break;
-					}
-				}
-				return readBytes;
-			}
-
-			/// <inheritdoc/>
-			public override long Seek(long offset, SeekOrigin origin) => throw new NotImplementedException();
-
-			/// <inheritdoc/>
-			public override void SetLength(long value) => throw new NotImplementedException();
-
-			/// <inheritdoc/>
-			public override void Write(byte[] buffer, int offset, int count) => throw new NotImplementedException();
-		}
-
-		#endregion
 	}
 }
diff --git a/Engine/Source/Programs/Horde/HordeServer/Logs/Storage/LocalLogStorage.cs b/Engine/Source/Programs/Horde/HordeServer/Logs/Storage/LocalLogStorage.cs
deleted file mode 100644
diff --git a/Engine/Source/Programs/Horde/HordeServer/Logs/Storage/NullLogStorage.cs b/Engine/Source/Programs/Horde/HordeServer/Logs/Storage/NullLogStorage.cs
deleted file mode 100644
diff --git a/Engine/Source/Programs/Horde/HordeServer/Logs/Storage/PersistentLogStorage.cs b/Engine/Source/Programs/Horde/HordeServer/Logs/Storage/PersistentLogStorage.cs
deleted file mode 100644
diff --git a/Engine/Source/Programs/Horde/HordeServer/Logs/Storage/SequencedLogStorage.cs b/Engine/Source/Programs/Horde/HordeServer/Logs/Storage/SequencedLogStorage.cs
deleted file mode 100644
diff --git a/Engine/Source/Programs/Horde/HordeServer/ServerSettings.cs b/Engine/Source/Programs/Horde/HordeServer/ServerSettings.cs
--- a/Engine/Source/Programs/Horde/HordeServer/ServerSettings.cs
+++ b/Engine/Source/Programs/Horde/HordeServer/ServerSettings.cs
@@ -23,119 +23,6 @@ using HordeServer.Tools;
 
 namespace HordeServer
 {
-	/// <summary>
-	/// Types of storage backend to use
-	/// </summary>
-	public enum StorageBackendType
-	{
-		/// <summary>
-		/// Local filesystem
-		/// </summary>
-		FileSystem,
-
-		/// <summary>
-		/// AWS S3
-		/// </summary>
-		Aws,
-
-		/// <summary>
-		/// Azure blob store
-		/// </summary>
-		Azure,
-
-		/// <summary>
-		/// In-memory only (for testing)
-		/// </summary>
-		Memory,
-	};
-
-	/// <summary>
-	/// Common settings for different storage backends
-	/// </summary>
-	public interface IStorageBackendOptions : IAwsStorageOptions, IAzureStorageOptions
-	{
-		/// <summary>
-		/// Base directory for filesystem storage
-		/// </summary>
-		string? BaseDir { get; }
-
-		/// <summary>
-		/// The type of storage backend to use
-		/// </summary>
-		StorageBackendType? Type { get; }
-	}
-
-	/// <summary>
-	/// Common settings object for different providers
-	/// </summary>
-	public class StorageBackendOptions : IStorageBackendOptions
-	{
-		/// <inheritdoc/>
-		public StorageBackendType? Type { get; set; }
-
-		/// <inheritdoc/>
-		public string? BaseDir { get; set; }
-
-		/// <inheritdoc/>
-		public string? AwsBucketName { get; set; }
-
-		/// <inheritdoc/>
-		public string? AwsBucketPath { get; set; }
-
-		/// <inheritdoc/>
-		public AwsCredentialsType? AwsCredentials { get; set; }
-
-		/// <inheritdoc/>
-		public string? AwsRole { get; set; }
-
-		/// <inheritdoc/>
-		public string? AwsProfile { get; set; }
-
-		/// <inheritdoc/>
-		public string? AwsRegion { get; set; }
-
-		/// <inheritdoc/>
-		public string? AzureConnectionString { get; set; }
-
-		/// <inheritdoc/>
-		public string? AzureContainerName { get; set; }
-	}
-
-	/// <summary>
-	/// Options for configuring a blob store
-	/// </summary>
-	public class BlobStoreOptions : StorageBackendOptions
-	{
-	}
-
-	/// <summary>
-	/// Options for configuring the default tree store implementation
-	/// </summary>
-	public interface ITreeStoreOptions
-	{
-		/// <summary>
-		/// Options for creating bundles
-		/// </summary>
-		BundleOptions Bundle { get; }
-
-		/// <summary>
-		/// Options for chunking content
-		/// </summary>
-		ChunkingOptions Chunking { get; }
-	}
-
-	/// <summary>
-	/// Options for storing trees
-	/// </summary>
-	public class TreeStoreOptions : BlobStoreOptions, ITreeStoreOptions
-	{
-		/// <inheritdoc/>
-		public BundleOptions Bundle { get; set; } = new BundleOptions();
-
-		/// <inheritdoc/>
-		public ChunkingOptions Chunking { get; set; } = new ChunkingOptions();
-	}
-
 	/// <summary>
 	/// Type of run mode this process should use. Each carry different types of workloads. 
 	/// More than one mode can be active. But not all modes are not guaranteed to be compatible with each other and will
@@ -614,32 +501,11 @@ namespace HordeServer
 		/// </summary>
 		public bool RedisReadOnlyMode { get; set; }
 
-		/// <summary>
-		/// Type of write cache to use in log service
-		/// Currently Supported: "InMemory" or "Redis"
-		/// </summary>
-		public string LogServiceWriteCacheType { get; set; } = "InMemory";
-
 		/// <summary>
 		/// Overridden settings for storage backends. Useful for running against a production server with custom backends.
 		/// </summary>
 		public List<BackendConfig> StorageBackends { get; set; } = new List<BackendConfig>();
 
-		/// <summary>
-		/// Settings for artifact storage
-		/// </summary>
-		public StorageBackendOptions LogStorage { get; set; } = new StorageBackendOptions() { BaseDir = "Logs" };
-
-		/// <summary>
-		/// Settings for artifact storage
-		/// </summary>
-		public StorageBackendOptions ArtifactStorage { get; set; } = new StorageBackendOptions() { BaseDir = "Artifacts" };
-
-		/// <summary>
-		/// Configuration of tree storage
-		/// </summary>
-		public TreeStoreOptions CommitStorage { get; set; } = new TreeStoreOptions() { BaseDir = "Commits" };
-
 		/// <summary>
 		/// Whether to log json to stdout
 		/// </summary>
@@ -650,11 +516,6 @@ namespace HordeServer
 		/// </summary>
 		public bool LogSessionRequests { get; set; } = false;
 
-		/// <summary>
-		/// Whether to enable the hosted LogService running background jobs
-		/// </summary>
-		public bool EnableLogService { get; set; } = true;
-
 		/// <summary>
 		/// Default fleet manager to use (when not specified by pool)
 		/// </summary>
diff --git a/Engine/Source/Programs/Horde/HordeServer/Startup.cs b/Engine/Source/Programs/Horde/HordeServer/Startup.cs
--- a/Engine/Source/Programs/Horde/HordeServer/Startup.cs
+++ b/Engine/Source/Programs/Horde/HordeServer/Startup.cs
@@ -59,7 +59,6 @@ using HordeServer.Devices;
 using HordeServer.Issues;
 using HordeServer.Issues.External;
 using HordeServer.Jobs;
-using HordeServer.Jobs.Artifacts;
 using HordeServer.Jobs.Bisect;
 using HordeServer.Jobs.Graphs;
 using HordeServer.Jobs.Schedules;
@@ -67,7 +66,6 @@ using HordeServer.Jobs.Templates;
 using HordeServer.Jobs.TestData;
 using HordeServer.Jobs.Timing;
 using HordeServer.Logs;
-using HordeServer.Logs.Storage;
 using HordeServer.Notifications;
 using HordeServer.Notifications.Sinks;
 using HordeServer.Perforce;
@@ -258,21 +256,6 @@ namespace HordeServer
 			}
 		}
 
-		static IObjectStore CreateObjectStore(IServiceProvider sp, StorageBackendOptions options)
-		{
-			switch (options.Type ?? StorageBackendType.FileSystem)
-			{
-				case StorageBackendType.FileSystem:
-					return sp.GetRequiredService<FileObjectStoreFactory>().CreateStore(DirectoryReference.Combine(ServerApp.DataDir, options.BaseDir ?? "Storage"));
-				case StorageBackendType.Aws:
-					return sp.GetRequiredService<AwsObjectStoreFactory>().CreateStore(options);
-				case StorageBackendType.Memory:
-					return new MemoryObjectStore();
-				default:
-					throw new NotImplementedException();
-			}
-		}
-
 		public Startup(IConfiguration configuration)
 		{
 			Configuration = configuration;
@@ -404,7 +387,6 @@ namespace HordeServer
 			services.AddSingleton<IAccountCollection, AccountCollection>();
 			services.AddSingleton<IAgentCollection, AgentCollection>();
 			services.AddSingleton<IArtifactCollection, ArtifactCollection>();
-			services.AddSingleton<IArtifactCollectionV1, ArtifactCollectionV1>();
 			services.AddSingleton<IGraphCollection, GraphCollection>();
 			services.AddSingleton<IIssueCollection, IssueCollection>();
 			services.AddSingleton<IJobCollection, JobCollection>();
@@ -564,8 +546,6 @@ namespace HordeServer
 			services.AddSingleton<IObjectStoreFactory, ObjectStoreFactory>();
 			services.AddSingleton<AwsObjectStoreFactory>();
 			services.AddSingleton<FileObjectStoreFactory>();
-			services.AddSingleton(sp => CreateObjectStore(sp, settings.LogStorage).ForType<PersistentLogStorage>());
-			services.AddSingleton(sp => CreateObjectStore(sp, settings.ArtifactStorage).ForType<ArtifactCollectionV1>());
 
 			if (settings.WithAws)
 			{
@@ -587,8 +567,6 @@ namespace HordeServer
 				services.AddSingleton<IPoolSizeStrategyFactory, ComputeQueueAwsMetricStrategyFactory>();
 			}
 
-			ConfigureLogStorage(services);
-
 			AuthenticationBuilder authBuilder = services.AddAuthentication(options =>
 				{
 					switch (settings.AuthMethod)
@@ -953,19 +931,6 @@ namespace HordeServer
 			LogValueFormatter.RegisterTypeAnnotation<UserId>("UserId");
 		}
 
-		private static void ConfigureLogStorage(IServiceCollection services)
-		{
-			services.AddSingleton<PersistentLogStorage>();
-
-			services.AddSingleton<ILogStorage>(provider =>
-			{
-				ILogStorage storage = provider.GetRequiredService<PersistentLogStorage>();
-				storage = new SequencedLogStorage(storage);
-				storage = new LocalLogStorage(50, storage);
-				return storage;
-			});
-		}
-
 		public sealed class BlobLocatorBsonSerializer : SerializerBase<BlobLocator>
 		{
 			/// <inheritdoc/>
diff --git a/Engine/Source/Programs/Horde/HordeServer/Storage/StorageConfig.cs b/Engine/Source/Programs/Horde/HordeServer/Storage/StorageConfig.cs
--- a/Engine/Source/Programs/Horde/HordeServer/Storage/StorageConfig.cs
+++ b/Engine/Source/Programs/Horde/HordeServer/Storage/StorageConfig.cs
@@ -169,11 +169,37 @@ namespace HordeServer.Storage
 		public bool TryGetNamespace(NamespaceId namespaceId, [NotNullWhen(true)] out NamespaceConfig? namespaceConfig) => _namespaceLookup.TryGetValue(namespaceId, out namespaceConfig);
 	}
 
+	/// <summary>
+	/// Types of storage backend to use
+	/// </summary>
+	public enum StorageBackendType
+	{
+		/// <summary>
+		/// Local filesystem
+		/// </summary>
+		FileSystem,
+
+		/// <summary>
+		/// AWS S3
+		/// </summary>
+		Aws,
+
+		/// <summary>
+		/// Azure blob store
+		/// </summary>
+		Azure,
+
+		/// <summary>
+		/// In-memory only (for testing)
+		/// </summary>
+		Memory,
+	};
+
 	/// <summary>
 	/// Common settings object for different providers
 	/// </summary>
 	[DebuggerDisplay("{Id}")]
-	public class BackendConfig : IStorageBackendOptions
+	public class BackendConfig : IAwsStorageOptions
 	{
 		/// <summary>
 		/// The storage backend ID
