Horde: Build fixes.

#rnx

diff --git a/Engine/Source/Programs/Horde/HordeServer.Tests/ControllerIntegrationTest.cs b/Engine/Source/Programs/Horde/HordeServer.Tests/ControllerIntegrationTest.cs
--- a/Engine/Source/Programs/Horde/HordeServer.Tests/ControllerIntegrationTest.cs
+++ b/Engine/Source/Programs/Horde/HordeServer.Tests/ControllerIntegrationTest.cs
@@ -155,11 +155,10 @@ public class ControllerIntegrationTest : IAsyncDisposable
 		ConfigService configService = services.GetRequiredService<ConfigService>();
 		ITemplateCollection templateService = services.GetRequiredService<ITemplateCollection>();
 		JobService jobService = services.GetRequiredService<JobService>();
-		IArtifactCollectionV1 artifactCollection = services.GetRequiredService<IArtifactCollectionV1>();
 		AgentService agentService = services.GetRequiredService<AgentService>();
 		IGraphCollection graphCollection = services.GetRequiredService<IGraphCollection>();
 		IOptions<ServerSettings> serverSettings = services.GetRequiredService<IOptions<ServerSettings>>();
 
-		return await Fixture.CreateAsync(configService, graphCollection, templateService, jobService, artifactCollection, agentService, serverSettings.Value);
+		return await Fixture.CreateAsync(configService, graphCollection, templateService, jobService, agentService, serverSettings.Value);
 	}
 }
\ No newline at end of file
diff --git a/Engine/Source/Programs/Horde/HordeServer.Tests/Logs/LogChunkTests.cs b/Engine/Source/Programs/Horde/HordeServer.Tests/Logs/LogChunkTests.cs
deleted file mode 100644
diff --git a/Engine/Source/Programs/Horde/HordeServer.Tests/TestSetup.cs b/Engine/Source/Programs/Horde/HordeServer.Tests/TestSetup.cs
--- a/Engine/Source/Programs/Horde/HordeServer.Tests/TestSetup.cs
+++ b/Engine/Source/Programs/Horde/HordeServer.Tests/TestSetup.cs
@@ -224,7 +224,6 @@ namespace HordeServer.Tests
 			services.AddSingleton<IAgentCollection, AgentCollection>();
 			services.AddSingleton<IAgentTelemetryCollection, AgentTelemetryCollection>();
 			services.AddSingleton<IArtifactCollection, ArtifactCollection>();
-			services.AddSingleton<IArtifactCollectionV1, ArtifactCollectionV1>();
 			services.AddSingleton<ICommitService, CommitService>();
 			services.AddSingleton<IGraphCollection, GraphCollection>();
 			services.AddSingleton<IIssueCollection, IssueCollection>();
@@ -281,7 +280,6 @@ namespace HordeServer.Tests
 			services.AddSingleton<JobService>();
 			services.AddSingleton<JobExpirationService>();
 			services.AddSingleton<LifetimeService>();
-			services.AddSingleton<ILogStorage, NullLogStorage>();
 			services.AddSingleton<LogTailService>();
 			services.AddSingleton<INotificationService, NotificationService>();
 			services.AddSingleton<IPerforceService, PerforceServiceStub>();
@@ -303,8 +301,6 @@ namespace HordeServer.Tests
 
 			services.AddSingleton<FileObjectStoreFactory>();
 			services.AddSingleton<IObjectStoreFactory, ObjectStoreFactory>();
-			services.AddSingleton<IObjectStore<PersistentLogStorage>>(sp => new MemoryObjectStore().ForType<PersistentLogStorage>());
-			services.AddSingleton<IObjectStore<ArtifactCollectionV1>>(sp => new MemoryObjectStore().ForType<ArtifactCollectionV1>());
 
 			services.AddSingleton<StorageService>();
 			services.AddSingleton<StorageBackendCache>();
@@ -313,7 +309,7 @@ namespace HordeServer.Tests
 
 		public Task<Fixture> CreateFixtureAsync()
 		{
-			return Fixture.CreateAsync(ConfigService, GraphCollection, TemplateCollection, JobService, ArtifactCollection, AgentService, ServerSettings);
+			return Fixture.CreateAsync(ConfigService, GraphCollection, TemplateCollection, JobService, AgentService, ServerSettings);
 		}
 
 		private JobsController GetJobsController()
diff --git a/Engine/Source/Programs/Horde/HordeServer/Dockerfile b/Engine/Source/Programs/Horde/HordeServer/Dockerfile
--- a/Engine/Source/Programs/Horde/HordeServer/Dockerfile
+++ b/Engine/Source/Programs/Horde/HordeServer/Dockerfile
@@ -48,6 +48,7 @@ RUN dotnet restore Source/Programs/Horde/HordeServer/HordeServer.csproj -p:WithH
 COPY Binaries/DotNET/EpicGames.Perforce.Native/ Binaries/DotNET/EpicGames.Perforce.Native/
 COPY Source/Programs/Shared/ Source/Programs/Shared/
 COPY Source/Programs/Horde/.editorconfig Source/Programs/Horde/.editorconfig
+COPY Source/Programs/Horde/Plugins/ Source/Programs/Horde/Plugins/
 COPY Source/Programs/Horde/HordeServer/ Source/Programs/Horde/HordeServer/
 COPY Source/Programs/Horde/HordeServer.Shared/ Source/Programs/Horde/HordeServer.Shared/
 COPY Source/Programs/AutomationTool/AutomationUtils/Matchers/ Source/Programs/AutomationTool/AutomationUtils/Matchers/
