Horde: Fixes to plugin config initialization for tests.

#rnx

diff --git a/Engine/Source/Programs/Horde/HordeServer.Shared/Configuration/ConfigContext.cs b/Engine/Source/Programs/Horde/HordeServer.Shared/Configuration/ConfigContext.cs
--- a/Engine/Source/Programs/Horde/HordeServer.Shared/Configuration/ConfigContext.cs
+++ b/Engine/Source/Programs/Horde/HordeServer.Shared/Configuration/ConfigContext.cs
@@ -8,6 +8,7 @@ using System.Text.Json;
 using System.Threading;
 using System.Threading.Tasks;
 using EpicGames.Core;
+using EpicGames.Horde.Storage;
 using Microsoft.Extensions.Logging;
 
 namespace HordeServer.Configuration
@@ -180,45 +181,56 @@ namespace HordeServer.Configuration
 			return file;
 		}
 
+		/// <summary>
+		/// Parses a config file as a json object
+		/// </summary>
+		/// <param name="file">File to parse</param>
+		/// <param name="cancellationToken">Cancellation token for the operation</param>
+		/// <returns>The parsed file</returns>
+		public async ValueTask<JsonObject> ParseFileAsync(IConfigFile file, CancellationToken cancellationToken = default)
+		{
+			ReadOnlyMemory<byte> data = await file.ReadAsync(cancellationToken);
+
+			JsonObject? obj = JsonSerializer.Deserialize<JsonObject>(data.Span, JsonOptions);
+			if (obj == null)
+			{
+				throw new ConfigException(this, $"Config file {file.Uri} contains a null object.");
+			}
+
+			return obj;
+		}
+
 		/// <summary>
 		/// Reads an object from a particular URL
 		/// </summary>
-		/// <typeparam name="T">Type of object to read</typeparam>
 		/// <param name="uri">Location of the file to read</param>
+		/// <param name="rootNode">Root node describing the preprocessor fields</param>
 		/// <param name="cancellationToken">Cancellation token for the operation</param>
 		/// <returns></returns>
-		public async Task<T> ReadAsync<T>(Uri uri, CancellationToken cancellationToken) where T : class, new()
+		public async Task<JsonObject> PreprocessFileAsync(Uri uri, ObjectConfigNode rootNode, CancellationToken cancellationToken)
 		{
 			IConfigFile file = await ReadFileAsync(uri, cancellationToken);
 			JsonObject obj = await ParseFileAsync(file, cancellationToken);
 
 			IncludeStack.Push(file);
-
-			ObjectConfigNode type = new ObjectConfigNode(typeof(T));
-			obj = await type.PreprocessAsync(obj, this, cancellationToken);
-
+			obj = await rootNode.PreprocessAsync(obj, this, cancellationToken);
 			IncludeStack.Pop();
 
-			return JsonSerializer.Deserialize<T>(obj, JsonOptions) ?? new T();
+			return obj;
 		}
 
 		/// <summary>
-		/// Parses a config file as a json object
+		/// Reads an object from a particular URL
 		/// </summary>
-		/// <param name="file">File to parse</param>
+		/// <typeparam name="T">Type of object to read</typeparam>
+		/// <param name="uri">Location of the file to read</param>
 		/// <param name="cancellationToken">Cancellation token for the operation</param>
-		/// <returns>The parsed file</returns>
-		public async ValueTask<JsonObject> ParseFileAsync(IConfigFile file, CancellationToken cancellationToken = default)
+		/// <returns></returns>
+		public async Task<T> ReadAsync<T>(Uri uri, CancellationToken cancellationToken) where T : class, new()
 		{
-			ReadOnlyMemory<byte> data = await file.ReadAsync(cancellationToken);
-
-			JsonObject? obj = JsonSerializer.Deserialize<JsonObject>(data.Span, JsonOptions);
-			if (obj == null)
-			{
-				throw new ConfigException(this, $"Config file {file.Uri} contains a null object.");
-			}
-
-			return obj;
+			ObjectConfigNode type = new ObjectConfigNode(typeof(T));
+			JsonObject obj = await PreprocessFileAsync(uri, type, cancellationToken);
+			return JsonSerializer.Deserialize<T>(obj, JsonOptions) ?? new T();
 		}
 	}
 }
diff --git a/Engine/Source/Programs/Horde/HordeServer.Shared/Plugins/PluginName.cs b/Engine/Source/Programs/Horde/HordeServer.Shared/Plugins/PluginName.cs
new file mode 100644
--- /dev/null
+++ b/Engine/Source/Programs/Horde/HordeServer.Shared/Plugins/PluginName.cs
@@ -0,0 +1,44 @@
+// Copyright Epic Games, Inc. All Rights Reserved.
+
+using System.ComponentModel;
+using EpicGames.Core;
+using EpicGames.Horde;
+
+namespace HordeServer.Plugins
+{
+	/// <summary>
+	/// Name of a plugin
+	/// </summary>
+	/// <param name="Id"></param>
+	[LogValueType]
+	[JsonSchemaString]
+	[TypeConverter(typeof(StringIdTypeConverter<PluginName, PluginNameStringIdConverter>))]
+	[StringIdConverter(typeof(PluginNameStringIdConverter))]
+	public record struct PluginName(StringId Id)
+	{
+		/// <summary>
+		/// Constructor
+		/// </summary>
+		/// <param name="name"></param>
+		public PluginName(string name)
+			: this(new StringId(name))
+		{ }
+
+		/// <inheritdoc/>
+		public override string ToString() => Id.ToString();
+	}
+
+	/// <summary>
+	/// Converter from PluginName objects to regular StringId types
+	/// </summary>
+	class PluginNameStringIdConverter : StringIdConverter<PluginName>
+	{
+		/// <inheritdoc/>
+		public override PluginName FromStringId(StringId id)
+			=> new PluginName(id);
+
+		/// <inheritdoc/>
+		public override StringId ToStringId(PluginName value)
+			=> value.Id;
+	}
+}
diff --git a/Engine/Source/Programs/Horde/HordeServer.Tests/Compute/ComputeServiceTest.cs b/Engine/Source/Programs/Horde/HordeServer.Tests/Compute/ComputeServiceTest.cs
--- a/Engine/Source/Programs/Horde/HordeServer.Tests/Compute/ComputeServiceTest.cs
+++ b/Engine/Source/Programs/Horde/HordeServer.Tests/Compute/ComputeServiceTest.cs
@@ -11,6 +11,7 @@ using EpicGames.Horde.Agents.Pools;
 using EpicGames.Horde.Compute;
 using HordeServer.Agents;
 using HordeServer.Compute;
+using HordeServer.Plugins;
 using HordeServer.Server;
 using Microsoft.Extensions.Logging.Abstractions;
 using Microsoft.VisualStudio.TestTools.UnitTesting;
@@ -259,7 +260,7 @@ namespace HordeServer.Tests.Compute
 				new() { Id = cluster1, Condition = $"pool == '{poolFoo.ToString()}'" },
 				new() { Id = cluster2, Condition = $"pool == '{poolBar.ToString()}'" }
 			};
-			GlobalConfig.CurrentValue.PostLoad(new ServerSettings());
+			GlobalConfig.CurrentValue.PostLoad(new ServerSettings(), new List<ILoadedPlugin>());
 
 			List<string> props = ["ComputeIp=11.0.0.1", "ComputePort=5000"];
 			IAgent agent1 = await CreateAgentAsync(new PoolId("foo"), properties: props);
@@ -300,7 +301,7 @@ namespace HordeServer.Tests.Compute
 				new NetworkConfig { CidrBlock = "11.0.0.0/16", Id = "network1", ComputeId = "compute1" },
 				new NetworkConfig { CidrBlock = "12.0.0.0/16", Id = "network2", ComputeId = "compute2" },
 			] };
-			globalConfig.PostLoad(new ServerSettings());
+			globalConfig.PostLoad(new ServerSettings(), new List<ILoadedPlugin>());
 			Assert.AreEqual(new ClusterId("compute1"), ComputeService.FindBestComputeClusterId(globalConfig, IPAddress.Parse("11.0.0.1")));
 			Assert.AreEqual(new ClusterId("compute2"), ComputeService.FindBestComputeClusterId(globalConfig, IPAddress.Parse("12.0.1.1")));
 			Assert.ThrowsException<ComputeServiceException>(() => ComputeService.FindBestComputeClusterId(globalConfig, IPAddress.Parse("123.123.123.123")));
@@ -309,7 +310,7 @@ namespace HordeServer.Tests.Compute
 			{
 				Networks = [new NetworkConfig { CidrBlock = "0.0.0.0/0", Id = "catchAll", ComputeId = "catchAll" }]
 			};
-			globalConfigCatchAll.PostLoad(new ServerSettings());
+			globalConfigCatchAll.PostLoad(new ServerSettings(), new List<ILoadedPlugin>());
 			Assert.AreEqual(new ClusterId("catchAll"), ComputeService.FindBestComputeClusterId(globalConfigCatchAll, IPAddress.Parse("123.123.123.123")));
 		}
 
@@ -379,7 +380,7 @@ namespace HordeServer.Tests.Compute
 				NullLogger<ComputeService>.Instance);
 			await CreateAgentAsync(new PoolId(pool), properties: ["ComputeIp=11.0.0.1", "ComputePort=5000"]);
 			GlobalConfig.CurrentValue.Compute = [ccc];
-			GlobalConfig.CurrentValue.PostLoad(ss);
+			GlobalConfig.CurrentValue.PostLoad(ss, new List<ILoadedPlugin>());
 			return cs;
 		}
 
diff --git a/Engine/Source/Programs/Horde/HordeServer.Tests/Configuration/ConfigTests.cs b/Engine/Source/Programs/Horde/HordeServer.Tests/Configuration/ConfigTests.cs
--- a/Engine/Source/Programs/Horde/HordeServer.Tests/Configuration/ConfigTests.cs
+++ b/Engine/Source/Programs/Horde/HordeServer.Tests/Configuration/ConfigTests.cs
@@ -9,6 +9,7 @@ using System.Threading;
 using System.Threading.Tasks;
 using EpicGames.Core;
 using HordeServer.Configuration;
+using HordeServer.Plugins;
 using HordeServer.Projects;
 using HordeServer.Server;
 using HordeServer.Streams;
@@ -178,7 +179,7 @@ namespace HordeServer.Tests.Configuration
 			inputWorkspaces["subSubType"] = new WorkspaceConfig { Base = "subType", Identifier = "subSubType" };
 			
 			gc.Projects.Add(new ProjectConfig { Streams = [new StreamConfig { WorkspaceTypes = inputWorkspaces }]});
-			gc.PostLoad(new ServerSettings());
+			gc.PostLoad(new ServerSettings(), new List<ILoadedPlugin>());
 
 			Dictionary<string,WorkspaceConfig> workspaces = gc.Projects[0].Streams[0].WorkspaceTypes;
 			Assert.AreEqual(3, workspaces.Count);
@@ -210,7 +211,7 @@ namespace HordeServer.Tests.Configuration
 				]
 			});
 			
-			gc.PostLoad(new ServerSettings());
+			gc.PostLoad(new ServerSettings(), new List<ILoadedPlugin>());
 
 			Dictionary<string, WorkspaceConfig> workspaces = gc.Projects[0].Streams[0].WorkspaceTypes;
 			Assert.AreEqual(4, workspaces.Count);
diff --git a/Engine/Source/Programs/Horde/HordeServer.Tests/ControllerIntegrationTest.cs b/Engine/Source/Programs/Horde/HordeServer.Tests/ControllerIntegrationTest.cs
--- a/Engine/Source/Programs/Horde/HordeServer.Tests/ControllerIntegrationTest.cs
+++ b/Engine/Source/Programs/Horde/HordeServer.Tests/ControllerIntegrationTest.cs
@@ -11,6 +11,7 @@ using HordeServer.Jobs;
 using HordeServer.Jobs.Graphs;
 using HordeServer.Jobs.Templates;
 using HordeServer.Perforce;
+using HordeServer.Plugins;
 using HordeServer.Server;
 using HordeServer.Tests.Stubs.Services;
 using Microsoft.AspNetCore.Hosting;
@@ -157,8 +158,9 @@ public class ControllerIntegrationTest : IAsyncDisposable
 		JobService jobService = services.GetRequiredService<JobService>();
 		AgentService agentService = services.GetRequiredService<AgentService>();
 		IGraphCollection graphCollection = services.GetRequiredService<IGraphCollection>();
+		IPluginCollection pluginCollection = services.GetRequiredService<IPluginCollection>();
 		IOptions<ServerSettings> serverSettings = services.GetRequiredService<IOptions<ServerSettings>>();
 
-		return await Fixture.CreateAsync(configService, graphCollection, templateService, jobService, agentService, serverSettings.Value);
+		return await Fixture.CreateAsync(configService, graphCollection, templateService, jobService, agentService, pluginCollection, serverSettings.Value);
 	}
 }
\ No newline at end of file
diff --git a/Engine/Source/Programs/Horde/HordeServer.Tests/Ddc/Functional/ContentId/CompressedBlobTests.cs b/Engine/Source/Programs/Horde/HordeServer.Tests/Ddc/Functional/ContentId/CompressedBlobTests.cs
--- a/Engine/Source/Programs/Horde/HordeServer.Tests/Ddc/Functional/ContentId/CompressedBlobTests.cs
+++ b/Engine/Source/Programs/Horde/HordeServer.Tests/Ddc/Functional/ContentId/CompressedBlobTests.cs
@@ -17,6 +17,7 @@ using EpicGames.Core;
 using EpicGames.Horde.Storage;
 using HordeServer.Configuration;
 using HordeServer.Ddc;
+using HordeServer.Plugins;
 using HordeServer.Server;
 using HordeServer.Storage;
 using Microsoft.Extensions.DependencyInjection;
@@ -44,7 +45,7 @@ namespace HordeServer.Tests.Ddc.FunctionalTests.CompressedBlobs
 
 			globalConfig.Storage.Backends.Add(new BackendConfig { Id = new BackendId("default"), Type = StorageBackendType.Memory });
 			globalConfig.Storage.Namespaces.Add(new NamespaceConfig { Id = new NamespaceId(TestNamespace), Backend = globalConfig.Storage.Backends[^1].Id });
-			globalConfig.PostLoad(serverSettings);
+			globalConfig.PostLoad(serverSettings, new List<ILoadedPlugin>());
 
 			ConfigService configService = ServiceProvider.GetRequiredService<ConfigService>();
 			configService.OverrideConfig(globalConfig);
diff --git a/Engine/Source/Programs/Horde/HordeServer.Tests/Ddc/Functional/Objects/ReferencesTests.cs b/Engine/Source/Programs/Horde/HordeServer.Tests/Ddc/Functional/Objects/ReferencesTests.cs
--- a/Engine/Source/Programs/Horde/HordeServer.Tests/Ddc/Functional/Objects/ReferencesTests.cs
+++ b/Engine/Source/Programs/Horde/HordeServer.Tests/Ddc/Functional/Objects/ReferencesTests.cs
@@ -20,6 +20,7 @@ using EpicGames.Horde.Storage;
 using EpicGames.Serialization;
 using HordeServer.Configuration;
 using HordeServer.Ddc;
+using HordeServer.Plugins;
 using HordeServer.Server;
 using HordeServer.Storage;
 using Microsoft.AspNetCore.Mvc;
@@ -60,7 +61,7 @@ namespace HordeServer.Tests.Ddc.FunctionalTests.References
 
 			globalConfig.Storage.Backends.Add(new BackendConfig { Id = new BackendId("default"), Type = StorageBackendType.Memory });
 			globalConfig.Storage.Namespaces.Add(new NamespaceConfig { Id = TestNamespace, Backend = globalConfig.Storage.Backends[^1].Id });
-			globalConfig.PostLoad(serverSettings);
+			globalConfig.PostLoad(serverSettings, new List<ILoadedPlugin>());
 
 			ConfigService configService = ServiceProvider.GetRequiredService<ConfigService>();
 			configService.OverrideConfig(globalConfig);
diff --git a/Engine/Source/Programs/Horde/HordeServer.Tests/Fixture.cs b/Engine/Source/Programs/Horde/HordeServer.Tests/Fixture.cs
--- a/Engine/Source/Programs/Horde/HordeServer.Tests/Fixture.cs
+++ b/Engine/Source/Programs/Horde/HordeServer.Tests/Fixture.cs
@@ -15,6 +15,7 @@ using HordeServer.Configuration;
 using HordeServer.Jobs;
 using HordeServer.Jobs.Graphs;
 using HordeServer.Jobs.Templates;
+using HordeServer.Plugins;
 using HordeServer.Projects;
 using HordeServer.Server;
 using HordeServer.Streams;
@@ -35,10 +36,10 @@ namespace HordeServer.Tests
 		public string Agent1Name { get; private set; } = null!;
 		public const string PoolName = "TestingPool";
 
-		public static async Task<Fixture> CreateAsync(ConfigService configService, IGraphCollection graphCollection, ITemplateCollection templateCollection, JobService jobService, AgentService agentService, ServerSettings serverSettings)
+		public static async Task<Fixture> CreateAsync(ConfigService configService, IGraphCollection graphCollection, ITemplateCollection templateCollection, JobService jobService, AgentService agentService, IPluginCollection pluginCollection, ServerSettings serverSettings)
 		{
 			Fixture fixture = new Fixture();
-			await fixture.PopulateAsync(configService, graphCollection, templateCollection, jobService, agentService, serverSettings);
+			await fixture.PopulateAsync(configService, graphCollection, templateCollection, jobService, agentService, pluginCollection, serverSettings);
 
 			//			(PerforceService as PerforceServiceStub)?.AddChange("//UE5/Main", 112233, "leet.coder", "Did stuff", new []{"file.cpp"});
 			//			(PerforceService as PerforceServiceStub)?.AddChange("//UE5/Main", 1111, "swarm", "A shelved CL here", new []{"renderer.cpp"});
@@ -46,7 +47,7 @@ namespace HordeServer.Tests
 			return fixture;
 		}
 
-		private async Task PopulateAsync(ConfigService configService, IGraphCollection graphCollection, ITemplateCollection templateCollection, JobService jobService, AgentService agentService, ServerSettings serverSettings)
+		private async Task PopulateAsync(ConfigService configService, IGraphCollection graphCollection, ITemplateCollection templateCollection, JobService jobService, AgentService agentService, IPluginCollection pluginCollection, ServerSettings serverSettings)
 		{
 			FixtureGraph fg = new FixtureGraph();
 			fg.Id = ContentHash.Empty;
@@ -80,7 +81,7 @@ namespace HordeServer.Tests
 			ProjectConfig projectConfig = new ProjectConfig { Id = projectId, Name = "UE5", Streams = new List<StreamConfig> { streamConfig } };
 
 			GlobalConfig globalConfig = new GlobalConfig { Projects = new List<ProjectConfig> { projectConfig } };
-			globalConfig.PostLoad(serverSettings);
+			globalConfig.PostLoad(serverSettings, pluginCollection.LoadedPlugins);
 			configService.OverrideConfig(globalConfig);
 
 			StreamId = streamId;
diff --git a/Engine/Source/Programs/Horde/HordeServer.Tests/PluginConfigExtensions.cs b/Engine/Source/Programs/Horde/HordeServer.Tests/PluginConfigExtensions.cs
new file mode 100644
--- /dev/null
+++ b/Engine/Source/Programs/Horde/HordeServer.Tests/PluginConfigExtensions.cs
@@ -0,0 +1,18 @@
+// Copyright Epic Games, Inc. All Rights Reserved.
+
+using System;
+using System.Collections.Generic;
+using System.Linq;
+using System.Text;
+using System.Threading.Tasks;
+using HordeServer.Plugins;
+using HordeServer.Server;
+
+namespace HordeServer.Tests
+{
+	static class PluginConfigExtensions
+	{
+		public static void AddAnalytics(this PluginConfigCollection configCollection, AnalyticsGlobalConfig config)
+			=> configCollection[new PluginName("Analytics")] = config;
+	}
+}
diff --git a/Engine/Source/Programs/Horde/HordeServer.Tests/Telemetry/TelemetryTests.cs b/Engine/Source/Programs/Horde/HordeServer.Tests/Telemetry/TelemetryTests.cs
--- a/Engine/Source/Programs/Horde/HordeServer.Tests/Telemetry/TelemetryTests.cs
+++ b/Engine/Source/Programs/Horde/HordeServer.Tests/Telemetry/TelemetryTests.cs
@@ -48,8 +48,11 @@ namespace HordeServer.Tests.Telemetry
 			telemetryStoreConfig.Id = TelemetryStoreId.Default;
 			telemetryStoreConfig.Metrics.Add(metricConfig);
 
+			AnalyticsGlobalConfig analyticsConfig = new AnalyticsGlobalConfig();
+			analyticsConfig.TelemetryStores.Add(telemetryStoreConfig);
+
 			GlobalConfig globalConfig = new GlobalConfig();
-			globalConfig.TelemetryStores.Add(telemetryStoreConfig);
+			globalConfig.Plugins.AddAnalytics(analyticsConfig);
 			SetConfig(globalConfig);
 
 			MetricTelemetrySink sink = ServiceProvider.GetRequiredService<MetricTelemetrySink>();
@@ -86,8 +89,11 @@ namespace HordeServer.Tests.Telemetry
 			telemetryStoreConfig.Id = TelemetryStoreId.Default;
 			telemetryStoreConfig.Metrics.Add(metricConfig);
 
+			AnalyticsGlobalConfig analyticsConfig = new AnalyticsGlobalConfig();
+			analyticsConfig.TelemetryStores.Add(telemetryStoreConfig);
+
 			GlobalConfig globalConfig = new GlobalConfig();
-			globalConfig.TelemetryStores.Add(telemetryStoreConfig);
+			globalConfig.Plugins.AddAnalytics(analyticsConfig);
 			SetConfig(globalConfig);
 
 			MetricTelemetrySink sink = ServiceProvider.GetRequiredService<MetricTelemetrySink>();
@@ -164,8 +170,11 @@ namespace HordeServer.Tests.Telemetry
 			telemetryStoreConfig.Metrics.Add(metricConfig1);
 			telemetryStoreConfig.Metrics.Add(metricConfig2);
 
+			AnalyticsGlobalConfig analyticsConfig = new AnalyticsGlobalConfig();
+			analyticsConfig.TelemetryStores.Add(telemetryStoreConfig);
+
 			GlobalConfig globalConfig = new GlobalConfig();
-			globalConfig.TelemetryStores.Add(telemetryStoreConfig);
+			globalConfig.Plugins.AddAnalytics(analyticsConfig);
 			SetConfig(globalConfig);
 
 			MetricTelemetrySink sink = ServiceProvider.GetRequiredService<MetricTelemetrySink>();
@@ -228,8 +237,11 @@ namespace HordeServer.Tests.Telemetry
 			telemetryStoreConfig.Id = TelemetryStoreId.Default;
 			telemetryStoreConfig.Metrics.Add(metricConfig);
 
+			AnalyticsGlobalConfig analyticsConfig = new AnalyticsGlobalConfig();
+			analyticsConfig.TelemetryStores.Add(telemetryStoreConfig);
+
 			GlobalConfig globalConfig = new GlobalConfig();
-			globalConfig.TelemetryStores.Add(telemetryStoreConfig);
+			globalConfig.Plugins.AddAnalytics(analyticsConfig);
 			SetConfig(globalConfig);
 
 			MetricTelemetrySink sink = ServiceProvider.GetRequiredService<MetricTelemetrySink>();
@@ -265,8 +277,11 @@ namespace HordeServer.Tests.Telemetry
 			telemetryStoreConfig.Id = TelemetryStoreId.Default;
 			telemetryStoreConfig.Metrics.Add(metricConfig);
 
+			AnalyticsGlobalConfig analyticsConfig = new AnalyticsGlobalConfig();
+			analyticsConfig.TelemetryStores.Add(telemetryStoreConfig);
+
 			GlobalConfig globalConfig = new GlobalConfig();
-			globalConfig.TelemetryStores.Add(telemetryStoreConfig);
+			globalConfig.Plugins.AddAnalytics(analyticsConfig);
 			SetConfig(globalConfig);
 
 			MetricTelemetrySink sink = ServiceProvider.GetRequiredService<MetricTelemetrySink>();
@@ -313,8 +328,11 @@ namespace HordeServer.Tests.Telemetry
 			telemetryStoreConfig.Id = TelemetryStoreId.Default;
 			telemetryStoreConfig.Metrics.Add(metricConfig);
 
+			AnalyticsGlobalConfig analyticsConfig = new AnalyticsGlobalConfig();
+			analyticsConfig.TelemetryStores.Add(telemetryStoreConfig);
+
 			GlobalConfig globalConfig = new GlobalConfig();
-			globalConfig.TelemetryStores.Add(telemetryStoreConfig);
+			globalConfig.Plugins.AddAnalytics(analyticsConfig);
 			SetConfig(globalConfig);
 
 			MetricTelemetrySink sink = ServiceProvider.GetRequiredService<MetricTelemetrySink>();
@@ -371,8 +389,11 @@ namespace HordeServer.Tests.Telemetry
 			telemetryStoreConfig.Id = TelemetryStoreId.Default;
 			telemetryStoreConfig.Metrics.Add(metricConfig);
 
+			AnalyticsGlobalConfig analyticsConfig = new AnalyticsGlobalConfig();
+			analyticsConfig.TelemetryStores.Add(telemetryStoreConfig);
+
 			GlobalConfig globalConfig = new GlobalConfig();
-			globalConfig.TelemetryStores.Add(telemetryStoreConfig);
+			globalConfig.Plugins.AddAnalytics(analyticsConfig);
 			SetConfig(globalConfig);
 
 			MetricTelemetrySink sink = ServiceProvider.GetRequiredService<MetricTelemetrySink>();
diff --git a/Engine/Source/Programs/Horde/HordeServer.Tests/TestSetup.cs b/Engine/Source/Programs/Horde/HordeServer.Tests/TestSetup.cs
--- a/Engine/Source/Programs/Horde/HordeServer.Tests/TestSetup.cs
+++ b/Engine/Source/Programs/Horde/HordeServer.Tests/TestSetup.cs
@@ -147,16 +147,23 @@ namespace HordeServer.Tests
 		public IOptionsMonitor<GlobalConfig> GlobalConfig => ServiceProvider.GetRequiredService<IOptionsMonitor<GlobalConfig>>();
 		public IOptionsSnapshot<GlobalConfig> GlobalConfigSnapshot => ServiceProvider.GetRequiredService<IOptionsSnapshot<GlobalConfig>>();
 
+		public IPluginCollection PluginCollection => ServiceProvider.GetRequiredService<IPluginCollection>();
+
 		private static bool s_datadogWriterPatched;
 
+		readonly PluginCollection _pluginCollection;
+
 		public TestSetup()
 		{
+			_pluginCollection = new PluginCollection();
+			_pluginCollection.Add<AnalyticsPlugin>();
+
 			PatchDatadogWriter();
 		}
 
 		protected void SetConfig(GlobalConfig globalConfig)
 		{
-			globalConfig.PostLoad(ServerSettings);
+			globalConfig.PostLoad(ServerSettings, _pluginCollection.LoadedPlugins);
 			ConfigService.OverrideConfig(globalConfig);
 		}
 
@@ -164,7 +171,7 @@ namespace HordeServer.Tests
 		{
 			GlobalConfig globalConfig = GlobalConfig.CurrentValue;
 			action(globalConfig);
-			globalConfig.PostLoad(ServerSettings);
+			globalConfig.PostLoad(ServerSettings, _pluginCollection.LoadedPlugins);
 			ConfigService.OverrideConfig(globalConfig);
 		}
 
@@ -195,7 +202,7 @@ namespace HordeServer.Tests
 			services.AddHttpClient<RpcService>();
 
 			services.AddSingleton<IServerInfo>(new ServerInfo());
-			services.AddSingleton<IPluginCollection>(new PluginCollection());
+			services.AddSingleton<IPluginCollection>(_pluginCollection);
 
 			services.AddLogging(builder => { builder.AddConsole().SetMinimumLevel(LogLevel.Debug); });
 			services.AddSingleton<IMemoryCache>(sp => new MemoryCache(new MemoryCacheOptions { }));
@@ -305,11 +312,16 @@ namespace HordeServer.Tests
 			services.AddSingleton<StorageService>();
 			services.AddSingleton<StorageBackendCache>();
 			services.AddSingleton<BundleCache>();
+
+			foreach (ILoadedPlugin plugin in _pluginCollection.LoadedPlugins)
+			{
+				plugin.ConfigureServices(new ConfigurationBuilder().Build(), services);
+			}
 		}
 
 		public Task<Fixture> CreateFixtureAsync()
 		{
-			return Fixture.CreateAsync(ConfigService, GraphCollection, TemplateCollection, JobService, AgentService, ServerSettings);
+			return Fixture.CreateAsync(ConfigService, GraphCollection, TemplateCollection, JobService, AgentService, PluginCollection, ServerSettings);
 		}
 
 		private JobsController GetJobsController()
diff --git a/Engine/Source/Programs/Horde/HordeServer.Tests/Tools/ToolsControllerTests.cs b/Engine/Source/Programs/Horde/HordeServer.Tests/Tools/ToolsControllerTests.cs
--- a/Engine/Source/Programs/Horde/HordeServer.Tests/Tools/ToolsControllerTests.cs
+++ b/Engine/Source/Programs/Horde/HordeServer.Tests/Tools/ToolsControllerTests.cs
@@ -2,6 +2,7 @@
 
 using System.Collections.Generic;
 using System.IO;
+using System.Linq;
 using System.Net;
 using System.Net.Http;
 using System.Net.Http.Headers;
@@ -10,6 +11,7 @@ using EpicGames.Horde.Server;
 using EpicGames.Horde.Tools;
 using HordeServer.Acls;
 using HordeServer.Configuration;
+using HordeServer.Plugins;
 using HordeServer.Server;
 using HordeServer.ServiceAccounts;
 using HordeServer.Storage;
@@ -38,7 +40,7 @@ public class ToolsControllerTests
 		globalConfig.Tools.Add(new ToolConfig(toolId) { Name = "Foo", Description = "This is foo", Acl = new AclConfig() { Entries = [aclEntryConfig] }, Public = false });
 		
 		ServerSettings serverSettings = new() { AuthMethod = AuthMethod.Horde };
-		globalConfig.PostLoad(serverSettings);
+		globalConfig.PostLoad(serverSettings, new List<ILoadedPlugin>());
 		Dictionary<string, string> settings = new() { { "Horde:AuthMethod", AuthMethod.Horde.ToString() } };
 		await using FakeHordeWebApp app = new (settings);
 		
diff --git a/Engine/Source/Programs/Horde/HordeServer/Configuration/ConfigService.cs b/Engine/Source/Programs/Horde/HordeServer/Configuration/ConfigService.cs
--- a/Engine/Source/Programs/Horde/HordeServer/Configuration/ConfigService.cs
+++ b/Engine/Source/Programs/Horde/HordeServer/Configuration/ConfigService.cs
@@ -5,8 +5,10 @@ using System.Collections.Generic;
 using System.IO;
 using System.IO.Compression;
 using System.Linq;
+using System.Runtime.Versioning;
 using System.Text;
 using System.Text.Json;
+using System.Text.Json.Nodes;
 using System.Threading;
 using System.Threading.Tasks;
 using EpicGames.Core;
@@ -16,6 +18,7 @@ using HordeServer.Plugins;
 using HordeServer.Server;
 using HordeServer.Users;
 using HordeServer.Utilities;
+using Json.Path;
 using Microsoft.Extensions.Diagnostics.HealthChecks;
 using Microsoft.Extensions.Hosting;
 using Microsoft.Extensions.Logging;
@@ -27,6 +30,8 @@ using StackExchange.Redis;
 
 namespace HordeServer.Configuration
 {
+	using JsonObject = System.Text.Json.Nodes.JsonObject;
+
 	/// <summary>
 	/// Service which processes runtime configuration data.
 	/// </summary>
@@ -277,7 +282,7 @@ namespace HordeServer.Configuration
 			{
 				Uri globalConfigUri = GetGlobalConfigUri();
 				GlobalConfig globalConfig = await context.ReadAsync<GlobalConfig>(globalConfigUri, cancellationToken);
-				globalConfig.PostLoad(_serverSettings);
+				globalConfig.PostLoad(_serverSettings, _pluginCollection.LoadedPlugins);
 
 				foreach (OverrideConfigFile file in overrideFiles.Values)
 				{
@@ -581,7 +586,12 @@ namespace HordeServer.Configuration
 				// Read the new config in
 				Uri globalConfigUri = GetGlobalConfigUri();
 
-				GlobalConfig globalConfig = await context.ReadAsync<GlobalConfig>(globalConfigUri, cancellationToken);
+				ObjectConfigNode configNode = new ObjectConfigNode(typeof(GlobalConfig));
+
+				JsonObject configObj = await context.PreprocessFileAsync(GetGlobalConfigUri(), configNode, cancellationToken);
+				CopyJsonNode(configObj, "TelemetryStores", "Plugins.Analytics.TelemetryStores");
+
+				GlobalConfig globalConfig = JsonSerializer.Deserialize<GlobalConfig>(configObj, _jsonOptions)!;
 				if (globalConfig.VersionEnum < GlobalVersion.Latest)
 				{
 					List<string> message = new List<string>();
@@ -607,7 +617,7 @@ namespace HordeServer.Configuration
 				}
 
 				// Execute a PostLoad before returning so we can validate that everything is valid
-				globalConfig.PostLoad(_serverSettings);
+				globalConfig.PostLoad(_serverSettings, _pluginCollection.LoadedPlugins);
 
 				return snapshot;
 			}
@@ -617,6 +627,40 @@ namespace HordeServer.Configuration
 			}
 		}
 
+		static void CopyJsonNode(JsonObject rootObj, string sourcePath, string targetPath)
+		{
+			JsonNode? sourceNode = rootObj;
+
+			string[] sourceFragments = sourcePath.Split('.');
+			foreach (string sourceFragment in sourceFragments)
+			{
+				sourceNode = sourceNode[sourceFragment];
+				if (sourceNode == null)
+				{
+					return;
+				}
+			}
+
+			JsonObject? targetObj = rootObj;
+
+			string[] targetFragments = targetPath.Split('.');
+			for (int idx = 0; idx < targetFragments.Length - 1; idx++)
+			{
+				string targetFragment = targetFragments[idx];
+
+				JsonObject? nextObj = targetObj[targetFragment] as JsonObject;
+				if (nextObj == null)
+				{
+					nextObj = new JsonObject();
+					targetObj[targetFragment] = nextObj;
+				}
+
+				targetObj = nextObj;
+			}
+
+			targetObj[targetFragments[^1]] = sourceNode.DeepClone();
+		}
+
 		internal byte[] Serialize(GlobalConfig config)
 		{
 			return JsonSerializer.SerializeToUtf8Bytes(config, _jsonOptions);
@@ -746,18 +790,8 @@ namespace HordeServer.Configuration
 			GlobalConfig globalConfig = JsonSerializer.Deserialize<GlobalConfig>(snapshot.Data, _jsonOptions)!;
 			globalConfig.Revision = IoHash.Compute(data.Span).ToString();
 
-			// Ensure that all plugins have an entry in the global config so they can register their ACLs
-			foreach (ILoadedPlugin loadedPlugin in _pluginCollection.LoadedPlugins)
-			{
-				if (!globalConfig.Plugins.TryGetValue(loadedPlugin.Name, out _))
-				{
-					IPluginConfig pluginConfig = (IPluginConfig)Activator.CreateInstance(loadedPlugin.GlobalConfigType)!;
-					globalConfig.Plugins.Add(loadedPlugin.Name, pluginConfig);
-				}
-			}
-
 			// Run the postload callbacks on all the config objects
-			globalConfig.PostLoad(_serverSettings);
+			globalConfig.PostLoad(_serverSettings, _pluginCollection.LoadedPlugins);
 			return globalConfig;
 		}
 	}
diff --git a/Engine/Source/Programs/Horde/HordeServer/Plugins/PluginCollection.cs b/Engine/Source/Programs/Horde/HordeServer/Plugins/PluginCollection.cs
--- a/Engine/Source/Programs/Horde/HordeServer/Plugins/PluginCollection.cs
+++ b/Engine/Source/Programs/Horde/HordeServer/Plugins/PluginCollection.cs
@@ -153,6 +153,19 @@ namespace HordeServer.Plugins
 		/// Adds a plugin with the given startup class
 		/// </summary>
 		/// <typeparam name="T">Type of the startup class</typeparam>
+		public ILoadedPlugin Add<T>() where T : class, IPluginStartup
+		{
+			PluginAttribute attr = typeof(T).GetCustomAttribute<PluginAttribute>()
+				?? throw new InvalidOperationException($"Cannot add {typeof(T).Name} as a plugin. No {nameof(PluginAttribute)} was found.");
+
+			return Add<T>(new PluginMetadata{ Name = new PluginName(attr.Name) });
+		}
+
+		/// <summary>
+		/// Adds a plugin with the given startup class
+		/// </summary>
+		/// <typeparam name="T">Type of the startup class</typeparam>
+		/// <param name="metadata">Metadata for the plugin</param>
 		public ILoadedPlugin Add<T>(IPluginMetadata metadata) where T : class, IPluginStartup
 		{
 			ILoadedPlugin plugin = AddLoadedPlugin(metadata, typeof(T));
diff --git a/Engine/Source/Programs/Horde/HordeServer/Plugins/PluginName.cs b/Engine/Source/Programs/Horde/HordeServer/Plugins/PluginName.cs
deleted file mode 100644
diff --git a/Engine/Source/Programs/Horde/HordeServer/Server/GlobalConfig.cs b/Engine/Source/Programs/Horde/HordeServer/Server/GlobalConfig.cs
--- a/Engine/Source/Programs/Horde/HordeServer/Server/GlobalConfig.cs
+++ b/Engine/Source/Programs/Horde/HordeServer/Server/GlobalConfig.cs
@@ -241,11 +241,6 @@ namespace HordeServer.Server
 		/// </summary>
 		public List<ArtifactTypeConfig> ArtifactTypes { get; set; } = new List<ArtifactTypeConfig>();
 
-		/// <summary>
-		/// Metrics to aggregate on the Horde server
-		/// </summary>
-		public List<TelemetryStoreConfig> TelemetryStores { get; set; } = new List<TelemetryStoreConfig>();
-
 		/// <summary>
 		/// Plugin config objects
 		/// </summary>
@@ -286,7 +281,6 @@ namespace HordeServer.Server
 		private readonly Dictionary<ArtifactType, ArtifactTypeConfig> _artifactTypeLookup = new Dictionary<ArtifactType, ArtifactTypeConfig>();
 		private readonly Dictionary<SecretId, SecretConfig> _secretLookup = new Dictionary<SecretId, SecretConfig>();
 		private readonly Dictionary<PoolId, PoolConfig> _poolLookup = new Dictionary<PoolId, PoolConfig>();
-		private readonly Dictionary<TelemetryStoreId, TelemetryStoreConfig> _telemetryStoreLookup = new Dictionary<TelemetryStoreId, TelemetryStoreConfig>();
 
 		/// <inheritdoc cref="AclConfig.Authorize(AclAction, ClaimsPrincipal)"/>
 		public bool Authorize(AclAction action, ClaimsPrincipal user)
@@ -295,7 +289,7 @@ namespace HordeServer.Server
 		/// <summary>
 		/// Called after the config file has been read
 		/// </summary>
-		public void PostLoad(ServerSettings serverSettings)
+		public void PostLoad(ServerSettings serverSettings, IReadOnlyList<ILoadedPlugin> loadedPlugins)
 		{
 			ServerSettings = serverSettings;
 
@@ -352,15 +346,18 @@ namespace HordeServer.Server
 			ConfigObject.MergeDefaults<string, PoolConfig>(Pools.Select(x => (x.Id.ToString(), x.Base?.ToString(), x)));
 			UpdateWorkspacesForPools();
 
-			_telemetryStoreLookup.Clear();
-			foreach (TelemetryStoreConfig telemetryStore in TelemetryStores)
+			Storage.PostLoad(this);
+
+			// Ensure that all plugins have an entry in the global config so they can register their ACLs
+			foreach (ILoadedPlugin loadedPlugin in loadedPlugins)
 			{
-				_telemetryStoreLookup.Add(telemetryStore.Id, telemetryStore);
-				telemetryStore.PostLoad(Acl);
+				if (!Plugins.TryGetValue(loadedPlugin.Name, out _))
+				{
+					IPluginConfig pluginConfig = (IPluginConfig)Activator.CreateInstance(loadedPlugin.GlobalConfigType)!;
+					Plugins.Add(loadedPlugin.Name, pluginConfig);
+				}
 			}
 
-			Storage.PostLoad(this);
-
 			foreach (IPluginConfig pluginConfig in Plugins.Values)
 			{
 				pluginConfig.PostLoad(Acl);
@@ -581,14 +578,6 @@ namespace HordeServer.Server
 		/// <returns>True if the pool configuration was found</returns>
 		public bool TryGetPool(PoolId poolId, [NotNullWhen(true)] out PoolConfig? config) => _poolLookup.TryGetValue(poolId, out config);
 
-		/// <summary>
-		/// Attempts to get configuration for a pool from this object
-		/// </summary>
-		/// <param name="telemetryStoreId">The pool identifier</param>
-		/// <param name="config">Configuration for the telemetry store</param>
-		/// <returns>True if the telemetry configuration was found</returns>
-		public bool TryGetTelemetryStore(TelemetryStoreId telemetryStoreId, [NotNullWhen(true)] out TelemetryStoreConfig? config) => _telemetryStoreLookup.TryGetValue(telemetryStoreId, out config);
-
 		/// <summary>
 		/// Attempt to resolve an IP address to a network config
 		/// </summary>
diff --git a/Engine/Source/Programs/Horde/HordeServer/Startup.cs b/Engine/Source/Programs/Horde/HordeServer/Startup.cs
--- a/Engine/Source/Programs/Horde/HordeServer/Startup.cs
+++ b/Engine/Source/Programs/Horde/HordeServer/Startup.cs
@@ -321,7 +321,7 @@ namespace HordeServer
 
 			// Register the plugin collection
 			PluginCollection pluginCollection = new PluginCollection();
-			pluginCollection.Add<AnalyticsPlugin>(new PluginMetadata { Name = new PluginName("analytics") });
+			pluginCollection.Add<AnalyticsPlugin>();
 			services.AddSingleton<IPluginCollection>(pluginCollection);
 
 			// Register all the plugin services
