Horde: Fix static analysis warnings.

#rnx

diff --git a/Engine/Source/Programs/Horde/HordeServer.Tests/Fixture.cs b/Engine/Source/Programs/Horde/HordeServer.Tests/Fixture.cs
--- a/Engine/Source/Programs/Horde/HordeServer.Tests/Fixture.cs
+++ b/Engine/Source/Programs/Horde/HordeServer.Tests/Fixture.cs
@@ -1,8 +1,6 @@
 // Copyright Epic Games, Inc. All Rights Reserved.
 
 using System.Collections.Generic;
-using System.IO;
-using System.Text;
 using System.Threading.Tasks;
 using EpicGames.Core;
 using EpicGames.Horde.Agents.Pools;
diff --git a/Engine/Source/Programs/Horde/HordeServer.Tests/PluginConfigExtensions.cs b/Engine/Source/Programs/Horde/HordeServer.Tests/PluginConfigExtensions.cs
--- a/Engine/Source/Programs/Horde/HordeServer.Tests/PluginConfigExtensions.cs
+++ b/Engine/Source/Programs/Horde/HordeServer.Tests/PluginConfigExtensions.cs
@@ -1,12 +1,6 @@
 // Copyright Epic Games, Inc. All Rights Reserved.
 
-using System;
-using System.Collections.Generic;
-using System.Linq;
-using System.Text;
-using System.Threading.Tasks;
 using HordeServer.Plugins;
-using HordeServer.Server;
 
 namespace HordeServer.Tests
 {
diff --git a/Engine/Source/Programs/Horde/HordeServer.Tests/Tools/ToolsControllerTests.cs b/Engine/Source/Programs/Horde/HordeServer.Tests/Tools/ToolsControllerTests.cs
--- a/Engine/Source/Programs/Horde/HordeServer.Tests/Tools/ToolsControllerTests.cs
+++ b/Engine/Source/Programs/Horde/HordeServer.Tests/Tools/ToolsControllerTests.cs
@@ -2,7 +2,6 @@
 
 using System.Collections.Generic;
 using System.IO;
-using System.Linq;
 using System.Net;
 using System.Net.Http;
 using System.Net.Http.Headers;
@@ -29,21 +28,21 @@ public class ToolsControllerTests
 	[TestMethod]
 	public async Task DownloadToolWithAclAsync()
 	{
-		AclEntryConfig aclEntryConfig = new (HordeClaims.AgentRoleClaim, [ToolAclAction.DownloadTool]);
-		ToolId toolId = new ("foo");
+		AclEntryConfig aclEntryConfig = new(HordeClaims.AgentRoleClaim, [ToolAclAction.DownloadTool]);
+		ToolId toolId = new("foo");
 		GlobalConfig globalConfig = new();
-		
+
 		globalConfig.Storage.Backends.Clear();
 		globalConfig.Storage.Backends.Add(new BackendConfig { Id = new BackendId("tools-backend"), Type = StorageBackendType.Memory });
 		globalConfig.Storage.Namespaces.Clear();
 		globalConfig.Storage.Namespaces.Add(new NamespaceConfig { Id = Namespace.Tools, Backend = new BackendId("tools-backend") });
 		globalConfig.Tools.Add(new ToolConfig(toolId) { Name = "Foo", Description = "This is foo", Acl = new AclConfig() { Entries = [aclEntryConfig] }, Public = false });
-		
+
 		ServerSettings serverSettings = new() { AuthMethod = AuthMethod.Horde };
 		globalConfig.PostLoad(serverSettings, new List<ILoadedPlugin>());
 		Dictionary<string, string> settings = new() { { "Horde:AuthMethod", AuthMethod.Horde.ToString() } };
-		await using FakeHordeWebApp app = new (settings);
-		
+		await using FakeHordeWebApp app = new(settings);
+
 		ConfigService configService = app.ServiceProvider.GetRequiredService<ConfigService>();
 		IToolCollection tools = app.ServiceProvider.GetRequiredService<IToolCollection>();
 		IServiceAccountCollection serviceAccounts = app.ServiceProvider.GetRequiredService<IServiceAccountCollection>();
@@ -52,7 +51,7 @@ public class ToolsControllerTests
 
 		List<IUserClaim> claims = [new UserClaim("http://epicgames.com/ue/horde/role", "agent")];
 		(IServiceAccount _, string token) = await serviceAccounts.CreateAsync(new CreateServiceAccountOptions("myDesc", claims));
-		
+
 		// Create tool and deployment
 		using MemoryStream ms = new(await ToolTests.CreateZipFileDataAsync("foo.txt", "foo content"));
 		ITool? tool = await tools.GetAsync(toolId);
@@ -60,7 +59,7 @@ public class ToolsControllerTests
 		tool = await tool.CreateDeploymentAsync(new ToolDeploymentConfig() { Version = "1" }, ms);
 
 		HttpClient client = app.CreateHttpClient();
-		using HttpRequestMessage req = new (HttpMethod.Get, $"/api/v1/tools/{toolId.Id}?action=download");
+		using HttpRequestMessage req = new(HttpMethod.Get, $"/api/v1/tools/{toolId.Id}?action=download");
 		req.Headers.Authorization = new AuthenticationHeaderValue("ServiceAccount", token);
 		HttpResponseMessage res = await client.SendAsync(req);
 		Assert.AreEqual(HttpStatusCode.OK, res.StatusCode);
diff --git a/Engine/Source/Programs/Horde/HordeServer/Configuration/ConfigService.cs b/Engine/Source/Programs/Horde/HordeServer/Configuration/ConfigService.cs
--- a/Engine/Source/Programs/Horde/HordeServer/Configuration/ConfigService.cs
+++ b/Engine/Source/Programs/Horde/HordeServer/Configuration/ConfigService.cs
@@ -5,7 +5,6 @@ using System.Collections.Generic;
 using System.IO;
 using System.IO.Compression;
 using System.Linq;
-using System.Runtime.Versioning;
 using System.Text;
 using System.Text.Json;
 using System.Text.Json.Nodes;
@@ -18,7 +17,6 @@ using HordeServer.Plugins;
 using HordeServer.Server;
 using HordeServer.Users;
 using HordeServer.Utilities;
-using Json.Path;
 using Microsoft.Extensions.Diagnostics.HealthChecks;
 using Microsoft.Extensions.Hosting;
 using Microsoft.Extensions.Logging;
diff --git a/Engine/Source/Programs/Horde/HordeServer/Server/GlobalConfig.cs b/Engine/Source/Programs/Horde/HordeServer/Server/GlobalConfig.cs
--- a/Engine/Source/Programs/Horde/HordeServer/Server/GlobalConfig.cs
+++ b/Engine/Source/Programs/Horde/HordeServer/Server/GlobalConfig.cs
@@ -19,7 +19,6 @@ using EpicGames.Horde.Jobs.Templates;
 using EpicGames.Horde.Projects;
 using EpicGames.Horde.Secrets;
 using EpicGames.Horde.Streams;
-using EpicGames.Horde.Telemetry;
 using EpicGames.Horde.Tools;
 using EpicGames.Horde.Users;
 using EpicGames.Perforce;
@@ -43,7 +42,6 @@ using HordeServer.Projects;
 using HordeServer.Secrets;
 using HordeServer.Storage;
 using HordeServer.Streams;
-using HordeServer.Telemetry.Metrics;
 using HordeServer.Tools;
 using HordeServer.Utilities;
 
diff --git a/Engine/Source/Programs/Horde/HordeServer/ServerSettings.cs b/Engine/Source/Programs/Horde/HordeServer/ServerSettings.cs
--- a/Engine/Source/Programs/Horde/HordeServer/ServerSettings.cs
+++ b/Engine/Source/Programs/Horde/HordeServer/ServerSettings.cs
@@ -17,8 +17,6 @@ using EpicGames.Perforce;
 using HordeServer.Agents.Fleet;
 using HordeServer.Server;
 using HordeServer.Storage;
-using HordeServer.Storage.ObjectStores;
-using HordeServer.Telemetry.Sinks;
 using HordeServer.Tools;
 
 namespace HordeServer
@@ -242,8 +240,8 @@ namespace HordeServer
 		/// </summary>
 		[Obsolete("Use MongoConnectionString instead")]
 		public string? DatabaseConnectionString
-		{ 
-			get => MongoConnectionString; 
+		{
+			get => MongoConnectionString;
 			set => MongoConnectionString = value;
 		}
 
@@ -491,7 +489,7 @@ namespace HordeServer
 		/// <inheritdoc cref="RedisConnectionString"/>
 		[Obsolete("Use RedisConnectionString instead")]
 		public string? RedisConnectionConfig
-		{ 
+		{
 			get => RedisConnectionString;
 			set => RedisConnectionString = value;
 		}
diff --git a/Engine/Source/Programs/Horde/HordeServer/Startup.cs b/Engine/Source/Programs/Horde/HordeServer/Startup.cs
--- a/Engine/Source/Programs/Horde/HordeServer/Startup.cs
+++ b/Engine/Source/Programs/Horde/HordeServer/Startup.cs
@@ -80,9 +80,6 @@ using HordeServer.Storage;
 using HordeServer.Storage.ObjectStores;
 using HordeServer.Streams;
 using HordeServer.Tasks;
-using HordeServer.Telemetry;
-using HordeServer.Telemetry.Metrics;
-using HordeServer.Telemetry.Sinks;
 using HordeServer.Tools;
 using HordeServer.Ugs;
 using HordeServer.Users;
