Horde: More test fixes, and fix preprocessing of include directives and macros from legacy telemetry stores.

#rnx

diff --git a/Engine/Source/Programs/Horde/HordeServer.Tests/PluginConfigExtensions.cs b/Engine/Source/Programs/Horde/HordeServer.Tests/PluginConfigExtensions.cs
--- a/Engine/Source/Programs/Horde/HordeServer.Tests/PluginConfigExtensions.cs
+++ b/Engine/Source/Programs/Horde/HordeServer.Tests/PluginConfigExtensions.cs
@@ -6,7 +6,7 @@ namespace HordeServer.Tests
 {
 	static class PluginConfigExtensions
 	{
-		public static void AddAnalytics(this PluginConfigCollection configCollection, AnalyticsGlobalConfig config)
+		public static void AddAnalytics(this PluginConfigCollection configCollection, AnalyticsConfig config)
 			=> configCollection[new PluginName("Analytics")] = config;
 	}
 }
diff --git a/Engine/Source/Programs/Horde/HordeServer.Tests/Telemetry/TelemetryTests.cs b/Engine/Source/Programs/Horde/HordeServer.Tests/Telemetry/TelemetryTests.cs
--- a/Engine/Source/Programs/Horde/HordeServer.Tests/Telemetry/TelemetryTests.cs
+++ b/Engine/Source/Programs/Horde/HordeServer.Tests/Telemetry/TelemetryTests.cs
@@ -48,8 +48,8 @@ namespace HordeServer.Tests.Telemetry
 			telemetryStoreConfig.Id = TelemetryStoreId.Default;
 			telemetryStoreConfig.Metrics.Add(metricConfig);
 
-			AnalyticsGlobalConfig analyticsConfig = new AnalyticsGlobalConfig();
-			analyticsConfig.TelemetryStores.Add(telemetryStoreConfig);
+			AnalyticsConfig analyticsConfig = new AnalyticsConfig();
+			analyticsConfig.Stores.Add(telemetryStoreConfig);
 
 			GlobalConfig globalConfig = new GlobalConfig();
 			globalConfig.Plugins.AddAnalytics(analyticsConfig);
@@ -89,8 +89,8 @@ namespace HordeServer.Tests.Telemetry
 			telemetryStoreConfig.Id = TelemetryStoreId.Default;
 			telemetryStoreConfig.Metrics.Add(metricConfig);
 
-			AnalyticsGlobalConfig analyticsConfig = new AnalyticsGlobalConfig();
-			analyticsConfig.TelemetryStores.Add(telemetryStoreConfig);
+			AnalyticsConfig analyticsConfig = new AnalyticsConfig();
+			analyticsConfig.Stores.Add(telemetryStoreConfig);
 
 			GlobalConfig globalConfig = new GlobalConfig();
 			globalConfig.Plugins.AddAnalytics(analyticsConfig);
@@ -170,8 +170,8 @@ namespace HordeServer.Tests.Telemetry
 			telemetryStoreConfig.Metrics.Add(metricConfig1);
 			telemetryStoreConfig.Metrics.Add(metricConfig2);
 
-			AnalyticsGlobalConfig analyticsConfig = new AnalyticsGlobalConfig();
-			analyticsConfig.TelemetryStores.Add(telemetryStoreConfig);
+			AnalyticsConfig analyticsConfig = new AnalyticsConfig();
+			analyticsConfig.Stores.Add(telemetryStoreConfig);
 
 			GlobalConfig globalConfig = new GlobalConfig();
 			globalConfig.Plugins.AddAnalytics(analyticsConfig);
@@ -237,8 +237,8 @@ namespace HordeServer.Tests.Telemetry
 			telemetryStoreConfig.Id = TelemetryStoreId.Default;
 			telemetryStoreConfig.Metrics.Add(metricConfig);
 
-			AnalyticsGlobalConfig analyticsConfig = new AnalyticsGlobalConfig();
-			analyticsConfig.TelemetryStores.Add(telemetryStoreConfig);
+			AnalyticsConfig analyticsConfig = new AnalyticsConfig();
+			analyticsConfig.Stores.Add(telemetryStoreConfig);
 
 			GlobalConfig globalConfig = new GlobalConfig();
 			globalConfig.Plugins.AddAnalytics(analyticsConfig);
@@ -277,8 +277,8 @@ namespace HordeServer.Tests.Telemetry
 			telemetryStoreConfig.Id = TelemetryStoreId.Default;
 			telemetryStoreConfig.Metrics.Add(metricConfig);
 
-			AnalyticsGlobalConfig analyticsConfig = new AnalyticsGlobalConfig();
-			analyticsConfig.TelemetryStores.Add(telemetryStoreConfig);
+			AnalyticsConfig analyticsConfig = new AnalyticsConfig();
+			analyticsConfig.Stores.Add(telemetryStoreConfig);
 
 			GlobalConfig globalConfig = new GlobalConfig();
 			globalConfig.Plugins.AddAnalytics(analyticsConfig);
@@ -328,8 +328,8 @@ namespace HordeServer.Tests.Telemetry
 			telemetryStoreConfig.Id = TelemetryStoreId.Default;
 			telemetryStoreConfig.Metrics.Add(metricConfig);
 
-			AnalyticsGlobalConfig analyticsConfig = new AnalyticsGlobalConfig();
-			analyticsConfig.TelemetryStores.Add(telemetryStoreConfig);
+			AnalyticsConfig analyticsConfig = new AnalyticsConfig();
+			analyticsConfig.Stores.Add(telemetryStoreConfig);
 
 			GlobalConfig globalConfig = new GlobalConfig();
 			globalConfig.Plugins.AddAnalytics(analyticsConfig);
@@ -389,8 +389,8 @@ namespace HordeServer.Tests.Telemetry
 			telemetryStoreConfig.Id = TelemetryStoreId.Default;
 			telemetryStoreConfig.Metrics.Add(metricConfig);
 
-			AnalyticsGlobalConfig analyticsConfig = new AnalyticsGlobalConfig();
-			analyticsConfig.TelemetryStores.Add(telemetryStoreConfig);
+			AnalyticsConfig analyticsConfig = new AnalyticsConfig();
+			analyticsConfig.Stores.Add(telemetryStoreConfig);
 
 			GlobalConfig globalConfig = new GlobalConfig();
 			globalConfig.Plugins.AddAnalytics(analyticsConfig);
diff --git a/Engine/Source/Programs/Horde/HordeServer.Tests/Tools/ToolsControllerTests.cs b/Engine/Source/Programs/Horde/HordeServer.Tests/Tools/ToolsControllerTests.cs
--- a/Engine/Source/Programs/Horde/HordeServer.Tests/Tools/ToolsControllerTests.cs
+++ b/Engine/Source/Programs/Horde/HordeServer.Tests/Tools/ToolsControllerTests.cs
@@ -30,8 +30,11 @@ public class ToolsControllerTests
 	{
 		AclEntryConfig aclEntryConfig = new(HordeClaims.AgentRoleClaim, [ToolAclAction.DownloadTool]);
 		ToolId toolId = new("foo");
-		GlobalConfig globalConfig = new();
 
+		PluginCollection pluginCollection = new PluginCollection();
+		pluginCollection.Add<AnalyticsPlugin>();
+
+		GlobalConfig globalConfig = new();
 		globalConfig.Storage.Backends.Clear();
 		globalConfig.Storage.Backends.Add(new BackendConfig { Id = new BackendId("tools-backend"), Type = StorageBackendType.Memory });
 		globalConfig.Storage.Namespaces.Clear();
@@ -39,7 +42,8 @@ public class ToolsControllerTests
 		globalConfig.Tools.Add(new ToolConfig(toolId) { Name = "Foo", Description = "This is foo", Acl = new AclConfig() { Entries = [aclEntryConfig] }, Public = false });
 
 		ServerSettings serverSettings = new() { AuthMethod = AuthMethod.Horde };
-		globalConfig.PostLoad(serverSettings, new List<ILoadedPlugin>());
+		globalConfig.PostLoad(serverSettings, pluginCollection.LoadedPlugins);
+
 		Dictionary<string, string> settings = new() { { "Horde:AuthMethod", AuthMethod.Horde.ToString() } };
 		await using FakeHordeWebApp app = new(settings);
 
diff --git a/Engine/Source/Programs/Horde/HordeServer/Configuration/ConfigService.cs b/Engine/Source/Programs/Horde/HordeServer/Configuration/ConfigService.cs
--- a/Engine/Source/Programs/Horde/HordeServer/Configuration/ConfigService.cs
+++ b/Engine/Source/Programs/Horde/HordeServer/Configuration/ConfigService.cs
@@ -587,7 +587,7 @@ namespace HordeServer.Configuration
 				ObjectConfigNode configNode = new ObjectConfigNode(typeof(GlobalConfig));
 
 				JsonObject configObj = await context.PreprocessFileAsync(GetGlobalConfigUri(), configNode, cancellationToken);
-				CopyJsonNode(configObj, "TelemetryStores", "Plugins.Analytics.TelemetryStores");
+				CopyJsonNode(configObj, "TelemetryStores", "Plugins.Analytics.Stores");
 
 				GlobalConfig globalConfig = JsonSerializer.Deserialize<GlobalConfig>(configObj, _jsonOptions)!;
 				if (globalConfig.VersionEnum < GlobalVersion.Latest)
diff --git a/Engine/Source/Programs/Horde/HordeServer/Server/GlobalConfig.cs b/Engine/Source/Programs/Horde/HordeServer/Server/GlobalConfig.cs
--- a/Engine/Source/Programs/Horde/HordeServer/Server/GlobalConfig.cs
+++ b/Engine/Source/Programs/Horde/HordeServer/Server/GlobalConfig.cs
@@ -239,6 +239,12 @@ namespace HordeServer.Server
 		/// </summary>
 		public List<ArtifactTypeConfig> ArtifactTypes { get; set; } = new List<ArtifactTypeConfig>();
 
+		/// <summary>
+		/// Metrics to aggregate on the Horde server
+		/// </summary>
+		[Obsolete("Use Plugins.Analytics.Stores instead")]
+		public List<LegacyTelemetryStoreConfig> TelemetryStores { get; set; } = new List<LegacyTelemetryStoreConfig>();
+
 		/// <summary>
 		/// Plugin config objects
 		/// </summary>
@@ -790,6 +796,22 @@ namespace HordeServer.Server
 		}
 	}
 
+	/// <summary>
+	/// Shim for supporting includes and macros in telemetry settings before creating the analytics plugin
+	/// </summary>
+	public class LegacyTelemetryStoreConfig
+	{
+		/// <summary>
+		/// Includes for other configuration files
+		/// </summary>
+		public List<ConfigInclude> Include { get; set; } = new List<ConfigInclude>();
+
+		/// <summary>
+		/// Macros within this configuration
+		/// </summary>
+		public List<ConfigMacro> Macros { get; set; } = new List<ConfigMacro>();
+	}
+
 	/// <summary>
 	/// Profile for executing compute requests
 	/// </summary>
diff --git a/Engine/Source/Programs/Horde/Plugins/Analytics/HordeServer.Analytics/AnalyticsConfig.cs b/Engine/Source/Programs/Horde/Plugins/Analytics/HordeServer.Analytics/AnalyticsConfig.cs
new file mode 100644
--- /dev/null
+++ b/Engine/Source/Programs/Horde/Plugins/Analytics/HordeServer.Analytics/AnalyticsConfig.cs
@@ -0,0 +1,59 @@
+// Copyright Epic Games, Inc. All Rights Reserved.
+
+using System;
+using System.Collections.Generic;
+using System.Diagnostics.CodeAnalysis;
+using System.Linq;
+using System.Security.Claims;
+using System.Text;
+using System.Threading.Tasks;
+using EpicGames.Horde.Acls;
+using EpicGames.Horde.Telemetry;
+using HordeServer.Acls;
+using HordeServer.Plugins;
+using HordeServer.Telemetry.Metrics;
+
+#pragma warning disable CA2227 // Change X to be read-only by removing the property setter
+
+namespace HordeServer
+{
+	/// <summary>
+	/// Config settings for analytics
+	/// </summary>
+	public class AnalyticsConfig : IPluginConfig
+	{
+		/// <summary>
+		/// Metrics to aggregate on the Horde server
+		/// </summary>
+		public List<TelemetryStoreConfig> Stores { get; set; } = new List<TelemetryStoreConfig>();
+
+		private AclConfig _parentAcl = null!;
+		private readonly Dictionary<TelemetryStoreId, TelemetryStoreConfig> _telemetryStoreLookup = new Dictionary<TelemetryStoreId, TelemetryStoreConfig>();
+
+		public void PostLoad(AclConfig parentAcl)
+		{
+			_parentAcl = parentAcl;
+
+			_telemetryStoreLookup.Clear();
+			foreach (TelemetryStoreConfig telemetryStore in Stores)
+			{
+				_telemetryStoreLookup.Add(telemetryStore.Id, telemetryStore);
+				telemetryStore.PostLoad(parentAcl);
+			}
+		}
+
+		/// <inheritdoc cref="AclConfig.Authorize(AclAction, ClaimsPrincipal)"/>
+		public bool Authorize(AclAction action, ClaimsPrincipal user)
+			=> _parentAcl.Authorize(action, user);
+
+
+		/// <summary>
+		/// Attempts to get configuration for a pool from this object
+		/// </summary>
+		/// <param name="telemetryStoreId">The pool identifier</param>
+		/// <param name="config">Configuration for the telemetry store</param>
+		/// <returns>True if the telemetry configuration was found</returns>
+		public bool TryGetTelemetryStore(TelemetryStoreId telemetryStoreId, [NotNullWhen(true)] out TelemetryStoreConfig? config) 
+			=> _telemetryStoreLookup.TryGetValue(telemetryStoreId, out config);
+	}
+}
diff --git a/Engine/Source/Programs/Horde/Plugins/Analytics/HordeServer.Analytics/AnalyticsGlobalConfig.cs b/Engine/Source/Programs/Horde/Plugins/Analytics/HordeServer.Analytics/AnalyticsGlobalConfig.cs
deleted file mode 100644
diff --git a/Engine/Source/Programs/Horde/Plugins/Analytics/HordeServer.Analytics/AnalyticsPlugin.cs b/Engine/Source/Programs/Horde/Plugins/Analytics/HordeServer.Analytics/AnalyticsPlugin.cs
--- a/Engine/Source/Programs/Horde/Plugins/Analytics/HordeServer.Analytics/AnalyticsPlugin.cs
+++ b/Engine/Source/Programs/Horde/Plugins/Analytics/HordeServer.Analytics/AnalyticsPlugin.cs
@@ -17,7 +17,7 @@ namespace HordeServer
 	/// <summary>
 	/// Main entry point for the analytics plugin
 	/// </summary>
-	[Plugin("Analytics", ServerConfigType = typeof(AnalyticsServerConfig), GlobalConfigType = typeof(AnalyticsGlobalConfig))]
+	[Plugin("Analytics", GlobalConfigType = typeof(AnalyticsConfig), ServerConfigType = typeof(AnalyticsServerConfig))]
 	public class AnalyticsPlugin : IPluginStartup
 	{
 		/// <inheritdoc/>
diff --git a/Engine/Source/Programs/Horde/Plugins/Analytics/HordeServer.Analytics/Telemetry/Metrics/MetricCollection.cs b/Engine/Source/Programs/Horde/Plugins/Analytics/HordeServer.Analytics/Telemetry/Metrics/MetricCollection.cs
--- a/Engine/Source/Programs/Horde/Plugins/Analytics/HordeServer.Analytics/Telemetry/Metrics/MetricCollection.cs
+++ b/Engine/Source/Programs/Horde/Plugins/Analytics/HordeServer.Analytics/Telemetry/Metrics/MetricCollection.cs
@@ -58,12 +58,12 @@ namespace HordeServer.Telemetry.Metrics
 		readonly AsyncEvent _flushEvent = new AsyncEvent();
 		readonly BackgroundTask _flushTask;
 		readonly IClock _clock;
-		readonly IOptionsMonitor<AnalyticsGlobalConfig> _globalConfig;
+		readonly IOptionsMonitor<AnalyticsConfig> _analyticsConfig;
 		readonly ILogger _logger;
 
 		Dictionary<SampleKey, List<double>> _queuedSamples = new Dictionary<SampleKey, List<double>>();
 
-		public MetricCollection(IMongoService mongoService, IClock clock, IOptionsMonitor<AnalyticsGlobalConfig> globalConfig, ILogger<MetricCollection> logger)
+		public MetricCollection(IMongoService mongoService, IClock clock, IOptionsMonitor<AnalyticsConfig> analyticsConfig, ILogger<MetricCollection> logger)
 		{
 			List<MongoIndex<MetricDocument>> indexes = new List<MongoIndex<MetricDocument>>();
 			indexes.Add(MongoIndex.Create<MetricDocument>(keys => keys.Ascending(x => x.TelemetryStoreId).Descending(x => x.Time).Ascending(x => x.MetricId).Ascending(x => x.Group)));
@@ -71,7 +71,7 @@ namespace HordeServer.Telemetry.Metrics
 
 			_flushTask = new BackgroundTask(BackgroundTickAsync);
 			_clock = clock;
-			_globalConfig = globalConfig;
+			_analyticsConfig = analyticsConfig;
 			_logger = logger;
 		}
 
@@ -100,7 +100,7 @@ namespace HordeServer.Telemetry.Metrics
 			JsonArray array = new JsonArray { node };
 
 			TelemetryStoreConfig? telemetryStoreConfig;
-			if (_globalConfig.CurrentValue.TryGetTelemetryStore(storeId, out telemetryStoreConfig))
+			if (_analyticsConfig.CurrentValue.TryGetTelemetryStore(storeId, out telemetryStoreConfig))
 			{
 				foreach (MetricConfig metric in telemetryStoreConfig.Metrics)
 				{
@@ -258,7 +258,7 @@ namespace HordeServer.Telemetry.Metrics
 				foreach ((SampleKey sampleKey, List<double> sampleValues) in samples)
 				{
 					MetricConfig? metricConfig;
-					if (_globalConfig.CurrentValue.TryGetTelemetryStore(sampleKey.Store, out TelemetryStoreConfig? telemetryStoreConfig) && telemetryStoreConfig.TryGetMetric(sampleKey.Metric, out metricConfig))
+					if (_analyticsConfig.CurrentValue.TryGetTelemetryStore(sampleKey.Store, out TelemetryStoreConfig? telemetryStoreConfig) && telemetryStoreConfig.TryGetMetric(sampleKey.Metric, out metricConfig))
 					{
 						await CombineValuesAsync(sampleKey.Store, metricConfig, sampleKey.Group, sampleKey.Time, sampleValues, cancellationToken);
 					}
diff --git a/Engine/Source/Programs/Horde/Plugins/Analytics/HordeServer.Analytics/Telemetry/TelemetryController.cs b/Engine/Source/Programs/Horde/Plugins/Analytics/HordeServer.Analytics/Telemetry/TelemetryController.cs
--- a/Engine/Source/Programs/Horde/Plugins/Analytics/HordeServer.Analytics/Telemetry/TelemetryController.cs
+++ b/Engine/Source/Programs/Horde/Plugins/Analytics/HordeServer.Analytics/Telemetry/TelemetryController.cs
@@ -27,16 +27,16 @@ namespace HordeServer.Telemetry
 	{
 		readonly TelemetryManager _telemetryManager;
 		readonly IMetricCollection _metricCollection;
-		readonly IOptionsSnapshot<AnalyticsGlobalConfig> _globalConfig;
+		readonly IOptionsSnapshot<AnalyticsConfig> _analyticsConfig;
 
 		/// <summary>
 		/// Constructor
 		/// </summary>
-		public TelemetryController(TelemetryManager telemetryManager, IMetricCollection metricCollection, IOptionsSnapshot<AnalyticsGlobalConfig> globalConfig)
+		public TelemetryController(TelemetryManager telemetryManager, IMetricCollection metricCollection, IOptionsSnapshot<AnalyticsConfig> analyticsConfig)
 		{
 			_telemetryManager = telemetryManager;
 			_metricCollection = metricCollection;
-			_globalConfig = globalConfig;
+			_analyticsConfig = analyticsConfig;
 		}
 
 		/// <summary>
@@ -72,13 +72,13 @@ namespace HordeServer.Telemetry
 		[Route("/api/v1/telemetry/{storeId}/metrics")]
 		public async Task<ActionResult<List<GetTelemetryMetricsResponse>>> GetMetricsAsync(TelemetryStoreId storeId, [FromQuery] MetricId[] id, [FromQuery] DateTime? minTime = null, [FromQuery] DateTime? maxTime = null, [FromQuery] string? group = null, [FromQuery] int results = 50, CancellationToken cancellationToken = default)
 		{
-			if (!_globalConfig.Value.Authorize(TelemetryAclAction.QueryMetrics, User))
+			if (!_analyticsConfig.Value.Authorize(TelemetryAclAction.QueryMetrics, User))
 			{
 				return Forbid(TelemetryAclAction.QueryMetrics);
 			}
 
 			TelemetryStoreConfig? telemetryStoreConfig;
-			if (!_globalConfig.Value.TryGetTelemetryStore(storeId, out telemetryStoreConfig))
+			if (!_analyticsConfig.Value.TryGetTelemetryStore(storeId, out telemetryStoreConfig))
 			{
 				return NotFound(storeId);
 			}
