Horde: Do not log cancellation exception as a warning when checking if blob exists.

diff --git a/Engine/Source/Programs/Horde/HordeServer/Storage/StorageService.cs b/Engine/Source/Programs/Horde/HordeServer/Storage/StorageService.cs
--- a/Engine/Source/Programs/Horde/HordeServer/Storage/StorageService.cs
+++ b/Engine/Source/Programs/Horde/HordeServer/Storage/StorageService.cs
@@ -608,6 +608,10 @@ namespace HordeServer.Storage
 						_logger.LogWarning("Blob {BlobId} ({Locator}) accessed after being garbage collected", blobInfo.Id, locator);
 					}
 				}
+				catch (OperationCanceledException)
+				{
+					throw;
+				}
 				catch (Exception ex)
 				{
 					_logger.LogWarning(ex, "Exception checking if blob {NamespaceId}:{Locator} exists", namespaceId, locator);
