Horde: Move analytics tests into separate project.

#rnx

diff --git a/Engine/Source/Programs/Horde/HordeServer.Tests.Shared/AppHostStub.cs b/Engine/Source/Programs/Horde/HordeServer.Tests.Shared/AppHostStub.cs
new file mode 100644
--- /dev/null
+++ b/Engine/Source/Programs/Horde/HordeServer.Tests.Shared/AppHostStub.cs
@@ -0,0 +1,30 @@
+// Copyright Epic Games, Inc. All Rights Reserved.
+
+using System;
+using System.Collections.Generic;
+using System.Linq;
+using System.Text;
+using System.Threading.Tasks;
+using Microsoft.Extensions.Hosting;
+
+namespace HordeServer.Tests
+{
+	public class AppLifetimeStub : IHostApplicationLifetime
+	{
+		public CancellationToken ApplicationStarted { get; }
+		public CancellationToken ApplicationStopping { get; }
+		public CancellationToken ApplicationStopped { get; }
+
+		public AppLifetimeStub()
+		{
+			ApplicationStarted = new CancellationToken();
+			ApplicationStopping = new CancellationToken();
+			ApplicationStopped = new CancellationToken();
+		}
+
+		public void StopApplication()
+		{
+			throw new NotImplementedException();
+		}
+	}
+}
diff --git a/Engine/Source/Programs/Horde/HordeServer.Tests.Shared/DatabaseIntegrationTest.cs b/Engine/Source/Programs/Horde/HordeServer.Tests.Shared/DatabaseIntegrationTest.cs
new file mode 100644
--- /dev/null
+++ b/Engine/Source/Programs/Horde/HordeServer.Tests.Shared/DatabaseIntegrationTest.cs
@@ -0,0 +1,265 @@
+// Copyright Epic Games, Inc. All Rights Reserved.
+
+using System;
+using System.Collections.Generic;
+using System.Net;
+using System.Runtime.InteropServices;
+using System.Threading;
+using System.Threading.Tasks;
+using HordeServer.Server;
+using HordeServer.Utilities;
+using Microsoft.AspNetCore.Mvc;
+using Microsoft.Extensions.DependencyInjection;
+using Microsoft.Extensions.Logging;
+using Microsoft.Extensions.Logging.Abstractions;
+using Microsoft.Extensions.Options;
+using Microsoft.VisualStudio.TestTools.UnitTesting;
+using MongoDB.Driver;
+using StackExchange.Redis;
+
+namespace HordeServer.Tests
+{
+	// Stub for fulfilling IOptions interface during testing
+	public sealed class TestOptions<T> : IOptions<T> where T : class
+	{
+		public TestOptions(T options) => Value = options;
+		public T Value { get; }
+	}
+
+	// Stub for fulfilling IOptionsMonitor interface during testing
+	public sealed class TestOptionsMonitor<T> : IOptionsMonitor<T>
+		where T : class, new()
+	{
+		sealed class Disposable : IDisposable
+		{
+			public void Dispose() { }
+		}
+
+		public TestOptionsMonitor(T currentValue)
+		{
+			CurrentValue = currentValue;
+		}
+
+		public T Get(string? name)
+		{
+			return CurrentValue;
+		}
+
+		public IDisposable OnChange(Action<T, string> listener)
+		{
+			return new Disposable();
+		}
+
+		public T CurrentValue { get; }
+	}
+
+	public sealed class MongoInstance : IDisposable
+	{
+		public string DatabaseName { get; }
+		public string ConnectionString { get; }
+		MongoClient Client { get; }
+
+		private static readonly object s_lockObject = new object();
+		private static MongoDbRunnerLocal? s_mongoDbRunner;
+		private static int s_nextDatabaseIndex = 1;
+		public const string MongoDbDatabaseNamePrefix = "HordeServerTest_";
+
+		public MongoInstance()
+		{
+			int databaseIndex;
+			lock (s_lockObject)
+			{
+				if (s_mongoDbRunner == null)
+				{
+					// One-time setup per test run to avoid overhead of starting the external MongoDB process
+					Startup.ConfigureMongoDbClient();
+					s_mongoDbRunner = new MongoDbRunnerLocal();
+					s_mongoDbRunner.Start();
+
+					// Drop all the previous databases
+					MongoClientSettings mongoSettings = MongoClientSettings.FromConnectionString(s_mongoDbRunner.GetConnectionString());
+					MongoClient client = new MongoClient(mongoSettings);
+
+					List<string> dropDatabaseNames = client.ListDatabaseNames().ToList();
+					foreach (string dropDatabaseName in dropDatabaseNames)
+					{
+						if (dropDatabaseName.StartsWith(MongoDbDatabaseNamePrefix, StringComparison.Ordinal))
+						{
+							client.DropDatabase(dropDatabaseName);
+						}
+					}
+				}
+				databaseIndex = s_nextDatabaseIndex++;
+			}
+
+			DatabaseName = $"{MongoDbDatabaseNamePrefix}{databaseIndex}";
+			ConnectionString = $"{s_mongoDbRunner.GetConnectionString()}/{DatabaseName}";
+			Client = new MongoClient(MongoClientSettings.FromConnectionString(ConnectionString));
+		}
+
+		public void Dispose()
+		{
+			IMongoClient strictClient = Client.WithWriteConcern(new WriteConcern(journal: true));
+			for (int i = 0; i < 5; i++)
+			{
+				strictClient.DropDatabase(DatabaseName);
+				List<string> dbNames = strictClient.ListDatabaseNames().ToList();
+				if (!dbNames.Contains(DatabaseName))
+				{
+					return;
+				}
+				Thread.Sleep(300);
+			}
+
+			throw new Exception($"Unable to drop MongoDB database {DatabaseName}");
+		}
+	}
+
+	public sealed class RedisInstance : IDisposable
+	{
+		static readonly object s_lockObject = new object();
+
+		const bool UseExistingRedisInstance = true;
+		const int RedisPort = 6379;
+		const int RedisDbNum = 15;
+
+		private static string? s_redisConnectionString;
+		private static int s_redisDbNum;
+		private static RedisProcess? s_redisProcess;
+
+		public string ConnectionString { get; }
+		public int DatabaseNumber { get; }
+
+		public RedisInstance()
+		{
+			lock (s_lockObject)
+			{
+				if (s_redisConnectionString == null)
+				{
+					int port = GetRedisPortInternal();
+					s_redisConnectionString = $"127.0.0.1:{port},allowAdmin=true";
+					s_redisDbNum = RedisDbNum;
+				}
+
+				ConnectionString = s_redisConnectionString;
+				DatabaseNumber = s_redisDbNum;
+			}
+		}
+
+		public void Dispose()
+		{
+		}
+
+		static int GetRedisPortInternal()
+		{
+			if (UseExistingRedisInstance && !DatabaseRunner.IsPortAvailable(RedisPort))
+			{
+				Console.WriteLine("Using existing Redis instance on port {0}", RedisPort);
+				return RedisPort;
+			}
+
+			if (RuntimeInformation.IsOSPlatform(OSPlatform.Windows))
+			{
+				s_redisProcess = new RedisProcess(NullLogger.Instance);
+				s_redisProcess.Start("");
+
+				return s_redisProcess.Port;
+			}
+
+			throw new Exception("Unable to connect to Redis");
+		}
+	}
+
+	public class DatabaseIntegrationTest : ServiceTest
+	{
+		private static readonly object s_lockObject = new object();
+
+		private MongoInstance? _mongoInstance;
+		private MongoService? _mongoService;
+		private readonly LoggerFactory _loggerFactory = new LoggerFactory();
+
+		private RedisInstance? _redisInstance;
+		private RedisService? _redisService;
+
+		public DatabaseIntegrationTest()
+		{
+		}
+
+		protected override void ConfigureServices(IServiceCollection services)
+		{
+			base.ConfigureServices(services);
+
+			services.AddSingleton<IMongoService>(GetMongoServiceSingleton());
+			services.AddSingleton<IRedisService>(GetRedisServiceSingleton());
+		}
+
+		public override async ValueTask DisposeAsync()
+		{
+			await base.DisposeAsync();
+
+			GC.SuppressFinalize(this);
+
+			if (_mongoService != null)
+			{
+				await _mongoService.DisposeAsync();
+			}
+			_mongoInstance?.Dispose();
+
+			if (_redisService != null)
+			{
+				await _redisService.DisposeAsync();
+			}
+			_redisInstance?.Dispose();
+
+			_loggerFactory.Dispose();
+		}
+
+		public IMongoService GetMongoServiceSingleton()
+		{
+			lock (s_lockObject)
+			{
+				if (_mongoService == null)
+				{
+					IRedisService redisService = GetRedisServiceSingleton();
+					_mongoInstance = new MongoInstance();
+
+					ServerSettings ss = new ServerSettings();
+					ss.MongoDatabaseName = _mongoInstance.DatabaseName;
+					ss.MongoConnectionString = _mongoInstance.ConnectionString;
+
+					MongoCommandTracer mongoTracer = new (OpenTelemetryTracers.MongoDb, _loggerFactory.CreateLogger<MongoCommandTracer>());
+					_mongoService = new MongoService(Options.Create(ss), redisService, mongoTracer, OpenTelemetryTracers.Horde, _loggerFactory.CreateLogger<MongoService>(), _loggerFactory);
+				}
+			}
+			return _mongoService;
+		}
+
+		public IRedisService GetRedisServiceSingleton()
+		{
+			if (_redisService == null)
+			{
+				_redisInstance = new RedisInstance();
+				_redisService = new RedisService(_redisInstance.ConnectionString, _redisInstance.DatabaseNumber, _loggerFactory.CreateLogger<RedisService>());
+
+				IConnectionMultiplexer cm = _redisService.ConnectionPool.GetConnection();
+				foreach (EndPoint endpoint in cm.GetEndPoints())
+				{
+					cm.GetServer(endpoint).FlushDatabase(_redisInstance.DatabaseNumber);
+				}
+			}
+			return _redisService;
+		}
+
+		public static T Deref<T>(T? item)
+		{
+			Assert.IsNotNull(item);
+			return item!;
+		}
+
+		public static T Deref<T>(ActionResult<T>? item) where T : class
+		{
+			Assert.IsNotNull(item?.Value);
+			return item!.Value!;
+		}
+	}
+}
\ No newline at end of file
diff --git a/Engine/Source/Programs/Horde/HordeServer.Tests.Shared/DatabaseRunner.cs b/Engine/Source/Programs/Horde/HordeServer.Tests.Shared/DatabaseRunner.cs
new file mode 100644
--- /dev/null
+++ b/Engine/Source/Programs/Horde/HordeServer.Tests.Shared/DatabaseRunner.cs
@@ -0,0 +1,208 @@
+// Copyright Epic Games, Inc. All Rights Reserved.
+
+using System;
+using System.Diagnostics;
+using System.IO;
+using System.Linq;
+using System.Net;
+using System.Net.NetworkInformation;
+using System.Net.Sockets;
+using System.Reflection;
+using System.Runtime.InteropServices;
+using System.Threading.Tasks;
+using EpicGames.Core;
+
+namespace HordeServer.Tests;
+
+public abstract class DatabaseRunner : IDisposable
+{
+	private readonly string _binName;
+	private readonly int _defaultPort;
+	private readonly string _name;
+	private readonly bool _reuseProcess;
+	protected string TempDir { get; }
+
+	private ManagedProcessGroup? _processGroup;
+	private ManagedProcess? _process;
+
+	protected DatabaseRunner(string name, string binName, int defaultPort, bool reuseProcess)
+	{
+		_name = name;
+		_binName = binName;
+		_defaultPort = defaultPort;
+		_reuseProcess = reuseProcess;
+		TempDir = GetTemporaryDirectory();
+	}
+
+	/// <inheritdoc/>
+	public void Dispose()
+	{
+		Dispose(true);
+		GC.SuppressFinalize(this);
+	}
+
+	protected virtual void Dispose(bool disposing)
+	{
+		Stop();
+	}
+
+	protected int Port { get; private set; } = -1;
+
+	protected abstract string GetArguments();
+
+	public void Start()
+	{
+		if (_reuseProcess && !IsPortAvailable(_defaultPort))
+		{
+			Console.WriteLine($"Re-using already running {_name} process!");
+			Port = _defaultPort;
+			return;
+		}
+
+		if (!RuntimeInformation.IsOSPlatform(OSPlatform.Windows))
+		{
+			Console.WriteLine($"Unable to find a running {_name} process to use during testing!");
+			Console.WriteLine("This is required on any non-Windows as the runner can only start Windows binaries!");
+			Console.WriteLine($"Please ensure {_binName} is running on the default port {_defaultPort}.");
+			throw new Exception("Failed finding process to re-use! See stdout for info.");
+		}
+
+		if (_process != null)
+		{
+			return;
+		}
+
+		Port = GetAvailablePort();
+
+		_processGroup = new ManagedProcessGroup();
+		_process = new ManagedProcess(_processGroup, GetBinaryPath(), GetArguments(), TempDir, null, ProcessPriorityClass.Normal);
+		Task.Run(() => RelayOutputAsync(_process));
+
+		// Try detect when main .NET process exits and kill the runner
+		AppDomain.CurrentDomain.ProcessExit += (sender, eventArgs) =>
+		{
+			Console.WriteLine("Main process exiting!");
+			Stop();
+		};
+	}
+
+	static async Task RelayOutputAsync(ManagedProcess process)
+	{
+		for (; ; )
+		{
+			string? line = await process.ReadLineAsync();
+			if (line == null)
+			{
+				break;
+			}
+			//			Console.WriteLine("{0} output: {1}", _name, line);
+		}
+	}
+
+	public void Stop()
+	{
+		if (_process != null)
+		{
+			_process.Dispose();
+			_process = null;
+		}
+		if (_processGroup != null)
+		{
+			_processGroup.Dispose();
+			_processGroup = null;
+		}
+		DeleteDirectory(TempDir);
+	}
+
+	public (string Host, int Port) GetListenAddress()
+	{
+		return ("localhost", Port);
+	}
+
+	private string GetBinaryPath()
+	{
+		FileReference file = new(new Uri(Assembly.GetExecutingAssembly().Location).LocalPath);
+		FileReference binPath = FileReference.Combine(file.Directory, _binName);
+		return binPath.FullName;
+	}
+
+	private string GetTemporaryDirectory()
+	{
+		string temp = Path.Join(Path.GetTempPath(), $"horde-{_name}-" + Path.GetRandomFileName());
+		Directory.CreateDirectory(temp);
+		return temp;
+	}
+
+	private static int GetAvailablePort()
+	{
+		using TcpListener listener = new(IPAddress.Loopback, 0);
+		listener.Start();
+		int port = ((IPEndPoint)listener.LocalEndpoint).Port;
+		listener.Stop();
+		return port;
+	}
+
+	public static bool IsPortAvailable(int port)
+	{
+		if (RuntimeInformation.IsOSPlatform(OSPlatform.Windows))
+		{
+			IPGlobalProperties ipGlobalProperties = IPGlobalProperties.GetIPGlobalProperties();
+
+			IPEndPoint[] listeners = ipGlobalProperties.GetActiveTcpListeners();
+			if (listeners.Any(x => x.Port == port))
+			{
+				return false;
+			}
+
+			return true;
+		}
+
+		try
+		{
+			using TcpListener listenerAny = new TcpListener(IPAddress.Loopback, port);
+			listenerAny.Start();
+			using TcpListener listenerLoopback = new TcpListener(IPAddress.Any, port);
+			listenerLoopback.Start();
+			return true;
+		}
+		catch (SocketException)
+		{
+		}
+
+		return false;
+	}
+
+	private static void DeleteDirectory(string path)
+	{
+		DirectoryInfo dir = new DirectoryInfo(path);
+		if (dir.Exists)
+		{
+			dir.Attributes = FileAttributes.Normal;
+
+			foreach (FileSystemInfo info in dir.GetFileSystemInfos("*", SearchOption.AllDirectories))
+			{
+				info.Attributes = FileAttributes.Normal;
+			}
+
+			dir.Delete(true);
+		}
+	}
+}
+
+public class MongoDbRunnerLocal : DatabaseRunner
+{
+	public MongoDbRunnerLocal() : base("mongodb", "ThirdParty/Mongo/mongod.exe", 27017, true)
+	{
+	}
+
+	protected override string GetArguments()
+	{
+		return $"--dbpath {TempDir} --noauth --quiet --port {Port}";
+	}
+
+	public string GetConnectionString()
+	{
+		(string host, int listenPort) = GetListenAddress();
+		return $"mongodb://{host}:{listenPort}";
+	}
+}
diff --git a/Engine/Source/Programs/Horde/HordeServer.Tests.Shared/DowntimeServiceStub.cs b/Engine/Source/Programs/Horde/HordeServer.Tests.Shared/DowntimeServiceStub.cs
new file mode 100644
--- /dev/null
+++ b/Engine/Source/Programs/Horde/HordeServer.Tests.Shared/DowntimeServiceStub.cs
@@ -0,0 +1,21 @@
+// Copyright Epic Games, Inc. All Rights Reserved.
+
+using System;
+using System.Collections.Generic;
+using System.Linq;
+using System.Text;
+using System.Threading.Tasks;
+using HordeServer.Server;
+
+namespace HordeServer.Tests
+{
+	public class DowntimeServiceStub : IDowntimeService
+	{
+		public DowntimeServiceStub(bool isDowntimeActive = false)
+		{
+			IsDowntimeActive = isDowntimeActive;
+		}
+
+		public bool IsDowntimeActive { get; set; }
+	}
+}
diff --git a/Engine/Source/Programs/Horde/HordeServer.Tests.Shared/HordeServer.Tests.Shared.csproj b/Engine/Source/Programs/Horde/HordeServer.Tests.Shared/HordeServer.Tests.Shared.csproj
new file mode 100644
--- /dev/null
+++ b/Engine/Source/Programs/Horde/HordeServer.Tests.Shared/HordeServer.Tests.Shared.csproj
@@ -0,0 +1,27 @@
+﻿<Project Sdk="Microsoft.NET.Sdk">
+
+  <PropertyGroup>
+    <TargetFramework>net8.0</TargetFramework>
+    <ImplicitUsings>enable</ImplicitUsings>
+    <Nullable>enable</Nullable>
+    <RootNamespace>HordeServer.Tests</RootNamespace>
+  </PropertyGroup>
+
+  <ItemGroup>
+    <ProjectReference Include="..\HordeServer\HordeServer.csproj" />
+  </ItemGroup>
+
+  <ItemGroup>
+    <PackageReference Include="Microsoft.AspNetCore.Mvc.Testing" Version="6.0.25" />
+    <PackageReference Include="Microsoft.NET.Test.Sdk" Version="17.8.0" />
+    <PackageReference Include="Moq" Version="4.20.70" />
+    <PackageReference Include="MSTest.TestAdapter" Version="3.1.1" />
+    <PackageReference Include="MSTest.TestFramework" Version="3.1.1" />
+    <PackageReference Include="Microsoft.Extensions.Caching.Memory" Version="8.0.0" />
+    <PackageReference Include="coverlet.collector" Version="6.0.0">
+      <PrivateAssets>all</PrivateAssets>
+      <IncludeAssets>runtime; build; native; contentfiles; analyzers; buildtransitive</IncludeAssets>
+    </PackageReference>
+  </ItemGroup>
+
+</Project>
diff --git a/Engine/Source/Programs/Horde/HordeServer.Tests.Shared/ServerServiceTest.cs b/Engine/Source/Programs/Horde/HordeServer.Tests.Shared/ServerServiceTest.cs
new file mode 100644
--- /dev/null
+++ b/Engine/Source/Programs/Horde/HordeServer.Tests.Shared/ServerServiceTest.cs
@@ -0,0 +1,180 @@
+// Copyright Epic Games, Inc. All Rights Reserved.
+
+using System;
+using System.Collections.Generic;
+using System.Diagnostics.Metrics;
+using System.Linq;
+using System.Reflection;
+using System.Text;
+using System.Threading.Tasks;
+using EpicGames.Core;
+using HordeCommon;
+using HordeServer.Acls;
+using HordeServer.Auditing;
+using HordeServer.Configuration;
+using HordeServer.Plugins;
+using HordeServer.Server;
+using Microsoft.Extensions.Caching.Memory;
+using Microsoft.Extensions.Configuration;
+using Microsoft.Extensions.DependencyInjection;
+using Microsoft.Extensions.Hosting;
+using Microsoft.Extensions.Logging;
+using Microsoft.Extensions.Options;
+
+namespace HordeServer.Tests
+{
+	public class ServerServiceTest : DatabaseIntegrationTest
+	{
+		public FakeClock Clock => ServiceProvider.GetRequiredService<FakeClock>();
+		public IMemoryCache Cache => ServiceProvider.GetRequiredService<IMemoryCache>();
+
+		public AclService AclService => ServiceProvider.GetRequiredService<AclService>();
+		public IMongoService MongoService => ServiceProvider.GetRequiredService<IMongoService>();
+		public IDowntimeService DowntimeService => ServiceProvider.GetRequiredService<IDowntimeService>();
+		public LifetimeService LifetimeService => ServiceProvider.GetRequiredService<LifetimeService>();
+		public ServerStatusService ServerStatusService => ServiceProvider.GetRequiredService<ServerStatusService>();
+
+		public ServerSettings ServerSettings => ServiceProvider.GetRequiredService<IOptions<ServerSettings>>().Value;
+		public IOptionsMonitor<ServerSettings> ServerSettingsMon => ServiceProvider.GetRequiredService<IOptionsMonitor<ServerSettings>>();
+
+		public OpenTelemetry.Trace.Tracer Tracer => ServiceProvider.GetRequiredService<OpenTelemetry.Trace.Tracer>();
+		public Meter Meter => ServiceProvider.GetRequiredService<Meter>();
+		public ConfigService ConfigService => ServiceProvider.GetRequiredService<ConfigService>();
+		public IOptionsMonitor<GlobalConfig> GlobalConfig => ServiceProvider.GetRequiredService<IOptionsMonitor<GlobalConfig>>();
+		public IOptionsSnapshot<GlobalConfig> GlobalConfigSnapshot => ServiceProvider.GetRequiredService<IOptionsSnapshot<GlobalConfig>>();
+
+		public IPluginCollection PluginCollection => ServiceProvider.GetRequiredService<IPluginCollection>();
+
+		private static bool s_datadogWriterPatched;
+
+		readonly PluginCollection _pluginCollection;
+
+		public ServerServiceTest()
+		{
+			_pluginCollection = new PluginCollection();
+			_pluginCollection.Add<AnalyticsPlugin>();
+
+			PatchDatadogWriter();
+		}
+
+		protected void SetConfig(GlobalConfig globalConfig)
+		{
+			globalConfig.PostLoad(ServerSettings, _pluginCollection.LoadedPlugins);
+			ConfigService.OverrideConfig(globalConfig);
+		}
+
+		protected void UpdateConfig(Action<GlobalConfig> action)
+		{
+			GlobalConfig globalConfig = GlobalConfig.CurrentValue;
+			action(globalConfig);
+			globalConfig.PostLoad(ServerSettings, _pluginCollection.LoadedPlugins);
+			ConfigService.OverrideConfig(globalConfig);
+		}
+
+		protected override void ConfigureSettings(ServerSettings settings)
+		{
+			DirectoryReference baseDir = DirectoryReference.Combine(ServerApp.DataDir, "Tests");
+			try
+			{
+				FileUtils.ForceDeleteDirectoryContents(baseDir);
+			}
+			catch
+			{
+			}
+
+			settings.WithAws = true;
+
+			settings.ForceConfigUpdateOnStartup = true;
+		}
+
+		protected override void ConfigureServices(IServiceCollection services)
+		{
+			base.ConfigureServices(services);
+
+			services.AddSingleton<IServerInfo>(new ServerInfo());
+			services.AddSingleton<IPluginCollection>(_pluginCollection);
+
+			services.AddLogging(builder => { builder.AddConsole().SetMinimumLevel(LogLevel.Debug); });
+			services.AddSingleton<IMemoryCache>(sp => new MemoryCache(new MemoryCacheOptions { }));
+
+			services.AddSingleton(typeof(IAuditLogFactory<>), typeof(AuditLogFactory<>));
+
+			services.AddSingleton<ConfigService>();
+			services.AddSingleton<IOptionsFactory<GlobalConfig>>(sp => sp.GetRequiredService<ConfigService>());
+			services.AddSingleton<IOptionsChangeTokenSource<GlobalConfig>>(sp => sp.GetRequiredService<ConfigService>());
+
+			services.AddSingleton<IConfigSource, FileConfigSource>();
+
+			services.AddSingleton<FakeClock>();
+			services.AddSingleton<IClock>(sp => sp.GetRequiredService<FakeClock>());
+			services.AddSingleton<IHostApplicationLifetime, AppLifetimeStub>();
+			services.AddSingleton<IHostEnvironment, WebHostEnvironmentStub>();
+
+			services.AddSingleton<AclService>();
+			services.AddSingleton<IDowntimeService, DowntimeServiceStub>();
+			services.AddSingleton<LifetimeService>();
+
+			services.AddSingleton(typeof(IHealthMonitor<>), typeof(HealthMonitor<>));
+			services.AddSingleton<ServerStatusService>();
+
+			foreach (ILoadedPlugin plugin in _pluginCollection.LoadedPlugins)
+			{
+				plugin.ConfigureServices(new ConfigurationBuilder().Build(), services);
+			}
+		}
+
+		protected void AddPlugin<T>() where T : class, IPluginStartup
+		{
+			_pluginCollection.Add<T>();
+		}
+
+		/// <summary>
+		/// Hack the Datadog tracing library to not block during shutdown of tests.
+		/// Without this fix, the lib will try to send traces to a host that isn't running and block for +20 secs
+		///
+		/// Since so many of the interfaces and classes in the lib are internal it was difficult to replace Tracer.Instance
+		/// </summary>
+		private static void PatchDatadogWriter()
+		{
+			if (s_datadogWriterPatched)
+			{
+				return;
+			}
+
+			s_datadogWriterPatched = true;
+
+			/*			string msg = "Unable to patch Datadog agent writer! Tests will still work, but shutdown will block for +20 seconds.";
+
+						FieldInfo? agentWriterField = Datadog.Trace.Tracer.Instance.GetType().GetField("_agentWriter", BindingFlags.NonPublic | BindingFlags.Instance);
+						if (agentWriterField == null)
+						{
+							Console.Error.WriteLine(msg);
+							return;
+						}
+
+						object? agentWriterInstance = agentWriterField.GetValue(Datadog.Trace.Tracer.Instance);
+						if (agentWriterInstance == null)
+						{
+							Console.Error.WriteLine(msg);
+							return;
+						}
+
+						FieldInfo? processExitField = agentWriterInstance.GetType().GetField("_processExit", BindingFlags.NonPublic | BindingFlags.Instance);
+						if (processExitField == null)
+						{
+							Console.Error.WriteLine(msg);
+							return;
+						}
+
+						TaskCompletionSource<bool>? processExitInstance = (TaskCompletionSource<bool>?)processExitField.GetValue(agentWriterInstance);
+						if (processExitInstance == null)
+						{
+							Console.Error.WriteLine(msg);
+							return;
+						}
+
+						processExitInstance.TrySetResult(true);
+			*/
+		}
+	}
+}
diff --git a/Engine/Source/Programs/Horde/HordeServer.Tests.Shared/ServiceTest.cs b/Engine/Source/Programs/Horde/HordeServer.Tests.Shared/ServiceTest.cs
new file mode 100644
--- /dev/null
+++ b/Engine/Source/Programs/Horde/HordeServer.Tests.Shared/ServiceTest.cs
@@ -0,0 +1,53 @@
+// Copyright Epic Games, Inc. All Rights Reserved.
+
+using System;
+using System.Collections.Generic;
+using System.Linq;
+using System.Text;
+using System.Threading.Tasks;
+using Microsoft.Extensions.DependencyInjection;
+
+namespace HordeServer.Tests
+{
+	/// <summary>
+	/// Implements a test framework which manages a collection of services
+	/// </summary>
+	public class ServiceTest : IAsyncDisposable
+	{
+		private ServiceProvider? _serviceProvider = null;
+
+		public IServiceProvider ServiceProvider
+		{
+			get
+			{
+				if (_serviceProvider == null)
+				{
+					IServiceCollection services = new ServiceCollection();
+					ConfigureServices(services);
+
+					_serviceProvider = services.BuildServiceProvider();
+				}
+				return _serviceProvider;
+			}
+		}
+
+		public virtual async ValueTask DisposeAsync()
+		{
+			GC.SuppressFinalize(this);
+
+			if (_serviceProvider != null)
+			{
+				await _serviceProvider.DisposeAsync();
+				_serviceProvider = null;
+			}
+		}
+
+		protected virtual void ConfigureSettings(ServerSettings settings)
+		{
+		}
+
+		protected virtual void ConfigureServices(IServiceCollection services)
+		{
+		}
+	}
+}
diff --git a/Engine/Source/Programs/Horde/HordeServer.Tests.Shared/WebHostEnvironmentStub.cs b/Engine/Source/Programs/Horde/HordeServer.Tests.Shared/WebHostEnvironmentStub.cs
new file mode 100644
--- /dev/null
+++ b/Engine/Source/Programs/Horde/HordeServer.Tests.Shared/WebHostEnvironmentStub.cs
@@ -0,0 +1,26 @@
+// Copyright Epic Games, Inc. All Rights Reserved.
+
+using System;
+using System.Collections.Generic;
+using System.Linq;
+using System.Text;
+using System.Threading.Tasks;
+using Microsoft.Extensions.FileProviders;
+using Microsoft.Extensions.Hosting;
+
+namespace HordeServer.Tests
+{
+	public class WebHostEnvironmentStub : IHostEnvironment
+	{
+		public string ApplicationName { get; set; } = "HordeTest";
+		public IFileProvider ContentRootFileProvider { get; set; }
+		public string ContentRootPath { get; set; }
+		public string EnvironmentName { get; set; } = "Testing";
+
+		public WebHostEnvironmentStub()
+		{
+			ContentRootPath = Directory.CreateTempSubdirectory("HordeTest").FullName;
+			ContentRootFileProvider = new PhysicalFileProvider(ContentRootPath);
+		}
+	}
+}
diff --git a/Engine/Source/Programs/Horde/HordeServer.Tests/DatabaseIntegrationTest.cs b/Engine/Source/Programs/Horde/HordeServer.Tests/DatabaseIntegrationTest.cs
deleted file mode 100644
diff --git a/Engine/Source/Programs/Horde/HordeServer.Tests/DatabaseRunner.cs b/Engine/Source/Programs/Horde/HordeServer.Tests/DatabaseRunner.cs
deleted file mode 100644
diff --git a/Engine/Source/Programs/Horde/HordeServer.Tests/HordeServer.Tests.csproj b/Engine/Source/Programs/Horde/HordeServer.Tests/HordeServer.Tests.csproj
--- a/Engine/Source/Programs/Horde/HordeServer.Tests/HordeServer.Tests.csproj
+++ b/Engine/Source/Programs/Horde/HordeServer.Tests/HordeServer.Tests.csproj
@@ -33,6 +33,7 @@
 
   <ItemGroup>
     <ProjectReference Include="..\..\Shared\EpicGames.Serialization\EpicGames.Serialization.csproj" />
+    <ProjectReference Include="..\HordeServer.Tests.Shared\HordeServer.Tests.Shared.csproj"/>
     <ProjectReference Include="..\HordeAgent\HordeAgent.csproj">
       <Aliases>HordeAgent</Aliases>
     </ProjectReference>
diff --git a/Engine/Source/Programs/Horde/HordeServer.Tests/Server/RpcServiceTest.cs b/Engine/Source/Programs/Horde/HordeServer.Tests/Server/RpcServiceTest.cs
--- a/Engine/Source/Programs/Horde/HordeServer.Tests/Server/RpcServiceTest.cs
+++ b/Engine/Source/Programs/Horde/HordeServer.Tests/Server/RpcServiceTest.cs
@@ -27,39 +27,6 @@ namespace HordeServer.Tests.Server
 {
 	using ISession = Microsoft.AspNetCore.Http.ISession;
 
-	public class AppLifetimeStub : IHostApplicationLifetime
-	{
-		public CancellationToken ApplicationStarted { get; }
-		public CancellationToken ApplicationStopping { get; }
-		public CancellationToken ApplicationStopped { get; }
-
-		public AppLifetimeStub()
-		{
-			ApplicationStarted = new CancellationToken();
-			ApplicationStopping = new CancellationToken();
-			ApplicationStopped = new CancellationToken();
-		}
-
-		public void StopApplication()
-		{
-			throw new NotImplementedException();
-		}
-	}
-
-	public class WebHostEnvironmentStub : IHostEnvironment
-	{
-		public string ApplicationName { get; set; } = "HordeTest";
-		public IFileProvider ContentRootFileProvider { get; set; }
-		public string ContentRootPath { get; set; }
-		public string EnvironmentName { get; set; } = "Testing";
-
-		public WebHostEnvironmentStub()
-		{
-			ContentRootPath = Directory.CreateTempSubdirectory("HordeTest").FullName;
-			ContentRootFileProvider = new PhysicalFileProvider(ContentRootPath);
-		}
-	}
-
 	sealed class HttpContextStub : HttpContext
 	{
 		public override ConnectionInfo Connection { get; } = null!;
diff --git a/Engine/Source/Programs/Horde/HordeServer.Tests/Telemetry/TelemetryTests.cs b/Engine/Source/Programs/Horde/HordeServer.Tests/Telemetry/TelemetryTests.cs
deleted file mode 100644
diff --git a/Engine/Source/Programs/Horde/HordeServer.Tests/TestSetup.cs b/Engine/Source/Programs/Horde/HordeServer.Tests/TestSetup.cs
--- a/Engine/Source/Programs/Horde/HordeServer.Tests/TestSetup.cs
+++ b/Engine/Source/Programs/Horde/HordeServer.Tests/TestSetup.cs
@@ -80,11 +80,8 @@ namespace HordeServer.Tests
 	///
 	/// Easier to pass all these things around in a single object.
 	/// </summary>
-	public class TestSetup : DatabaseIntegrationTest
+	public class TestSetup : ServerServiceTest
 	{
-		public FakeClock Clock => ServiceProvider.GetRequiredService<FakeClock>();
-		public IMemoryCache Cache => ServiceProvider.GetRequiredService<IMemoryCache>();
-
 		public IGraphCollection GraphCollection => ServiceProvider.GetRequiredService<IGraphCollection>();
 		public INotificationTriggerCollection NotificationTriggerCollection => ServiceProvider.GetRequiredService<INotificationTriggerCollection>();
 		public IStreamCollection StreamCollection => ServiceProvider.GetRequiredService<IStreamCollection>();
@@ -104,13 +101,11 @@ namespace HordeServer.Tests
 		public IDashboardPreviewCollection DashboardPreviewCollection => ServiceProvider.GetRequiredService<IDashboardPreviewCollection>();
 		public IBisectTaskCollection BisectTaskCollection => ServiceProvider.GetRequiredService<IBisectTaskCollection>();
 
-		public AclService AclService => ServiceProvider.GetRequiredService<AclService>();
 		public FleetService FleetService => ServiceProvider.GetRequiredService<FleetService>();
 		public AgentService AgentService => ServiceProvider.GetRequiredService<AgentService>();
 		public AgentRelayService AgentRelayService => ServiceProvider.GetRequiredService<AgentRelayService>();
 		public ICommitService CommitService => ServiceProvider.GetRequiredService<ICommitService>();
 		public GlobalsService GlobalsService => ServiceProvider.GetRequiredService<GlobalsService>();
-		public IMongoService MongoService => ServiceProvider.GetRequiredService<IMongoService>();
 		public ITemplateCollection TemplateCollection => ServiceProvider.GetRequiredService<ITemplateCollection>();
 		internal PerforceServiceStub PerforceService => (PerforceServiceStub)ServiceProvider.GetRequiredService<IPerforceService>();
 		public ISubscriptionCollection SubscriptionCollection => ServiceProvider.GetRequiredService<ISubscriptionCollection>();
@@ -118,19 +113,13 @@ namespace HordeServer.Tests
 		public IssueService IssueService => ServiceProvider.GetRequiredService<IssueService>();
 		public JobTaskSource JobTaskSource => ServiceProvider.GetRequiredService<JobTaskSource>();
 		public JobService JobService => ServiceProvider.GetRequiredService<JobService>();
-		public IDowntimeService DowntimeService => ServiceProvider.GetRequiredService<IDowntimeService>();
 		public RpcService RpcService => ServiceProvider.GetRequiredService<RpcService>();
 		public PoolService PoolService => ServiceProvider.GetRequiredService<PoolService>();
-		public LifetimeService LifetimeService => ServiceProvider.GetRequiredService<LifetimeService>();
 		public ScheduleService ScheduleService => ServiceProvider.GetRequiredService<ScheduleService>();
 		public DeviceService DeviceService => ServiceProvider.GetRequiredService<DeviceService>();
 		public StorageService StorageService => ServiceProvider.GetRequiredService<StorageService>();
 		public TestDataService TestDataService => ServiceProvider.GetRequiredService<TestDataService>();
 		public ComputeService ComputeService => ServiceProvider.GetRequiredService<ComputeService>();
-		public ServerStatusService ServerStatusService => ServiceProvider.GetRequiredService<ServerStatusService>();
-
-		public ServerSettings ServerSettings => ServiceProvider.GetRequiredService<IOptions<ServerSettings>>().Value;
-		public IOptionsMonitor<ServerSettings> ServerSettingsMon => ServiceProvider.GetRequiredService<IOptionsMonitor<ServerSettings>>();
 
 		public JobsController JobsController => GetJobsController();
 		public AgentsController AgentsController => GetAgentsController();
@@ -141,56 +130,6 @@ namespace HordeServer.Tests
 		public TestDataController TestDataController => GetTestDataController();
 		public BisectTasksController BisectTasksController => GetBisectTasksController();
 
-		public OpenTelemetry.Trace.Tracer Tracer => ServiceProvider.GetRequiredService<OpenTelemetry.Trace.Tracer>();
-		public Meter Meter => ServiceProvider.GetRequiredService<Meter>();
-		public ConfigService ConfigService => ServiceProvider.GetRequiredService<ConfigService>();
-		public IOptionsMonitor<GlobalConfig> GlobalConfig => ServiceProvider.GetRequiredService<IOptionsMonitor<GlobalConfig>>();
-		public IOptionsSnapshot<GlobalConfig> GlobalConfigSnapshot => ServiceProvider.GetRequiredService<IOptionsSnapshot<GlobalConfig>>();
-
-		public IPluginCollection PluginCollection => ServiceProvider.GetRequiredService<IPluginCollection>();
-
-		private static bool s_datadogWriterPatched;
-
-		readonly PluginCollection _pluginCollection;
-
-		public TestSetup()
-		{
-			_pluginCollection = new PluginCollection();
-			_pluginCollection.Add<AnalyticsPlugin>();
-
-			PatchDatadogWriter();
-		}
-
-		protected void SetConfig(GlobalConfig globalConfig)
-		{
-			globalConfig.PostLoad(ServerSettings, _pluginCollection.LoadedPlugins);
-			ConfigService.OverrideConfig(globalConfig);
-		}
-
-		protected void UpdateConfig(Action<GlobalConfig> action)
-		{
-			GlobalConfig globalConfig = GlobalConfig.CurrentValue;
-			action(globalConfig);
-			globalConfig.PostLoad(ServerSettings, _pluginCollection.LoadedPlugins);
-			ConfigService.OverrideConfig(globalConfig);
-		}
-
-		protected override void ConfigureSettings(ServerSettings settings)
-		{
-			DirectoryReference baseDir = DirectoryReference.Combine(ServerApp.DataDir, "Tests");
-			try
-			{
-				FileUtils.ForceDeleteDirectoryContents(baseDir);
-			}
-			catch
-			{
-			}
-
-			settings.WithAws = true;
-
-			settings.ForceConfigUpdateOnStartup = true;
-		}
-
 		protected override void ConfigureServices(IServiceCollection services)
 		{
 			base.ConfigureServices(services);
@@ -201,31 +140,10 @@ namespace HordeServer.Tests
 
 			services.AddHttpClient<RpcService>();
 
-			services.AddSingleton<IServerInfo>(new ServerInfo());
-			services.AddSingleton<IPluginCollection>(_pluginCollection);
-
-			services.AddLogging(builder => { builder.AddConsole().SetMinimumLevel(LogLevel.Debug); });
-			services.AddSingleton<IMemoryCache>(sp => new MemoryCache(new MemoryCacheOptions { }));
-
-			services.AddSingleton(typeof(IAuditLogFactory<>), typeof(AuditLogFactory<>));
 			services.AddSingleton<IAuditLog<AgentId>>(sp => sp.GetRequiredService<IAuditLogFactory<AgentId>>().Create("Agents.Log", "AgentId"));
-			services.AddSingleton<ITelemetrySink, NullTelemetrySink>();
-			services.AddSingleton<NullTelemetrySink>();
-			services.AddSingleton<ITelemetrySink, MetricTelemetrySink>();
-			services.AddSingleton<MetricTelemetrySink>();
 
-			services.AddSingleton<TelemetryManager>();
-			services.AddSingleton<ITelemetryWriter>(sp => sp.GetRequiredService<TelemetryManager>());
 			services.AddSingleton<OpenTelemetry.Trace.Tracer>(sp => TracerProvider.Default.GetTracer("TestTracer"));
 			services.AddSingleton(sp => new Meter("TestMeter"));
-			services.AddSingleton<MetricCollection>();
-			services.AddSingleton<IMetricCollection, MetricCollection>(sp => sp.GetRequiredService<MetricCollection>());
-
-			services.AddSingleton<ConfigService>();
-			services.AddSingleton<IOptionsFactory<GlobalConfig>>(sp => sp.GetRequiredService<ConfigService>());
-			services.AddSingleton<IOptionsChangeTokenSource<GlobalConfig>>(sp => sp.GetRequiredService<ConfigService>());
-
-			services.AddSingleton<IConfigSource, FileConfigSource>();
 
 			services.AddSingleton<IAccountCollection, AccountCollection>();
 			services.AddSingleton<IAgentCollection, AgentCollection>();
@@ -255,11 +173,6 @@ namespace HordeServer.Tests
 			services.AddSingleton<IDeviceCollection, DeviceCollection>();
 			services.AddSingleton<IDashboardPreviewCollection, DashboardPreviewCollection>();
 
-			services.AddSingleton<FakeClock>();
-			services.AddSingleton<IClock>(sp => sp.GetRequiredService<FakeClock>());
-			services.AddSingleton<IHostApplicationLifetime, AppLifetimeStub>();
-			services.AddSingleton<IHostEnvironment, WebHostEnvironmentStub>();
-
 			// Empty mocked object to satisfy basic test runs
 			services.AddSingleton<IAmazonEC2>(sp => new Mock<IAmazonEC2>().Object);
 			services.AddSingleton<IAmazonAutoScaling>(sp => new Mock<IAmazonAutoScaling>().Object);
@@ -271,7 +184,6 @@ namespace HordeServer.Tests
 			services.AddSingleton<IPoolSizeStrategyFactory, LeaseUtilizationStrategyFactory>();
 			services.AddSingleton<IPoolSizeStrategyFactory, LeaseUtilizationAwsMetricStrategyFactory>();
 
-			services.AddSingleton<AclService>();
 			services.AddSingleton<AgentService>();
 			services.AddSingleton(provider => new Lazy<AgentService>(provider.GetRequiredService<AgentService>));
 			services.AddSingleton<AgentRelayService>();
@@ -282,11 +194,9 @@ namespace HordeServer.Tests
 			services.AddSingleton<RequestTrackerService>();
 			services.AddSingleton<GlobalsService>();
 			services.AddSingleton<JobTaskSource>();
-			services.AddSingleton<IDowntimeService, DowntimeServiceStub>();
 			services.AddSingleton<IssueService>();
 			services.AddSingleton<JobService>();
 			services.AddSingleton<JobExpirationService>();
-			services.AddSingleton<LifetimeService>();
 			services.AddSingleton<LogTailService>();
 			services.AddSingleton<INotificationService, NotificationService>();
 			services.AddSingleton<IPerforceService, PerforceServiceStub>();
@@ -300,9 +210,6 @@ namespace HordeServer.Tests
 			services.AddSingleton<ComputeService>();
 			services.AddSingleton<EnrollmentService>();
 
-			services.AddSingleton(typeof(IHealthMonitor<>), typeof(HealthMonitor<>));
-			services.AddSingleton<ServerStatusService>();
-
 			services.AddSingleton<ConformTaskSource>();
 			services.AddSingleton<ICommitService, CommitService>();
 
@@ -312,11 +219,6 @@ namespace HordeServer.Tests
 			services.AddSingleton<StorageService>();
 			services.AddSingleton<StorageBackendCache>();
 			services.AddSingleton<BundleCache>();
-
-			foreach (ILoadedPlugin plugin in _pluginCollection.LoadedPlugins)
-			{
-				plugin.ConfigureServices(new ConfigurationBuilder().Build(), services);
-			}
 		}
 
 		public Task<Fixture> CreateFixtureAsync()
@@ -460,54 +362,6 @@ namespace HordeServer.Tests
 			return agent;
 		}
 
-		/// <summary>
-		/// Hack the Datadog tracing library to not block during shutdown of tests.
-		/// Without this fix, the lib will try to send traces to a host that isn't running and block for +20 secs
-		///
-		/// Since so many of the interfaces and classes in the lib are internal it was difficult to replace Tracer.Instance
-		/// </summary>
-		private static void PatchDatadogWriter()
-		{
-			if (s_datadogWriterPatched)
-			{
-				return;
-			}
-
-			s_datadogWriterPatched = true;
-
-			string msg = "Unable to patch Datadog agent writer! Tests will still work, but shutdown will block for +20 seconds.";
-
-			FieldInfo? agentWriterField = Datadog.Trace.Tracer.Instance.GetType().GetField("_agentWriter", BindingFlags.NonPublic | BindingFlags.Instance);
-			if (agentWriterField == null)
-			{
-				Console.Error.WriteLine(msg);
-				return;
-			}
-
-			object? agentWriterInstance = agentWriterField.GetValue(Datadog.Trace.Tracer.Instance);
-			if (agentWriterInstance == null)
-			{
-				Console.Error.WriteLine(msg);
-				return;
-			}
-
-			FieldInfo? processExitField = agentWriterInstance.GetType().GetField("_processExit", BindingFlags.NonPublic | BindingFlags.Instance);
-			if (processExitField == null)
-			{
-				Console.Error.WriteLine(msg);
-				return;
-			}
-
-			TaskCompletionSource<bool>? processExitInstance = (TaskCompletionSource<bool>?)processExitField.GetValue(agentWriterInstance);
-			if (processExitInstance == null)
-			{
-				Console.Error.WriteLine(msg);
-				return;
-			}
-
-			processExitInstance.TrySetResult(true);
-		}
-
 		/// <summary>
 		/// Find an available TCP/IP port
 		/// </summary>
@@ -543,14 +397,4 @@ namespace HordeServer.Tests
 			return await PoolCollection.GetAsync(poolConfig.Id) ?? throw new NotImplementedException();
 		}
 	}
-
-	public class DowntimeServiceStub : IDowntimeService
-	{
-		public DowntimeServiceStub(bool isDowntimeActive = false)
-		{
-			IsDowntimeActive = isDowntimeActive;
-		}
-
-		public bool IsDowntimeActive { get; set; }
-	}
 }
\ No newline at end of file
diff --git a/Engine/Source/Programs/Horde/HordeServer/Properties/AssemblyInfo.cs b/Engine/Source/Programs/Horde/HordeServer/Properties/AssemblyInfo.cs
--- a/Engine/Source/Programs/Horde/HordeServer/Properties/AssemblyInfo.cs
+++ b/Engine/Source/Programs/Horde/HordeServer/Properties/AssemblyInfo.cs
@@ -4,4 +4,5 @@ using System.Resources;
 using System.Runtime.CompilerServices;
 
 [assembly: InternalsVisibleTo("HordeServer.Tests")]
+[assembly: InternalsVisibleTo("HordeServer.Tests.Shared")]
 [assembly: NeutralResourcesLanguage("en-US")]
diff --git a/Engine/Source/Programs/Horde/Plugins/Analytics/HordeServer.Analytics.Tests/HordeServer.Analytics.Tests.csproj b/Engine/Source/Programs/Horde/Plugins/Analytics/HordeServer.Analytics.Tests/HordeServer.Analytics.Tests.csproj
new file mode 100644
--- /dev/null
+++ b/Engine/Source/Programs/Horde/Plugins/Analytics/HordeServer.Analytics.Tests/HordeServer.Analytics.Tests.csproj
@@ -0,0 +1,28 @@
+<Project Sdk="Microsoft.NET.Sdk">
+
+  <PropertyGroup>
+    <TargetFramework>net8.0</TargetFramework>
+    <ImplicitUsings>enable</ImplicitUsings>
+    <Nullable>enable</Nullable>
+
+    <IsPackable>false</IsPackable>
+    <IsTestProject>true</IsTestProject>
+  </PropertyGroup>
+
+  <ItemGroup>
+    <PackageReference Include="coverlet.collector" Version="6.0.0" />
+    <PackageReference Include="Microsoft.NET.Test.Sdk" Version="17.8.0" />
+    <PackageReference Include="MSTest.TestAdapter" Version="3.1.1" />
+    <PackageReference Include="MSTest.TestFramework" Version="3.1.1" />
+  </ItemGroup>
+
+  <ItemGroup>
+    <ProjectReference Include="..\HordeServer.Analytics\HordeServer.Analytics.csproj" />
+    <ProjectReference Include="..\..\..\HordeServer.Tests.Shared\HordeServer.Tests.Shared.csproj" />
+  </ItemGroup>
+
+  <ItemGroup>
+    <Using Include="Microsoft.VisualStudio.TestTools.UnitTesting" />
+  </ItemGroup>
+
+</Project>
diff --git a/Engine/Source/Programs/Horde/Plugins/Analytics/HordeServer.Analytics.Tests/Telemetry/TelemetryTests.cs b/Engine/Source/Programs/Horde/Plugins/Analytics/HordeServer.Analytics.Tests/Telemetry/TelemetryTests.cs
new file mode 100644
--- /dev/null
+++ b/Engine/Source/Programs/Horde/Plugins/Analytics/HordeServer.Analytics.Tests/Telemetry/TelemetryTests.cs
@@ -0,0 +1,454 @@
+// Copyright Epic Games, Inc. All Rights Reserved.
+
+using System;
+using System.Collections.Generic;
+using System.Linq;
+using System.Text.Json;
+using System.Threading;
+using System.Threading.Tasks;
+using EpicGames.Horde.Telemetry;
+using EpicGames.Horde.Telemetry.Metrics;
+using HordeServer.Server;
+using HordeServer.Telemetry;
+using HordeServer.Telemetry.Metrics;
+using HordeServer.Telemetry.Sinks;
+using HordeServer.Tests;
+using HordeServer.Utilities;
+using Json.Path;
+using Microsoft.Extensions.DependencyInjection;
+using Microsoft.VisualStudio.TestTools.UnitTesting;
+
+namespace HordeServer.Analytics.Tests.Telemetry
+{
+	[TestClass]
+	public class TelemetryTests : ServerServiceTest
+	{
+		static readonly TelemetryRecordMeta s_metadata = new TelemetryRecordMeta();
+
+		protected override void ConfigureServices(IServiceCollection services)
+		{
+			base.ConfigureServices(services);
+
+			services.AddAnalytics();
+		}
+
+		[TestMethod]
+		public async Task FilterAsync()
+		{
+			Clock.UtcNow = new DateTime(2023, 6, 8, 4, 30, 0, DateTimeKind.Utc);
+
+			MetricConfig metricConfig = new MetricConfig();
+			metricConfig.Id = new MetricId("test-metric");
+			metricConfig.Function = AggregationFunction.Sum;
+			metricConfig.Interval = TimeSpan.FromHours(1.0);
+			metricConfig.Filter = JsonPath.Parse("$[?(@.Payload.EventName == 'Included')]");
+			metricConfig.Property = JsonPath.Parse("$.Payload.foo");
+
+			TelemetryStoreConfig telemetryStoreConfig = new TelemetryStoreConfig();
+			telemetryStoreConfig.Id = TelemetryStoreId.Default;
+			telemetryStoreConfig.Metrics.Add(metricConfig);
+
+			AnalyticsConfig analyticsConfig = new AnalyticsConfig();
+			analyticsConfig.Stores.Add(telemetryStoreConfig);
+
+			GlobalConfig globalConfig = new GlobalConfig();
+			globalConfig.Plugins.AddAnalytics(analyticsConfig);
+			SetConfig(globalConfig);
+
+			MetricTelemetrySink sink = ServiceProvider.GetRequiredService<MetricTelemetrySink>();
+			IMetricCollection collection = ServiceProvider.GetRequiredService<IMetricCollection>();
+
+			// Test 1
+			{
+				sink.SendEvent(telemetryStoreConfig.Id, new TelemetryEvent(s_metadata, new { EventName = "Included", foo = 1 }));
+				sink.SendEvent(telemetryStoreConfig.Id, new TelemetryEvent(s_metadata, new { EventName = "Included", foo = 2 }));
+				sink.SendEvent(telemetryStoreConfig.Id, new TelemetryEvent(s_metadata, new { EventName = "Excluded", foo = 3 }));
+				await sink.FlushAsync(CancellationToken.None);
+				await collection.FlushAsync(CancellationToken.None);
+
+				List<IMetric> metrics = await collection.FindAsync(telemetryStoreConfig.Id, [metricConfig.Id]);
+				Assert.AreEqual(1, metrics.Count);
+				Assert.AreEqual(new DateTime(2023, 6, 8, 4, 0, 0), metrics[0].Time);
+				Assert.AreEqual(3, metrics[0].Value);
+				Assert.AreEqual(2, metrics[0].Count);
+			}
+		}
+
+		[TestMethod]
+		public async Task SingleMetricAsync()
+		{
+			Clock.UtcNow = new DateTime(2023, 6, 8, 4, 30, 0, DateTimeKind.Utc);
+
+			MetricConfig metricConfig = new MetricConfig();
+			metricConfig.Id = new MetricId("test-metric");
+			metricConfig.Function = AggregationFunction.Sum;
+			metricConfig.Interval = TimeSpan.FromHours(1.0);
+			metricConfig.Property = JsonPath.Parse("$.Payload.foo");
+
+			TelemetryStoreConfig telemetryStoreConfig = new TelemetryStoreConfig();
+			telemetryStoreConfig.Id = TelemetryStoreId.Default;
+			telemetryStoreConfig.Metrics.Add(metricConfig);
+
+			AnalyticsConfig analyticsConfig = new AnalyticsConfig();
+			analyticsConfig.Stores.Add(telemetryStoreConfig);
+
+			GlobalConfig globalConfig = new GlobalConfig();
+			globalConfig.Plugins.AddAnalytics(analyticsConfig);
+			SetConfig(globalConfig);
+
+			MetricTelemetrySink sink = ServiceProvider.GetRequiredService<MetricTelemetrySink>();
+			IMetricCollection collection = ServiceProvider.GetRequiredService<IMetricCollection>();
+
+			// Test 1
+			{
+				sink.SendEvent(telemetryStoreConfig.Id, new TelemetryEvent(s_metadata, new { foo = 1 }));
+				sink.SendEvent(telemetryStoreConfig.Id, new TelemetryEvent(s_metadata, new { foo = 2 }));
+				sink.SendEvent(telemetryStoreConfig.Id, new TelemetryEvent(s_metadata, new { foo = 3 }));
+				await sink.FlushAsync(CancellationToken.None);
+				await collection.FlushAsync(CancellationToken.None);
+
+				List<IMetric> metrics = await collection.FindAsync(telemetryStoreConfig.Id, [metricConfig.Id]);
+				Assert.AreEqual(1, metrics.Count);
+				Assert.AreEqual(new DateTime(2023, 6, 8, 4, 0, 0), metrics[0].Time);
+				Assert.AreEqual(6, metrics[0].Value);
+				Assert.AreEqual(3, metrics[0].Count);
+			}
+
+			// Test 2
+			{
+				sink.SendEvent(telemetryStoreConfig.Id, new TelemetryEvent(s_metadata, new { foo = 3 }));
+				await sink.FlushAsync(CancellationToken.None);
+				await collection.FlushAsync(CancellationToken.None);
+
+				List<IMetric> metrics = await collection.FindAsync(telemetryStoreConfig.Id, [metricConfig.Id]);
+				Assert.AreEqual(1, metrics.Count);
+				Assert.AreEqual(new DateTime(2023, 6, 8, 4, 0, 0), metrics[0].Time);
+				Assert.AreEqual(9, metrics[0].Value);
+				Assert.AreEqual(4, metrics[0].Count);
+			}
+
+			// Test 3
+			{
+				await Clock.AdvanceAsync(TimeSpan.FromHours(1.0));
+
+				sink.SendEvent(telemetryStoreConfig.Id, new TelemetryEvent(s_metadata, new { foo = 4 }));
+				await sink.FlushAsync(CancellationToken.None);
+				await collection.FlushAsync(CancellationToken.None);
+
+				List<IMetric> metrics = await collection.FindAsync(telemetryStoreConfig.Id, [metricConfig.Id]);
+				Assert.AreEqual(2, metrics.Count);
+
+				Assert.AreEqual(new DateTime(2023, 6, 8, 5, 0, 0), metrics[0].Time);
+				Assert.AreEqual(4, metrics[0].Value);
+				Assert.AreEqual(1, metrics[0].Count);
+
+				Assert.AreEqual(new DateTime(2023, 6, 8, 4, 0, 0), metrics[1].Time);
+				Assert.AreEqual(9, metrics[1].Value);
+				Assert.AreEqual(4, metrics[1].Count);
+			}
+		}
+
+		[TestMethod]
+		public async Task SeparateMetricsAsync()
+		{
+			Clock.UtcNow = new DateTime(2023, 6, 8, 4, 30, 0, DateTimeKind.Utc);
+
+			MetricConfig metricConfig1 = new MetricConfig();
+			metricConfig1.Id = new MetricId("test-metric-1");
+			metricConfig1.Function = AggregationFunction.Sum;
+			metricConfig1.Interval = TimeSpan.FromHours(1.0);
+			metricConfig1.Property = JsonPath.Parse("$.Payload.foo");
+
+			MetricConfig metricConfig2 = new MetricConfig();
+			metricConfig2.Id = new MetricId("test-metric-2");
+			metricConfig2.Function = AggregationFunction.Sum;
+			metricConfig2.Interval = TimeSpan.FromHours(1.0);
+			metricConfig2.Property = JsonPath.Parse("$.Payload.bar.baz");
+
+			TelemetryStoreConfig telemetryStoreConfig = new TelemetryStoreConfig();
+			telemetryStoreConfig.Id = TelemetryStoreId.Default;
+			telemetryStoreConfig.Metrics.Add(metricConfig1);
+			telemetryStoreConfig.Metrics.Add(metricConfig2);
+
+			AnalyticsConfig analyticsConfig = new AnalyticsConfig();
+			analyticsConfig.Stores.Add(telemetryStoreConfig);
+
+			GlobalConfig globalConfig = new GlobalConfig();
+			globalConfig.Plugins.AddAnalytics(analyticsConfig);
+			SetConfig(globalConfig);
+
+			MetricTelemetrySink sink = ServiceProvider.GetRequiredService<MetricTelemetrySink>();
+			IMetricCollection collection = ServiceProvider.GetRequiredService<IMetricCollection>();
+
+			// Test 1
+			{
+				sink.SendEvent(telemetryStoreConfig.Id, new TelemetryEvent(s_metadata, new { foo = 1 }));
+				sink.SendEvent(telemetryStoreConfig.Id, new TelemetryEvent(s_metadata, new { foo = 2, bar = 201 }));
+				sink.SendEvent(telemetryStoreConfig.Id, new TelemetryEvent(s_metadata, new { foo = 3, bar = new { baz = 101 } }));
+				await sink.FlushAsync(CancellationToken.None);
+				await collection.FlushAsync(CancellationToken.None);
+
+				List<IMetric> metrics = await collection.FindAsync(telemetryStoreConfig.Id, [metricConfig1.Id]);
+				Assert.AreEqual(1, metrics.Count);
+				Assert.AreEqual(new DateTime(2023, 6, 8, 4, 0, 0), metrics[0].Time);
+				Assert.AreEqual(6, metrics[0].Value);
+				Assert.AreEqual(3, metrics[0].Count);
+
+				List<IMetric> metrics2 = await collection.FindAsync(telemetryStoreConfig.Id, [metricConfig2.Id]);
+				Assert.AreEqual(1, metrics2.Count);
+				Assert.AreEqual(new DateTime(2023, 6, 8, 4, 0, 0), metrics2[0].Time);
+				Assert.AreEqual(101, metrics2[0].Value);
+				Assert.AreEqual(1, metrics2[0].Count);
+			}
+		}
+		/*
+	[TestMethod]
+	public async Task FunctionTestAsync()
+	{
+		await SingleFunctionTestAsync(AggregationFunction.Count, new double[] { 5, 4, 3, -1, 2 }, 5);
+		await SingleFunctionTestAsync(AggregationFunction.Min, new double[] { 5, 4, 3, -1, 2 }, -1);
+		await SingleFunctionTestAsync(AggregationFunction.Max, new double[] { 5, 4, 3, -1, 2 }, 5);
+		await SingleFunctionTestAsync(AggregationFunction.Sum, new double[] { 5, 4, 3, -1, 2 }, 13);
+		await SingleFunctionTestAsync(AggregationFunction.Average, new double[] { 5, 4, 3, -1, 2 }, 2.6);
+		await SingleFunctionTestAsync(AggregationFunction.Percentile, 4.25);
+	}
+		*/
+		[TestMethod]
+		[DataRow(AggregationFunction.Count, 5)]
+		[DataRow(AggregationFunction.Min, -1)]
+		[DataRow(AggregationFunction.Max, 5)]
+		[DataRow(AggregationFunction.Sum, 13)]
+		[DataRow(AggregationFunction.Average, 2.6)]
+		[DataRow(AggregationFunction.Percentile, 4.25)]
+		public async Task SingleFunctionTestAsync(AggregationFunction function, double result)
+		{
+			double[] values = new double[] { 5, 4, 3, -1, 2 };
+
+			await Clock.AdvanceAsync(TimeSpan.FromDays(1.0));
+
+			MetricConfig metricConfig = new MetricConfig();
+			metricConfig.Id = new MetricId("test-metric-2");
+			metricConfig.Function = function;
+			metricConfig.Interval = TimeSpan.FromHours(1.0);
+			metricConfig.Property = JsonPath.Parse("$.Payload.foo");
+			metricConfig.Percentile = 75;
+
+			TelemetryStoreConfig telemetryStoreConfig = new TelemetryStoreConfig();
+			telemetryStoreConfig.Id = TelemetryStoreId.Default;
+			telemetryStoreConfig.Metrics.Add(metricConfig);
+
+			AnalyticsConfig analyticsConfig = new AnalyticsConfig();
+			analyticsConfig.Stores.Add(telemetryStoreConfig);
+
+			GlobalConfig globalConfig = new GlobalConfig();
+			globalConfig.Plugins.AddAnalytics(analyticsConfig);
+			SetConfig(globalConfig);
+
+			MetricTelemetrySink sink = ServiceProvider.GetRequiredService<MetricTelemetrySink>();
+			IMetricCollection collection = ServiceProvider.GetRequiredService<IMetricCollection>();
+
+			foreach (double value in values)
+			{
+				sink.SendEvent(telemetryStoreConfig.Id, new TelemetryEvent(s_metadata, new { foo = value }));
+			}
+
+			await sink.FlushAsync(CancellationToken.None);
+			await collection.FlushAsync(CancellationToken.None);
+
+			List<IMetric> metrics = await collection.FindAsync(telemetryStoreConfig.Id, [metricConfig.Id], maxResults: 1);
+			Assert.AreEqual(1, metrics.Count);
+			Assert.AreEqual(values.Length, metrics[0].Count);
+			Assert.AreEqual(result, metrics[0].Value);
+		}
+
+		[TestMethod]
+		public async Task GroupingTestAsync()
+		{
+			Clock.UtcNow = new DateTime(2023, 6, 8, 4, 30, 0, DateTimeKind.Utc);
+
+			MetricConfig metricConfig = new MetricConfig();
+			metricConfig.Id = new MetricId("test-metric-1");
+			metricConfig.Function = AggregationFunction.Sum;
+			metricConfig.Interval = TimeSpan.FromHours(1.0);
+			metricConfig.Property = JsonPath.Parse("$.Payload.foo");
+			metricConfig.GroupBy = "$.Payload.group";
+
+			TelemetryStoreConfig telemetryStoreConfig = new TelemetryStoreConfig();
+			telemetryStoreConfig.Id = TelemetryStoreId.Default;
+			telemetryStoreConfig.Metrics.Add(metricConfig);
+
+			AnalyticsConfig analyticsConfig = new AnalyticsConfig();
+			analyticsConfig.Stores.Add(telemetryStoreConfig);
+
+			GlobalConfig globalConfig = new GlobalConfig();
+			globalConfig.Plugins.AddAnalytics(analyticsConfig);
+			SetConfig(globalConfig);
+
+			MetricTelemetrySink sink = ServiceProvider.GetRequiredService<MetricTelemetrySink>();
+			IMetricCollection collection = ServiceProvider.GetRequiredService<IMetricCollection>();
+
+			sink.SendEvent(TelemetryStoreId.Default, new TelemetryEvent(s_metadata, new { foo = 1, group = "first" }));
+			sink.SendEvent(TelemetryStoreId.Default, new TelemetryEvent(s_metadata, new { foo = 2, group = "first" }));
+			sink.SendEvent(TelemetryStoreId.Default, new TelemetryEvent(s_metadata, new { foo = 3, group = "first" }));
+			sink.SendEvent(TelemetryStoreId.Default, new TelemetryEvent(s_metadata, new { foo = 4, group = "second" }));
+			sink.SendEvent(TelemetryStoreId.Default, new TelemetryEvent(s_metadata, new { foo = 5 }));
+			await sink.FlushAsync(CancellationToken.None);
+			await collection.FlushAsync(CancellationToken.None);
+
+			List<IMetric> metrics = await collection.FindAsync(telemetryStoreConfig.Id, [metricConfig.Id]);
+			metrics = metrics.OrderBy(x => x.Group).ToList();
+			Assert.AreEqual(3, metrics.Count);
+
+			Assert.AreEqual("", metrics[0].Group);
+			Assert.AreEqual(1, metrics[0].Count);
+			Assert.AreEqual(5, metrics[0].Value);
+
+			Assert.AreEqual("first", metrics[1].Group);
+			Assert.AreEqual(3, metrics[1].Count);
+			Assert.AreEqual(6, metrics[1].Value);
+
+			Assert.AreEqual("second", metrics[2].Group);
+			Assert.AreEqual(1, metrics[2].Count);
+			Assert.AreEqual(4, metrics[2].Value);
+		}
+
+		[TestMethod]
+		public async Task GroupingEscapeTestAsync()
+		{
+			Clock.UtcNow = new DateTime(2023, 6, 8, 4, 30, 0, DateTimeKind.Utc);
+
+			MetricConfig metricConfig = new MetricConfig();
+			metricConfig.Id = new MetricId("test-metric-1");
+			metricConfig.Function = AggregationFunction.Sum;
+			metricConfig.Interval = TimeSpan.FromHours(1.0);
+			metricConfig.Property = JsonPath.Parse("$.Payload.foo");
+			metricConfig.GroupBy = "$.Payload.group";
+
+			TelemetryStoreConfig telemetryStoreConfig = new TelemetryStoreConfig();
+			telemetryStoreConfig.Id = TelemetryStoreId.Default;
+			telemetryStoreConfig.Metrics.Add(metricConfig);
+
+			AnalyticsConfig analyticsConfig = new AnalyticsConfig();
+			analyticsConfig.Stores.Add(telemetryStoreConfig);
+
+			GlobalConfig globalConfig = new GlobalConfig();
+			globalConfig.Plugins.AddAnalytics(analyticsConfig);
+			SetConfig(globalConfig);
+
+			MetricTelemetrySink sink = ServiceProvider.GetRequiredService<MetricTelemetrySink>();
+			IMetricCollection collection = ServiceProvider.GetRequiredService<IMetricCollection>();
+
+			sink.SendEvent(telemetryStoreConfig.Id, new TelemetryEvent(s_metadata, new { foo = 123, group = "first, second & third" }));
+			await sink.FlushAsync(CancellationToken.None);
+			await collection.FlushAsync(CancellationToken.None);
+
+			List<IMetric> metrics = await collection.FindAsync(telemetryStoreConfig.Id, [metricConfig.Id]);
+			metrics = metrics.OrderBy(x => x.Group).ToList();
+			Assert.AreEqual(1, metrics.Count);
+
+			Assert.AreEqual("\"first, second & third\"", metrics[0].Group);
+			Assert.AreEqual(1, metrics[0].Count);
+			Assert.AreEqual(123, metrics[0].Value);
+		}
+
+		[TestMethod]
+		public void GroupConverterTest()
+		{
+			JsonSerializerOptions serializerOptions = new JsonSerializerOptions();
+			JsonUtils.ConfigureJsonSerializer(serializerOptions);
+
+			MetricConfig? config = JsonSerializer.Deserialize<MetricConfig>("{ \"groupBy\": \"$.foo,$.bar\" }", serializerOptions);
+			Assert.IsNotNull(config);
+			Assert.AreEqual(2, config.GroupByPaths.Count);
+			Assert.AreEqual("$.foo", config.GroupByPaths[0].ToString());
+			Assert.AreEqual("$.bar", config.GroupByPaths[1].ToString());
+
+			string text = JsonSerializer.Serialize(config, serializerOptions);
+			Assert.IsTrue(text.Contains("\"groupBy\":\"$.foo,$.bar\"", StringComparison.Ordinal));
+		}
+
+		[TestMethod]
+		[DataRow(AggregationFunction.Count)]
+		[DataRow(AggregationFunction.Sum)]
+		public async Task MultiGroupingTestAsync(AggregationFunction function)
+		{
+			Clock.UtcNow = new DateTime(2023, 6, 8, 4, 30, 0, DateTimeKind.Utc);
+
+			MetricConfig metricConfig = new MetricConfig();
+			metricConfig.Id = new MetricId("test-metric");
+			metricConfig.Function = function;
+			metricConfig.Interval = TimeSpan.FromHours(1.0);
+			metricConfig.Filter = JsonPath.Parse("$[?(@.Payload.EventName == 'Included')]");
+			metricConfig.GroupBy = "$.Payload.groupFacetA, $.Payload.groupFacetB, $.Payload.groupFacetC";
+			if (function == AggregationFunction.Sum)
+			{
+				metricConfig.Property = JsonPath.Parse("$.Payload.foo");
+			}
+
+			TelemetryStoreConfig telemetryStoreConfig = new TelemetryStoreConfig();
+			telemetryStoreConfig.Id = TelemetryStoreId.Default;
+			telemetryStoreConfig.Metrics.Add(metricConfig);
+
+			AnalyticsConfig analyticsConfig = new AnalyticsConfig();
+			analyticsConfig.Stores.Add(telemetryStoreConfig);
+
+			GlobalConfig globalConfig = new GlobalConfig();
+			globalConfig.Plugins.AddAnalytics(analyticsConfig);
+			SetConfig(globalConfig);
+
+			MetricTelemetrySink sink = ServiceProvider.GetRequiredService<MetricTelemetrySink>();
+			IMetricCollection collection = ServiceProvider.GetRequiredService<IMetricCollection>();
+
+			// Test 1
+			{
+				sink.SendEvent(telemetryStoreConfig.Id, new TelemetryEvent(s_metadata, new { EventName = "Included", foo = 1, groupFacetA = "groupA" }));
+				sink.SendEvent(telemetryStoreConfig.Id, new TelemetryEvent(s_metadata, new { EventName = "Included", foo = 2, groupFacetA = "groupA" }));
+				sink.SendEvent(telemetryStoreConfig.Id, new TelemetryEvent(s_metadata, new { EventName = "Included", foo = 3, groupFacetA = "groupA", groupFacetB = "groupB" }));
+				sink.SendEvent(telemetryStoreConfig.Id, new TelemetryEvent(s_metadata, new { EventName = "Included", foo = 4, groupFacetB = "groupB" }));
+				sink.SendEvent(telemetryStoreConfig.Id, new TelemetryEvent(s_metadata, new { EventName = "Included", foo = 5, groupFacetB = "groupA,groupB" }));
+				sink.SendEvent(telemetryStoreConfig.Id, new TelemetryEvent(s_metadata, new { EventName = "Included", foo = 6, groupFacetA = "groupA", groupFacetB = "groupB", groupFacetC = "groupC" }));
+				sink.SendEvent(telemetryStoreConfig.Id, new TelemetryEvent(s_metadata, new { EventName = "Excluded", foo = 6 }));
+				await sink.FlushAsync(CancellationToken.None);
+				await collection.FlushAsync(CancellationToken.None);
+
+				List<IMetric> metrics = await collection.FindAsync(telemetryStoreConfig.Id, [metricConfig.Id]);
+				metrics = metrics.OrderBy(x => x.Group).ToList();
+
+				Assert.AreEqual(5, metrics.Count);
+
+				Assert.AreEqual(",\"groupA,groupB\",", metrics[0].Group);
+				Assert.AreEqual(",groupB,", metrics[1].Group);
+				Assert.AreEqual("groupA,,", metrics[2].Group);
+				Assert.AreEqual("groupA,groupB,", metrics[3].Group);
+				Assert.AreEqual("groupA,groupB,groupC", metrics[4].Group);
+
+				if (function == AggregationFunction.Sum)
+				{
+					Assert.AreEqual(1, metrics[0].Count);
+					Assert.AreEqual(5, metrics[0].Value);
+
+					Assert.AreEqual(1, metrics[1].Count);
+					Assert.AreEqual(4, metrics[1].Value);
+
+					Assert.AreEqual(2, metrics[2].Count);
+					Assert.AreEqual(3, metrics[2].Value);
+
+					Assert.AreEqual(1, metrics[3].Count);
+					Assert.AreEqual(3, metrics[3].Value);
+
+					Assert.AreEqual(1, metrics[4].Count);
+					Assert.AreEqual(6, metrics[4].Value);
+				}
+
+				if (function == AggregationFunction.Count)
+				{
+					for (int i = 0; i < 5; i++)
+					{
+						Assert.AreEqual(i == 2 ? 2 : 1, metrics[i].Count);
+						Assert.AreEqual(i == 2 ? 2 : 1, metrics[i].Value);
+					}
+				}
+			}
+		}
+	}
+}
diff --git a/Engine/Source/Programs/Horde/Plugins/Analytics/HordeServer.Analytics/Properties/AssemblyInfo.cs b/Engine/Source/Programs/Horde/Plugins/Analytics/HordeServer.Analytics/Properties/AssemblyInfo.cs
--- a/Engine/Source/Programs/Horde/Plugins/Analytics/HordeServer.Analytics/Properties/AssemblyInfo.cs
+++ b/Engine/Source/Programs/Horde/Plugins/Analytics/HordeServer.Analytics/Properties/AssemblyInfo.cs
@@ -3,5 +3,5 @@
 using System.Resources;
 using System.Runtime.CompilerServices;
 
-[assembly: InternalsVisibleTo("HordeServer.Tests")]
+[assembly: InternalsVisibleTo("HordeServer.Analytics.Tests")]
 [assembly: NeutralResourcesLanguage("en-US")]
