Horde: Fix static analysis warnings.

#rnx

diff --git a/Engine/Source/Programs/Horde/HordeServer.Tests/Server/RpcServiceTest.cs b/Engine/Source/Programs/Horde/HordeServer.Tests/Server/RpcServiceTest.cs
--- a/Engine/Source/Programs/Horde/HordeServer.Tests/Server/RpcServiceTest.cs
+++ b/Engine/Source/Programs/Horde/HordeServer.Tests/Server/RpcServiceTest.cs
@@ -2,7 +2,6 @@
 
 using System;
 using System.Collections.Generic;
-using System.IO;
 using System.Linq;
 using System.Reflection;
 using System.Security.Claims;
@@ -12,15 +11,13 @@ using System.Threading.Tasks;
 using EpicGames.Horde.Agents;
 using EpicGames.Horde.Agents.Pools;
 using Grpc.Core;
+using HordeCommon.Rpc;
+using HordeCommon.Rpc.Messages;
 using HordeServer.Agents;
 using HordeServer.Server;
 using HordeServer.Utilities;
-using HordeCommon.Rpc;
-using HordeCommon.Rpc.Messages;
 using Microsoft.AspNetCore.Http;
 using Microsoft.AspNetCore.Http.Features;
-using Microsoft.Extensions.FileProviders;
-using Microsoft.Extensions.Hosting;
 using Microsoft.VisualStudio.TestTools.UnitTesting;
 
 namespace HordeServer.Tests.Server
diff --git a/Engine/Source/Programs/Horde/HordeServer.Tests/TestSetup.cs b/Engine/Source/Programs/Horde/HordeServer.Tests/TestSetup.cs
--- a/Engine/Source/Programs/Horde/HordeServer.Tests/TestSetup.cs
+++ b/Engine/Source/Programs/Horde/HordeServer.Tests/TestSetup.cs
@@ -5,20 +5,17 @@ using System.Collections.Generic;
 using System.Diagnostics.Metrics;
 using System.Net;
 using System.Net.Sockets;
-using System.Reflection;
 using System.Security.Claims;
 using System.Threading.Tasks;
 using Amazon.AutoScaling;
 using Amazon.CloudWatch;
 using Amazon.EC2;
-using EpicGames.Core;
 using EpicGames.Horde.Agents;
 using EpicGames.Horde.Agents.Pools;
 using EpicGames.Horde.Storage;
 using EpicGames.Horde.Storage.Bundles;
 using EpicGames.Horde.Storage.ObjectStores;
 using HordeServer.Accounts;
-using HordeServer.Acls;
 using HordeServer.Agents;
 using HordeServer.Agents.Enrollment;
 using HordeServer.Agents.Fleet;
@@ -32,7 +29,6 @@ using HordeServer.Artifacts;
 using HordeServer.Auditing;
 using HordeServer.Commits;
 using HordeServer.Compute;
-using HordeServer.Configuration;
 using HordeServer.Dashboard;
 using HordeServer.Devices;
 using HordeServer.Issues;
@@ -50,28 +46,19 @@ using HordeServer.Server;
 using HordeServer.Storage;
 using HordeServer.Streams;
 using HordeServer.Tasks;
-using HordeServer.Telemetry;
-using HordeServer.Telemetry.Sinks;
-using HordeServer.Tests.Server;
 using HordeServer.Tests.Stubs.Services;
 using HordeServer.Tools;
 using HordeServer.Ugs;
 using HordeServer.Users;
 using HordeServer.Utilities;
-using HordeCommon;
-using HordeServer.Telemetry.Metrics;
 using Microsoft.AspNetCore.Http;
 using Microsoft.AspNetCore.Mvc;
-using Microsoft.Extensions.Caching.Memory;
 using Microsoft.Extensions.Configuration;
 using Microsoft.Extensions.DependencyInjection;
-using Microsoft.Extensions.Hosting;
 using Microsoft.Extensions.Logging;
-using Microsoft.Extensions.Options;
 using Microsoft.VisualStudio.TestTools.UnitTesting;
 using Moq;
 using OpenTelemetry.Trace;
-using HordeServer.Plugins;
 
 namespace HordeServer.Tests
 {
@@ -347,7 +334,7 @@ namespace HordeServer.Tests
 			{
 				await agent.TryUpdateWorkspacesAsync(workspaces, false);
 			}
-			
+
 			if (requestShutdown)
 			{
 				await agent.TryUpdateAsync(new UpdateAgentOptions { RequestShutdown = true });
