Horde: Enable analytics tests in Docker build.

#rnx

diff --git a/Engine/Source/Programs/Horde/Scripts/test.sh b/Engine/Source/Programs/Horde/Scripts/test.sh
--- a/Engine/Source/Programs/Horde/Scripts/test.sh
+++ b/Engine/Source/Programs/Horde/Scripts/test.sh
@@ -14,7 +14,8 @@ testProjects=(
 	"Source/Programs/Shared/EpicGames.IoHash.Tests/EpicGames.IoHash.Tests.csproj"
 	"Source/Programs/Shared/EpicGames.Redis.Tests/EpicGames.Redis.Tests.csproj"
 	"Source/Programs/Shared/EpicGames.Serialization.Tests/EpicGames.Serialization.Tests.csproj"
-	"Source/Programs/Horde/HordeServer.Tests/HordeServer.Tests.csproj"
+	"Source/Programs/Horde/HordeServer.Tests/HordeServer.Tests.csproj",
+	"Source/Programs/Horde/Plugins/Analytics/HordeServer.Analytics.Tests/HordeServer.Analytics.Tests.csproj"
 )
 
 for csProj in "${testProjects[@]}"; do
