Horde: Include shared test project in docker image.

#rnx

diff --git a/Engine/Source/Programs/Horde/HordeServer/Dockerfile b/Engine/Source/Programs/Horde/HordeServer/Dockerfile
--- a/Engine/Source/Programs/Horde/HordeServer/Dockerfile
+++ b/Engine/Source/Programs/Horde/HordeServer/Dockerfile
@@ -60,6 +60,7 @@ RUN dotnet publish Source/Programs/Horde/HordeServer/HordeServer.csproj -c Relea
 COPY Source/Programs/Horde/Drivers/JobDriver/ Source/Programs/Horde/Drivers/JobDriver/
 COPY Source/Programs/Horde/HordeAgent/ Source/Programs/Horde/HordeAgent/
 COPY Source/Programs/Horde/HordeServer.Tests/ Source/Programs/Horde/HordeServer.Tests/
+COPY Source/Programs/Horde/HordeServer.Tests.Shared/ Source/Programs/Horde/HordeServer.Tests.Shared/
 COPY Source/Programs/Shared/EpicGames.Serialization.Tests/ Source/Programs/Shared/EpicGames.Serialization.Tests/
 
 ARG code_coverage
