Horde: Set agent property if Wine executable is configured

Needed for scheduling Wine-compatible compute tasks.

diff --git a/Engine/Source/Programs/Horde/HordeAgent/Services/CapabilitiesService.cs b/Engine/Source/Programs/Horde/HordeAgent/Services/CapabilitiesService.cs
--- a/Engine/Source/Programs/Horde/HordeAgent/Services/CapabilitiesService.cs
+++ b/Engine/Source/Programs/Horde/HordeAgent/Services/CapabilitiesService.cs
@@ -244,7 +244,12 @@ namespace HordeAgent.Services
 				{
 					await AddAwsPropertiesAsync(primaryDevice.Properties, logger);
 				}
-
+				
+				if (_settings.WineExecutablePath != null)
+				{
+					primaryDevice.Properties.Add($"{KnownPropertyNames.WineEnabled}=true");
+				}
+				
 				// Parse the CPU info
 				List<Dictionary<string, string>>? cpuRecords = await ReadLinuxHwPropsAsync("/proc/cpuinfo", logger);
 				if (cpuRecords != null)
diff --git a/Engine/Source/Programs/Shared/EpicGames.Horde/Agents/KnownPropertyNames.cs b/Engine/Source/Programs/Shared/EpicGames.Horde/Agents/KnownPropertyNames.cs
--- a/Engine/Source/Programs/Shared/EpicGames.Horde/Agents/KnownPropertyNames.cs
+++ b/Engine/Source/Programs/Shared/EpicGames.Horde/Agents/KnownPropertyNames.cs
@@ -76,4 +76,9 @@ public static class KnownPropertyNames
 	/// AWS: Instance type
 	/// </summary>
 	public const string AwsInstanceType = "aws-instance-type";
+	
+	/// <summary>
+	/// Whether the Wine compatibility layer is enabled (for running Windows applications on Linux)
+	/// </summary>
+	public const string WineEnabled = "WineEnabled";
 }
\ No newline at end of file
