Horde: Show invalid pool assignment with no pools filter
#rnx
#skipci

diff --git a/Engine/Source/Programs/Horde/HordeDashboard/src/components/AgentView.tsx b/Engine/Source/Programs/Horde/HordeDashboard/src/components/AgentView.tsx
--- a/Engine/Source/Programs/Horde/HordeDashboard/src/components/AgentView.tsx
+++ b/Engine/Source/Programs/Horde/HordeDashboard/src/components/AgentView.tsx
@@ -2287,7 +2287,25 @@ export const AgentViewInner: React.FC<{ agentId?: string, poolId?: string, searc
       if (filter.has("No Pools")) {
          if (!item.pools?.length) {
             filtered = false;
+         } else {
+            let anyInvalid = false;
+            item.pools.forEach(pid => {
+
+               if (anyInvalid) {
+                  return;
+               }
+
+               if (!agentStore.pools.find(pool => { return pool.id === pid; })) {
+                  anyInvalid = true;
+               }
+
+            });
+
+            if (anyInvalid) {
+               filtered = false;
+            }
          }
+
       }
 
       if (filter.has("Ephemeral")) {
@@ -2520,8 +2538,9 @@ export const AgentViewInner: React.FC<{ agentId?: string, poolId?: string, searc
                </Stack>
             );
          case 'pools':
-            const poolItems:any = [];
-            const poolSearchNames:string[] = [];
+            const poolItems: any = [];
+            const poolSearchNames: string[] = [];
+            const invalidPoolIds = new Set<string>();
             if (agent.pools) {
                const poolObjs: PoolData[] = [];
                // get actual pool objects
@@ -2530,6 +2549,15 @@ export const AgentViewInner: React.FC<{ agentId?: string, poolId?: string, searc
                   const pool = agentStore.pools.find(pool => { return pool.id === id; });
                   if (pool) {
                      poolObjs.push(pool);
+                  } else {
+                     invalidPoolIds.add(id);
+                     poolObjs.push({
+                        id: id,
+                        name: id ?? "",
+                        colorValue: "#797979",
+                        enableAutoscaling: false,
+                        workspaces: []
+                     })
                   }
                }
                // sort them by name
@@ -2538,7 +2566,7 @@ export const AgentViewInner: React.FC<{ agentId?: string, poolId?: string, searc
                });
                for (let idx = 0; idx < poolObjs.length; idx++) {
                   let color = poolObjs[idx].colorValue;
-                  const textColor = "white";
+                  const textColor = invalidPoolIds.has(poolObjs[idx]?.id) ? "#DDDDDD" : "white";
                   if (agent.pendingConform || agent.pendingFullConform) {
                      const pendingConformColor = hexToRGB(color);
                      color = `rgb(${pendingConformColor.r},${pendingConformColor.g},${pendingConformColor.b}, .5)`;
@@ -2570,13 +2598,15 @@ export const AgentViewInner: React.FC<{ agentId?: string, poolId?: string, searc
                      ],
                   };
 
+                  const disabled = invalidPoolIds.has(poolObjs[idx].id);
+
 
                   poolItems.push(
                      <Stack.Item align={"center"} key={"pool_" + agent.id + "_" + poolObjs[idx].id}>
                         <DefaultButton key={poolObjs[idx].id}
                            text={poolObjs[idx].name}
                            primary
-                           menuProps={menuProps}
+                           menuProps={!disabled ? menuProps : undefined}
                            menuIconProps={{ iconName: "" }}
                            className={agentStyles.buttonFont}
                            onClick={(ev) => { ev.preventDefault(); ev.stopPropagation() }}
@@ -2592,10 +2622,10 @@ export const AgentViewInner: React.FC<{ agentId?: string, poolId?: string, searc
                }
             }
             localState.columnSearchState['pools'][agent.id] = poolSearchNames;
-            return <Stack horizontal horizontalAlign={"start"} styles={{ root: { overflow: "auto", height: '100%' } }} tokens={{ childrenGap: 4 }}>{poolItems}</Stack>;
+            return <Stack horizontal horizontalAlign={"start"} styles={{ root: { overflow: "auto", height: '100%' } }} tokens={{ childrenGap: 6 }}>{poolItems}</Stack>;
          case 'status':
             const leases: any = [];
-            const leaseSearchItems:string[] = [];
+            const leaseSearchItems: string[] = [];
             if (agent.leases) {
                agent.leases.forEach(lease => {
 
@@ -2645,7 +2675,7 @@ export const AgentViewInner: React.FC<{ agentId?: string, poolId?: string, searc
                   }
                   if (lease.name) {
                      leaseSearchItems.push(lease.name);
-                  }                  
+                  }
                });
             }
             // if there are no leases, we'll push some other state.
