Horde: Add a find artfiacts to stream view, fix issue with name being populated due to lack of a key on the freeform artifact type entry
#rnx

diff --git a/Engine/Source/Programs/Horde/HordeDashboard/src/components/StreamView.tsx b/Engine/Source/Programs/Horde/HordeDashboard/src/components/StreamView.tsx
--- a/Engine/Source/Programs/Horde/HordeDashboard/src/components/StreamView.tsx
+++ b/Engine/Source/Programs/Horde/HordeDashboard/src/components/StreamView.tsx
@@ -24,6 +24,7 @@ import { StreamSummary } from './StreamSummary';
 import { TopNav } from './TopNav';
 import { NewBuildV2 } from './build/NewBuildV2';
 import { getSiteConfig } from '../backend/Config';
+import { FindArtifactsModal } from './artifacts/ArtifactsSearch';
 
 export const SummaryPage: React.FC = () => {
 
@@ -112,15 +113,15 @@ class IncrementalState {
       const templateMap = new Map<string, GetTemplateRefResponse>();
       this.project.streams.forEach(s => {
          const tab = s.tabs.find(t => t.title === "Incremental") as GetJobsTabResponse;
-         if (!tab) { 
+         if (!tab) {
             return;
          }
-         
-         if (tab.templates?.length === 1) {            
+
+         if (tab.templates?.length === 1) {
             const template = s.templates.find(t => t.id === tab.templates![0]);
             if (template)
                templateMap.set(s.id, template);
-         }                  
+         }
       })
 
       let incrementals: StreamIncemental[] = [];
@@ -131,10 +132,10 @@ class IncrementalState {
             template: template,
             jobs: [],
             labelState: LabelState.Unspecified,
-            labelOutcome: LabelOutcome.Success            
-         })         
+            labelOutcome: LabelOutcome.Success
+         })
       })
-      
+
       if (!incrementals.length) {
          return;
       }
@@ -150,12 +151,12 @@ class IncrementalState {
 
       while (rincrementals.length) {
 
-         const batch = rincrementals.slice(0, 5);         
+         const batch = rincrementals.slice(0, 5);
 
          await Promise.all(batch.map(b => {
             return backend.getStreamJobs(b.streamId, { template: [b.template.id], count: 5, filter: "id,labels,createTime,streamId,defaultLabel,preflightChange" })
             // eslint-disable-next-line no-loop-func
-         })).then((r) => {            
+         })).then((r) => {
             for (let i = 0; i < r.length; i++) {
                let jobs = r[i];
                // filter out jobs > 3 days
@@ -180,9 +181,9 @@ class IncrementalState {
                   labels.forEach((label, index) => {
 
                      if (label.outcome === LabelOutcome.Failure || label.outcome === LabelOutcome.Warnings) {
-                        labelStatus.push({ index: index, outcome: label.outcome }); 
+                        labelStatus.push({ index: index, outcome: label.outcome });
                      } else if (label.state === LabelState.Complete && label.outcome === LabelOutcome.Success) {
-                        labelStatus.push({ index: index, outcome: label.outcome }); 
+                        labelStatus.push({ index: index, outcome: label.outcome });
                      }
                   });
 
@@ -222,31 +223,30 @@ class IncrementalState {
          let jobs = i.jobs.sort((a, b) => new Date(b.createTime).getTime() - new Date(a.createTime).getTime());
 
          const labelStatus = new Map<number, LabelOutcome>();
-      
+
          for (let x = 0; x < jobs.length; x++) {
 
             const j = jobs[x];
             const status = jobStatus.get(j.id);
-   
+
             if (!status?.length) {
                continue;
             }
 
-            for (let y = 0; y < status.length; y++)
-            {
+            for (let y = 0; y < status.length; y++) {
                const label = status[y];
 
                if (labelStatus.get(label.index)) {
                   continue;
                }
-               
+
                labelStatus.set(label.index, label.outcome);
-            }         
+            }
          }
 
          const error = Array.from(labelStatus.values()).find(outcome => outcome === LabelOutcome.Failure)
          const warning = Array.from(labelStatus.values()).find(outcome => outcome === LabelOutcome.Warnings)
-         
+
          if (!!error) {
             streamOutcome.set(i.streamId, LabelOutcome.Failure);
          } else if (!!warning) {
@@ -284,7 +284,7 @@ const IncrementalPanel: React.FC<{ project: ProjectData }> = observer(({ project
    if (incrementalState.querying) {
       return <Stack key={`Incrementalpanel_spinner_${incrementalState.updated}`} style={{ padding: 32 }} tokens={{ childrenGap: 24 }} >
          <Stack>
-            <Text style={{fontWeight: 600} } variant="medium">Querying Stream Labels</Text>
+            <Text style={{ fontWeight: 600 }} variant="medium">Querying Stream Labels</Text>
          </Stack>
          <Stack>
             <Spinner size={SpinnerSize.large} />
@@ -318,14 +318,14 @@ const IncrementalPanel: React.FC<{ project: ProjectData }> = observer(({ project
 
    if (!items.length) {
       return <Stack key={`Incrementalpanel_${incrementalState.updated}`} style={{ padding: 32 }} tokens={{ childrenGap: 12 }}>
-         <Text style={{fontWeight: 600} }variant='medium'>No Label Issues</Text>
+         <Text style={{ fontWeight: 600 }} variant='medium'>No Label Issues</Text>
       </Stack>
    }
 
 
    return <Stack key={`Incrementalpanel_${incrementalState.updated}`} style={{ padding: 32 }} tokens={{ childrenGap: 18 }}>
       <Stack>
-         <Text style={{fontWeight: 600} } variant='medium'>Label Issues</Text>
+         <Text style={{ fontWeight: 600 }} variant='medium'>Label Issues</Text>
       </Stack>
       <Stack tokens={{ childrenGap: 12 }} >
          {items}
@@ -345,6 +345,7 @@ const StreamViewInner: React.FC = observer(() => {
    const [showOthersPreflights, setShowOthersPreflights] = useState<boolean | undefined>(dashboard.showPreflights);
 
    const [shown, setShown] = useState(query.get("newbuild") ? true : false);
+   const [artifactSearchShown, setArtifactSearchShown] = useState(false);
    const [findJobsShown, setFindJobsShown] = useState(false);
 
    const { projectStore } = useBackend();
@@ -357,7 +358,7 @@ const StreamViewInner: React.FC = observer(() => {
    const project = stream?.project;
 
    //let newBuildVersion: string = query.get("newbuildversion") ? query.get("newbuildversion") : "2";
-   const newBuildVersion:string = "2"
+   const newBuildVersion: string = "2"
 
    if (!stream || !project) {
       console.error("Bad stream or project id in StreamView");
@@ -464,6 +465,14 @@ const StreamViewInner: React.FC = observer(() => {
                   </Stack>
                }
             },
+            {
+               key: "find_artifacts",
+               text: "Find Artifacts",
+               iconProps: { iconName: "Search" },
+               onClick: () => {
+                  setArtifactSearchShown(true);
+               }
+            },
             {
                key: 'show_other_preflights',
                text: 'Show preflights for all users',
@@ -490,11 +499,11 @@ const StreamViewInner: React.FC = observer(() => {
       filterKeyword: filter
    }
 
-
    return (
       <Stack className={hordeClasses.horde}>
          <TopNav />
          <Breadcrumbs items={crumbItems} title={crumbTitle} />
+         {artifactSearchShown && <FindArtifactsModal streamId={stream.id} onClose={() => setArtifactSearchShown(false)} />}
          {shown && newBuildVersion == "2" && <NewBuildV2 streamId={streamId!} jobKey={newBuildTab!} show={true} onClose={(newJobId) => {
             setShown(false);
             if (newJobId) {
@@ -504,7 +513,7 @@ const StreamViewInner: React.FC = observer(() => {
                   navigate(`/stream/${streamId}?tab=${queryTab}`, { replace: true });
                }
             }
-         }} />}         
+         }} />}
          {shown && newBuildVersion == "1" && <NewBuild streamId={streamId!} jobKey={newBuildTab!} show={true} onClose={(newJobId) => {
             setShown(false);
             if (newJobId) {
@@ -540,7 +549,7 @@ const StreamViewInner: React.FC = observer(() => {
                      <Stack horizontal verticalAlign="center" horizontalAlign={"end"} tokens={{ childrenGap: 8 }}>
                         <Stack horizontal tokens={{ childrenGap: 18 }}>
                            {!isSwarmTab && <DefaultButton
-                              styles={{ root: { fontFamily: "Horde Open Sans SemiBold !important"  } }}
+                              styles={{ root: { fontFamily: "Horde Open Sans SemiBold !important" } }}
                               text="Search"
                               split={!isSummary}
                               menuProps={!isSummary ? findJobsItems : undefined}
diff --git a/Engine/Source/Programs/Horde/HordeDashboard/src/components/artifacts/ArtifactsSearch.tsx b/Engine/Source/Programs/Horde/HordeDashboard/src/components/artifacts/ArtifactsSearch.tsx
--- a/Engine/Source/Programs/Horde/HordeDashboard/src/components/artifacts/ArtifactsSearch.tsx
+++ b/Engine/Source/Programs/Horde/HordeDashboard/src/components/artifacts/ArtifactsSearch.tsx
@@ -119,7 +119,7 @@ const artifactTypes: IComboBoxOption[] = [
    }
 ]
 
-export const FindArtifactsModal: React.FC<{ onClose: () => void }> = ({ onClose }) => {
+export const FindArtifactsModal: React.FC<{ streamId?: string, onClose: () => void }> = ({ streamId, onClose }) => {
 
    const [state, setState] = useState<{ searching?: boolean, artifacts?: GetArtifactResponseV2[] }>({});
    const streamRef = React.useRef<IStreamChooser>(null);
@@ -221,20 +221,20 @@ export const FindArtifactsModal: React.FC<{ onClose: () => void }> = ({ onClose
                   <Stack horizontal verticalAlign="center" tokens={{ childrenGap: 20 }}>
                      <Stack >
                         <Label>Stream</Label>
-                        <StreamChooser ref={streamRef} />
+                        <StreamChooser defaultStreamId={streamId} ref={streamRef} />
                      </Stack>
                      <Stack >
-                        <TextField componentRef={minChangeRef} style={{ width: 92 }} label="Min Changelist" />
+                        <TextField key="min_change_option" componentRef={minChangeRef} style={{ width: 92 }} label="Min Changelist" />
                      </Stack>
                      <Stack >
-                        <TextField componentRef={maxChangeRef} style={{ width: 92 }} label="Max Changelist" />
+                        <TextField key="max_change_option" componentRef={maxChangeRef} style={{ width: 92 }} label="Max Changelist" />
                      </Stack>
                      <Stack >
-                        <TextField componentRef={nameRef} style={{ width: 220 }} label="Name" spellCheck={false} autoComplete="off" />
+                        <TextField key="name_option" componentRef={nameRef} style={{ width: 220 }} label="Name" spellCheck={false} autoComplete="off" />
                      </Stack>
                      <Stack>
                         <Label>Artifact Type</Label>
-                        <ComboBox componentRef={typeRef} allowFreeform={true} autoComplete="off" spellCheck={false} style={{ width: 144, textAlign: "left" }} defaultSelectedKey="step-all" options={filterTypes} calloutProps={{ doNotLayer: true }} />
+                        <ComboBox key="type_option" componentRef={typeRef} allowFreeform={true} autoComplete="off" spellCheck={false} style={{ width: 144, textAlign: "left" }} defaultSelectedKey="step-all" options={filterTypes} calloutProps={{ doNotLayer: true }} />
                      </Stack>
                   </Stack>
                   <Stack horizontal>
diff --git a/Engine/Source/Programs/Horde/HordeDashboard/src/components/projects/StreamChooser.tsx b/Engine/Source/Programs/Horde/HordeDashboard/src/components/projects/StreamChooser.tsx
--- a/Engine/Source/Programs/Horde/HordeDashboard/src/components/projects/StreamChooser.tsx
+++ b/Engine/Source/Programs/Horde/HordeDashboard/src/components/projects/StreamChooser.tsx
@@ -9,12 +9,13 @@ export interface IStreamChooser {
 }
 
 export interface IStreamChooserProps {
+   defaultStreamId?: string;
 
 }
 
 export const StreamChooser = React.forwardRef<IStreamChooser, IStreamChooserProps>((props, ref) => {
 
-   const [chosen, setChosen] = useState("");
+   const [chosen, setChosen] = useState(props?.defaultStreamId ?? "");
 
    useImperativeHandle(ref, () => ({
       streamId: chosen,
diff --git a/Engine/Source/Programs/Horde/HordeDashboard/src/styles/Styles.tsx b/Engine/Source/Programs/Horde/HordeDashboard/src/styles/Styles.tsx
--- a/Engine/Source/Programs/Horde/HordeDashboard/src/styles/Styles.tsx
+++ b/Engine/Source/Programs/Horde/HordeDashboard/src/styles/Styles.tsx
@@ -73,6 +73,7 @@ registerIcons({
       "Dot": IconContents.DOT,
       "Pin": IconContents.PIN,
       "Unpin": IconContents.UNPIN,
+      "Search": IconContents.SEARCH,
       "SearchTemplate": IconContents.SEARCH_TEMPLATE,
       "History": IconContents.HISTORY,
       "Import": IconContents.IMPORT,
