Horde: Move storage system into a plugin.

#rnx

diff --git a/Engine/Source/Programs/Horde/HordeServer.Shared/Configuration/ConfigNode.cs b/Engine/Source/Programs/Horde/HordeServer.Shared/Configuration/ConfigNode.cs
--- a/Engine/Source/Programs/Horde/HordeServer.Shared/Configuration/ConfigNode.cs
+++ b/Engine/Source/Programs/Horde/HordeServer.Shared/Configuration/ConfigNode.cs
@@ -48,6 +48,11 @@ namespace HordeServer.Configuration
 	/// </summary>
 	public abstract class ConfigNode
 	{
+		/// <summary>
+		/// Default options for new json node objects
+		/// </summary>
+		internal static JsonNodeOptions s_defaultJsonNodeOptions { get; } = new JsonNodeOptions { PropertyNameCaseInsensitive = true };
+
 		/// <summary>
 		/// Parses macro definitions from this object
 		/// </summary>
@@ -120,7 +125,7 @@ namespace HordeServer.Configuration
 			}
 			else if (node is JsonObject obj)
 			{
-				JsonObject result = new JsonObject();
+				JsonObject result = new JsonObject(s_defaultJsonNodeOptions);
 				foreach ((string propertyName, JsonNode? propertyNode) in obj)
 				{
 					result[propertyName] = ExpandMacros(propertyNode, context);
@@ -321,7 +326,7 @@ namespace HordeServer.Configuration
 		/// <inheritdoc/>
 		public override async Task<JsonNode?> PreprocessAsync(JsonNode? node, JsonNode? existingNode, ConfigContext context, CancellationToken cancellationToken)
 		{
-			JsonObject? targetObject = ((JsonObject?)existingNode) ?? new JsonObject();
+			JsonObject? targetObject = ((JsonObject?)existingNode) ?? new JsonObject(s_defaultJsonNodeOptions);
 			foreach ((string key, JsonNode? element) in (JsonObject)node!)
 			{
 				context.EnterScope($"[{key}]");
@@ -545,7 +550,7 @@ namespace HordeServer.Configuration
 			}
 
 			// Ensure that the target object is valid so we can write properties into it
-			target ??= new JsonObject();
+			target ??= new JsonObject(s_defaultJsonNodeOptions);
 
 			// Parse all the macros for this scope
 			if (MacroScope)
diff --git a/Engine/Source/Programs/Horde/HordeServer.Shared/HordeServer.Shared.csproj b/Engine/Source/Programs/Horde/HordeServer.Shared/HordeServer.Shared.csproj
--- a/Engine/Source/Programs/Horde/HordeServer.Shared/HordeServer.Shared.csproj
+++ b/Engine/Source/Programs/Horde/HordeServer.Shared/HordeServer.Shared.csproj
@@ -20,6 +20,8 @@
   <ItemGroup>
     <PackageReference Include="MongoDB.Driver" Version="2.22.0" />
     <PackageReference Include="OpenTelemetry" Version="1.8.0" />
+    <PackageReference Include="OpenTelemetry.Exporter.Console" Version="1.8.0" />
+    <PackageReference Include="OpenTelemetry.Exporter.OpenTelemetryProtocol" Version="1.8.0" />
   </ItemGroup>
 
 </Project>
diff --git a/Engine/Source/Programs/Horde/HordeServer.Tests.Shared/ServerServiceTest.cs b/Engine/Source/Programs/Horde/HordeServer.Tests.Shared/ServerServiceTest.cs
--- a/Engine/Source/Programs/Horde/HordeServer.Tests.Shared/ServerServiceTest.cs
+++ b/Engine/Source/Programs/Horde/HordeServer.Tests.Shared/ServerServiceTest.cs
@@ -20,6 +20,7 @@ using Microsoft.Extensions.DependencyInjection;
 using Microsoft.Extensions.Hosting;
 using Microsoft.Extensions.Logging;
 using Microsoft.Extensions.Options;
+using OpenTelemetry.Trace;
 
 namespace HordeServer.Tests
 {
@@ -52,7 +53,6 @@ namespace HordeServer.Tests
 		public ServerServiceTest()
 		{
 			_pluginCollection = new PluginCollection();
-			_pluginCollection.Add<AnalyticsPlugin>();
 
 			PatchDatadogWriter();
 		}
@@ -97,6 +97,9 @@ namespace HordeServer.Tests
 			services.AddLogging(builder => { builder.AddConsole().SetMinimumLevel(LogLevel.Debug); });
 			services.AddSingleton<IMemoryCache>(sp => new MemoryCache(new MemoryCacheOptions { }));
 
+			services.AddSingleton<OpenTelemetry.Trace.Tracer>(sp => TracerProvider.Default.GetTracer("TestTracer"));
+			services.AddSingleton(sp => new Meter("TestMeter"));
+
 			services.AddSingleton(typeof(IAuditLogFactory<>), typeof(AuditLogFactory<>));
 
 			services.AddSingleton<ConfigService>();
diff --git a/Engine/Source/Programs/Horde/HordeServer.Tests/Ddc/Functional/ContentId/CompressedBlobTests.cs b/Engine/Source/Programs/Horde/HordeServer.Tests/Ddc/Functional/ContentId/CompressedBlobTests.cs
--- a/Engine/Source/Programs/Horde/HordeServer.Tests/Ddc/Functional/ContentId/CompressedBlobTests.cs
+++ b/Engine/Source/Programs/Horde/HordeServer.Tests/Ddc/Functional/ContentId/CompressedBlobTests.cs
@@ -41,10 +41,12 @@ namespace HordeServer.Tests.Ddc.FunctionalTests.CompressedBlobs
 		{
 			ServerSettings serverSettings = ServiceProvider.GetRequiredService<IOptions<ServerSettings>>().Value;
 
-			GlobalConfig globalConfig = ServiceProvider.GetRequiredService<IOptions<GlobalConfig>>().Value;
+			StorageConfig storageConfig = new StorageConfig();
+			storageConfig.Backends.Add(new BackendConfig { Id = new BackendId("default"), Type = StorageBackendType.Memory });
+			storageConfig.Namespaces.Add(new NamespaceConfig { Id = new NamespaceId(TestNamespace), Backend = storageConfig.Backends[^1].Id });
 
-			globalConfig.Storage.Backends.Add(new BackendConfig { Id = new BackendId("default"), Type = StorageBackendType.Memory });
-			globalConfig.Storage.Namespaces.Add(new NamespaceConfig { Id = new NamespaceId(TestNamespace), Backend = globalConfig.Storage.Backends[^1].Id });
+			GlobalConfig globalConfig = ServiceProvider.GetRequiredService<IOptions<GlobalConfig>>().Value;
+			globalConfig.Plugins.AddStorageTestConfig(storageConfig);
 			globalConfig.PostLoad(serverSettings, new List<ILoadedPlugin>());
 
 			ConfigService configService = ServiceProvider.GetRequiredService<ConfigService>();
diff --git a/Engine/Source/Programs/Horde/HordeServer.Tests/Ddc/Functional/Objects/ReferencesTests.cs b/Engine/Source/Programs/Horde/HordeServer.Tests/Ddc/Functional/Objects/ReferencesTests.cs
--- a/Engine/Source/Programs/Horde/HordeServer.Tests/Ddc/Functional/Objects/ReferencesTests.cs
+++ b/Engine/Source/Programs/Horde/HordeServer.Tests/Ddc/Functional/Objects/ReferencesTests.cs
@@ -57,10 +57,12 @@ namespace HordeServer.Tests.Ddc.FunctionalTests.References
 
 			ServerSettings serverSettings = ServiceProvider.GetRequiredService<IOptions<ServerSettings>>().Value;
 
-			GlobalConfig globalConfig = ServiceProvider.GetRequiredService<IOptions<GlobalConfig>>().Value;
+			StorageConfig storageConfig = new StorageConfig();
+			storageConfig.Backends.Add(new BackendConfig { Id = new BackendId("default"), Type = StorageBackendType.Memory });
+			storageConfig.Namespaces.Add(new NamespaceConfig { Id = TestNamespace, Backend = storageConfig.Backends[^1].Id });
 
-			globalConfig.Storage.Backends.Add(new BackendConfig { Id = new BackendId("default"), Type = StorageBackendType.Memory });
-			globalConfig.Storage.Namespaces.Add(new NamespaceConfig { Id = TestNamespace, Backend = globalConfig.Storage.Backends[^1].Id });
+			GlobalConfig globalConfig = ServiceProvider.GetRequiredService<IOptions<GlobalConfig>>().Value;
+			globalConfig.Plugins.AddStorageTestConfig(storageConfig);
 			globalConfig.PostLoad(serverSettings, new List<ILoadedPlugin>());
 
 			ConfigService configService = ServiceProvider.GetRequiredService<ConfigService>();
diff --git a/Engine/Source/Programs/Horde/HordeServer.Tests/HordeServer.Tests.csproj b/Engine/Source/Programs/Horde/HordeServer.Tests/HordeServer.Tests.csproj
--- a/Engine/Source/Programs/Horde/HordeServer.Tests/HordeServer.Tests.csproj
+++ b/Engine/Source/Programs/Horde/HordeServer.Tests/HordeServer.Tests.csproj
@@ -41,7 +41,8 @@
       <Aliases>JobDriver</Aliases>
     </ProjectReference>
     <ProjectReference Include="..\HordeServer\HordeServer.csproj" />
-    <ProjectReference Include="..\Plugins\Analytics\HordeServer.Analytics\HordeServer.Analytics.csproj" />
+    <ProjectReference Include="..\Plugins\Analytics\HordeServer.Analytics.Tests\HordeServer.Analytics.Tests.csproj" />
+    <ProjectReference Include="..\Plugins\Storage\HordeServer.Storage.Tests\HordeServer.Storage.Tests.csproj" />
   </ItemGroup>
 
   <ItemGroup>
diff --git a/Engine/Source/Programs/Horde/HordeServer.Tests/Issues/IssueServiceTests.cs b/Engine/Source/Programs/Horde/HordeServer.Tests/Issues/IssueServiceTests.cs
--- a/Engine/Source/Programs/Horde/HordeServer.Tests/Issues/IssueServiceTests.cs
+++ b/Engine/Source/Programs/Horde/HordeServer.Tests/Issues/IssueServiceTests.cs
@@ -23,6 +23,7 @@ using HordeServer.Issues;
 using HordeServer.Jobs;
 using HordeServer.Jobs.Graphs;
 using HordeServer.Logs;
+using HordeServer.Plugins;
 using HordeServer.Projects;
 using HordeServer.Server;
 using HordeServer.Storage;
@@ -149,11 +150,13 @@ namespace HordeServer.Tests.Issues
 			projectConfig.Streams.Add(CreateStream(_releaseStreamId, ReleaseStreamName));
 			projectConfig.Streams.Add(CreateStream(_devStreamId, DevStreamName));
 
+			StorageConfig storageConfig = new StorageConfig();
+			storageConfig.Backends.Add(new BackendConfig { Id = new BackendId("default-backend"), Type = StorageBackendType.Memory });
+			storageConfig.Namespaces.Add(new NamespaceConfig { Id = new NamespaceId("horde-logs"), Backend = new BackendId("default-backend"), GcDelayHrs = 0.0 });
+
 			GlobalConfig globalConfig = new GlobalConfig();
 			globalConfig.Projects.Add(projectConfig);
-			globalConfig.Storage.Backends.Add(new BackendConfig { Id = new BackendId("default-backend"), Type = StorageBackendType.Memory });
-			globalConfig.Storage.Namespaces.Add(new NamespaceConfig { Id = new NamespaceId("horde-logs"), Backend = new BackendId("default-backend"), GcDelayHrs = 0.0 });
-
+			globalConfig.Plugins.Add(new PluginName("storage"), storageConfig);
 			SetConfig(globalConfig);
 
 			static StreamConfig CreateStream(StreamId streamId, string streamName)
diff --git a/Engine/Source/Programs/Horde/HordeServer.Tests/Storage/BlobStoreTests.cs b/Engine/Source/Programs/Horde/HordeServer.Tests/Storage/BlobStoreTests.cs
deleted file mode 100644
diff --git a/Engine/Source/Programs/Horde/HordeServer.Tests/Storage/GcServiceTests.cs b/Engine/Source/Programs/Horde/HordeServer.Tests/Storage/GcServiceTests.cs
deleted file mode 100644
diff --git a/Engine/Source/Programs/Horde/HordeServer.Tests/Storage/StorageControllerTests.cs b/Engine/Source/Programs/Horde/HordeServer.Tests/Storage/StorageControllerTests.cs
deleted file mode 100644
diff --git a/Engine/Source/Programs/Horde/HordeServer.Tests/Storage/StorageServiceTests.cs b/Engine/Source/Programs/Horde/HordeServer.Tests/Storage/StorageServiceTests.cs
deleted file mode 100644
diff --git a/Engine/Source/Programs/Horde/HordeServer.Tests/TestSetup.cs b/Engine/Source/Programs/Horde/HordeServer.Tests/TestSetup.cs
--- a/Engine/Source/Programs/Horde/HordeServer.Tests/TestSetup.cs
+++ b/Engine/Source/Programs/Horde/HordeServer.Tests/TestSetup.cs
@@ -117,6 +117,12 @@ namespace HordeServer.Tests
 		public TestDataController TestDataController => GetTestDataController();
 		public BisectTasksController BisectTasksController => GetBisectTasksController();
 
+		public TestSetup()
+		{
+			AddPlugin<AnalyticsPlugin>();
+			AddPlugin<StoragePlugin>();
+		}
+
 		protected override void ConfigureServices(IServiceCollection services)
 		{
 			base.ConfigureServices(services);
@@ -129,9 +135,6 @@ namespace HordeServer.Tests
 
 			services.AddSingleton<IAuditLog<AgentId>>(sp => sp.GetRequiredService<IAuditLogFactory<AgentId>>().Create("Agents.Log", "AgentId"));
 
-			services.AddSingleton<OpenTelemetry.Trace.Tracer>(sp => TracerProvider.Default.GetTracer("TestTracer"));
-			services.AddSingleton(sp => new Meter("TestMeter"));
-
 			services.AddSingleton<IAccountCollection, AccountCollection>();
 			services.AddSingleton<IAgentCollection, AgentCollection>();
 			services.AddSingleton<IAgentTelemetryCollection, AgentTelemetryCollection>();
@@ -199,13 +202,6 @@ namespace HordeServer.Tests
 
 			services.AddSingleton<ConformTaskSource>();
 			services.AddSingleton<ICommitService, CommitService>();
-
-			services.AddSingleton<FileObjectStoreFactory>();
-			services.AddSingleton<IObjectStoreFactory, ObjectStoreFactory>();
-
-			services.AddSingleton<StorageService>();
-			services.AddSingleton<StorageBackendCache>();
-			services.AddSingleton<BundleCache>();
 		}
 
 		public Task<Fixture> CreateFixtureAsync()
@@ -334,7 +330,7 @@ namespace HordeServer.Tests
 			{
 				await agent.TryUpdateWorkspacesAsync(workspaces, false);
 			}
-
+			
 			if (requestShutdown)
 			{
 				await agent.TryUpdateAsync(new UpdateAgentOptions { RequestShutdown = true });
diff --git a/Engine/Source/Programs/Horde/HordeServer.Tests/Tools/ToolTests.cs b/Engine/Source/Programs/Horde/HordeServer.Tests/Tools/ToolTests.cs
--- a/Engine/Source/Programs/Horde/HordeServer.Tests/Tools/ToolTests.cs
+++ b/Engine/Source/Programs/Horde/HordeServer.Tests/Tools/ToolTests.cs
@@ -14,6 +14,7 @@ using HordeServer.Tools;
 using HordeCommon;
 using Microsoft.Extensions.DependencyInjection;
 using Microsoft.VisualStudio.TestTools.UnitTesting;
+using HordeServer.Plugins;
 
 namespace HordeServer.Tests.Tools
 {
@@ -24,11 +25,14 @@ namespace HordeServer.Tests.Tools
 
 		public ToolTests()
 		{
+			StorageConfig storageConfig = new StorageConfig();
+			storageConfig.Backends.Clear();
+			storageConfig.Backends.Add(new BackendConfig { Id = new BackendId("tools-backend"), Type = StorageBackendType.Memory });
+			storageConfig.Namespaces.Clear();
+			storageConfig.Namespaces.Add(new NamespaceConfig { Id = Namespace.Tools, Backend = new BackendId("tools-backend") });
+
 			GlobalConfig globalConfig = new GlobalConfig();
-			globalConfig.Storage.Backends.Clear();
-			globalConfig.Storage.Backends.Add(new BackendConfig { Id = new BackendId("tools-backend"), Type = StorageBackendType.Memory });
-			globalConfig.Storage.Namespaces.Clear();
-			globalConfig.Storage.Namespaces.Add(new NamespaceConfig { Id = Namespace.Tools, Backend = new BackendId("tools-backend") });
+			globalConfig.Plugins.Add(new PluginName("storage"), storageConfig);
 			globalConfig.Tools.Add(new ToolConfig(_toolId) { Name = "UnrealGameSync", Description = "Tool for syncing content from source control" });
 			SetConfig(globalConfig);
 		}
diff --git a/Engine/Source/Programs/Horde/HordeServer.Tests/Tools/ToolsControllerTests.cs b/Engine/Source/Programs/Horde/HordeServer.Tests/Tools/ToolsControllerTests.cs
--- a/Engine/Source/Programs/Horde/HordeServer.Tests/Tools/ToolsControllerTests.cs
+++ b/Engine/Source/Programs/Horde/HordeServer.Tests/Tools/ToolsControllerTests.cs
@@ -33,20 +33,24 @@ public class ToolsControllerTests
 
 		PluginCollection pluginCollection = new PluginCollection();
 		pluginCollection.Add<AnalyticsPlugin>();
+		pluginCollection.Add<StoragePlugin>();
+
+		StorageConfig storageConfig = new StorageConfig();
+		storageConfig.Backends.Clear();
+		storageConfig.Backends.Add(new BackendConfig { Id = new BackendId("tools-backend"), Type = StorageBackendType.Memory });
+		storageConfig.Namespaces.Clear();
+		storageConfig.Namespaces.Add(new NamespaceConfig { Id = Namespace.Tools, Backend = new BackendId("tools-backend") });
 
 		GlobalConfig globalConfig = new();
-		globalConfig.Storage.Backends.Clear();
-		globalConfig.Storage.Backends.Add(new BackendConfig { Id = new BackendId("tools-backend"), Type = StorageBackendType.Memory });
-		globalConfig.Storage.Namespaces.Clear();
-		globalConfig.Storage.Namespaces.Add(new NamespaceConfig { Id = Namespace.Tools, Backend = new BackendId("tools-backend") });
+		globalConfig.Plugins.AddStorageTestConfig(storageConfig);
 		globalConfig.Tools.Add(new ToolConfig(toolId) { Name = "Foo", Description = "This is foo", Acl = new AclConfig() { Entries = [aclEntryConfig] }, Public = false });
-
+		
 		ServerSettings serverSettings = new() { AuthMethod = AuthMethod.Horde };
 		globalConfig.PostLoad(serverSettings, pluginCollection.LoadedPlugins);
 
 		Dictionary<string, string> settings = new() { { "Horde:AuthMethod", AuthMethod.Horde.ToString() } };
 		await using FakeHordeWebApp app = new(settings);
-
+		
 		ConfigService configService = app.ServiceProvider.GetRequiredService<ConfigService>();
 		IToolCollection tools = app.ServiceProvider.GetRequiredService<IToolCollection>();
 		IServiceAccountCollection serviceAccounts = app.ServiceProvider.GetRequiredService<IServiceAccountCollection>();
@@ -55,7 +59,7 @@ public class ToolsControllerTests
 
 		List<IUserClaim> claims = [new UserClaim("http://epicgames.com/ue/horde/role", "agent")];
 		(IServiceAccount _, string token) = await serviceAccounts.CreateAsync(new CreateServiceAccountOptions("myDesc", claims));
-
+		
 		// Create tool and deployment
 		using MemoryStream ms = new(await ToolTests.CreateZipFileDataAsync("foo.txt", "foo content"));
 		ITool? tool = await tools.GetAsync(toolId);
diff --git a/Engine/Source/Programs/Horde/HordeServer/Configuration/ConfigService.cs b/Engine/Source/Programs/Horde/HordeServer/Configuration/ConfigService.cs
--- a/Engine/Source/Programs/Horde/HordeServer/Configuration/ConfigService.cs
+++ b/Engine/Source/Programs/Horde/HordeServer/Configuration/ConfigService.cs
@@ -588,6 +588,7 @@ namespace HordeServer.Configuration
 
 				JsonObject configObj = await context.PreprocessFileAsync(GetGlobalConfigUri(), configNode, cancellationToken);
 				CopyJsonNode(configObj, "TelemetryStores", "Plugins.Analytics.Stores");
+				CopyJsonNode(configObj, "Storage", "Plugins.Storage");
 
 				GlobalConfig globalConfig = JsonSerializer.Deserialize<GlobalConfig>(configObj, _jsonOptions)!;
 				if (globalConfig.VersionEnum < GlobalVersion.Latest)
@@ -624,7 +625,19 @@ namespace HordeServer.Configuration
 				throw new ConfigException(context, ex.Message, ex);
 			}
 		}
-
+		/*
+		static bool TryGetChildNode(JsonNode node, string name, [NotNullWhen(true)] out JsonNode? child)
+		{
+			child = node[name];
+			if (child != null)
+			{
+				JsonObject? parentObj = node as JsonObject;
+				parentObj
+				parentObj.TryGetPropertyValue(
+			}
+			return child != null;
+		}
+		*/
 		static void CopyJsonNode(JsonObject rootObj, string sourcePath, string targetPath)
 		{
 			JsonNode? sourceNode = rootObj;
diff --git a/Engine/Source/Programs/Horde/HordeServer/Ddc/RequestHelper.cs b/Engine/Source/Programs/Horde/HordeServer/Ddc/RequestHelper.cs
--- a/Engine/Source/Programs/Horde/HordeServer/Ddc/RequestHelper.cs
+++ b/Engine/Source/Programs/Horde/HordeServer/Ddc/RequestHelper.cs
@@ -5,7 +5,6 @@ using System.Security.Claims;
 using System.Threading.Tasks;
 using EpicGames.Horde.Acls;
 using EpicGames.Horde.Storage;
-using HordeServer.Server;
 using HordeServer.Storage;
 using Microsoft.AspNetCore.Http;
 using Microsoft.AspNetCore.Mvc;
@@ -16,12 +15,12 @@ namespace HordeServer.Ddc
 	class RequestHelper : IRequestHelper
 	{
 		readonly IStorageClientFactory _storageClientFactory;
-		readonly GlobalConfig _globalConfig;
+		readonly StorageConfig _storageConfig;
 
-		public RequestHelper(IStorageClientFactory storageClientFactory, IOptionsSnapshot<GlobalConfig> globalConfig)
+		public RequestHelper(IStorageClientFactory storageClientFactory, IOptionsSnapshot<StorageConfig> storageConfig)
 		{
 			_storageClientFactory = storageClientFactory;
-			_globalConfig = globalConfig.Value;
+			_storageConfig = storageConfig.Value;
 		}
 
 		public Task<ActionResult?> HasAccessForGlobalOperationsAsync(ClaimsPrincipal user, AclAction[] aclActions)
@@ -36,7 +35,7 @@ namespace HordeServer.Ddc
 			{
 				return Task.FromResult<ActionResult?>(new ForbidResult());
 			}
-			if (!_globalConfig.Storage.TryGetNamespace(ns, out NamespaceConfig? namespaceConfig) || aclActions.Any(x => !namespaceConfig.Authorize(x, user)))
+			if (!_storageConfig.TryGetNamespace(ns, out NamespaceConfig? namespaceConfig) || aclActions.Any(x => !namespaceConfig.Authorize(x, user)))
 			{
 				return Task.FromResult<ActionResult?>(new ForbidResult());
 			}
diff --git a/Engine/Source/Programs/Horde/HordeServer/HordeServer.csproj b/Engine/Source/Programs/Horde/HordeServer/HordeServer.csproj
--- a/Engine/Source/Programs/Horde/HordeServer/HordeServer.csproj
+++ b/Engine/Source/Programs/Horde/HordeServer/HordeServer.csproj
@@ -135,6 +135,7 @@
     <ProjectReference Include="..\..\Shared\EpicGames.Slack\EpicGames.Slack.csproj" />
     <ProjectReference Include="..\HordeServer.Shared\HordeServer.Shared.csproj" />
     <ProjectReference Include="..\Plugins\Analytics\HordeServer.Analytics\HordeServer.Analytics.csproj" />
+    <ProjectReference Include="..\Plugins\Storage\HordeServer.Storage\HordeServer.Storage.csproj" />
   </ItemGroup>
 
   <ItemGroup>
diff --git a/Engine/Source/Programs/Horde/HordeServer/Plugins/PluginCollection.cs b/Engine/Source/Programs/Horde/HordeServer/Plugins/PluginCollection.cs
--- a/Engine/Source/Programs/Horde/HordeServer/Plugins/PluginCollection.cs
+++ b/Engine/Source/Programs/Horde/HordeServer/Plugins/PluginCollection.cs
@@ -136,7 +136,7 @@ namespace HordeServer.Plugins
 		{
 			if (!_loadedPluginNames.Add(metadata.Name))
 			{
-				throw new NotImplementedException();
+				throw new InvalidOperationException($"An implementation of the {metadata.Name} plugin has already been added");
 			}
 
 			PluginAttribute attr = startupType.GetCustomAttribute<PluginAttribute>()
diff --git a/Engine/Source/Programs/Horde/HordeServer/Server/GlobalConfig.cs b/Engine/Source/Programs/Horde/HordeServer/Server/GlobalConfig.cs
--- a/Engine/Source/Programs/Horde/HordeServer/Server/GlobalConfig.cs
+++ b/Engine/Source/Programs/Horde/HordeServer/Server/GlobalConfig.cs
@@ -40,7 +40,6 @@ using HordeServer.Perforce;
 using HordeServer.Plugins;
 using HordeServer.Projects;
 using HordeServer.Secrets;
-using HordeServer.Storage;
 using HordeServer.Streams;
 using HordeServer.Tools;
 using HordeServer.Utilities;
@@ -229,11 +228,6 @@ namespace HordeServer.Server
 		/// </summary>
 		public TimeSpan? AgentShutdownIfDisabledGracePeriod { get; set; } = null;
 
-		/// <summary>
-		/// Storage configuration
-		/// </summary>
-		public StorageConfig Storage { get; set; } = new StorageConfig();
-
 		/// <summary>
 		/// Configuration for different artifact types
 		/// </summary>
@@ -350,8 +344,6 @@ namespace HordeServer.Server
 			ConfigObject.MergeDefaults<string, PoolConfig>(Pools.Select(x => (x.Id.ToString(), x.Base?.ToString(), x)));
 			UpdateWorkspacesForPools();
 
-			Storage.PostLoad(this);
-
 			// Ensure that all plugins have an entry in the global config so they can register their ACLs
 			foreach (ILoadedPlugin loadedPlugin in loadedPlugins)
 			{
diff --git a/Engine/Source/Programs/Horde/HordeServer/ServerSettings.cs b/Engine/Source/Programs/Horde/HordeServer/ServerSettings.cs
--- a/Engine/Source/Programs/Horde/HordeServer/ServerSettings.cs
+++ b/Engine/Source/Programs/Horde/HordeServer/ServerSettings.cs
@@ -16,7 +16,6 @@ using EpicGames.Horde.Tools;
 using EpicGames.Perforce;
 using HordeServer.Agents.Fleet;
 using HordeServer.Server;
-using HordeServer.Storage;
 using HordeServer.Tools;
 
 namespace HordeServer
@@ -502,7 +501,7 @@ namespace HordeServer
 		/// <summary>
 		/// Overridden settings for storage backends. Useful for running against a production server with custom backends.
 		/// </summary>
-		public List<BackendConfig> StorageBackends { get; set; } = new List<BackendConfig>();
+		public string LogServiceWriteCacheType { get; set; } = "InMemory";
 
 		/// <summary>
 		/// Whether to log json to stdout
@@ -713,21 +712,6 @@ namespace HordeServer
 		/// </summary>
 		public bool OpenBrowser { get; set; } = false;
 
-		/// <summary>
-		/// Directory to use for the coarse-grained backend cache. This caches full bundles downloaded from the upstream object store.
-		/// </summary>
-		public string? BundleCacheDir { get; set; }
-
-		/// <summary>
-		/// Maximum size of the storage cache on disk. Accepts standard binary suffixes (kb, mb, gb, tb, etc...)
-		/// </summary>
-		public string BundleCacheSize { get; set; } = "1gb";
-
-		/// <summary>
-		/// Accessor for the bundle cache size in bytes
-		/// </summary>
-		public long BundleCacheSizeBytes => StringUtils.ParseBytesString(BundleCacheSize);
-
 		/// <summary>
 		/// Directory to store the fine-grained block cache. This caches individual exports embedded in bundles.
 		/// </summary>
diff --git a/Engine/Source/Programs/Horde/HordeServer/Startup.cs b/Engine/Source/Programs/Horde/HordeServer/Startup.cs
--- a/Engine/Source/Programs/Horde/HordeServer/Startup.cs
+++ b/Engine/Source/Programs/Horde/HordeServer/Startup.cs
@@ -27,8 +27,6 @@ using EpicGames.Horde.Logs;
 using EpicGames.Horde.Projects;
 using EpicGames.Horde.Server;
 using EpicGames.Horde.Storage;
-using EpicGames.Horde.Storage.Bundles;
-using EpicGames.Horde.Storage.ObjectStores;
 using EpicGames.Horde.Streams;
 using EpicGames.Horde.Users;
 using EpicGames.Redis;
@@ -77,7 +75,6 @@ using HordeServer.Server;
 using HordeServer.Server.Notices;
 using HordeServer.ServiceAccounts;
 using HordeServer.Storage;
-using HordeServer.Storage.ObjectStores;
 using HordeServer.Streams;
 using HordeServer.Tasks;
 using HordeServer.Tools;
@@ -319,6 +316,7 @@ namespace HordeServer
 			// Register the plugin collection
 			PluginCollection pluginCollection = new PluginCollection();
 			pluginCollection.Add<AnalyticsPlugin>();
+			pluginCollection.Add<StoragePlugin>();
 			services.AddSingleton<IPluginCollection>(pluginCollection);
 
 			// Register all the plugin services
@@ -506,13 +504,8 @@ namespace HordeServer
 
 			services.AddSingleton<DeviceService>();
 			services.AddSingleton<NoticeService>();
-			services.AddSingleton<StorageService>();
-			services.AddSingleton<IStorageService>(sp => sp.GetRequiredService<StorageService>());
-			services.AddScoped(sp => sp.GetRequiredService<StorageService>().CreateStorageClientFactory(sp.GetRequiredService<IOptionsSnapshot<GlobalConfig>>().Value));
 			services.AddSingleton<IBlockCache>(sp => CreateBlockCache(sp));
 			services.AddSingleton<TestDataService>();
-			services.AddSingleton<BundleCache>();
-			services.AddSingleton<StorageBackendCache>(CreateStorageBackendCache);
 
 			if (settings.JiraUrl != null)
 			{
@@ -539,11 +532,6 @@ namespace HordeServer
 			services.AddScoped<ISecretCollection, SecretCollection>();
 			services.AddSingleton<ISecretProvider, AwsParameterStoreSecretProvider>();
 
-			// Storage providers
-			services.AddSingleton<IObjectStoreFactory, ObjectStoreFactory>();
-			services.AddSingleton<AwsObjectStoreFactory>();
-			services.AddSingleton<FileObjectStoreFactory>();
-
 			if (settings.WithAws)
 			{
 				AWSOptions awsOptions = Configuration.GetAWSOptions();
@@ -912,13 +900,6 @@ namespace HordeServer
 			return BlockCache.Create(cacheDir, (int)(serverSettings.BlockCacheSizeBytes / (1024 * 1024 * 1024)));
 		}
 
-		static StorageBackendCache CreateStorageBackendCache(IServiceProvider serviceProvider)
-		{
-			ServerSettings serverSettings = serviceProvider.GetRequiredService<IOptions<ServerSettings>>().Value;
-			DirectoryReference cacheDir = DirectoryReference.Combine(ServerApp.DataDir, String.IsNullOrEmpty(serverSettings.BundleCacheDir) ? "Cache" : serverSettings.BundleCacheDir);
-			return new StorageBackendCache(cacheDir, serverSettings.BundleCacheSizeBytes, serviceProvider.GetRequiredService<ILogger<StorageBackendCache>>());
-		}
-
 		public static void ConfigureFormatters()
 		{
 			LogValueFormatter.RegisterTypeAnnotation<AgentId>("AgentId");
diff --git a/Engine/Source/Programs/Horde/HordeServer/Storage/BackendId.cs b/Engine/Source/Programs/Horde/HordeServer/Storage/BackendId.cs
deleted file mode 100644
diff --git a/Engine/Source/Programs/Horde/HordeServer/Storage/IObjectStoreFactory.cs b/Engine/Source/Programs/Horde/HordeServer/Storage/IObjectStoreFactory.cs
deleted file mode 100644
diff --git a/Engine/Source/Programs/Horde/HordeServer/Storage/ObjectStoreFactory.cs b/Engine/Source/Programs/Horde/HordeServer/Storage/ObjectStoreFactory.cs
deleted file mode 100644
diff --git a/Engine/Source/Programs/Horde/HordeServer/Storage/ObjectStores/AwsObjectStore.cs b/Engine/Source/Programs/Horde/HordeServer/Storage/ObjectStores/AwsObjectStore.cs
deleted file mode 100644
diff --git a/Engine/Source/Programs/Horde/HordeServer/Storage/ObjectStores/AzureObjectStore.cs b/Engine/Source/Programs/Horde/HordeServer/Storage/ObjectStores/AzureObjectStore.cs
deleted file mode 100644
diff --git a/Engine/Source/Programs/Horde/HordeServer/Storage/StorageAclAction.cs b/Engine/Source/Programs/Horde/HordeServer/Storage/StorageAclAction.cs
deleted file mode 100644
diff --git a/Engine/Source/Programs/Horde/HordeServer/Storage/StorageConfig.cs b/Engine/Source/Programs/Horde/HordeServer/Storage/StorageConfig.cs
deleted file mode 100644
diff --git a/Engine/Source/Programs/Horde/HordeServer/Storage/StorageController.cs b/Engine/Source/Programs/Horde/HordeServer/Storage/StorageController.cs
deleted file mode 100644
diff --git a/Engine/Source/Programs/Horde/HordeServer/Storage/StorageService.cs b/Engine/Source/Programs/Horde/HordeServer/Storage/StorageService.cs
deleted file mode 100644
diff --git a/Engine/Source/Programs/Horde/HordeServer/appsettings.json b/Engine/Source/Programs/Horde/HordeServer/appsettings.json
--- a/Engine/Source/Programs/Horde/HordeServer/appsettings.json
+++ b/Engine/Source/Programs/Horde/HordeServer/appsettings.json
@@ -3,7 +3,12 @@
 		"ConsoleLogLevel": "Debug",
 		"DashboardUrl": "https://localhost:3000",
 		"CorsEnabled": true,
-		"CorsOrigin": "https://localhost:5003;http://localhost:3000"
+		"CorsOrigin": "https://localhost:5003;http://localhost:3000",
+		"Plugins": {
+			"analytics": {
+				"Foo": 1234
+			}
+		}
 	},
 	"Serilog": {
 		"MinimumLevel": {
diff --git a/Engine/Source/Programs/Horde/Plugins/Analytics/HordeServer.Analytics.Tests/AnalyticsTestExtensions.cs b/Engine/Source/Programs/Horde/Plugins/Analytics/HordeServer.Analytics.Tests/AnalyticsTestExtensions.cs
new file mode 100644
--- /dev/null
+++ b/Engine/Source/Programs/Horde/Plugins/Analytics/HordeServer.Analytics.Tests/AnalyticsTestExtensions.cs
@@ -0,0 +1,21 @@
+// Copyright Epic Games, Inc. All Rights Reserved.
+
+using System;
+using System.Collections.Generic;
+using System.Linq;
+using System.Text;
+using System.Threading.Tasks;
+using HordeServer.Plugins;
+using HordeServer.Telemetry;
+using HordeServer.Telemetry.Metrics;
+using HordeServer.Telemetry.Sinks;
+using Microsoft.Extensions.DependencyInjection;
+
+namespace HordeServer.Analytics.Tests
+{
+	public static class AnalyticsTestExtensions
+	{
+		public static void AddAnalyticsTestConfig(this PluginConfigCollection configCollection, AnalyticsConfig config)
+			=> configCollection[new PluginName("Analytics")] = config;
+	}
+}
diff --git a/Engine/Source/Programs/Horde/Plugins/Analytics/HordeServer.Analytics.Tests/Telemetry/TelemetryTests.cs b/Engine/Source/Programs/Horde/Plugins/Analytics/HordeServer.Analytics.Tests/Telemetry/TelemetryTests.cs
--- a/Engine/Source/Programs/Horde/Plugins/Analytics/HordeServer.Analytics.Tests/Telemetry/TelemetryTests.cs
+++ b/Engine/Source/Programs/Horde/Plugins/Analytics/HordeServer.Analytics.Tests/Telemetry/TelemetryTests.cs
@@ -25,11 +25,9 @@ namespace HordeServer.Analytics.Tests.Telemetry
 	{
 		static readonly TelemetryRecordMeta s_metadata = new TelemetryRecordMeta();
 
-		protected override void ConfigureServices(IServiceCollection services)
+		public TelemetryTests()
 		{
-			base.ConfigureServices(services);
-
-			services.AddAnalytics();
+			AddPlugin<AnalyticsPlugin>();
 		}
 
 		[TestMethod]
@@ -52,7 +50,7 @@ namespace HordeServer.Analytics.Tests.Telemetry
 			analyticsConfig.Stores.Add(telemetryStoreConfig);
 
 			GlobalConfig globalConfig = new GlobalConfig();
-			globalConfig.Plugins.AddAnalytics(analyticsConfig);
+			globalConfig.Plugins.AddAnalyticsTestConfig(analyticsConfig);
 			SetConfig(globalConfig);
 
 			MetricTelemetrySink sink = ServiceProvider.GetRequiredService<MetricTelemetrySink>();
@@ -93,7 +91,7 @@ namespace HordeServer.Analytics.Tests.Telemetry
 			analyticsConfig.Stores.Add(telemetryStoreConfig);
 
 			GlobalConfig globalConfig = new GlobalConfig();
-			globalConfig.Plugins.AddAnalytics(analyticsConfig);
+			globalConfig.Plugins.AddAnalyticsTestConfig(analyticsConfig);
 			SetConfig(globalConfig);
 
 			MetricTelemetrySink sink = ServiceProvider.GetRequiredService<MetricTelemetrySink>();
@@ -174,7 +172,7 @@ namespace HordeServer.Analytics.Tests.Telemetry
 			analyticsConfig.Stores.Add(telemetryStoreConfig);
 
 			GlobalConfig globalConfig = new GlobalConfig();
-			globalConfig.Plugins.AddAnalytics(analyticsConfig);
+			globalConfig.Plugins.AddAnalyticsTestConfig(analyticsConfig);
 			SetConfig(globalConfig);
 
 			MetricTelemetrySink sink = ServiceProvider.GetRequiredService<MetricTelemetrySink>();
@@ -241,7 +239,7 @@ namespace HordeServer.Analytics.Tests.Telemetry
 			analyticsConfig.Stores.Add(telemetryStoreConfig);
 
 			GlobalConfig globalConfig = new GlobalConfig();
-			globalConfig.Plugins.AddAnalytics(analyticsConfig);
+			globalConfig.Plugins.AddAnalyticsTestConfig(analyticsConfig);
 			SetConfig(globalConfig);
 
 			MetricTelemetrySink sink = ServiceProvider.GetRequiredService<MetricTelemetrySink>();
@@ -281,7 +279,7 @@ namespace HordeServer.Analytics.Tests.Telemetry
 			analyticsConfig.Stores.Add(telemetryStoreConfig);
 
 			GlobalConfig globalConfig = new GlobalConfig();
-			globalConfig.Plugins.AddAnalytics(analyticsConfig);
+			globalConfig.Plugins.AddAnalyticsTestConfig(analyticsConfig);
 			SetConfig(globalConfig);
 
 			MetricTelemetrySink sink = ServiceProvider.GetRequiredService<MetricTelemetrySink>();
@@ -332,7 +330,7 @@ namespace HordeServer.Analytics.Tests.Telemetry
 			analyticsConfig.Stores.Add(telemetryStoreConfig);
 
 			GlobalConfig globalConfig = new GlobalConfig();
-			globalConfig.Plugins.AddAnalytics(analyticsConfig);
+			globalConfig.Plugins.AddAnalyticsTestConfig(analyticsConfig);
 			SetConfig(globalConfig);
 
 			MetricTelemetrySink sink = ServiceProvider.GetRequiredService<MetricTelemetrySink>();
@@ -393,7 +391,7 @@ namespace HordeServer.Analytics.Tests.Telemetry
 			analyticsConfig.Stores.Add(telemetryStoreConfig);
 
 			GlobalConfig globalConfig = new GlobalConfig();
-			globalConfig.Plugins.AddAnalytics(analyticsConfig);
+			globalConfig.Plugins.AddAnalyticsTestConfig(analyticsConfig);
 			SetConfig(globalConfig);
 
 			MetricTelemetrySink sink = ServiceProvider.GetRequiredService<MetricTelemetrySink>();
diff --git a/Engine/Source/Programs/Horde/Plugins/Storage/HordeServer.Storage.Tests/HordeServer.Storage.Tests.csproj b/Engine/Source/Programs/Horde/Plugins/Storage/HordeServer.Storage.Tests/HordeServer.Storage.Tests.csproj
new file mode 100644
--- /dev/null
+++ b/Engine/Source/Programs/Horde/Plugins/Storage/HordeServer.Storage.Tests/HordeServer.Storage.Tests.csproj
@@ -0,0 +1,29 @@
+<Project Sdk="Microsoft.NET.Sdk">
+
+  <PropertyGroup>
+    <TargetFramework>net8.0</TargetFramework>
+    <ImplicitUsings>enable</ImplicitUsings>
+    <Nullable>enable</Nullable>
+
+    <IsPackable>false</IsPackable>
+    <IsTestProject>true</IsTestProject>
+    <RootNamespace>HordeServer.Tests</RootNamespace>
+  </PropertyGroup>
+
+  <ItemGroup>
+    <PackageReference Include="coverlet.collector" Version="6.0.0" />
+    <PackageReference Include="Microsoft.NET.Test.Sdk" Version="17.8.0" />
+    <PackageReference Include="MSTest.TestAdapter" Version="3.1.1" />
+    <PackageReference Include="MSTest.TestFramework" Version="3.1.1" />
+  </ItemGroup>
+
+  <ItemGroup>
+    <ProjectReference Include="..\HordeServer.Storage\HordeServer.Storage.csproj" />
+    <ProjectReference Include="..\..\..\HordeServer.Tests.Shared\HordeServer.Tests.Shared.csproj" />
+  </ItemGroup>
+
+  <ItemGroup>
+    <Using Include="Microsoft.VisualStudio.TestTools.UnitTesting" />
+  </ItemGroup>
+
+</Project>
diff --git a/Engine/Source/Programs/Horde/Plugins/Storage/HordeServer.Storage.Tests/Storage/BlobStoreTests.cs b/Engine/Source/Programs/Horde/Plugins/Storage/HordeServer.Storage.Tests/Storage/BlobStoreTests.cs
new file mode 100644
--- /dev/null
+++ b/Engine/Source/Programs/Horde/Plugins/Storage/HordeServer.Storage.Tests/Storage/BlobStoreTests.cs
@@ -0,0 +1,192 @@
+// Copyright Epic Games, Inc. All Rights Reserved.
+
+using System;
+using System.Collections.Generic;
+using System.Linq;
+using System.Threading.Tasks;
+using EpicGames.Core;
+using EpicGames.Horde.Storage;
+using HordeServer.Plugins;
+using HordeServer.Server;
+using HordeServer.Storage;
+using Microsoft.Extensions.DependencyInjection;
+using Microsoft.VisualStudio.TestTools.UnitTesting;
+
+namespace HordeServer.Tests.Storage
+{
+	[TestClass]
+	public class BlobStoreTests : ServerServiceTest
+	{
+		static readonly BlobType s_blobType = new BlobType("{AFDF76A7-4DEE-5333-F5B5-37B8451251CA}", 1);
+
+		public StorageService StorageService => ServiceProvider.GetRequiredService<StorageService>();
+
+		public BlobStoreTests()
+		{
+			AddPlugin<StoragePlugin>();
+		}
+
+		IStorageClient CreateStorageClient()
+		{
+			StorageConfig storageConfig = new StorageConfig();
+			storageConfig.Backends.Add(new BackendConfig { Id = new BackendId("default-backend"), Type = StorageBackendType.Memory });
+			storageConfig.Namespaces.Add(new NamespaceConfig { Id = new NamespaceId("default"), Backend = new BackendId("default-backend"), GcDelayHrs = 0.0 });
+
+			GlobalConfig globalConfig = new GlobalConfig();
+			globalConfig.Plugins.Add(new PluginName("storage"), storageConfig);
+			SetConfig(globalConfig);
+
+			return StorageService.CreateClient(new NamespaceId("default"));
+		}
+
+		static byte[] CreateTestData(int length, int seed)
+		{
+			byte[] data = new byte[length];
+			new Random(seed).NextBytes(data);
+			return data;
+		}
+
+		class Blob
+		{
+			public ReadOnlyMemory<byte> Data { get; set; }
+			public List<BlobRefValue> References { get; set; }
+
+			public Blob() : this(ReadOnlyMemory<byte>.Empty, new List<BlobRefValue>())
+			{
+			}
+
+			public Blob(ReadOnlyMemory<byte> data, IEnumerable<BlobRefValue> references)
+			{
+				Data = data;
+				References = new List<BlobRefValue>(references);
+			}
+		}
+
+		static async ValueTask<Blob> ReadBlobAsync(IStorageClient store, BlobRefValue locator)
+		{
+			using (BlobData blobData = await store.CreateBlobRef(locator.Hash, locator.Locator).ReadBlobDataAsync())
+			{
+				BlobReader reader = new BlobReader(blobData, null);
+				ReadOnlyMemory<byte> data = reader.ReadVariableLengthBytes();
+
+				List<IBlobRef> handles = new List<IBlobRef>();
+				while (reader.RemainingMemory.Length > 0)
+				{
+					handles.Add(reader.ReadBlobRef<object>());
+				}
+
+				List<BlobRefValue> locators = handles.ConvertAll(x => x.GetRefValue());
+				return new Blob(data, locators);
+			}
+		}
+
+		static async ValueTask<BlobRefValue> WriteBlobAsync(IStorageClient store, Blob blob)
+		{
+			await using IBlobWriter writer = store.CreateBlobWriter();
+			writer.WriteVariableLengthBytes(blob.Data.Span);
+
+			foreach (BlobRefValue reference in blob.References)
+			{
+				writer.WriteBlobRef(store.CreateBlobRef(reference));
+			}
+
+			IBlobRef blobRef = await writer.CompleteAsync(s_blobType);
+			await blobRef.FlushAsync();
+
+			return new BlobRefValue(blobRef.Hash, blobRef.GetLocator());
+		}
+
+		[TestMethod]
+		public async Task LeafTestAsync()
+		{
+			using IStorageClient store = CreateStorageClient();
+
+			byte[] input = CreateTestData(256, 0);
+
+			Blob inputBlob = new Blob(input, Array.Empty<BlobRefValue>());
+			BlobRefValue locator = await WriteBlobAsync(store, inputBlob);
+			Blob outputBlob = await ReadBlobAsync(store, locator);
+
+			Assert.IsTrue(outputBlob.Data.Span.SequenceEqual(input));
+			Assert.AreEqual(0, outputBlob.References.Count);
+		}
+
+		[TestMethod]
+		public async Task ReferenceTestAsync()
+		{
+			using IStorageClient store = CreateStorageClient();
+
+			byte[] input1 = CreateTestData(256, 1);
+			BlobRefValue locator1 = await WriteBlobAsync(store, new Blob(input1, Array.Empty<BlobRefValue>()));
+			Blob blob1 = await ReadBlobAsync(store, locator1);
+			Assert.IsTrue(blob1.Data.Span.SequenceEqual(input1));
+			Assert.IsTrue(blob1.References.SequenceEqual(Array.Empty<BlobRefValue>()));
+
+			byte[] input2 = CreateTestData(256, 2);
+			BlobRefValue locator2 = await WriteBlobAsync(store, new Blob(input2, new BlobRefValue[] { locator1 }));
+			Blob blob2 = await ReadBlobAsync(store, locator2);
+			Assert.IsTrue(blob2.Data.Span.SequenceEqual(input2));
+			Assert.IsTrue(blob2.References.SequenceEqual(new[] { locator1 }));
+
+			byte[] input3 = CreateTestData(256, 3);
+			BlobRefValue locator3 = await WriteBlobAsync(store, new Blob(input3, new BlobRefValue[] { locator1, locator2, locator1 }));
+			Blob blob3 = await ReadBlobAsync(store, locator3);
+			Assert.IsTrue(blob3.Data.Span.SequenceEqual(input3));
+			Assert.IsTrue(blob3.References.SequenceEqual(new[] { locator1, locator2, locator1 }));
+
+			for (int idx = 0; idx < 2; idx++)
+			{
+				RefName refName = new RefName("hello");
+				await store.WriteRefAsync(refName, store.CreateBlobRef(locator3));
+				IBlobRef refTarget = await store.ReadRefAsync(refName);
+				Assert.AreEqual(locator3.Locator, refTarget.GetLocator());
+			}
+		}
+
+		[TestMethod]
+		public async Task RefExpiryTestAsync()
+		{
+			using IStorageClient store = CreateStorageClient();
+
+			Blob blob1 = new Blob(new byte[] { 1, 2, 3 }, Array.Empty<BlobRefValue>());
+			BlobRefValue target = await WriteBlobAsync(store, blob1);
+
+			await store.WriteRefAsync("test-ref-1", store.CreateBlobRef(target));
+			await store.WriteRefAsync("test-ref-2", store.CreateBlobRef(target), new RefOptions { Lifetime = TimeSpan.FromMinutes(30.0), Extend = true });
+			await store.WriteRefAsync("test-ref-3", store.CreateBlobRef(target), new RefOptions { Lifetime = TimeSpan.FromMinutes(30.0), Extend = false });
+
+			Assert.AreEqual(target, await TryReadRefValueAsync(store, "test-ref-1"));
+			Assert.AreEqual(target, await TryReadRefValueAsync(store, "test-ref-2"));
+			Assert.AreEqual(target, await TryReadRefValueAsync(store, "test-ref-3"));
+
+			await Clock.AdvanceAsync(TimeSpan.FromMinutes(25.0));
+
+			Assert.AreEqual(target, await TryReadRefValueAsync(store, "test-ref-1"));
+			Assert.AreEqual(target, await TryReadRefValueAsync(store, "test-ref-2"));
+			Assert.AreEqual(target, await TryReadRefValueAsync(store, "test-ref-3"));
+
+			await Clock.AdvanceAsync(TimeSpan.FromMinutes(25.0));
+
+			Assert.AreEqual(target, await TryReadRefValueAsync(store, "test-ref-1"));
+			Assert.AreEqual(target, await TryReadRefValueAsync(store, "test-ref-2"));
+			Assert.AreEqual(default, await TryReadRefValueAsync(store, "test-ref-3"));
+
+			await Clock.AdvanceAsync(TimeSpan.FromMinutes(35.0));
+
+			Assert.AreEqual(target, await TryReadRefValueAsync(store, "test-ref-1"));
+			Assert.AreEqual(default, await TryReadRefValueAsync(store, "test-ref-2"));
+			Assert.AreEqual(default, await TryReadRefValueAsync(store, "test-ref-3"));
+		}
+
+		static async Task<BlobRefValue?> TryReadRefValueAsync(IStorageClient store, RefName name)
+		{
+			IBlobRef? handle = await store.TryReadRefAsync(name);
+			if (handle == null)
+			{
+				return default;
+			}
+			return handle.GetRefValue();
+		}
+	}
+}
+
diff --git a/Engine/Source/Programs/Horde/Plugins/Storage/HordeServer.Storage.Tests/Storage/GcServiceTests.cs b/Engine/Source/Programs/Horde/Plugins/Storage/HordeServer.Storage.Tests/Storage/GcServiceTests.cs
new file mode 100644
--- /dev/null
+++ b/Engine/Source/Programs/Horde/Plugins/Storage/HordeServer.Storage.Tests/Storage/GcServiceTests.cs
@@ -0,0 +1,137 @@
+// Copyright Epic Games, Inc. All Rights Reserved.
+
+using System;
+using System.Collections.Generic;
+using System.Linq;
+using System.Threading;
+using System.Threading.Tasks;
+using EpicGames.Horde.Storage;
+using HordeServer.Plugins;
+using HordeServer.Server;
+using HordeServer.Storage;
+using Microsoft.Extensions.DependencyInjection;
+using Microsoft.VisualStudio.TestTools.UnitTesting;
+
+namespace HordeServer.Tests.Storage
+{
+	[TestClass]
+	public sealed class GcServiceTests : ServerServiceTest
+	{
+		public StorageService StorageService => ServiceProvider.GetRequiredService<StorageService>();
+
+		public GcServiceTests()
+		{
+			AddPlugin<StoragePlugin>();
+		}
+
+		[TestMethod]
+		public async Task CreateBasicTreeAsync()
+		{
+			await StorageService.StartAsync(CancellationToken.None);
+
+			StorageConfig storageConfig = new StorageConfig();
+			storageConfig.Backends.Add(new BackendConfig { Id = new BackendId("default-backend"), Type = StorageBackendType.Memory });
+			storageConfig.Namespaces.Add(new NamespaceConfig { Id = new NamespaceId("default"), Backend = new BackendId("default-backend"), GcDelayHrs = 0.0 });
+
+			GlobalConfig globalConfig = new GlobalConfig();
+			globalConfig.Plugins.Add(new PluginName("storage"), storageConfig);
+			SetConfig(globalConfig);
+
+			using IStorageClient store = StorageService.CreateClient(new NamespaceId("default"));
+
+			Random random = new Random(0);
+			IBlobRef[] blobs = await CreateTestDataAsync(store, 30, 50, 30, 5, random);
+
+			HashSet<IBlobHandle> roots = new HashSet<IBlobHandle>();
+			for (int idx = 0; idx < 10; idx++)
+			{
+				int blobIdx = (int)(random.NextDouble() * blobs.Length);
+				if (roots.Add(blobs[blobIdx]))
+				{
+					IBlobRef handle = blobs[blobIdx];
+					await store.WriteRefAsync(new RefName($"ref-{idx}"), handle);
+				}
+			}
+
+			HashSet<BlobLocator> nodes = await FindNodesAsync(store, roots);
+
+			await Clock.AdvanceAsync(TimeSpan.FromDays(1.0));
+
+			IObjectStore backend = ServiceProvider.GetRequiredService<IObjectStoreFactory>().CreateObjectStore(storageConfig.Backends[0]);
+
+			ObjectKey[] remaining = await backend.EnumerateAsync().ToArrayAsync();
+			Assert.AreEqual(nodes.Count, remaining.Length);
+
+			HashSet<ObjectKey> nodePaths = new HashSet<ObjectKey>(nodes.Select(x => StorageService.GetObjectKey(x.BaseLocator)));
+			Assert.IsTrue(remaining.All(x => nodePaths.Contains(x)));
+		}
+
+		static async Task<HashSet<BlobLocator>> FindNodesAsync(IStorageClient store, IEnumerable<IBlobHandle> roots)
+		{
+			HashSet<BlobLocator> nodes = new HashSet<BlobLocator>();
+			await FindNodesAsync(store, roots, nodes);
+			return nodes;
+		}
+
+		static async Task FindNodesAsync(IStorageClient store, IEnumerable<IBlobHandle> roots, HashSet<BlobLocator> nodes)
+		{
+			foreach (IBlobHandle root in roots)
+			{
+				BlobLocator locator = root.GetLocator();
+				if (nodes.Add(locator))
+				{
+					BlobData data = await root.ReadBlobDataAsync();
+					await FindNodesAsync(store, data.Imports, nodes);
+				}
+			}
+		}
+
+		static async ValueTask<IBlobRef[]> CreateTestDataAsync(IStorageClient store, int numRoots, int numInterior, int numLeaves, int avgChildren, Random random)
+		{
+			int firstRoot = 0;
+			int firstInterior = firstRoot + numRoots;
+			int firstLeaf = firstInterior + numInterior;
+			int numNodes = firstLeaf + numLeaves;
+
+			List<int>[] children = new List<int>[numNodes];
+			for (int idx = 0; idx < numNodes; idx++)
+			{
+				children[idx] = new List<int>();
+			}
+
+			if (numRoots < numNodes)
+			{
+				double maxParents = ((numRoots + numInterior) * avgChildren) / (numInterior + numLeaves);
+				for (int idx = numRoots; idx < numNodes; idx++)
+				{
+					int numParents = 1 + (int)(random.NextDouble() * maxParents);
+					for (; numParents > 0; numParents--)
+					{
+						int parentIdx = Math.Min((int)(random.NextDouble() * Math.Min(idx, numRoots + numInterior)), idx - 1);
+						children[parentIdx].Add(idx);
+					}
+				}
+			}
+
+			BlobType blobType = new BlobType(Guid.Parse("{AFDF76A7-4DEE-5333-F5B5-37B8451251CA}"), 0);
+
+			IBlobRef[] handles = new IBlobRef[children.Length];
+			for (int idx = numNodes - 1; idx >= 0; idx--)
+			{
+				IBlobRef handle;
+				await using (IBlobWriter writer = store.CreateBlobWriter("gctest"))
+				{
+					List<IBlobRef> imports = children[idx].ConvertAll(x => handles[x]);
+					foreach (IBlobRef import in imports)
+					{
+						writer.WriteBlobRef(import);
+					}
+					handle = await writer.CompleteAsync<object>(blobType);
+				}
+				handles[idx] = handle;
+			}
+
+			return handles;
+		}
+	}
+}
diff --git a/Engine/Source/Programs/Horde/Plugins/Storage/HordeServer.Storage.Tests/Storage/StorageControllerTests.cs b/Engine/Source/Programs/Horde/Plugins/Storage/HordeServer.Storage.Tests/Storage/StorageControllerTests.cs
new file mode 100644
--- /dev/null
+++ b/Engine/Source/Programs/Horde/Plugins/Storage/HordeServer.Storage.Tests/Storage/StorageControllerTests.cs
@@ -0,0 +1,62 @@
+﻿// Copyright Epic Games, Inc. All Rights Reserved.
+
+using System;
+using System.Linq;
+using System.Threading.Tasks;
+using EpicGames.Core;
+using EpicGames.Horde.Storage;
+using EpicGames.Horde.Storage.Backends;
+using EpicGames.Horde.Storage.Bundles;
+using HordeServer.Storage;
+using Microsoft.Extensions.DependencyInjection;
+using Microsoft.Extensions.Logging;
+using Microsoft.VisualStudio.TestTools.UnitTesting;
+using MongoDB.Driver;
+
+namespace HordeServer.Tests.Storage
+{
+#if false
+	[TestClass]
+	public class StorageControllerTests : ControllerIntegrationTest
+	{
+		static readonly Guid s_guid = Guid.Parse("{8B9701C9-9677-4E3C-895F-2D654267C4CB}");
+
+		[TestMethod]
+		public async Task TestBlobImportsAsync()
+		{
+			ILogger logger = ServiceProvider.GetRequiredService<ILogger<StorageControllerTests>>();
+			HttpStorageBackend backend = new HttpStorageBackend("api/v1/storage/default", () => CreateHttpClient(), () => CreateHttpClient(), logger);
+			await using BundleCache cache = new BundleCache();
+			using BundleStorageClient client = new BundleStorageClient(backend, cache, null, logger);
+
+			IBlobRef blobRef1;
+			IBlobRef blobRef2;
+			await using (IBlobWriter writer = client.CreateBlobWriter())
+			{
+				writer.WriteInt32(1234);
+				blobRef1 = await writer.CompleteAsync(new BlobType(s_guid, 0));
+				await writer.FlushAsync();
+
+				writer.WriteBlobRef(blobRef1);
+				blobRef2 = await writer.CompleteAsync(new BlobType(s_guid, 0));
+			}
+
+			BlobLocator locator1 = blobRef1.GetLocator();
+			BlobLocator locator2 = blobRef2.GetLocator();
+
+			StorageService storageService = ServiceProvider.GetRequiredService<StorageService>();
+			IMongoCollection<StorageService.BlobInfo> blobCollection = storageService.BlobCollection;
+
+			StorageService.BlobInfo blob2 = await blobCollection.Find(x => x.Path == locator2.BaseLocator.ToString()).FirstAsync();
+			Assert.IsNotNull(blob2.Imports);
+			Assert.AreEqual(1, blob2.Imports.Count);
+			Assert.AreEqual(blob2.Path, locator2.BaseLocator.ToString());
+
+			StorageService.BlobInfo blob1 = await blobCollection.Find(x => x.Id == blob2.Imports![0]).FirstAsync();
+			Assert.IsNotNull(blob1.Imports);
+			Assert.AreEqual(0, blob1.Imports.Count);
+			Assert.AreEqual(blob1.Path, locator1.BaseLocator.ToString());
+		}
+	}
+#endif
+}
diff --git a/Engine/Source/Programs/Horde/Plugins/Storage/HordeServer.Storage.Tests/Storage/StorageServiceTests.cs b/Engine/Source/Programs/Horde/Plugins/Storage/HordeServer.Storage.Tests/Storage/StorageServiceTests.cs
new file mode 100644
--- /dev/null
+++ b/Engine/Source/Programs/Horde/Plugins/Storage/HordeServer.Storage.Tests/Storage/StorageServiceTests.cs
@@ -0,0 +1,66 @@
+// Copyright Epic Games, Inc. All Rights Reserved.
+
+using System;
+using System.Threading.Tasks;
+using EpicGames.Core;
+using EpicGames.Horde.Storage;
+using HordeServer.Storage;
+using Microsoft.Extensions.DependencyInjection;
+using Microsoft.VisualStudio.TestTools.UnitTesting;
+
+namespace HordeServer.Tests.Storage
+{
+	[TestClass]
+	public class StorageServiceTests : ServerServiceTest
+	{
+		public StorageService StorageService => ServiceProvider.GetRequiredService<StorageService>();
+
+		public StorageServiceTests()
+		{
+			AddPlugin<StoragePlugin>();
+		}
+
+		[TestMethod]
+		public async Task BlobCollectionTestAsync()
+		{
+			StorageService storageService = ServiceProvider.GetRequiredService<StorageService>();
+			using IStorageClient client = storageService.CreateClient(new NamespaceId("memory"));
+
+			BlobType type1 = new BlobType(Guid.Parse("{11C2D886-4164-3349-D1E9-6F943D2ED10B}"), 0);
+			byte[] data1 = new byte[] { 1, 2, 3 };
+
+			BlobType type2 = new BlobType(Guid.Parse("{6CB3A005-4787-26BA-3E79-D286CB7137D1}"), 0);
+			byte[] data2 = new byte[] { 4, 5, 6 };
+
+			IBlobRef handle1a;
+			IBlobRef handle1b;
+			IBlobRef handle2;
+			await using (IBlobWriter writer = client.CreateBlobWriter())
+			{
+				writer.WriteFixedLengthBytes(data1);
+				writer.AddAlias("foo", 2);
+				handle1a = await writer.CompleteAsync(type1);
+
+				writer.WriteFixedLengthBytes(data1);
+				writer.AddAlias("foo", 1);
+				handle1b = await writer.CompleteAsync(type1);
+
+				writer.WriteFixedLengthBytes(data2);
+				writer.AddAlias("bar", 0);
+				handle2 = await writer.CompleteAsync(type2);
+			}
+
+			BlobAlias[] aliases;
+
+			aliases = await client.FindAliasesAsync("foo");
+			Assert.AreEqual(2, aliases.Length);
+			Assert.AreEqual(handle1a.GetLocator(), aliases[0].Target.GetLocator());
+			Assert.AreEqual(handle1b.GetLocator(), aliases[1].Target.GetLocator());
+
+			aliases = await client.FindAliasesAsync("bar");
+			Assert.AreEqual(1, aliases.Length);
+			Assert.AreEqual(handle2.GetLocator(), aliases[0].Target.GetLocator());
+		}
+	}
+}
+
diff --git a/Engine/Source/Programs/Horde/Plugins/Storage/HordeServer.Storage.Tests/StorageTestExtensions.cs b/Engine/Source/Programs/Horde/Plugins/Storage/HordeServer.Storage.Tests/StorageTestExtensions.cs
new file mode 100644
--- /dev/null
+++ b/Engine/Source/Programs/Horde/Plugins/Storage/HordeServer.Storage.Tests/StorageTestExtensions.cs
@@ -0,0 +1,22 @@
+// Copyright Epic Games, Inc. All Rights Reserved.
+
+using System;
+using System.Collections.Generic;
+using System.Linq;
+using System.Text;
+using System.Threading.Tasks;
+using EpicGames.Horde.Storage;
+using EpicGames.Horde.Storage.Bundles;
+using EpicGames.Horde.Storage.ObjectStores;
+using HordeServer.Plugins;
+using HordeServer.Storage;
+using Microsoft.Extensions.DependencyInjection;
+
+namespace HordeServer.Tests
+{
+	public static class StorageTestExtensions
+	{
+		public static void AddStorageTestConfig(this PluginConfigCollection configCollection, StorageConfig config)
+			=> configCollection[new PluginName("Storage")] = config;
+	}
+}
diff --git a/Engine/Source/Programs/Horde/Plugins/Storage/HordeServer.Storage/HordeServer.Storage.csproj b/Engine/Source/Programs/Horde/Plugins/Storage/HordeServer.Storage/HordeServer.Storage.csproj
new file mode 100644
--- /dev/null
+++ b/Engine/Source/Programs/Horde/Plugins/Storage/HordeServer.Storage/HordeServer.Storage.csproj
@@ -0,0 +1,28 @@
+﻿<Project Sdk="Microsoft.NET.Sdk">
+
+    <PropertyGroup>
+    <TargetFramework>net8.0</TargetFramework>
+    <ImplicitUsings>enable</ImplicitUsings>
+    <Nullable>enable</Nullable>
+  </PropertyGroup>
+
+  <ItemGroup>
+    <FrameworkReference Include="Microsoft.AspNetCore.App" />
+  </ItemGroup>
+
+  <ItemGroup>
+    <PackageReference Include="AWSSDK.S3" Version="3.7.304.1" />
+    <PackageReference Include="AWSSDK.Extensions.NETCore.Setup" Version="3.7.300" />
+    <PackageReference Include="Azure.Storage.Blobs" Version="12.19.1" />
+    <PackageReference Include="MongoDB.Driver" Version="2.22.0" />
+  </ItemGroup>
+
+  <ItemGroup>
+    <ProjectReference Include="..\..\..\..\Shared\EpicGames.Core\EpicGames.Core.csproj" />
+    <ProjectReference Include="..\..\..\..\Shared\EpicGames.Horde\EpicGames.Horde.csproj" />
+    <ProjectReference Include="..\..\..\..\Shared\EpicGames.IoHash\EpicGames.IoHash.csproj" />
+    <ProjectReference Include="..\..\..\..\Shared\EpicGames.Redis\EpicGames.Redis.csproj" />
+    <ProjectReference Include="..\..\..\HordeServer.Shared\HordeServer.Shared.csproj" />
+  </ItemGroup>
+
+</Project>
diff --git a/Engine/Source/Programs/Horde/Plugins/Storage/HordeServer.Storage/Properties/AssemblyInfo.cs b/Engine/Source/Programs/Horde/Plugins/Storage/HordeServer.Storage/Properties/AssemblyInfo.cs
new file mode 100644
--- /dev/null
+++ b/Engine/Source/Programs/Horde/Plugins/Storage/HordeServer.Storage/Properties/AssemblyInfo.cs
@@ -0,0 +1,7 @@
+// Copyright Epic Games, Inc. All Rights Reserved.
+
+using System.Resources;
+using System.Runtime.CompilerServices;
+
+[assembly: InternalsVisibleTo("HordeServer.Storage.Tests")]
+[assembly: NeutralResourcesLanguage("en-US")]
diff --git a/Engine/Source/Programs/Horde/Plugins/Storage/HordeServer.Storage/Storage/BackendId.cs b/Engine/Source/Programs/Horde/Plugins/Storage/HordeServer.Storage/Storage/BackendId.cs
new file mode 100644
--- /dev/null
+++ b/Engine/Source/Programs/Horde/Plugins/Storage/HordeServer.Storage/Storage/BackendId.cs
@@ -0,0 +1,50 @@
+// Copyright Epic Games, Inc. All Rights Reserved.
+
+using System.ComponentModel;
+using EpicGames.Core;
+using EpicGames.Horde;
+using EpicGames.Serialization;
+
+namespace HordeServer.Storage
+{
+	/// <summary>
+	/// Identifier for a pool
+	/// </summary>
+	/// <param name="Id">Id to construct from</param>
+	[JsonSchemaString]
+	[TypeConverter(typeof(StringIdTypeConverter<BackendId, BackendIdConverter>))]
+	[StringIdConverter(typeof(BackendIdConverter))]
+	[CbConverter(typeof(StringIdCbConverter<BackendId, BackendIdConverter>))]
+	public record struct BackendId(StringId Id)
+	{
+		/// <summary>
+		/// Empty backend id
+		/// </summary>
+		public static BackendId Empty { get; } = default;
+
+		/// <summary>
+		/// Constructor
+		/// </summary>
+		public BackendId(string id) : this(new StringId(id))
+		{
+		}
+
+		/// <inheritdoc cref="StringId.IsEmpty"/>
+		public bool IsEmpty => Id.IsEmpty;
+
+		/// <inheritdoc/>
+		public override string ToString() => Id.ToString();
+	}
+
+	/// <summary>
+	/// Converter to and from <see cref="StringId"/> instances.
+	/// </summary>
+	class BackendIdConverter : StringIdConverter<BackendId>
+	{
+		/// <inheritdoc/>
+		public override BackendId FromStringId(StringId id) => new BackendId(id);
+
+		/// <inheritdoc/>
+		public override StringId ToStringId(BackendId value) => value.Id;
+	}
+}
diff --git a/Engine/Source/Programs/Horde/Plugins/Storage/HordeServer.Storage/Storage/IObjectStoreFactory.cs b/Engine/Source/Programs/Horde/Plugins/Storage/HordeServer.Storage/Storage/IObjectStoreFactory.cs
new file mode 100644
--- /dev/null
+++ b/Engine/Source/Programs/Horde/Plugins/Storage/HordeServer.Storage/Storage/IObjectStoreFactory.cs
@@ -0,0 +1,19 @@
+// Copyright Epic Games, Inc. All Rights Reserved.
+
+using EpicGames.Horde.Storage;
+
+namespace HordeServer.Storage
+{
+	/// <summary>
+	/// Service which resolves storage backend instances
+	/// </summary>
+	public interface IObjectStoreFactory
+	{
+		/// <summary>
+		/// Creates a new storage backend with the given identifier
+		/// </summary>
+		/// <param name="backendConfig">Configuration object for the storage backend</param>
+		/// <returns>New backend instance. This instance should be disposed when the client has finished with it.</returns>
+		IObjectStore CreateObjectStore(BackendConfig backendConfig);
+	}
+}
diff --git a/Engine/Source/Programs/Horde/Plugins/Storage/HordeServer.Storage/Storage/ObjectStoreFactory.cs b/Engine/Source/Programs/Horde/Plugins/Storage/HordeServer.Storage/Storage/ObjectStoreFactory.cs
new file mode 100644
--- /dev/null
+++ b/Engine/Source/Programs/Horde/Plugins/Storage/HordeServer.Storage/Storage/ObjectStoreFactory.cs
@@ -0,0 +1,80 @@
+// Copyright Epic Games, Inc. All Rights Reserved.
+
+using System;
+using System.Collections.Generic;
+using EpicGames.Core;
+using EpicGames.Horde.Storage;
+using EpicGames.Horde.Storage.ObjectStores;
+using HordeServer.Storage.ObjectStores;
+using Microsoft.Extensions.DependencyInjection;
+using Microsoft.Extensions.Logging;
+
+namespace HordeServer.Storage
+{
+	/// <summary>
+	/// Default implementation of <see cref="IObjectStoreFactory"/>
+	/// </summary>
+	class ObjectStoreFactory : IObjectStoreFactory
+	{
+		readonly IServiceProvider _serviceProvider;
+		readonly IServerInfo _serverInfo;
+		readonly object _lockObject = new object();
+		IReadOnlyDictionary<IoHash, IObjectStore> _objectStores = new Dictionary<IoHash, IObjectStore>();
+		readonly ILogger _logger;
+
+		/// <summary>
+		/// Constructor
+		/// </summary>
+		public ObjectStoreFactory(IServiceProvider serviceProvider, IServerInfo serverInfo, ILogger<ObjectStoreFactory> logger)
+		{
+			_serviceProvider = serviceProvider;
+			_serverInfo = serverInfo;
+			_logger = logger;
+		}
+
+		/// <inheritdoc/>
+		public IObjectStore CreateObjectStore(BackendConfig config)
+		{
+			// See if we've got an existing backend we can use
+			IObjectStore? objectStore;
+			if (!_objectStores.TryGetValue(config.Hash, out objectStore))
+			{
+				lock (_lockObject)
+				{
+					if (!_objectStores.TryGetValue(config.Hash, out objectStore))
+					{
+						objectStore = CreateObjectStoreInternal(config);
+
+						Dictionary<IoHash, IObjectStore> newObjectStores = new Dictionary<IoHash, IObjectStore>(_objectStores);
+						newObjectStores.Add(config.Hash, objectStore);
+						_objectStores = newObjectStores;
+
+						_logger.LogInformation("Created object store {Id}@{Hash}", config.Id, config.Hash);
+					}
+				}
+			}
+
+			return objectStore;
+		}
+
+		/// <summary>
+		/// Creates a storage backend with the given configuration
+		/// </summary>
+		/// <param name="config">Configuration for the backend</param>
+		/// <returns>New storage backend instance</returns>
+		IObjectStore CreateObjectStoreInternal(BackendConfig config)
+		{
+			switch (config.Type ?? StorageBackendType.FileSystem)
+			{
+				case StorageBackendType.FileSystem:
+					return _serviceProvider.GetRequiredService<FileObjectStoreFactory>().CreateStore(DirectoryReference.Combine(_serverInfo.DataDir, config.BaseDir ?? "Storage"));
+				case StorageBackendType.Aws:
+					return _serviceProvider.GetRequiredService<AwsObjectStoreFactory>().CreateStore(config);
+				case StorageBackendType.Memory:
+					return new MemoryObjectStore();
+				default:
+					throw new NotImplementedException();
+			}
+		}
+	}
+}
diff --git a/Engine/Source/Programs/Horde/Plugins/Storage/HordeServer.Storage/Storage/ObjectStores/AwsObjectStore.cs b/Engine/Source/Programs/Horde/Plugins/Storage/HordeServer.Storage/Storage/ObjectStores/AwsObjectStore.cs
new file mode 100644
--- /dev/null
+++ b/Engine/Source/Programs/Horde/Plugins/Storage/HordeServer.Storage/Storage/ObjectStores/AwsObjectStore.cs
@@ -0,0 +1,683 @@
+// Copyright Epic Games, Inc. All Rights Reserved.
+
+using System;
+using System.Collections.Generic;
+using System.Globalization;
+using System.IO;
+using System.Net;
+using System.Runtime.CompilerServices;
+using System.Threading;
+using System.Threading.Tasks;
+using Amazon;
+using Amazon.Extensions.NETCore.Setup;
+using Amazon.Runtime;
+using Amazon.S3;
+using Amazon.S3.Model;
+using EpicGames.Core;
+using EpicGames.Horde.Storage;
+using HordeServer.Utilities;
+using Microsoft.Extensions.Configuration;
+using Microsoft.Extensions.Logging;
+using OpenTelemetry.Trace;
+
+namespace HordeServer.Storage.ObjectStores
+{
+	/// <summary>
+	/// Exception wrapper for S3 requests
+	/// </summary>
+	public sealed class AwsException : StorageException
+	{
+		/// <summary>
+		/// Constructor
+		/// </summary>
+		/// <param name="message">Message for the exception</param>
+		/// <param name="innerException">Inner exception data</param>
+		public AwsException(string message, Exception? innerException)
+			: base(message, innerException)
+		{
+		}
+	}
+
+	/// <summary>
+	/// Credentials to use for AWS
+	/// </summary>
+	public enum AwsCredentialsType
+	{
+		/// <summary>
+		/// Use default credentials from the AWS SDK
+		/// </summary>
+		Default,
+
+		/// <summary>
+		/// Read credentials from the <see cref="IAwsStorageOptions.AwsProfile"/> profile in the AWS config file
+		/// </summary>
+		Profile,
+
+		/// <summary>
+		/// Assume a particular role. Should specify ARN in <see cref="IAwsStorageOptions.AwsRole"/>
+		/// </summary>
+		AssumeRole,
+
+		/// <summary>
+		/// Assume a particular role using the current environment variables.
+		/// </summary>
+		AssumeRoleWebIdentity,
+	}
+
+	/// <summary>
+	/// Options for AWS
+	/// </summary>
+	public interface IAwsStorageOptions
+	{
+		/// <summary>
+		/// Type of credentials to use
+		/// </summary>
+		public AwsCredentialsType? AwsCredentials { get; }
+
+		/// <summary>
+		/// Name of the bucket to use
+		/// </summary>
+		public string? AwsBucketName { get; }
+
+		/// <summary>
+		/// Base path within the bucket 
+		/// </summary>
+		public string? AwsBucketPath { get; }
+
+		/// <summary>
+		/// ARN of a role to assume
+		/// </summary>
+		public string? AwsRole { get; }
+
+		/// <summary>
+		/// The AWS profile to read credentials form
+		/// </summary>
+		public string? AwsProfile { get; }
+
+		/// <summary>
+		/// Region to connect to
+		/// </summary>
+		public string? AwsRegion { get; }
+	}
+
+	/// <summary>
+	/// Storage backend using AWS S3
+	/// </summary>
+	public sealed class AwsObjectStore : IObjectStore
+	{
+		/// <summary>
+		/// S3 Client
+		/// </summary>
+		private readonly IAmazonS3 _client;
+
+		/// <summary>
+		/// Options for AWs
+		/// </summary>
+		private readonly IAwsStorageOptions _options;
+
+		/// <summary>
+		/// Semaphore for connecting to AWS
+		/// </summary>
+		private readonly SemaphoreSlim _semaphore;
+
+		/// <summary>
+		/// Prefix for objects in the bucket
+		/// </summary>
+		private readonly string _pathPrefix;
+
+		/// <summary>
+		/// Logger interface
+		/// </summary>
+		private readonly ILogger _logger;
+
+		/// <inheritdoc/>
+		public bool SupportsRedirects => true;
+
+		/// <summary>
+		/// Constructor
+		/// </summary>
+		/// <param name="client">Client instance</param>
+		/// <param name="options">Storage options</param>
+		/// <param name="semaphore">Semaphore for making requests</param>
+		/// <param name="logger">Logger interface</param>
+		public AwsObjectStore(IAmazonS3 client, IAwsStorageOptions options, SemaphoreSlim semaphore, ILogger<AwsObjectStore> logger)
+		{
+			_client = client;
+			_options = options;
+			_semaphore = semaphore;
+			_logger = logger;
+
+			_pathPrefix = (_options.AwsBucketPath ?? String.Empty).TrimEnd('/');
+			if (_pathPrefix.Length > 0)
+			{
+				_pathPrefix += '/';
+			}
+		}
+
+		class WrappedResponseStream : Stream
+		{
+			readonly IDisposable _semaphore;
+			readonly TelemetrySpan _semaphoreSpan;
+			readonly GetObjectResponse _response;
+			readonly Stream _responseStream;
+
+			public WrappedResponseStream(IDisposable semaphore, TelemetrySpan semaphoreSpan, GetObjectResponse response)
+			{
+				_semaphore = semaphore;
+				_semaphoreSpan = semaphoreSpan;
+				_response = response;
+				_responseStream = response.ResponseStream;
+			}
+
+			public override bool CanRead => true;
+			public override bool CanSeek => false;
+			public override bool CanWrite => false;
+			public override long Length => _response.ContentLength;
+			public override long Position { get => _responseStream.Position; set => throw new NotSupportedException(); }
+
+			public override void Flush() => _responseStream.Flush();
+
+			public override int Read(byte[] buffer, int offset, int count) => _responseStream.Read(buffer, offset, count);
+			public override ValueTask<int> ReadAsync(Memory<byte> buffer, CancellationToken cancellationToken = default) => _responseStream.ReadAsync(buffer, cancellationToken);
+			public override Task<int> ReadAsync(byte[] buffer, int offset, int count, CancellationToken cancellationToken) => _responseStream.ReadAsync(buffer, offset, count, cancellationToken);
+
+			public override long Seek(long offset, SeekOrigin origin) => throw new NotSupportedException();
+			public override void SetLength(long value) => throw new NotSupportedException();
+			public override void Write(byte[] buffer, int offset, int count) => throw new NotSupportedException();
+
+			protected override void Dispose(bool disposing)
+			{
+				base.Dispose(disposing);
+
+				if (disposing)
+				{
+					_semaphore.Dispose();
+					_semaphoreSpan.Dispose();
+					_response.Dispose();
+					_responseStream.Dispose();
+				}
+			}
+
+			public override async ValueTask DisposeAsync()
+			{
+				await base.DisposeAsync();
+
+				_semaphore.Dispose();
+				_response.Dispose();
+				await _responseStream.DisposeAsync();
+			}
+		}
+
+		string GetFullPath(ObjectKey key) => $"{_pathPrefix}{key}";
+
+		/// <inheritdoc/>
+		public Task<Stream> OpenAsync(ObjectKey key, int offset, int? length, CancellationToken cancellationToken)
+		{
+			string range;
+			if (length == null)
+			{
+				range = $"bytes={offset}-";
+			}
+			else if (length.Value == 0)
+			{
+				throw new ArgumentException("Cannot read empty stream from AWS backend", nameof(length));
+			}
+			else
+			{
+				range = $"bytes={offset}-{offset + length.Value - 1}";
+			}
+			return OpenAsync(key, new ByteRange(range), cancellationToken);
+		}
+
+		async Task<Stream> OpenAsync(ObjectKey key, ByteRange? byteRange, CancellationToken cancellationToken)
+		{
+			using TelemetrySpan span = OpenTelemetryTracers.Horde.StartActiveSpan($"{nameof(AwsObjectStore)}.{nameof(OpenAsync)}");
+			span.SetAttribute("path", key.ToString());
+
+			string fullPath = GetFullPath(key);
+
+			IDisposable? semaLock = null;
+			TelemetrySpan? semaphoreSpan = null;
+			GetObjectResponse? response = null;
+			try
+			{
+				semaLock = await _semaphore.WaitDisposableAsync(cancellationToken);
+
+				semaphoreSpan = OpenTelemetryTracers.Horde.StartActiveSpan($"{nameof(AwsObjectStore)}.{nameof(OpenAsync)}.Semaphore");
+				semaphoreSpan.SetAttribute("path", key.ToString());
+
+				GetObjectRequest newGetRequest = new GetObjectRequest();
+				newGetRequest.BucketName = _options.AwsBucketName;
+				newGetRequest.Key = fullPath;
+				newGetRequest.ByteRange = byteRange;
+
+				try
+				{
+					response = await _client.GetObjectAsync(newGetRequest, cancellationToken);
+				}
+				catch (Exception ex) when (ex is not OperationCanceledException)
+				{
+					_logger.LogInformation(ex, "S3 read of {BucketId} {Key} failed: {Message}", newGetRequest.BucketName, newGetRequest.Key, ex.Message);
+
+					// Temp hack for files losing '.blob' extension
+					const string BlobExtension = ".blob";
+					if (fullPath.EndsWith(BlobExtension, StringComparison.OrdinalIgnoreCase))
+					{
+						try
+						{
+							newGetRequest.Key = fullPath.Substring(0, fullPath.Length - 5);
+							_logger.LogInformation("Attempting S3 read of {BucketId} {Key}...", newGetRequest.BucketName, newGetRequest.Key);
+							response = await _client.GetObjectAsync(newGetRequest, cancellationToken);
+							return new WrappedResponseStream(semaLock, semaphoreSpan, response);
+						}
+						catch (Exception ex2)
+						{
+							_logger.LogInformation(ex2, "Alternate S3 read (no extension) of {BucketId} {Key} failed: {Message}", newGetRequest.BucketName, newGetRequest.Key, ex2.Message);
+						}
+					}
+
+					// Temp hack for case changes with sanitized object keys
+					string newFullPath = GetFullPath(ObjectKey.Sanitize(key.Path));
+					if (!String.Equals(fullPath, newFullPath, StringComparison.Ordinal))
+					{
+						try
+						{
+							newGetRequest.Key = newFullPath;
+							_logger.LogInformation("Attempting S3 read of {BucketId} {Key}...", newGetRequest.BucketName, newGetRequest.Key);
+							response = await _client.GetObjectAsync(newGetRequest, cancellationToken);
+							return new WrappedResponseStream(semaLock, semaphoreSpan, response);
+						}
+						catch (Exception ex2)
+						{
+							_logger.LogInformation(ex2, "Alternate S3 read (sanitized path) of {BucketId} {Key} failed: {Message}", newGetRequest.BucketName, newGetRequest.Key, ex2.Message);
+						}
+					}
+
+					// Temp hack for case changes with sanitized object keys AND no extension
+					if (newFullPath.EndsWith(BlobExtension, StringComparison.OrdinalIgnoreCase))
+					{
+						try
+						{
+							newGetRequest.Key = newFullPath.Substring(0, newFullPath.Length - 5);
+							_logger.LogInformation("Attempting S3 read of {BucketId} {Key}...", newGetRequest.BucketName, newGetRequest.Key);
+							response = await _client.GetObjectAsync(newGetRequest, cancellationToken);
+							return new WrappedResponseStream(semaLock, semaphoreSpan, response);
+						}
+						catch (Exception ex2)
+						{
+							_logger.LogInformation(ex2, "Alternate S3 read (sanitized path, no extension) of {BucketId} {Key} failed: {Message}", newGetRequest.BucketName, newGetRequest.Key, ex2.Message);
+						}
+					}
+
+					throw;
+				}
+
+				return new WrappedResponseStream(semaLock, semaphoreSpan, response);
+			}
+			catch (Exception ex)
+			{
+				semaLock?.Dispose();
+				response?.Dispose();
+				semaphoreSpan?.Dispose();
+
+				if (ex is OperationCanceledException)
+				{
+					throw;
+				}
+
+				_logger.LogWarning(ex, "Unable to read {Path} from S3", fullPath);
+
+				if (ex is AmazonS3Exception s3ex && s3ex.StatusCode == HttpStatusCode.NotFound)
+				{
+					throw new ObjectNotFoundException(key, $"Object {key} not found in bucket {_options.AwsBucketName}", ex);
+				}
+				else
+				{
+					throw new StorageException($"Unable to read {fullPath} from {_options.AwsBucketName}: {ex.Message}", ex);
+				}
+			}
+		}
+
+		/// <inheritdoc/>
+		public async Task<IReadOnlyMemoryOwner<byte>> ReadAsync(ObjectKey key, int offset, int? length, CancellationToken cancellationToken)
+		{
+			using Stream stream = await OpenAsync(key, offset, length, cancellationToken);
+			return ReadOnlyMemoryOwner.Create(await stream.ReadAllBytesAsync(cancellationToken));
+		}
+
+		/// <inheritdoc/>
+		public ValueTask<Uri?> TryGetReadRedirectAsync(ObjectKey key, CancellationToken cancellationToken = default)
+			=> new ValueTask<Uri?>(GetPresignedUrl(key, HttpVerb.GET));
+
+		/// <inheritdoc/>
+		public ValueTask<Uri?> TryGetWriteRedirectAsync(ObjectKey key, CancellationToken cancellationToken = default)
+			=> new ValueTask<Uri?>(GetPresignedUrl(key, HttpVerb.PUT));
+
+		/// <summary>
+		/// Helper method to generate a presigned URL for a request
+		/// </summary>
+		Uri? GetPresignedUrl(ObjectKey key, HttpVerb verb)
+		{
+			using TelemetrySpan span = OpenTelemetryTracers.Horde.StartActiveSpan($"{nameof(AwsObjectStore)}.{nameof(GetPresignedUrl)}");
+			span.SetAttribute("path", key.ToString());
+
+			string fullPath = GetFullPath(key);
+
+			try
+			{
+				GetPreSignedUrlRequest newGetRequest = new GetPreSignedUrlRequest();
+				newGetRequest.BucketName = _options.AwsBucketName;
+				newGetRequest.Key = fullPath;
+				newGetRequest.Verb = verb;
+				newGetRequest.Expires = DateTime.UtcNow.AddHours(3.0);
+				newGetRequest.ResponseHeaderOverrides.CacheControl = "private, max-age=2592000, immutable"; // 30 days
+
+				string url = _client.GetPreSignedURL(newGetRequest);
+				_logger.LogDebug("Creating presigned URL for {Verb} to {Path}", verb, fullPath);
+				return new Uri(url);
+			}
+			catch (Exception ex)
+			{
+				_logger.LogWarning(ex, "Unable to get presigned url for {Path} from S3", fullPath);
+				return null;
+			}
+		}
+
+		/// <inheritdoc/>
+		public async Task WriteAsync(ObjectKey key, Stream inputStream, CancellationToken cancellationToken = default)
+		{
+			TimeSpan[] retryTimes =
+			{
+				TimeSpan.FromSeconds(1.0),
+				TimeSpan.FromSeconds(5.0),
+				TimeSpan.FromSeconds(10.0),
+			};
+
+			string fullPath = GetFullPath(key);
+			for (int attempt = 0; ; attempt++)
+			{
+				try
+				{
+					using IDisposable semaLock = await _semaphore.WaitDisposableAsync(cancellationToken);
+					await WriteInternalAsync(fullPath, inputStream, cancellationToken);
+					_logger.LogDebug("Written data to {Path}", fullPath);
+					break;
+				}
+				catch (Exception ex) when (ex is not OperationCanceledException)
+				{
+					_logger.LogError(ex, "Unable to write data to {Path} ({Attempt}/{AttemptCount})", fullPath, attempt + 1, retryTimes.Length + 1);
+					if (attempt >= retryTimes.Length)
+					{
+						throw new AwsException($"Unable to write to bucket {_options.AwsBucketName} path {fullPath}", ex);
+					}
+				}
+
+				await Task.Delay(retryTimes[attempt], cancellationToken);
+			}
+		}
+
+		/// <inheritdoc/>
+		public async Task WriteInternalAsync(string fullPath, Stream stream, CancellationToken cancellationToken)
+		{
+			using TelemetrySpan span = OpenTelemetryTracers.Horde.StartActiveSpan($"{nameof(AwsObjectStore)}.{nameof(WriteInternalAsync)}");
+			span.SetAttribute("path", fullPath);
+
+			const int MinPartSize = 5 * 1024 * 1024;
+
+			long streamLen = stream.Length;
+			if (streamLen < MinPartSize)
+			{
+				// Read the data into memory. Errors with hash value not matching if we don't do this (?)
+				byte[] buffer = new byte[streamLen];
+				await ReadExactLengthAsync(stream, buffer, (int)streamLen, cancellationToken);
+
+				// Upload it to S3
+				using (MemoryStream inputStream = new MemoryStream(buffer))
+				{
+					PutObjectRequest uploadRequest = new PutObjectRequest();
+					uploadRequest.BucketName = _options.AwsBucketName;
+					uploadRequest.Key = fullPath;
+					uploadRequest.InputStream = inputStream;
+					uploadRequest.Metadata.Add("bytes-written", streamLen.ToString(CultureInfo.InvariantCulture));
+					await _client.PutObjectAsync(uploadRequest, cancellationToken);
+				}
+			}
+			else
+			{
+				// Initiate a multi-part upload
+				InitiateMultipartUploadRequest initiateRequest = new InitiateMultipartUploadRequest();
+				initiateRequest.BucketName = _options.AwsBucketName;
+				initiateRequest.Key = fullPath;
+
+				InitiateMultipartUploadResponse initiateResponse = await _client.InitiateMultipartUploadAsync(initiateRequest, cancellationToken);
+				try
+				{
+					// Buffer for reading the data
+					byte[] buffer = new byte[MinPartSize];
+
+					// Upload all the parts
+					List<PartETag> partTags = new List<PartETag>();
+					for (long streamPos = 0; streamPos < streamLen;)
+					{
+						// Read the next chunk of data into the buffer
+						int bufferLen = (int)Math.Min((long)MinPartSize, streamLen - streamPos);
+						await ReadExactLengthAsync(stream, buffer, bufferLen, cancellationToken);
+						streamPos += bufferLen;
+
+						// Upload the part
+						using (MemoryStream inputStream = new MemoryStream(buffer, 0, bufferLen))
+						{
+							UploadPartRequest partRequest = new UploadPartRequest();
+							partRequest.BucketName = _options.AwsBucketName;
+							partRequest.Key = fullPath;
+							partRequest.UploadId = initiateResponse.UploadId;
+							partRequest.InputStream = inputStream;
+							partRequest.PartSize = bufferLen;
+							partRequest.PartNumber = partTags.Count + 1;
+							partRequest.IsLastPart = (streamPos == streamLen);
+
+							UploadPartResponse partResponse = await _client.UploadPartAsync(partRequest, cancellationToken);
+							partTags.Add(new PartETag(partResponse.PartNumber, partResponse.ETag));
+						}
+					}
+
+					// Mark the upload as complete
+					CompleteMultipartUploadRequest completeRequest = new CompleteMultipartUploadRequest();
+					completeRequest.BucketName = _options.AwsBucketName;
+					completeRequest.Key = fullPath;
+					completeRequest.UploadId = initiateResponse.UploadId;
+					completeRequest.PartETags = partTags;
+					await _client.CompleteMultipartUploadAsync(completeRequest, cancellationToken);
+				}
+				catch
+				{
+					// Abort the upload
+					AbortMultipartUploadRequest abortRequest = new AbortMultipartUploadRequest();
+					abortRequest.BucketName = _options.AwsBucketName;
+					abortRequest.Key = fullPath;
+					abortRequest.UploadId = initiateResponse.UploadId;
+					await _client.AbortMultipartUploadAsync(abortRequest, cancellationToken);
+
+					throw;
+				}
+			}
+		}
+
+		/// <summary>
+		/// Reads data of an exact length into a stream
+		/// </summary>
+		/// <param name="stream">The stream to read from</param>
+		/// <param name="buffer">The buffer to read into</param>
+		/// <param name="length">Length of the data to read</param>
+		/// <param name="cancellationToken">Cancellation token for the operation</param>
+		/// <returns>Async task</returns>
+		static async Task ReadExactLengthAsync(System.IO.Stream stream, byte[] buffer, int length, CancellationToken cancellationToken)
+		{
+			int bufferPos = 0;
+			while (bufferPos < length)
+			{
+				int bytesRead = await stream.ReadAsync(buffer, bufferPos, length - bufferPos, cancellationToken);
+				if (bytesRead == 0)
+				{
+					throw new InvalidOperationException("Unexpected end of stream");
+				}
+				bufferPos += bytesRead;
+			}
+		}
+
+		/// <inheritdoc/>
+		public async Task DeleteAsync(ObjectKey key, CancellationToken cancellationToken)
+		{
+			using TelemetrySpan span = OpenTelemetryTracers.Horde.StartActiveSpan($"{nameof(AwsObjectStore)}.{nameof(DeleteAsync)}");
+			span.SetAttribute("path", key.ToString());
+
+			DeleteObjectRequest newDeleteRequest = new DeleteObjectRequest();
+			newDeleteRequest.BucketName = _options.AwsBucketName;
+			newDeleteRequest.Key = GetFullPath(key);
+			await _client.DeleteObjectAsync(newDeleteRequest, cancellationToken);
+		}
+
+		/// <inheritdoc/>
+		public async Task<bool> ExistsAsync(ObjectKey key, CancellationToken cancellationToken)
+		{
+			using TelemetrySpan span = OpenTelemetryTracers.Horde.StartActiveSpan($"{nameof(AwsObjectStore)}.{nameof(ExistsAsync)}");
+			span.SetAttribute("path", key.ToString());
+
+			try
+			{
+				GetObjectMetadataRequest request = new GetObjectMetadataRequest();
+				request.BucketName = _options.AwsBucketName;
+				request.Key = GetFullPath(key);
+				await _client.GetObjectMetadataAsync(request, cancellationToken);
+				return true;
+			}
+			catch (AmazonS3Exception ex) when (ex.StatusCode == System.Net.HttpStatusCode.NotFound)
+			{
+				return false;
+			}
+		}
+
+		/// <inheritdoc/>
+		public async IAsyncEnumerable<ObjectKey> EnumerateAsync([EnumeratorCancellation] CancellationToken cancellationToken = default)
+		{
+			ListObjectsV2Request request = new ListObjectsV2Request();
+			request.BucketName = _options.AwsBucketName;
+			if (_pathPrefix.Length > 0)
+			{
+				request.Prefix = _pathPrefix;
+			}
+
+			for (; ; )
+			{
+				ListObjectsV2Response response = await _client.ListObjectsV2Async(request, cancellationToken);
+				foreach (S3Object obj in response.S3Objects)
+				{
+					string path = obj.Key;
+					if (path.StartsWith(_pathPrefix, StringComparison.Ordinal))
+					{
+						yield return new ObjectKey(path.Substring(_pathPrefix.Length));
+					}
+					else
+					{
+						_logger.LogError("Unexpected object enumerated from {AwsBucketName} - expected object \"{Path}\" to start with \"{AwsBucketPath}\"", _options.AwsBucketName, path, _pathPrefix);
+					}
+				}
+
+				if (response.IsTruncated)
+				{
+					request.ContinuationToken = response.NextContinuationToken;
+				}
+				else
+				{
+					break;
+				}
+			}
+		}
+
+		/// <inheritdoc/>
+		public void GetStats(StorageStats stats) { }
+	}
+
+	/// <summary>
+	/// Factory for constructing <see cref="AwsObjectStore"/> instances
+	/// </summary>
+	public sealed class AwsObjectStoreFactory : IDisposable
+	{
+		readonly IConfiguration _configuration;
+		readonly SemaphoreSlim _semaphore;
+		readonly ILogger<AwsObjectStore> _logger;
+
+		/// <summary>
+		/// Constructor
+		/// </summary>
+		public AwsObjectStoreFactory(IConfiguration configuration, ILogger<AwsObjectStore> logger)
+		{
+			_configuration = configuration;
+			_semaphore = new SemaphoreSlim(16);
+			_logger = logger;
+		}
+
+		/// <inheritdoc/>
+		public void Dispose()
+		{
+			_semaphore.Dispose();
+		}
+
+		/// <summary>
+		/// Create a new object store with the given configuration
+		/// </summary>
+		/// <param name="options">Configuration for the store</param>
+		public AwsObjectStore CreateStore(IAwsStorageOptions options)
+		{
+			AWSOptions awsOptions = GetAwsOptions(_configuration, options);
+
+			IAmazonS3 client = awsOptions.CreateServiceClient<IAmazonS3>();
+			_logger.LogInformation("Created AWS storage backend for bucket {BucketName} using credentials {Credentials} {CredentialsStr}", options.AwsBucketName, awsOptions.Credentials.GetType(), awsOptions.Credentials.ToString());
+
+			return new AwsObjectStore(client, options, _semaphore, _logger);
+		}
+
+		/// <summary>
+		/// Gets the AWS options 
+		/// </summary>
+		/// <param name="configuration">Global configuration object</param>
+		/// <param name="options">AWS storage options</param>
+		/// <returns></returns>
+		static AWSOptions GetAwsOptions(IConfiguration configuration, IAwsStorageOptions options)
+		{
+			AWSOptions awsOptions = configuration.GetAWSOptions();
+			if (options.AwsRegion != null)
+			{
+				awsOptions.Region = RegionEndpoint.GetBySystemName(options.AwsRegion);
+			}
+
+			switch (options.AwsCredentials ?? AwsCredentialsType.Default)
+			{
+				case AwsCredentialsType.Default:
+					// Using the fallback credentials from the AWS SDK, it will pick up credentials through a number of default mechanisms.
+					awsOptions.Credentials = FallbackCredentialsFactory.GetCredentials();
+					break;
+				case AwsCredentialsType.Profile:
+					if (options.AwsProfile == null)
+					{
+						throw new AwsException($"Missing {nameof(IAwsStorageOptions.AwsProfile)} setting for configuring {nameof(AwsObjectStore)}", null);
+					}
+
+					(string accessKey, string secretAccessKey, string secretToken) = AwsHelper.ReadAwsCredentials(options.AwsProfile);
+					awsOptions.Credentials = new Amazon.SecurityToken.Model.Credentials(accessKey, secretAccessKey, secretToken, DateTime.Now + TimeSpan.FromHours(12));
+					break;
+				case AwsCredentialsType.AssumeRole:
+					if (options.AwsRole == null)
+					{
+						throw new AwsException($"Missing {nameof(IAwsStorageOptions.AwsRole)} setting for configuring {nameof(AwsObjectStore)}", null);
+					}
+					awsOptions.Credentials = new AssumeRoleAWSCredentials(FallbackCredentialsFactory.GetCredentials(), options.AwsRole, "Horde");
+					break;
+				case AwsCredentialsType.AssumeRoleWebIdentity:
+					awsOptions.Credentials = AssumeRoleWithWebIdentityCredentials.FromEnvironmentVariables();
+					break;
+			}
+			return awsOptions;
+		}
+	}
+}
diff --git a/Engine/Source/Programs/Horde/Plugins/Storage/HordeServer.Storage/Storage/ObjectStores/AzureObjectStore.cs b/Engine/Source/Programs/Horde/Plugins/Storage/HordeServer.Storage/Storage/ObjectStores/AzureObjectStore.cs
new file mode 100644
--- /dev/null
+++ b/Engine/Source/Programs/Horde/Plugins/Storage/HordeServer.Storage/Storage/ObjectStores/AzureObjectStore.cs
@@ -0,0 +1,212 @@
+// Copyright Epic Games, Inc. All Rights Reserved.
+
+using System;
+using System.Collections.Generic;
+using System.IO;
+using System.Runtime.CompilerServices;
+using System.Threading;
+using System.Threading.Tasks;
+using Azure;
+using Azure.Storage.Blobs;
+using Azure.Storage.Blobs.Models;
+using Azure.Storage.Sas;
+using EpicGames.Core;
+using EpicGames.Horde.Storage;
+using Microsoft.Extensions.Logging;
+using OpenTelemetry.Trace;
+
+namespace HordeServer.Storage.ObjectStores
+{
+	/// <summary>
+	/// Exception wrapper for S3 requests
+	/// </summary>
+	public sealed class AzureException : Exception
+	{
+		/// <summary>
+		/// Constructor
+		/// </summary>
+		/// <param name="message">Message for the exception</param>
+		/// <param name="innerException">Inner exception data</param>
+		public AzureException(string? message, Exception? innerException)
+			: base(message, innerException)
+		{
+		}
+	}
+
+	/// <summary>
+	/// Options for Azure
+	/// </summary>
+	public interface IAzureStorageOptions
+	{
+		/// <summary>
+		/// Connection string for Azure
+		/// </summary>
+		public string? AzureConnectionString { get; }
+
+		/// <summary>
+		/// Name of the container
+		/// </summary>
+		public string? AzureContainerName { get; }
+	}
+
+	/// <summary>
+	/// Storage backend using Azure
+	/// </summary>
+	public sealed class AzureObjectStore : IObjectStore
+	{
+		private readonly ILogger _logger;
+		private readonly Tracer _tracer;
+		private readonly BlobContainerClient _blobContainer;
+
+		/// <inheritdoc/>
+		public bool SupportsRedirects => true;
+
+		/// <summary>
+		/// Constructor
+		/// </summary>
+		public AzureObjectStore(IAzureStorageOptions options, ILogger<AzureObjectStore> logger, Tracer tracer)
+		{
+			if (options.AzureConnectionString == null)
+			{
+				throw new AzureException($"Missing {nameof(IAzureStorageOptions.AzureConnectionString)} setting for Azure storage backend", null);
+			}
+			if (options.AzureContainerName == null)
+			{
+				throw new AzureException($"Missing {nameof(IAzureStorageOptions.AzureContainerName)} setting for Azure storage backend", null);
+			}
+
+			_logger = logger;
+			_tracer = tracer;
+			_blobContainer = new BlobContainerClient(options.AzureConnectionString, options.AzureContainerName);
+		}
+
+		/// <inheritdoc/>
+		public async Task<Stream> OpenAsync(ObjectKey key, int offset, int? length, CancellationToken cancellationToken = default)
+		{
+			if (!await _blobContainer.ExistsAsync(cancellationToken))
+			{
+				throw new InvalidOperationException($"Container {_blobContainer.Name} did not exist");
+			}
+
+			try
+			{
+				BlobClient blob = _blobContainer.GetBlobClient(key.ToString());
+				Response<BlobDownloadInfo> blobInfo = await blob.DownloadAsync(cancellationToken);
+				return blobInfo.Value.Content;
+			}
+			catch (RequestFailedException ex)
+			{
+				throw new AzureException($"Unable to read {key} from Azure", ex);
+			}
+		}
+
+		/// <inheritdoc/>
+		public async Task<IReadOnlyMemoryOwner<byte>> ReadAsync(ObjectKey key, int offset, int? length, CancellationToken cancellationToken = default)
+		{
+			using Stream stream = await OpenAsync(key, offset, length, cancellationToken);
+			return ReadOnlyMemoryOwner.Create(await stream.ReadAllBytesAsync(cancellationToken));
+		}
+
+		/// <inheritdoc/>
+		public async Task WriteAsync(ObjectKey key, Stream stream, CancellationToken cancellationToken = default)
+		{
+			await _blobContainer.CreateIfNotExistsAsync(cancellationToken: cancellationToken);
+
+			_logger.LogDebug("Fetching blob reference with name {ObjectName}", key);
+			try
+			{
+				await _blobContainer.GetBlobClient(key.ToString()).UploadAsync(stream, cancellationToken);
+			}
+			catch (Exception ex)
+			{
+				throw new AzureException("Unable to upload blob to Azure", ex);
+			}
+		}
+
+		/// <inheritdoc/>
+		public async Task<bool> ExistsAsync(ObjectKey key, CancellationToken cancellationToken)
+		{
+			if (!await _blobContainer.ExistsAsync(cancellationToken))
+			{
+				return false;
+			}
+
+			BlobClient blob = _blobContainer.GetBlobClient(key.ToString());
+			return await blob.ExistsAsync(cancellationToken);
+		}
+
+		/// <inheritdoc/>
+		public async Task DeleteAsync(ObjectKey key, CancellationToken cancellationToken)
+		{
+			if (!await _blobContainer.ExistsAsync(cancellationToken))
+			{
+				throw new InvalidOperationException($"Container {_blobContainer.Name} does not exist");
+			}
+
+			await _blobContainer.DeleteBlobAsync(key.ToString(), cancellationToken: cancellationToken);
+		}
+
+		/// <inheritdoc/>
+		public async IAsyncEnumerable<ObjectKey> EnumerateAsync([EnumeratorCancellation] CancellationToken cancellationToken = default)
+		{
+			bool exists = await _blobContainer.ExistsAsync(cancellationToken);
+			if (exists)
+			{
+				await foreach (BlobItem? item in _blobContainer.GetBlobsAsync(BlobTraits.Metadata, cancellationToken: cancellationToken))
+				{
+					yield return new ObjectKey(item.Name);
+				}
+			}
+		}
+
+		/// <inheritdoc/>
+		public async ValueTask<Uri?> TryGetReadRedirectAsync(ObjectKey key, CancellationToken cancellationToken = default)
+		{
+			if (!await _blobContainer.ExistsAsync(cancellationToken))
+			{
+				throw new InvalidOperationException($"Container {_blobContainer.Name} does not exist");
+			}
+
+			return TryGetPresignedUrl(key, BlobSasPermissions.Read);
+		}
+
+		/// <inheritdoc/>
+		public async ValueTask<Uri?> TryGetWriteRedirectAsync(ObjectKey key, CancellationToken cancellationToken = default)
+		{
+			if (!await _blobContainer.ExistsAsync(cancellationToken))
+			{
+				throw new InvalidOperationException($"Container {_blobContainer.Name} does not exist");
+			}
+
+			return TryGetPresignedUrl(key, BlobSasPermissions.Write);
+		}
+
+		/// <summary>
+		/// Helper method to generate a presigned URL for a request
+		/// </summary>
+		Uri? TryGetPresignedUrl(ObjectKey key, BlobSasPermissions permissions)
+		{
+			using TelemetrySpan span = _tracer.StartActiveSpan("azure.BuildPresignedUrl").SetAttribute("Path", key.ToString());
+
+			try
+			{
+				BlobClient blob = _blobContainer.GetBlobClient(key.ToString());
+				return blob.GenerateSasUri(permissions, DateTimeOffset.Now.AddHours(1.0));
+			}
+			catch (RequestFailedException e)
+			{
+				if (e.Status == 404)
+				{
+					return null;
+				}
+
+				throw;
+			}
+		}
+
+		/// <inheritdoc/>
+		public void GetStats(StorageStats stats)
+		{
+		}
+	}
+}
diff --git a/Engine/Source/Programs/Horde/Plugins/Storage/HordeServer.Storage/Storage/StorageAclAction.cs b/Engine/Source/Programs/Horde/Plugins/Storage/HordeServer.Storage/Storage/StorageAclAction.cs
new file mode 100644
--- /dev/null
+++ b/Engine/Source/Programs/Horde/Plugins/Storage/HordeServer.Storage/Storage/StorageAclAction.cs
@@ -0,0 +1,37 @@
+// Copyright Epic Games, Inc. All Rights Reserved.
+
+using EpicGames.Horde.Acls;
+
+namespace HordeServer.Storage
+{
+	/// <summary>
+	/// ACL actions which apply to storage
+	/// </summary>
+	public static class StorageAclAction
+	{
+		/// <summary>
+		/// Ability to read blobs from the storage service
+		/// </summary>
+		public static AclAction ReadBlobs { get; } = new AclAction("ReadBlobs");
+
+		/// <summary>
+		/// Ability to write blobs to the storage service
+		/// </summary>
+		public static AclAction WriteBlobs { get; } = new AclAction("WriteBlobs");
+
+		/// <summary>
+		/// Ability to read refs from the storage service
+		/// </summary>
+		public static AclAction ReadRefs { get; } = new AclAction("ReadRefs");
+
+		/// <summary>
+		/// Ability to write refs to the storage service
+		/// </summary>
+		public static AclAction WriteRefs { get; } = new AclAction("WriteRefs");
+
+		/// <summary>
+		/// Ability to delete refs
+		/// </summary>
+		public static AclAction DeleteRefs { get; } = new AclAction("DeleteRefs");
+	}
+}
diff --git a/Engine/Source/Programs/Horde/Plugins/Storage/HordeServer.Storage/Storage/StorageController.cs b/Engine/Source/Programs/Horde/Plugins/Storage/HordeServer.Storage/Storage/StorageController.cs
new file mode 100644
--- /dev/null
+++ b/Engine/Source/Programs/Horde/Plugins/Storage/HordeServer.Storage/Storage/StorageController.cs
@@ -0,0 +1,566 @@
+// Copyright Epic Games, Inc. All Rights Reserved.
+
+using System;
+using System.Collections.Generic;
+using System.Diagnostics.CodeAnalysis;
+using System.IO;
+using System.Linq;
+using System.Net.Http.Headers;
+using System.Security.Claims;
+using System.Text.RegularExpressions;
+using System.Threading;
+using System.Threading.Tasks;
+using EpicGames.Core;
+using EpicGames.Horde.Acls;
+using EpicGames.Horde.Storage;
+using EpicGames.Horde.Storage.Nodes;
+using EpicGames.Serialization;
+using HordeServer.Server;
+using HordeServer.Utilities;
+using Microsoft.AspNetCore.Authorization;
+using Microsoft.AspNetCore.Http;
+using Microsoft.AspNetCore.Mvc;
+using Microsoft.Extensions.Logging;
+using Microsoft.Extensions.Options;
+
+#pragma warning disable CA2227 // Change x to be read-only by removing the property setter
+
+namespace HordeServer.Storage
+{
+	/// <summary>
+	/// Controller for the /api/v1/storage endpoint
+	/// </summary>
+	[Authorize]
+	[ApiController]
+	[Route("[controller]")]
+	public class StorageController : HordeControllerBase
+	{
+		readonly StorageService _storageService;
+		readonly IOptionsSnapshot<StorageConfig> _storageConfig;
+
+		/// <summary>
+		/// Constructor
+		/// </summary>
+		public StorageController(StorageService storageService, IOptionsSnapshot<StorageConfig> storageConfig)
+		{
+			_storageService = storageService;
+			_storageConfig = storageConfig;
+		}
+
+		bool Authorize(NamespaceId namespaceId, AclAction action)
+		{
+			return _storageConfig.Value.TryGetNamespace(namespaceId, out NamespaceConfig? namespaceConfig) && namespaceConfig.Authorize(action, User);
+		}
+
+		/// <summary>
+		/// Uploads data to the storage service. 
+		/// </summary>
+		/// <param name="namespaceId">Namespace to fetch from</param>
+		/// <param name="request">Information about the blob to write</param>
+		/// <param name="cancellationToken">Cancellation token for the operation</param>
+		[HttpPost]
+		[Route("/api/v1/storage/{namespaceId}/blobs")]
+		public async Task<ActionResult<WriteBlobResponse>> WriteBlobAsync(NamespaceId namespaceId, WriteBlobRequest request, CancellationToken cancellationToken = default)
+		{
+			IStorageBackend? storageBackend = _storageService.TryCreateBackend(namespaceId);
+			if (storageBackend == null)
+			{
+				return NotFound(namespaceId);
+			}
+			if (!Authorize(namespaceId, StorageAclAction.WriteBlobs) && !HasPathClaim(User, HordeClaimTypes.WriteNamespace, namespaceId, request.Prefix ?? String.Empty))
+			{
+				return Forbid(StorageAclAction.WriteBlobs, namespaceId);
+			}
+
+			return await WriteBlobAsync(storageBackend, request, cancellationToken);
+		}
+
+		/// <summary>
+		/// Writes a blob to storage. Exposed as a public utility method to allow other routes with their own authentication methods to wrap their own authentication/redirection.
+		/// </summary>
+		/// <param name="storageBackend">The backend to write to</param>
+		/// <param name="request">Information about the blob to write</param>
+		/// <param name="cancellationToken">Cancellation token</param>
+		/// <returns>Information about the written blob, or redirect information</returns>
+		public static async Task<ActionResult<WriteBlobResponse>> WriteBlobAsync(IStorageBackend storageBackend, WriteBlobRequest request, CancellationToken cancellationToken = default)
+		{
+			if (request.File == null)
+			{
+				(BlobLocator Path, Uri UploadUrl)? result = await storageBackend.TryGetBlobWriteRedirectAsync(request.Imports, request.Prefix ?? String.Empty, cancellationToken);
+				if (result == null)
+				{
+					return new WriteBlobResponse { SupportsRedirects = false };
+				}
+
+				return new WriteBlobResponse { Blob = result.Value.Path.ToString(), UploadUrl = result.Value.UploadUrl };
+			}
+			else
+			{
+				using Stream stream = request.File.OpenReadStream();
+				BlobLocator locator = await storageBackend.WriteBlobAsync(stream, request.Imports, request.Prefix, cancellationToken);
+				return new WriteBlobResponse { Blob = locator.ToString(), SupportsRedirects = storageBackend.SupportsRedirects };
+			}
+		}
+
+		/// <summary>
+		/// Retrieves data from the storage service. 
+		/// </summary>
+		/// <param name="namespaceId">Namespace to fetch from</param>
+		/// <param name="locator">Bundle to retrieve</param>
+		/// <param name="cancellationToken">Cancellation token for the operation</param>
+		[HttpGet]
+		[Route("/api/v1/storage/{namespaceId}/blobs/{*locator}")]
+		public async Task<ActionResult> ReadBlobAsync(NamespaceId namespaceId, BlobLocator locator, CancellationToken cancellationToken = default)
+		{
+			IStorageBackend? backend = _storageService.TryCreateBackend(namespaceId);
+			if (backend == null)
+			{
+				return NotFound(namespaceId);
+			}
+			if (!Authorize(namespaceId, StorageAclAction.ReadBlobs) && !HasPathClaim(User, HordeClaimTypes.ReadNamespace, namespaceId, locator.Path.ToString()))
+			{
+				return Forbid(StorageAclAction.ReadBlobs, namespaceId);
+			}
+
+			return await ReadBlobInternalAsync(backend, locator, Request.Headers, cancellationToken);
+		}
+
+		/// <summary>
+		/// Reads a blob from storage, without performing namespace access checks.
+		/// </summary>
+		public static async Task<ActionResult> ReadBlobInternalAsync(IStorageBackend storageBackend, BlobLocator locator, IHeaderDictionary headers, CancellationToken cancellationToken)
+		{
+			Uri? redirectUrl = await storageBackend.TryGetBlobReadRedirectAsync(locator, cancellationToken);
+			if (redirectUrl != null)
+			{
+				return new RedirectResult(redirectUrl.ToString());
+			}
+
+			// Parse the range header
+			int offset = 0;
+			int? length = null;
+
+			if (headers.Range.Count > 0)
+			{
+				if (headers.Range.Count > 1)
+				{
+					return new BadRequestObjectResult(LogEvent.Create(LogLevel.Error, "Unsupported range header; only one range is allowed"));
+				}
+
+				string? value = headers.Range[0];
+				if (value == null)
+				{
+					return new BadRequestObjectResult(LogEvent.Create(LogLevel.Error, "Unsupported range header; only one range is allowed"));
+				}
+
+				Match match = Regex.Match(value, @"^\s*bytes\s*=\s*(\d*)-(\d*)$");
+				if (!match.Success)
+				{
+					return new BadRequestObjectResult(LogEvent.Create(LogLevel.Error, "Unsupported range header syntax; cannot parse {Value}", value));
+				}
+
+				if (match.Groups[1].Length > 0 && !Int32.TryParse(match.Groups[1].Value, out offset))
+				{
+					return new BadRequestObjectResult(LogEvent.Create(LogLevel.Error, "Unable to parse start for range: {Value}", value));
+				}
+				if (match.Groups[2].Length > 0)
+				{
+					int end;
+					if (Int32.TryParse(match.Groups[2].Value, out end) && end > offset)
+					{
+						length = (end + 1) - offset;
+					}
+					else
+					{
+						return new BadRequestObjectResult(LogEvent.Create(LogLevel.Error, "Unable to parse end for range: {Value}", value));
+					}
+				}
+			}
+
+#pragma warning disable CA2000 // Dispose objects before losing scope
+			Stream stream = await storageBackend.OpenBlobAsync(locator, offset, length, cancellationToken);
+			return new FileStreamResult(stream, "application/octet-stream");
+#pragma warning restore CA2000 // Dispose objects before losing scope
+		}
+
+		/// <summary>
+		/// Retrieves data from the storage service. 
+		/// </summary>
+		/// <param name="namespaceId">Namespace to fetch from</param>
+		/// <param name="alias">Alias of the node to find</param>
+		/// <param name="maxResults">Maximum number of results to return</param>
+		/// <param name="cancellationToken">Cancellation token for the operation</param>
+		[HttpGet]
+		[Route("/api/v1/storage/{namespaceId}/nodes")]
+		public async Task<ActionResult<FindNodesResponse>> FindNodesAsync(NamespaceId namespaceId, [FromQuery] string alias, [FromQuery] int? maxResults = null, CancellationToken cancellationToken = default)
+		{
+			IStorageBackend? backend = _storageService.TryCreateBackend(namespaceId);
+			if (backend == null)
+			{
+				return NotFound(namespaceId);
+			}
+			if (!Authorize(namespaceId, StorageAclAction.ReadBlobs))
+			{
+				return Forbid(StorageAclAction.ReadBlobs, namespaceId);
+			}
+
+			BlobAliasLocator[] aliases = await backend.FindAliasesAsync(alias, maxResults, cancellationToken);
+
+			FindNodesResponse response = new FindNodesResponse();
+			response.Nodes.AddRange(aliases.Select(x => new FindNodeResponse(x.Target, x.Rank, x.Data.ToArray())));
+
+			if (response.Nodes.Count == 0)
+			{
+				return NotFound();
+			}
+
+			return response;
+		}
+
+		/// <summary>
+		/// Writes a ref to the storage service.
+		/// </summary>
+		/// <param name="namespaceId">Namespace to write to</param>
+		/// <param name="refName">Name of the ref</param>
+		/// <param name="request">Request for the ref to write</param>
+		/// <param name="cancellationToken">Cancellation token for the operation</param>
+		[HttpPut]
+		[Route("/api/v1/storage/{namespaceId}/refs/{*refName}")]
+		public async Task<ActionResult> WriteRefAsync(NamespaceId namespaceId, RefName refName, [FromBody] WriteRefRequest request, CancellationToken cancellationToken)
+		{
+			IStorageBackend? backend = _storageService.TryCreateBackend(namespaceId);
+			if (backend == null)
+			{
+				return NotFound(namespaceId);
+			}
+			if (!Authorize(namespaceId, StorageAclAction.WriteRefs) && !HasPathClaim(User, HordeClaimTypes.WriteNamespace, namespaceId, refName.ToString()))
+			{
+				return Forbid(StorageAclAction.WriteRefs, namespaceId);
+			}
+
+#pragma warning disable CS0618 // Type or member is obsolete
+			if (request.Blob != null && request.ExportIdx != null)
+			{
+				request.Target = new BlobLocator($"{request.Blob.Value}#{request.ExportIdx.Value}");
+			}
+#pragma warning restore CS0618 // Type or member is obsolete
+
+			await backend.WriteRefAsync(refName, new BlobRefValue(request.Hash, request.Target), request.Options, cancellationToken);
+			return Ok();
+		}
+
+		/// <summary>
+		/// Retrieves a ref from the storage service. 
+		/// </summary>
+		/// <param name="namespaceId"></param>
+		/// <param name="refName"></param>
+		/// <param name="cancellationToken"></param>
+		[HttpGet]
+		[Route("/api/v1/storage/{namespaceId}/refs/{*refName}")]
+		public async Task<ActionResult<ReadRefResponse>> ReadRefAsync(NamespaceId namespaceId, RefName refName, CancellationToken cancellationToken)
+		{
+			NamespaceConfig? namespaceConfig;
+			if (!_storageConfig.Value.TryGetNamespace(namespaceId, out namespaceConfig))
+			{
+				return NotFound(namespaceId);
+			}
+			if (!namespaceConfig.Authorize(StorageAclAction.ReadRefs, User) && !HasPathClaim(User, HordeClaimTypes.ReadNamespace, namespaceId, refName.ToString()))
+			{
+				return Forbid(StorageAclAction.ReadRefs, namespaceId);
+			}
+
+			return await ReadRefInternalAsync(_storageService, namespaceId, refName, Request.Headers, cancellationToken);
+		}
+
+		/// <summary>
+		/// Reads a ref from storage, without performing namespace access checks.
+		/// </summary>
+		public static async Task<ActionResult<ReadRefResponse>> ReadRefInternalAsync(IStorageClientFactory storageService, NamespaceId namespaceId, RefName refName, IHeaderDictionary headers, CancellationToken cancellationToken)
+		{
+			using IStorageClient client = storageService.CreateClient(namespaceId);
+
+			RefCacheTime cacheTime = new RefCacheTime();
+			foreach (string? entry in headers.CacheControl)
+			{
+				if (entry != null && CacheControlHeaderValue.TryParse(entry, out CacheControlHeaderValue? value) && value?.MaxAge != null)
+				{
+					cacheTime = new RefCacheTime(value.MaxAge.Value);
+				}
+			}
+
+			IBlobRef? target = await client.TryReadRefAsync(refName, cacheTime, cancellationToken: cancellationToken);
+			if (target == null)
+			{
+				return new NotFoundResult();
+			}
+
+			return new ReadRefResponse { Hash = target.Hash, Target = target.GetLocator(), Link = GetNodeLink(namespaceId, target) };
+		}
+
+		/// <summary>
+		/// Checks whether the user has an explicit claim to read or write to a path within a namespace
+		/// </summary>
+		/// <param name="user">User to query</param>
+		/// <param name="claimType">The claim name</param>
+		/// <param name="namespaceId">Namespace id to check for</param>
+		/// <param name="entity">Path to the entity to query</param>
+		/// <returns>True if the user is authorized for access to the given path</returns>
+		static bool HasPathClaim(ClaimsPrincipal user, string claimType, NamespaceId namespaceId, string entity)
+		{
+			foreach (Claim claim in user.Claims)
+			{
+				if (claim.Type.Equals(claimType, StringComparison.Ordinal))
+				{
+					int colonIdx = claim.Value.IndexOf(':', StringComparison.Ordinal);
+					if (colonIdx == -1)
+					{
+						if (namespaceId.Text.Equals(claim.Value))
+						{
+							return true;
+						}
+					}
+					else
+					{
+						if (namespaceId.Text.Equals(claim.Value.AsMemory(0, colonIdx)) && HasPathPrefix(entity, claim.Value.AsSpan(colonIdx + 1)))
+						{
+							return true;
+						}
+					}
+				}
+			}
+			return false;
+		}
+
+		static bool HasPathPrefix(string name, ReadOnlySpan<char> prefix)
+		{
+			if (name.Length > prefix.Length)
+			{
+				return name[prefix.Length] == '/' && name.AsSpan(0, prefix.Length).SequenceEqual(prefix);
+			}
+			else if (name.Length == prefix.Length)
+			{
+				return name.AsSpan().SequenceEqual(prefix);
+			}
+			else
+			{
+				return false;
+			}
+		}
+
+		static readonly IReadOnlyDictionary<Guid, Type> s_blobGuidToType = GetBlobGuidTypeMap();
+
+		static Dictionary<Guid, Type> GetBlobGuidTypeMap()
+		{
+			Dictionary<Guid, Type> guidTypeMap = new Dictionary<Guid, Type>();
+			guidTypeMap.Add(CbNode.BlobTypeGuid, typeof(CbNode));
+			guidTypeMap.Add(LeafChunkedDataNode.BlobTypeGuid, typeof(LeafChunkedDataNode));
+			guidTypeMap.Add(InteriorChunkedDataNode.BlobTypeGuid, typeof(InteriorChunkedDataNode));
+			guidTypeMap.Add(CommitNode.BlobTypeGuid, typeof(CommitNode));
+			guidTypeMap.Add(DdcRefNode.BlobTypeGuid, typeof(DdcRefNode));
+			guidTypeMap.Add(DirectoryNode.BlobTypeGuid, typeof(DirectoryNode));
+			guidTypeMap.Add(RedirectNode.BlobTypeGuid, typeof(RedirectNode));
+			return guidTypeMap;
+		}
+
+		/// <summary>
+		/// Gets information about a particular bundle in storage
+		/// </summary>
+		/// <param name="namespaceId">Namespace containing the blob</param>
+		/// <param name="locator">Blob identifier</param>
+		/// <param name="pkt">Packet string</param>
+		/// <param name="exp">Export index</param>
+		/// <param name="data">Whether to download the blob data</param>
+		/// <param name="cancellationToken">Cancellation token for the operation</param>
+		/// <returns></returns>
+		[HttpGet]
+		[Route("/api/v1/storage/{namespaceId}/nodes/{*locator}")]
+		public async Task<ActionResult<object>> GetNodeAsync(NamespaceId namespaceId, BlobLocator locator, [FromQuery] string? pkt = null, [FromQuery] string? exp = null, [FromQuery] bool data = false, CancellationToken cancellationToken = default)
+		{
+			NamespaceConfig? namespaceConfig;
+			if (!_storageConfig.Value.TryGetNamespace(namespaceId, out namespaceConfig))
+			{
+				return NotFound(namespaceId);
+			}
+			if (!namespaceConfig.Authorize(StorageAclAction.ReadBlobs, User))
+			{
+				return Forbid(StorageAclAction.ReadBlobs, namespaceId);
+			}
+
+			List<string> fragments = new List<string>();
+			if (pkt != null)
+			{
+				fragments.Add($"pkt={pkt}");
+			}
+			if (exp != null)
+			{
+				fragments.Add($"exp={exp}");
+			}
+			if (fragments.Count > 0)
+			{
+				locator = new BlobLocator(locator, String.Join("&", fragments));
+			}
+
+			using IStorageClient storageClient = _storageService.CreateClient(namespaceId);
+
+			object content;
+
+			using BlobData blobData = await storageClient.CreateBlobHandle(locator).ReadBlobDataAsync(cancellationToken);
+			if (data)
+			{
+				ReadOnlyMemoryStream stream = new ReadOnlyMemoryStream(blobData.Data.ToArray());
+				return new FileStreamResult(stream, "application/octet-stream");
+			}
+
+			if (blobData.Type.Guid == DirectoryNode.BlobTypeGuid)
+			{
+				DirectoryNode directoryNode = BlobSerializer.Deserialize<DirectoryNode>(blobData);
+
+				List<object> directories = new List<object>();
+				foreach ((string name, DirectoryEntry entry) in directoryNode.NameToDirectory)
+				{
+					directories.Add(new { name = name.ToString(), length = entry.Length, target = GetNodeHandleLink(namespaceId, entry.Handle) });
+				}
+
+				List<object> files = new List<object>();
+				foreach ((string name, FileEntry entry) in directoryNode.NameToFile)
+				{
+					files.Add(new { name = name.ToString(), length = entry.Length, flags = entry.Flags, hash = entry.StreamHash, target = GetNodeHandleLink(namespaceId, entry.Target.Handle) });
+				}
+
+				content = new { directoryNode.Length, directories, files };
+			}
+			else if (blobData.Type.Guid == InteriorChunkedDataNode.BlobTypeGuid)
+			{
+				InteriorChunkedDataNode interiorNode = BlobSerializer.Deserialize<InteriorChunkedDataNode>(blobData);
+
+				List<object> children = new List<object>();
+				foreach (ChunkedDataNodeRef nodeRef in interiorNode.Children)
+				{
+					children.Add(new { nodeRef.Type, nodeRef.Length, hash = nodeRef.Handle.Hash, link = GetNodeLink(namespaceId, nodeRef.Handle) });
+				}
+
+				content = new { children };
+			}
+			else if (blobData.Type.Guid == CommitNode.BlobTypeGuid)
+			{
+				CommitNode commitNode = BlobSerializer.Deserialize<CommitNode>(blobData);
+
+				Dictionary<Guid, object>? metadata = null;
+				if (commitNode.Metadata.Count > 0)
+				{
+					metadata = new Dictionary<Guid, object>();
+					foreach ((Guid blobGuid, IBlobRef handle) in commitNode.Metadata)
+					{
+						metadata.Add(blobGuid, GetNodeHandleLink(namespaceId, handle));
+					}
+				}
+
+				content = new { commitNode.Number, parent = GetNodeHandleLink(namespaceId, commitNode.Parent), commitNode.Author, commitNode.AuthorId, commitNode.Committer, commitNode.CommitterId, commitNode.Message, commitNode.Time, contents = GetNodeObject(namespaceId, commitNode.Contents), metadata };
+			}
+			else if (blobData.Type.Guid == CbNode.BlobTypeGuid)
+			{
+				CbNode cbNode = BlobSerializer.Deserialize<CbNode>(blobData);
+				content = GetCbNodeObject(namespaceId, cbNode.Object.AsField(), cbNode.Imports.GetEnumerator()) ?? new object();
+			}
+			else
+			{
+				IEnumerable<string>? references = null;
+				if (blobData.Imports.Count > 0)
+				{
+					references = blobData.Imports.Select(x => GetNodeLink(namespaceId, x));
+				}
+				content = new { length = blobData.Data.Length, references };
+			}
+
+			string? typeName = null;
+			if (s_blobGuidToType.TryGetValue(blobData.Type.Guid, out Type? type))
+			{
+				typeName = type.Name;
+			}
+
+			return new { type = typeName, guid = blobData.Type.Guid, data = $"{GetNodeLink(namespaceId, locator)}&data=true", content = content };
+		}
+
+		static object? GetCbNodeObject(NamespaceId namespaceId, CbField field, IEnumerator<IBlobHandle> imports)
+		{
+			if (field.IsAttachment())
+			{
+				object? link = GetNodeLink(namespaceId, imports.Current);
+				imports.MoveNext();
+				return link;
+			}
+			else if (field.IsObject())
+			{
+				Dictionary<string, object?> fields = new Dictionary<string, object?>();
+
+				CbObject obj = field.AsObject();
+				foreach (CbField member in obj)
+				{
+					fields[member.Name.ToString()] = GetCbNodeObject(namespaceId, member, imports);
+				}
+
+				return fields;
+			}
+			else if (field.IsArray())
+			{
+				List<object?> elements = new List<object?>();
+
+				CbArray arr = field.AsArray();
+				foreach (CbField member in arr)
+				{
+					elements.Add(GetCbNodeObject(namespaceId, member, imports));
+				}
+
+				return elements;
+			}
+			else
+			{
+				return field.Value;
+			}
+		}
+
+		[return: NotNullIfNotNull("nodeRef")]
+		static object? GetNodeObject(NamespaceId namespaceId, DirectoryNodeRef? nodeRef) => (nodeRef == null) ? null : new { nodeRef.Length, nodeRef.Handle.Hash, link = GetNodeLink(namespaceId, nodeRef.Handle.GetLocator()) };
+
+		[return: NotNullIfNotNull("handle")]
+		static object? GetNodeHandleLink(NamespaceId namespaceId, IBlobRef? handle) => (handle == null) ? null : new { handle.Hash, link = GetNodeLink(namespaceId, handle.GetLocator()) };
+
+		static string GetNodeLink(NamespaceId namespaceId, IBlobHandle handle) => GetNodeLink(namespaceId, handle.GetLocator());
+
+		static string GetNodeLink(NamespaceId namespaceId, BlobLocator locator) => $"/api/v1/storage/{namespaceId}/nodes/{locator.BaseLocator}?{locator.Fragment}";
+	}
+
+	/// <summary>
+	/// Request to upload a blob
+	/// </summary>
+	public class WriteBlobRequest
+	{
+		/// <summary>
+		/// Set to indicate that the blob does not have any imports. Used to distinguish from an "unknown" set of imports.
+		/// </summary>
+		[FromForm(Name = "leaf")]
+		public bool Leaf
+		{
+			get => Imports != null && Imports.Count == 0;
+			set => Imports = (value ? new List<BlobLocator>() : Imports);
+		}
+
+		/// <summary>
+		/// Imported blobs
+		/// </summary>
+		[FromForm(Name = "import")]
+		public List<BlobLocator>? Imports { get; set; }
+
+		/// <summary>
+		/// Prefix for the uploaded file
+		/// </summary>
+		[FromForm(Name = "prefix")]
+		public string? Prefix { get; set; }
+
+		/// <summary>
+		/// Data to be uploaded. May be null, in which case the server may return a separate url.
+		/// </summary>
+		[FromForm(Name = "file")]
+		public IFormFile? File { get; set; }
+	}
+}
diff --git a/Engine/Source/Programs/Horde/Plugins/Storage/HordeServer.Storage/Storage/StorageService.cs b/Engine/Source/Programs/Horde/Plugins/Storage/HordeServer.Storage/Storage/StorageService.cs
new file mode 100644
--- /dev/null
+++ b/Engine/Source/Programs/Horde/Plugins/Storage/HordeServer.Storage/Storage/StorageService.cs
@@ -0,0 +1,1200 @@
+// Copyright Epic Games, Inc. All Rights Reserved.
+
+using System;
+using System.Buffers;
+using System.Collections.Generic;
+using System.ComponentModel;
+using System.Diagnostics;
+using System.IO;
+using System.Linq;
+using System.Threading;
+using System.Threading.Tasks;
+using EpicGames.Core;
+using EpicGames.Horde;
+using EpicGames.Horde.Storage;
+using EpicGames.Horde.Storage.Bundles;
+using EpicGames.Horde.Storage.ObjectStores;
+using EpicGames.Redis;
+using EpicGames.Redis.Utility;
+using HordeServer.Server;
+using HordeServer.Utilities;
+using Microsoft.Extensions.Caching.Memory;
+using Microsoft.Extensions.Configuration;
+using Microsoft.Extensions.DependencyInjection;
+using Microsoft.Extensions.Hosting;
+using Microsoft.Extensions.Logging;
+using Microsoft.Extensions.Options;
+using MongoDB.Bson;
+using MongoDB.Bson.Serialization.Attributes;
+using MongoDB.Driver;
+using OpenTelemetry.Trace;
+using StackExchange.Redis;
+
+namespace HordeServer.Storage
+{
+	/// <summary>
+	/// Interface for the storage service
+	/// </summary>
+	public interface IStorageService : IStorageClientFactory
+	{
+	}
+
+	/// <summary>
+	/// Functionality related to the storage service
+	/// </summary>
+	public sealed class StorageService : IHostedService, IStorageService, IAsyncDisposable
+	{
+		class RefCount
+		{
+			int _value;
+
+			public RefCount() => _value = 1;
+			public void AddRef() => Interlocked.Increment(ref _value);
+			public int Release() => Interlocked.Decrement(ref _value);
+		}
+
+		sealed class StorageBackendImpl : IStorageBackend
+		{
+			readonly StorageService _outer;
+			readonly IObjectStore _store;
+			readonly NamespaceConfig _config;
+
+			public NamespaceConfig Config => _config;
+			public NamespaceId NamespaceId => _config.Id;
+
+			/// <inheritdoc/>
+			public bool SupportsRedirects { get; }
+
+			public StorageBackendImpl(StorageService outer, NamespaceConfig config, IObjectStore store)
+			{
+				_outer = outer;
+				_store = store;
+				_config = config;
+
+				SupportsRedirects = store.SupportsRedirects && !config.EnableAliases;
+			}
+
+			#region Blobs
+
+			public async Task<Stream> OpenBlobAsync(BlobLocator locator, int offset, int? length, CancellationToken cancellationToken = default)
+			{
+				await _outer.CheckBlobExistsAsync(NamespaceId, locator, cancellationToken);
+				return await _store.OpenAsync(GetObjectKey(locator), offset, length, cancellationToken);
+			}
+
+			/// <inheritdoc/>
+			public async Task<IReadOnlyMemoryOwner<byte>> ReadBlobAsync(BlobLocator locator, int offset, int? length, CancellationToken cancellationToken = default)
+			{
+				await _outer.CheckBlobExistsAsync(NamespaceId, locator, cancellationToken);
+				return await _store.ReadAsync(GetObjectKey(locator), offset, length, cancellationToken);
+			}
+
+			/// <inheritdoc/>
+			public async Task<BlobLocator> WriteBlobAsync(Stream stream, IReadOnlyCollection<BlobLocator>? imports, string? basePath = null, CancellationToken cancellationToken = default)
+			{
+				BlobLocator locator = StorageHelpers.CreateUniqueLocator(basePath);
+
+				await _store.WriteAsync(GetObjectKey(locator), stream, cancellationToken);
+				await _outer.AddBlobAsync(NamespaceId, locator, imports, null, cancellationToken);
+
+				return locator;
+			}
+
+			/// <inheritdoc/>
+			public async ValueTask<Uri?> TryGetBlobReadRedirectAsync(BlobLocator locator, CancellationToken cancellationToken = default)
+			{
+				await _outer.CheckBlobExistsAsync(NamespaceId, locator, cancellationToken);
+				return await _store.TryGetReadRedirectAsync(GetObjectKey(locator), cancellationToken);
+			}
+
+			/// <inheritdoc/>
+			public async ValueTask<(BlobLocator, Uri)?> TryGetBlobWriteRedirectAsync(IReadOnlyCollection<BlobLocator>? imports = null, string? prefix = null, CancellationToken cancellationToken = default)
+			{
+				if (!_store.SupportsRedirects)
+				{
+					return null;
+				}
+
+				BlobLocator locator = StorageHelpers.CreateUniqueLocator(prefix);
+
+				Uri? url = await _store.TryGetWriteRedirectAsync(GetObjectKey(locator), cancellationToken);
+				if (url == null)
+				{
+					return null;
+				}
+
+				await _outer.AddBlobAsync(NamespaceId, locator, imports, null, cancellationToken);
+
+				return (locator, url);
+			}
+
+			#endregion
+
+			#region Aliases
+
+			/// <inheritdoc/>
+			public Task AddAliasAsync(string name, BlobLocator locator, int rank = 0, ReadOnlyMemory<byte> data = default, CancellationToken cancellationToken = default)
+				=> _outer.AddAliasAsync(NamespaceId, name, locator, rank, data, cancellationToken);
+
+			/// <inheritdoc/>
+			public Task RemoveAliasAsync(string name, BlobLocator locator, CancellationToken cancellationToken = default)
+				=> _outer.RemoveAliasAsync(NamespaceId, name, locator, cancellationToken);
+
+			/// <inheritdoc/>
+			public async Task<BlobAliasLocator[]> FindAliasesAsync(string alias, int? maxResults, CancellationToken cancellationToken = default)
+			{
+				List<(BlobLocator, AliasInfo)> aliases = await _outer.FindAliasesAsync(NamespaceId, alias, cancellationToken);
+				if (maxResults != null && maxResults.Value < aliases.Count)
+				{
+					aliases.RemoveRange(maxResults.Value, aliases.Count - maxResults.Value);
+				}
+				return aliases.Select(x => new BlobAliasLocator(x.Item1, x.Item2.Rank, x.Item2.Data)).ToArray();
+			}
+
+			#endregion
+
+			#region Refs
+
+			/// <inheritdoc/>
+			public async Task<BlobRefValue?> TryReadRefAsync(RefName name, RefCacheTime cacheTime, CancellationToken cancellationToken)
+			{
+				RefInfo? refInfo = await _outer.TryReadRefAsync(NamespaceId, name, cacheTime, cancellationToken);
+				if (refInfo == null)
+				{
+					return null;
+				}
+				return new BlobRefValue(refInfo.Hash, refInfo.Target);
+			}
+
+			/// <inheritdoc/>
+			public Task WriteRefAsync(RefName name, BlobRefValue value, RefOptions? options = null, CancellationToken cancellationToken = default)
+				=> _outer.WriteRefAsync(NamespaceId, name, value, options, cancellationToken);
+
+			/// <inheritdoc/>
+			public Task<bool> DeleteRefAsync(RefName name, CancellationToken cancellationToken = default)
+				=> _outer.DeleteRefAsync(NamespaceId, name, cancellationToken);
+
+			#endregion
+
+			/// <inheritdoc/>
+			public void GetStats(StorageStats stats) => _store.GetStats(stats);
+		}
+
+		class State
+		{
+			public StorageConfig Config { get; }
+			public Dictionary<NamespaceId, NamespaceInfo> Namespaces { get; } = new Dictionary<NamespaceId, NamespaceInfo>();
+
+			public State(StorageConfig config)
+			{
+				Config = config;
+			}
+
+			public IStorageBackend? TryCreateBackend(NamespaceId namespaceId)
+			{
+				NamespaceInfo? namespaceInfo;
+				if (!Namespaces.TryGetValue(namespaceId, out namespaceInfo))
+				{
+					return null;
+				}
+				return namespaceInfo.Backend;
+			}
+		}
+
+		class NamespaceInfo
+		{
+			public NamespaceId Id => Config.Id;
+			public NamespaceConfig Config { get; }
+			public IObjectStore Store { get; }
+			public StorageBackendImpl Backend { get; }
+
+			public NamespaceInfo(NamespaceConfig config, IObjectStore store, StorageBackendImpl backend)
+			{
+				Config = config;
+				Store = store;
+				Backend = backend;
+			}
+		}
+
+		internal class AliasInfo
+		{
+			[BsonElement("nam")]
+			public string Name { get; set; }
+
+			[BsonElement("frg")]
+			public string Fragment { get; set; }
+
+			[BsonElement("rnk"), BsonIgnoreIfDefault]
+			public int Rank { get; set; }
+
+			[BsonElement("dat"), BsonIgnoreIfNull]
+			public byte[]? Data { get; set; }
+
+			[BsonConstructor]
+			public AliasInfo()
+			{
+				Name = String.Empty;
+				Fragment = String.Empty;
+			}
+
+			public AliasInfo(string name, string fragment, byte[]? data, int rank)
+			{
+				Name = name;
+				Fragment = fragment;
+				Rank = rank;
+				Data = (data == null || data.Length == 0) ? null : data;
+			}
+		}
+
+		const int CurrentGcVersion = 2;
+
+		internal class BlobInfo
+		{
+			public ObjectId Id { get; set; }
+
+			[BsonElement("ns")]
+			public NamespaceId NamespaceId { get; set; }
+
+			[BsonElement("blob")]
+			public string Path { get; set; }
+
+			[BsonElement("imp"), BsonIgnoreIfNull]
+			public List<ObjectId>? Imports { get; set; }
+
+			[BsonElement("ali"), BsonIgnoreIfNull]
+			public List<AliasInfo>? Aliases { get; set; }
+
+			[BsonElement("del"), BsonIgnoreIfDefault]
+			public int GcVersion { get; set; }
+
+			[BsonIgnore]
+			public BlobLocator Locator => new BlobLocator(Path);
+
+			public BlobInfo()
+			{
+				Path = String.Empty;
+			}
+
+			public BlobInfo(ObjectId id, NamespaceId namespaceId, BlobLocator locator)
+			{
+				Id = id;
+				NamespaceId = namespaceId;
+				Path = locator.ToString();
+			}
+		}
+
+		internal class RefInfo : ISupportInitialize
+		{
+			[BsonIgnoreIfDefault]
+			public ObjectId Id { get; set; }
+
+			[BsonElement("ns")]
+			public NamespaceId NamespaceId { get; set; }
+
+			[BsonElement("name")]
+			public RefName Name { get; set; }
+
+			[BsonElement("hash")]
+			public IoHash Hash { get; set; }
+
+			[BsonElement("tgt")]
+			public BlobLocator Target { get; set; }
+
+			[BsonElement("binf")]
+			public ObjectId TargetBlobId { get; set; }
+
+			[BsonElement("xa"), BsonIgnoreIfDefault]
+			public DateTime? ExpiresAtUtc { get; set; }
+
+			[BsonElement("xt"), BsonIgnoreIfDefault]
+			public TimeSpan? Lifetime { get; set; }
+
+#pragma warning disable IDE0051 // Remove unused private members
+			[BsonExtraElements]
+			BsonDocument? ExtraElements { get; set; }
+#pragma warning restore IDE0051 // Remove unused private members
+
+			[BsonConstructor]
+			public RefInfo()
+			{
+				Name = RefName.Empty;
+			}
+
+			public RefInfo(NamespaceId namespaceId, RefName name, IoHash hash, BlobLocator target, ObjectId targetBlobId)
+			{
+				NamespaceId = namespaceId;
+				Name = name;
+				Hash = hash;
+				Target = target;
+				TargetBlobId = targetBlobId;
+			}
+
+			public bool HasExpired(DateTime utcNow) => ExpiresAtUtc.HasValue && utcNow >= ExpiresAtUtc.Value;
+
+			public bool RequiresTouch(DateTime utcNow) => ExpiresAtUtc.HasValue && Lifetime.HasValue && utcNow >= ExpiresAtUtc.Value - new TimeSpan(Lifetime.Value.Ticks / 4);
+
+			void ISupportInitialize.BeginInit() { }
+
+			void ISupportInitialize.EndInit()
+			{
+				if (ExtraElements != null)
+				{
+					if (ExtraElements.TryGetValue("blob", out BsonValue blob) && ExtraElements.TryGetValue("idx", out BsonValue idx))
+					{
+						Target = new BlobLocator($"{blob.AsString}#{idx.AsInt32}");
+					}
+					ExtraElements = null;
+				}
+			}
+		}
+
+		[SingletonDocument("gc-state")]
+		class GcState : SingletonBase
+		{
+			public ObjectId LastImportBlobInfoId { get; set; }
+			public List<GcNamespaceState> Namespaces { get; set; } = new List<GcNamespaceState>();
+			public bool Reset { get; set; }
+
+			public GcNamespaceState FindOrAddNamespace(NamespaceId namespaceId)
+			{
+				GcNamespaceState? namespaceState = Namespaces.FirstOrDefault(x => x.Id == namespaceId);
+				if (namespaceState == null)
+				{
+					namespaceState = new GcNamespaceState { Id = namespaceId, LastTime = DateTime.UtcNow };
+					Namespaces.Add(namespaceState);
+				}
+				return namespaceState;
+			}
+		}
+
+		class GcNamespaceState
+		{
+			public NamespaceId Id { get; set; }
+			public DateTime LastTime { get; set; }
+		}
+
+		readonly IRedisService _redisService;
+		readonly IClock _clock;
+		readonly BundleCache _bundleCache;
+		readonly IMemoryCache _memoryCache;
+		readonly IObjectStoreFactory _objectStoreFactory;
+		readonly IOptionsMonitor<StorageConfig> _storageConfig;
+		readonly IOptions<StaticStorageConfig> _staticStorageConfig;
+		readonly Tracer _tracer;
+		readonly ILogger _logger;
+
+		readonly IMongoCollection<BlobInfo> _blobCollection;
+		readonly IMongoCollection<RefInfo> _refCollection;
+
+		readonly ITicker _blobTicker;
+		readonly ITicker _refTicker;
+
+		readonly SingletonDocument<GcState> _gcState;
+		readonly ITicker _gcTicker;
+
+		readonly object _lockObject = new object();
+
+		string? _lastConfigRevision;
+		State? _lastState;
+
+		internal IMongoCollection<BlobInfo> BlobCollection => _blobCollection;
+
+		/// <summary>
+		/// Constructor
+		/// </summary>
+		public StorageService(IMongoService mongoService, IRedisService redisService, IClock clock, BundleCache bundleCache, IMemoryCache memoryCache, IObjectStoreFactory objectStoreFactory, IOptionsMonitor<StorageConfig> storageConfig, IOptions<StaticStorageConfig> staticStorageConfig, Tracer tracer, ILogger<StorageService> logger)
+		{
+			_redisService = redisService;
+			_clock = clock;
+			_bundleCache = bundleCache;
+			_memoryCache = memoryCache;
+			_objectStoreFactory = objectStoreFactory;
+			_storageConfig = storageConfig;
+			_staticStorageConfig = staticStorageConfig;
+			_tracer = tracer;
+			_logger = logger;
+
+			List<MongoIndex<BlobInfo>> blobIndexes = new List<MongoIndex<BlobInfo>>();
+			blobIndexes.Add(keys => keys.Ascending(x => x.Imports));
+			blobIndexes.Add(keys => keys.Ascending(x => x.NamespaceId).Ascending(x => x.Path), unique: true);
+			blobIndexes.Add(keys => keys.Ascending(x => x.NamespaceId).Ascending($"{nameof(BlobInfo.Aliases)}.{nameof(AliasInfo.Name)}"));
+			_blobCollection = mongoService.GetCollection<BlobInfo>("Storage.Blobs", blobIndexes);
+
+			List<MongoIndex<RefInfo>> refIndexes = new List<MongoIndex<RefInfo>>();
+			refIndexes.Add(keys => keys.Ascending(x => x.NamespaceId).Ascending(x => x.Name), unique: true);
+			refIndexes.Add(keys => keys.Ascending(x => x.TargetBlobId));
+			refIndexes.Add(keys => keys.Descending(x => x.ExpiresAtUtc), sparse: true);
+			_refCollection = mongoService.GetCollection<RefInfo>("Storage.Refs", refIndexes);
+
+			_blobTicker = clock.AddSharedTicker("Storage:Blobs", TimeSpan.FromMinutes(5.0), TickBlobsAsync, _logger);
+			_refTicker = clock.AddSharedTicker("Storage:Refs", TimeSpan.FromMinutes(5.0), TickRefsAsync, _logger);
+
+			_gcState = new SingletonDocument<GcState>(mongoService);
+			_gcTicker = clock.AddTicker("Storage:GC", TimeSpan.FromMinutes(5.0), TickGcAsync, logger);
+		}
+
+		/// <inheritdoc/>
+		public async ValueTask DisposeAsync()
+		{
+			await _blobTicker.DisposeAsync();
+			await _refTicker.DisposeAsync();
+			await _gcTicker.DisposeAsync();
+		}
+
+		internal static ObjectKey GetObjectKey(BlobLocator locator) => new ObjectKey($"{locator.Path}.blob");
+
+		class StorageClientFactory : IStorageClientFactory
+		{
+			readonly StorageService _storageService;
+			readonly StorageConfig _storageConfig;
+
+			public StorageClientFactory(StorageService storageService, StorageConfig storageConfig)
+			{
+				_storageService = storageService;
+				_storageConfig = storageConfig;
+			}
+
+			public IStorageClient? TryCreateClient(NamespaceId namespaceId)
+				=> _storageService.TryCreateClient(_storageConfig, namespaceId);
+		}
+
+		/// <summary>
+		/// Creates a new storage client factory using the current global config value
+		/// </summary>
+		public IStorageClientFactory CreateStorageClientFactory(StorageConfig storageConfig)
+			=> new StorageClientFactory(this, storageConfig);
+
+		/// <inheritdoc/>
+		public async Task StartAsync(CancellationToken cancellationToken)
+		{
+			await _blobTicker.StartAsync();
+			await _refTicker.StartAsync();
+			await _gcTicker.StartAsync();
+		}
+
+		/// <inheritdoc/>
+		public async Task StopAsync(CancellationToken cancellationToken)
+		{
+			await _gcTicker.StopAsync();
+			await _refTicker.StopAsync();
+			await _blobTicker.StopAsync();
+		}
+
+		/// <inheritdoc/>
+		public IStorageBackend CreateBackend(NamespaceId namespaceId)
+		{
+			return TryCreateBackend(namespaceId) ?? throw new StorageException($"Namespace '{namespaceId}' not found");
+		}
+
+		/// <inheritdoc/>
+		public IStorageBackend? TryCreateBackend(NamespaceId namespaceId)
+			=> TryCreateBackend(_storageConfig.CurrentValue, namespaceId);
+
+		/// <inheritdoc/>
+		public IStorageBackend? TryCreateBackend(StorageConfig storageConfig, NamespaceId namespaceId)
+		{
+			State snapshot = CreateState(storageConfig);
+			return snapshot.TryCreateBackend(namespaceId);
+		}
+
+		/// <inheritdoc/>
+		public IStorageClient? TryCreateClient(NamespaceId namespaceId)
+			=> TryCreateClient(namespaceId, null);
+
+		/// <inheritdoc/>
+		public IStorageClient? TryCreateClient(NamespaceId namespaceId, BundleOptions? bundleOptions = null)
+			=> TryCreateClient(_storageConfig.CurrentValue, namespaceId, bundleOptions);
+
+		/// <inheritdoc/>
+		public IStorageClient? TryCreateClient(StorageConfig storageConfig, NamespaceId namespaceId, BundleOptions? bundleOptions = null)
+		{
+#pragma warning disable CA2000 // Call dispose on backend; will be disposed by BundleStorageClient
+			IStorageBackend? backend = TryCreateBackend(storageConfig, namespaceId);
+#pragma warning restore CA2000
+			if (backend == null)
+			{
+				return null;
+			}
+			else
+			{
+				return new BundleStorageClient(backend, _bundleCache, bundleOptions, _logger);
+			}
+		}
+
+		#region Config
+
+		State CreateState(StorageConfig storageConfig)
+		{
+			lock (_lockObject)
+			{
+				if (_lastState == null || !String.Equals(_lastConfigRevision, storageConfig.Revision, StringComparison.Ordinal))
+				{
+					_logger.LogDebug("Updating storage providers for config {Revision}", storageConfig.Revision);
+
+					State nextState = new State(storageConfig);
+
+					foreach (NamespaceConfig namespaceConfig in storageConfig.Namespaces)
+					{
+						NamespaceId namespaceId = namespaceConfig.Id;
+						try
+						{
+							BackendId backendId = namespaceConfig.Backend;
+
+							BackendConfig? backendConfig = _staticStorageConfig.Value.Backends.FirstOrDefault(x => x.Id == namespaceConfig.Backend);
+							if (backendConfig == null && !storageConfig.TryGetBackend(namespaceConfig.Backend, out backendConfig))
+							{
+								throw new StorageException($"Missing or invalid backend identifier for namespace {namespaceConfig.Id}");
+							}
+
+							IObjectStore objectStore = _objectStoreFactory.CreateObjectStore(backendConfig);
+
+							if (!String.IsNullOrEmpty(namespaceConfig.Prefix))
+							{
+								objectStore = new PrefixedObjectStore(namespaceConfig.Prefix, objectStore);
+							}
+
+							StorageBackendImpl backend = new StorageBackendImpl(this, namespaceConfig, objectStore);
+
+							NamespaceInfo namespaceInfo = new NamespaceInfo(namespaceConfig, objectStore, backend);
+							nextState.Namespaces.Add(namespaceId, namespaceInfo);
+						}
+						catch (Exception ex)
+						{
+							_logger.LogError(ex, "Unable to create storage backend for {NamespaceId}: ", namespaceId);
+						}
+					}
+
+					_lastState = nextState;
+					_lastConfigRevision = storageConfig.Revision;
+				}
+				return _lastState;
+			}
+		}
+
+		#endregion
+
+		#region Blobs
+
+		/// <inheritdoc/>
+		async Task AddBlobAsync(NamespaceId namespaceId, BlobLocator locator, IReadOnlyCollection<BlobLocator>? imports = null, List<AliasInfo>? exports = null, CancellationToken cancellationToken = default)
+		{
+			ObjectId id = ObjectId.GenerateNewId(_clock.UtcNow);
+			BlobInfo blobInfo = new BlobInfo(id, namespaceId, locator);
+			blobInfo.Aliases = exports;
+
+			if (imports != null)
+			{
+				if (imports.Count == 0)
+				{
+					blobInfo.Imports = new List<ObjectId>();
+				}
+				else
+				{
+					List<string> paths = imports.Select(x => x.BaseLocator.ToString()).ToList();
+
+					FilterDefinition<BlobInfo> filter =
+						Builders<BlobInfo>.Filter.Eq(x => x.NamespaceId, namespaceId) &
+						Builders<BlobInfo>.Filter.In(x => x.Path, paths);
+
+					blobInfo.Imports = await _blobCollection.Find(filter).Project(x => x.Id).ToListAsync(cancellationToken);
+					blobInfo.Imports.Sort();
+				}
+			}
+
+			await _blobCollection.InsertOneAsync(blobInfo, new InsertOneOptions { }, cancellationToken);
+		}
+
+		async ValueTask CheckBlobExistsAsync(NamespaceId namespaceId, BlobLocator locator, CancellationToken cancellationToken)
+		{
+			if (_storageConfig.CurrentValue.EnableGcVerification)
+			{
+				try
+				{
+					string path = locator.Path.ToString();
+
+					BlobInfo? blobInfo = await _blobCollection.Find(x => x.NamespaceId == namespaceId && x.Path == path && x.GcVersion >= CurrentGcVersion).FirstOrDefaultAsync(cancellationToken);
+					if (blobInfo != null)
+					{
+						_logger.LogWarning("Blob {BlobId} ({Locator}) accessed after being garbage collected", blobInfo.Id, locator);
+					}
+				}
+				catch (OperationCanceledException)
+				{
+					throw;
+				}
+				catch (Exception ex)
+				{
+					_logger.LogWarning(ex, "Exception checking if blob {NamespaceId}:{Locator} exists", namespaceId, locator);
+				}
+			}
+		}
+
+		async Task<bool> IsBlobReferencedAsync(ObjectId blobInfoId, CancellationToken cancellationToken = default)
+		{
+			FilterDefinition<BlobInfo> blobFilter = Builders<BlobInfo>.Filter.AnyEq(x => x.Imports, blobInfoId);
+			if (await _blobCollection.Find(blobFilter).Limit(1).CountDocumentsAsync(cancellationToken) > 0)
+			{
+				return true;
+			}
+
+			FilterDefinition<RefInfo> refFilter = Builders<RefInfo>.Filter.Eq(x => x.TargetBlobId, blobInfoId);
+			if (await _refCollection.Find(refFilter).Limit(1).CountDocumentsAsync(cancellationToken) > 0)
+			{
+				return true;
+			}
+
+			return false;
+		}
+
+		/// <summary>
+		/// Finds blobs at least 30 minutes old and computes import metadata for them. Done with a delay to allow write redirects.
+		/// </summary>
+		/// <param name="cancellationToken">Cancellation token for the operation</param>
+		async ValueTask TickBlobsAsync(CancellationToken cancellationToken)
+		{
+			using TelemetrySpan span = _tracer.StartActiveSpan($"{nameof(StorageService)}.{nameof(TickBlobsAsync)}");
+
+			GcState gcState = await _gcState.GetAsync(cancellationToken);
+			if(gcState.Reset)
+			{
+				_logger.LogInformation("Resetting scan for new blobs...");
+				gcState = await _gcState.UpdateAsync(x => { 
+					x.Reset = false; 
+					x.LastImportBlobInfoId = ObjectId.Empty; 
+				}, cancellationToken);
+			}
+
+			DateTime ingestTimeUtc = _clock.UtcNow - TimeSpan.FromMinutes(30.0);
+
+			// Get the current state of the storage system
+			State state = CreateState(_storageConfig.CurrentValue);
+
+			Dictionary<NamespaceId, BundleStorageClient> cachedClients = new();
+			try
+			{
+				long ingestedCount = 0;
+
+				// Compute missing import info, by searching for blobs with an ObjectId timestamp after the last import compute cycle
+				ObjectId latestInfoId = ObjectId.GenerateNewId(ingestTimeUtc);
+				while (!gcState.Reset)
+				{
+					// Fetch the next batch of blobs
+					List<BlobInfo> current = await _blobCollection.Find(x => x.Id > gcState.LastImportBlobInfoId && x.Id < latestInfoId).SortBy(x => x.Id).Limit(500).ToListAsync(cancellationToken);
+					if (current.Count == 0)
+					{
+						break;
+					}
+
+					// Find imports, and add a check record for each new blob
+					using TelemetrySpan innerSpan = _tracer.StartActiveSpan($"{nameof(StorageService)}.{nameof(TickBlobsAsync)}.Batch");
+					innerSpan.SetAttribute("First", current[0].Id.ToString());
+					innerSpan.SetAttribute("Last", current[^1].Id.ToString());
+					innerSpan.SetAttribute("Count", current.Count);
+
+					_logger.LogDebug("Ticking {NumBlobs} blobs from {FirstId} to {LastId}", current.Count, current[0].Id, current[^1].Id);
+
+					foreach (BlobInfo blobInfo in current)
+					{
+						NamespaceInfo? namespaceInfo;
+						if (state.Namespaces.TryGetValue(blobInfo.NamespaceId, out namespaceInfo))
+						{
+							BundleStorageClient? storageClient;
+							if (!cachedClients.TryGetValue(namespaceInfo.Id, out storageClient))
+							{
+								storageClient = new BundleStorageClient(namespaceInfo.Backend, _bundleCache, null, _logger);
+								cachedClients.Add(namespaceInfo.Id, storageClient);
+							}
+
+							try
+							{
+								await TickBlobAsync(storageClient, blobInfo, cancellationToken);
+							}
+							catch (ObjectNotFoundException ex)
+							{
+								_logger.LogInformation(ex, "Unable to read references for {NamespaceId} blob {BlobId}: {Message}", blobInfo.NamespaceId, blobInfo.Id, ex.Message);
+							}
+							catch (Exception ex)
+							{
+								_logger.LogWarning(ex, "Unable to read references for {NamespaceId} blob {BlobId} (key: {ObjectKey}): {Message}", blobInfo.NamespaceId, blobInfo.Id, GetObjectKey(blobInfo.Locator), ex.Message);
+							}
+						}
+					}
+
+					// Update the last imported blob id
+					gcState = await _gcState.UpdateAsync(state => state.LastImportBlobInfoId = current[^1].Id, cancellationToken);
+					ingestedCount += current.Count;
+				}
+
+				_logger.LogInformation("Added {NumBlobs} blobs for GC (upper time: {Time})", ingestedCount, ingestTimeUtc);
+			}
+			finally
+			{
+				foreach (BundleStorageClient client in cachedClients.Values)
+				{
+					client.Dispose();
+				}
+			}
+		}
+
+		async Task TickBlobAsync(BundleStorageClient storageClient, BlobInfo blobInfo, CancellationToken cancellationToken)
+		{
+			List<ObjectId> importInfoIds = new List<ObjectId>();
+
+			IEnumerable<BlobLocator> importLocators = await storageClient.ReadBundleReferencesAsync(blobInfo.Locator, cancellationToken);
+			foreach (BlobLocator importLocator in importLocators)
+			{
+				string importPath = importLocator.BaseLocator.ToString();
+
+				FilterDefinition<BlobInfo> filter = Builders<BlobInfo>.Filter.Expr(x => x.NamespaceId == blobInfo.NamespaceId && x.Path == importPath);
+				UpdateDefinition<BlobInfo> update = Builders<BlobInfo>.Update.SetOnInsert(x => x.Imports, null);
+				BlobInfo blobInfoDoc = await _blobCollection.FindOneAndUpdateAsync(filter, update, new FindOneAndUpdateOptions<BlobInfo> { IsUpsert = true, ReturnDocument = ReturnDocument.After }, cancellationToken);
+
+				importInfoIds.Add(blobInfoDoc.Id);
+			}
+
+			importInfoIds.Sort();
+
+			if (blobInfo.Imports != null && !Enumerable.SequenceEqual(importInfoIds, blobInfo.Imports))
+			{
+				List<ObjectId> missing = importInfoIds.Except(blobInfo.Imports).ToList();
+				if (missing.Count > 0)
+				{
+					_logger.LogWarning("Missing imports for blob {Locator}: {Missing}", blobInfo.Path, String.Join(", ", missing.Select(x => x.ToString())));
+				}
+
+				List<ObjectId> extra = blobInfo.Imports.Except(importInfoIds).ToList();
+				if (extra.Count > 0)
+				{
+					_logger.LogWarning("Extra imports for blob {Locator}: {Extra}", blobInfo.Path, String.Join(", ", extra.Select(x => x.ToString())));
+				}
+			}
+
+			await _blobCollection.UpdateOneAsync(x => x.Id == blobInfo.Id, Builders<BlobInfo>.Update.Set(x => x.Imports, importInfoIds).Unset(x => x.GcVersion), null, cancellationToken);
+
+			AddGcCheckRecord(blobInfo.NamespaceId, blobInfo.Id);
+			_logger.LogDebug("Added {Count} imports for {NamespaceId} blob {BlobId}", importInfoIds.Count, blobInfo.NamespaceId, blobInfo.Id);
+		}
+
+		#endregion
+
+		#region Nodes
+
+		/// <summary>
+		/// Adds a node alias
+		/// </summary>
+		/// <param name="namespaceId">Namespace to search</param>
+		/// <param name="name">Alias for the node</param>
+		/// <param name="target">Target node for the alias</param>
+		/// <param name="rank">Rank for the alias. Higher ranked aliases are preferred by default.</param>
+		/// <param name="data">Inline data to store with this alias</param>
+		/// <param name="cancellationToken">Cancellation token for the operation</param>
+		/// <returns>Sequence of thandles</returns>
+		async Task AddAliasAsync(NamespaceId namespaceId, string name, BlobLocator target, int rank, ReadOnlyMemory<byte> data, CancellationToken cancellationToken = default)
+		{
+			string blobPath = target.BaseLocator.ToString();
+			string blobFragment = target.Fragment.ToString();
+
+			BlobInfo? blobInfo = await _blobCollection.Find(x => x.NamespaceId == namespaceId && x.Path == blobPath).FirstOrDefaultAsync(cancellationToken);
+			if (blobInfo == null)
+			{
+				throw new KeyNotFoundException($"Missing blob {blobPath}");
+			}
+			if (blobInfo.Aliases != null && blobInfo.Aliases.Any(x => x.Name.Equals(name, StringComparison.Ordinal) && x.Fragment.Equals(blobFragment, StringComparison.Ordinal)))
+			{
+				return;
+			}
+
+			FilterDefinition<BlobInfo> filter = Builders<BlobInfo>.Filter.Expr(x => x.NamespaceId == blobInfo.NamespaceId && x.Path == blobPath);
+			UpdateDefinition<BlobInfo> update = Builders<BlobInfo>.Update.Push(x => x.Aliases, new AliasInfo(name, blobFragment, data.ToArray(), rank));
+			await _blobCollection.UpdateOneAsync(filter, update, cancellationToken: cancellationToken);
+		}
+
+		/// <summary>
+		/// Removes a node alias
+		/// </summary>
+		/// <param name="namespaceId">Namespace to search</param>
+		/// <param name="name">Alias for the node</param>
+		/// <param name="target">Target node for the alias</param>
+		/// <param name="cancellationToken">Cancellation token for the operation</param>
+		/// <returns>Sequence of thandles</returns>
+		async Task RemoveAliasAsync(NamespaceId namespaceId, string name, BlobLocator target, CancellationToken cancellationToken = default)
+		{
+			string blobPath = target.BaseLocator.ToString();
+			string blobFragment = target.Fragment.ToString();
+
+			FilterDefinition<BlobInfo> filter = Builders<BlobInfo>.Filter.Expr(x => x.NamespaceId == namespaceId && x.Path == blobPath);
+			UpdateDefinition<BlobInfo> update = Builders<BlobInfo>.Update.PullFilter(x => x.Aliases, Builders<AliasInfo>.Filter.Expr(x => x.Name == name && x.Fragment == blobFragment));
+			await _blobCollection.UpdateOneAsync(filter, update, cancellationToken: cancellationToken);
+		}
+
+		/// <summary>
+		/// Finds nodes with the given type and hash
+		/// </summary>
+		/// <param name="namespaceId">Namespace to search</param>
+		/// <param name="name">Alias for the node</param>
+		/// <param name="cancellationToken">Cancellation token for the operation</param>
+		/// <returns>Sequence of thandles</returns>
+		async Task<List<(BlobLocator, AliasInfo)>> FindAliasesAsync(NamespaceId namespaceId, string name, CancellationToken cancellationToken = default)
+		{
+			List<(BlobLocator, AliasInfo)> results = new List<(BlobLocator, AliasInfo)>();
+			await foreach (BlobInfo blobInfo in _blobCollection.Find(x => x.NamespaceId == namespaceId && x.Aliases!.Any(y => y.Name == name)).ToAsyncEnumerable(cancellationToken))
+			{
+				if (blobInfo.Aliases != null)
+				{
+					foreach (AliasInfo aliasInfo in blobInfo.Aliases)
+					{
+						if (String.Equals(aliasInfo.Name, name, StringComparison.Ordinal))
+						{
+							BlobLocator locator = new BlobLocator(blobInfo.Locator, aliasInfo.Fragment);
+							results.Add((locator, aliasInfo));
+						}
+					}
+				}
+			}
+			return results.OrderByDescending(x => x.Item2.Rank).ToList();
+		}
+
+		#endregion
+
+		#region Refs
+
+		record RefCacheKey(NamespaceId NamespaceId, RefName Name);
+		record RefCacheValue(RefInfo? Value, DateTime Time);
+
+		/// <summary>
+		/// Adds a ref value to the cache
+		/// </summary>
+		/// <param name="namespaceId">Namespace containing the ref</param>
+		/// <param name="name">Name of the ref</param>
+		/// <param name="value">New target for the ref</param>
+		/// <returns>The cached value</returns>
+		RefCacheValue AddRefToCache(NamespaceId namespaceId, RefName name, RefInfo? value)
+		{
+			RefCacheValue cacheValue = new RefCacheValue(value, DateTime.UtcNow);
+			using (ICacheEntry newEntry = _memoryCache.CreateEntry(new RefCacheKey(namespaceId, name)))
+			{
+				newEntry.Value = cacheValue;
+				newEntry.SetSize(name.Text.Length);
+				newEntry.SetSlidingExpiration(TimeSpan.FromMinutes(5.0));
+			}
+			return cacheValue;
+		}
+
+		/// <summary>
+		/// Expires any refs that are no longer valid
+		/// </summary>
+		/// <param name="cancellationToken">Cancellation token for the operation</param>
+		async ValueTask TickRefsAsync(CancellationToken cancellationToken)
+		{
+			DateTime utcNow = DateTime.UtcNow;
+
+			using TelemetrySpan span = _tracer.StartActiveSpan($"{nameof(StorageService)}.{nameof(TickRefsAsync)}");
+
+			FilterDefinition<RefInfo> queryFilter = Builders<RefInfo>.Filter.Exists(x => x.ExpiresAtUtc) & Builders<RefInfo>.Filter.Lt(x => x.ExpiresAtUtc, utcNow);
+			using (IAsyncCursor<RefInfo> cursor = await _refCollection.Find(queryFilter).ToCursorAsync(cancellationToken))
+			{
+				List<DeleteOneModel<RefInfo>> requests = new List<DeleteOneModel<RefInfo>>();
+				while (await cursor.MoveNextAsync(cancellationToken))
+				{
+					requests.Clear();
+
+					foreach (RefInfo refInfo in cursor.Current)
+					{
+						_logger.LogInformation("Expired ref {NamespaceId}:{RefName}", refInfo.NamespaceId, refInfo.Name);
+						FilterDefinition<RefInfo> filter = Builders<RefInfo>.Filter.Expr(x => x.Id == refInfo.Id && x.ExpiresAtUtc == refInfo.ExpiresAtUtc);
+						requests.Add(new DeleteOneModel<RefInfo>(filter));
+						AddGcCheckRecord(refInfo.NamespaceId, refInfo.TargetBlobId);
+						AddRefToCache(refInfo.NamespaceId, refInfo.Name, default);
+					}
+
+					if (requests.Count > 0)
+					{
+						await _refCollection.BulkWriteAsync(requests, cancellationToken: cancellationToken);
+					}
+				}
+			}
+		}
+
+		/// <inheritdoc/>
+		async Task<bool> DeleteRefAsync(NamespaceId namespaceId, RefName name, CancellationToken cancellationToken = default)
+		{
+			FilterDefinition<RefInfo> filter = Builders<RefInfo>.Filter.Expr(x => x.NamespaceId == namespaceId && x.Name == name);
+			return await DeleteRefInternalAsync(namespaceId, name, filter, cancellationToken);
+		}
+
+		/// <summary>
+		/// Deletes a ref document that has reached its expiry time
+		/// </summary>
+		async Task DeleteExpiredRefAsync(RefInfo refDocument, CancellationToken cancellationToken = default)
+		{
+			FilterDefinition<RefInfo> filter = Builders<RefInfo>.Filter.Expr(x => x.Id == refDocument.Id && x.ExpiresAtUtc == refDocument.ExpiresAtUtc);
+			await DeleteRefInternalAsync(refDocument.NamespaceId, refDocument.Name, filter, cancellationToken);
+		}
+
+		async Task<bool> DeleteRefInternalAsync(NamespaceId namespaceId, RefName name, FilterDefinition<RefInfo> filter, CancellationToken cancellationToken = default)
+		{
+			RefInfo? oldRefInfo = await _refCollection.FindOneAndDeleteAsync<RefInfo>(filter, cancellationToken: cancellationToken);
+			AddRefToCache(namespaceId, name, default);
+
+			if (oldRefInfo != null)
+			{
+				_logger.LogInformation("Deleted ref {NamespaceId}:{RefName}", namespaceId, name);
+				AddGcCheckRecord(namespaceId, oldRefInfo.TargetBlobId);
+				return true;
+			}
+
+			return false;
+		}
+
+		/// <inheritdoc/>
+		async Task<RefInfo?> TryReadRefAsync(NamespaceId namespaceId, RefName name, RefCacheTime cacheTime = default, CancellationToken cancellationToken = default)
+		{
+			RefCacheValue? entry;
+			if (!_memoryCache.TryGetValue(name, out entry) || entry == null || RefCacheTime.IsStaleCacheEntry(entry.Time, cacheTime))
+			{
+				RefInfo? refDocument = await _refCollection.Find(x => x.NamespaceId == namespaceId && x.Name == name).FirstOrDefaultAsync(cancellationToken);
+				entry = AddRefToCache(namespaceId, name, refDocument);
+			}
+
+			if (entry.Value == null)
+			{
+				return null;
+			}
+
+			if (entry.Value.ExpiresAtUtc != null)
+			{
+				DateTime utcNow = _clock.UtcNow;
+				if (entry.Value.HasExpired(utcNow))
+				{
+					await DeleteExpiredRefAsync(entry.Value, cancellationToken);
+					return default;
+				}
+				if (entry.Value.RequiresTouch(utcNow))
+				{
+					await _refCollection.UpdateOneAsync(x => x.Id == entry.Value.Id, Builders<RefInfo>.Update.Set(x => x.ExpiresAtUtc, utcNow + entry.Value.Lifetime!.Value), cancellationToken: cancellationToken);
+				}
+			}
+
+			return entry.Value;
+		}
+
+		/// <inheritdoc/>
+		async Task WriteRefAsync(NamespaceId namespaceId, RefName name, BlobRefValue value, RefOptions? options = null, CancellationToken cancellationToken = default)
+		{
+			string path = value.Locator.BaseLocator.ToString();
+
+			BlobInfo? newBlobInfo = await _blobCollection.Find(x => x.NamespaceId == namespaceId && x.Path == path).FirstOrDefaultAsync(cancellationToken);
+			if (newBlobInfo == null)
+			{
+				throw new Exception($"Invalid/unknown blob identifier '{path}' in namespace {namespaceId}");
+			}
+
+			RefInfo newRefInfo = new RefInfo(namespaceId, name, value.Hash, value.Locator, newBlobInfo.Id);
+
+			if (options != null && options.Lifetime.HasValue)
+			{
+				newRefInfo.ExpiresAtUtc = _clock.UtcNow + options.Lifetime.Value;
+				if (options.Extend ?? true)
+				{
+					newRefInfo.Lifetime = options.Lifetime;
+				}
+			}
+
+			RefInfo? oldRefInfo = await _refCollection.FindOneAndReplaceAsync<RefInfo>(x => x.NamespaceId == namespaceId && x.Name == name, newRefInfo, new FindOneAndReplaceOptions<RefInfo> { IsUpsert = true }, cancellationToken);
+			if (oldRefInfo != null)
+			{
+				AddGcCheckRecord(namespaceId, oldRefInfo.TargetBlobId);
+			}
+
+			_logger.LogInformation("Updated ref {NamespaceId}:{RefName}", namespaceId, name);
+			AddRefToCache(namespaceId, name, newRefInfo);
+		}
+
+		#endregion
+
+		#region GC
+
+		uint GetGcTimestamp() => GetGcTimestamp(_clock.UtcNow);
+		static uint GetGcTimestamp(DateTime utcTime) => (uint)((utcTime - DateTime.UnixEpoch).Ticks / TimeSpan.TicksPerMinute);
+
+		static RedisKey GetRedisKey(string suffix) => $"storage:{suffix}";
+		static RedisKey GetRedisKey(NamespaceId namespaceId, string suffix) => GetRedisKey($"{namespaceId}:{suffix}");
+
+		static RedisSortedSetKey<RedisValue> GetGcCheckSet(NamespaceId namespaceId) => new RedisSortedSetKey<RedisValue>(GetRedisKey(namespaceId, "check"));
+
+		/// <summary>
+		/// Find the next namespace to run GC on
+		/// </summary>
+		/// <param name="cancellationToken"></param>
+		/// <returns></returns>
+		async ValueTask TickGcAsync(CancellationToken cancellationToken)
+		{
+			HashSet<NamespaceId> ranNamespaceIds = new HashSet<NamespaceId>();
+
+			DateTime utcNow = _clock.UtcNow;
+			for (; ; )
+			{
+				StorageConfig storageConfig = _storageConfig.CurrentValue;
+				if (!storageConfig.EnableGc && !storageConfig.EnableGcVerification)
+				{
+					break;
+				}
+
+				State state = CreateState(storageConfig);
+
+				// Synchronize the list of configured namespaces with the GC state object
+				GcState gcState = await _gcState.GetAsync(cancellationToken);
+				if (!Enumerable.SequenceEqual(storageConfig.Namespaces.Select(x => x.Id.Text.Text).OrderBy(x => x), gcState.Namespaces.Select(x => x.Id.Text.Text).OrderBy(x => x)))
+				{
+					gcState = await _gcState.UpdateAsync(s => SyncNamespaceList(s, storageConfig.Namespaces), cancellationToken);
+				}
+
+				// Find all the namespaces that need to have GC run on them
+				List<(DateTime, GcNamespaceState)> pending = new List<(DateTime, GcNamespaceState)>();
+				foreach (GcNamespaceState namespaceState in gcState.Namespaces)
+				{
+					if (!ranNamespaceIds.Contains(namespaceState.Id))
+					{
+						NamespaceConfig? namespaceConfig;
+						if (storageConfig.TryGetNamespace(namespaceState.Id, out namespaceConfig))
+						{
+							DateTime time = namespaceState.LastTime + TimeSpan.FromHours(namespaceConfig.GcFrequencyHrs);
+							if (time < utcNow)
+							{
+								pending.Add((time, namespaceState));
+							}
+						}
+					}
+				}
+				pending.SortBy(x => x.Item1);
+
+				// If there's nothing left to GC, bail out
+				if (pending.Count == 0)
+				{
+					break;
+				}
+
+				// Update the first one we can acquire a lock for
+				foreach ((_, GcNamespaceState namespaceState) in pending)
+				{
+					NamespaceId namespaceId = namespaceState.Id;
+					if (ranNamespaceIds.Add(namespaceId))
+					{
+						RedisKey key = GetRedisKey(namespaceId, "lock");
+#pragma warning disable CA2000 // Dispose objects before losing scope
+						using (RedisLock namespaceLock = new RedisLock(_redisService.GetDatabase(), key))
+						{
+							if (await namespaceLock.AcquireAsync(TimeSpan.FromMinutes(20.0)))
+							{
+								try
+								{
+									await TickGcForNamespaceAsync(state.Namespaces[namespaceId], gcState.LastImportBlobInfoId, utcNow, storageConfig.EnableGc, cancellationToken);
+								}
+								catch (Exception ex)
+								{
+									_logger.LogError(ex, "Exception while running garbage collection: {Message}", ex.Message);
+								}
+								break;
+							}
+						}
+#pragma warning restore CA2000 // Dispose objects before losing scope
+					}
+				}
+			}
+		}
+
+		async Task TickGcForNamespaceAsync(NamespaceInfo namespaceInfo, ObjectId lastImportBlobInfoId, DateTime utcNow, bool deleteObjects, CancellationToken cancellationToken)
+		{
+			using IStorageClient client = this.CreateClient(namespaceInfo.Id);
+
+			Stopwatch timer = Stopwatch.StartNew();
+			_logger.LogInformation("Running garbage collection for namespace {NamespaceId}...", namespaceInfo.Id);
+			int numItemsRemoved = 0;
+
+			double score = GetGcTimestamp(utcNow);
+
+			RedisSortedSetKey<RedisValue> checkSet = GetGcCheckSet(namespaceInfo.Id);
+			while (_storageConfig.CurrentValue.EnableGc || _storageConfig.CurrentValue.EnableGcVerification)
+			{
+				long length = await _redisService.GetDatabase().SortedSetLengthAsync(checkSet);
+				int batchSize = (int)Math.Min(length, 1024);
+				_logger.LogInformation("Garbage collection queue for namespace {NamespaceId} ({QueueName}) has {Length} entries; taking {Count}", namespaceInfo.Id, checkSet.Inner, length, batchSize);
+
+				if (length == 0)
+				{
+					await _gcState.UpdateAsync(state => state.FindOrAddNamespace(namespaceInfo.Id).LastTime = utcNow, cancellationToken);
+					_logger.LogInformation("Finished garbage collection for namespace {NamespaceId} in {TimeSecs}s ({NumItems} removed)", namespaceInfo.Id, timer.Elapsed.TotalSeconds, numItemsRemoved);
+					break;
+				}
+
+				RedisValue[] values = await _redisService.GetDatabase().SortedSetRangeByRankAsync(checkSet, 0, batchSize);
+				foreach (RedisValue value in values)
+				{
+					ObjectId blobInfoId = new ObjectId(((byte[]?)value)!);
+
+					using TelemetrySpan span = _tracer.StartActiveSpan($"{nameof(StorageService)}.{nameof(TickGcForNamespaceAsync)}");
+					span.SetAttribute("BlobId", blobInfoId.ToString());
+
+					if (blobInfoId <= lastImportBlobInfoId && !await IsBlobReferencedAsync(blobInfoId, cancellationToken))
+					{
+						BlobInfo? info;
+						if (deleteObjects)
+						{
+							info = await _blobCollection.FindOneAndDeleteAsync(x => x.Id == blobInfoId, cancellationToken: cancellationToken);
+						}
+						else
+						{
+							info = await _blobCollection.FindOneAndUpdateAsync(x => x.Id == blobInfoId, Builders<BlobInfo>.Update.Set(x => x.GcVersion, CurrentGcVersion), cancellationToken: cancellationToken);
+						}
+
+						if (info != null)
+						{
+							if (info.Imports != null)
+							{
+								SortedSetEntry<RedisValue>[] entries = info.Imports.Select(x => new SortedSetEntry<RedisValue>(x.ToByteArray(), score)).ToArray();
+								_ = _redisService.GetDatabase().SortedSetAddAsync(checkSet, entries, flags: CommandFlags.FireAndForget);
+								score = Math.BitIncrement(score);
+							}
+
+							ObjectKey objectKey = GetObjectKey(new BlobLocator(info.Path));
+							_logger.LogDebug("Deleting {NamespaceId} blob {BlobId}, key: {ObjectKey} ({ImportCount} imports)", namespaceInfo.Id, blobInfoId, objectKey, info.Imports?.Count ?? 0);
+
+							if (deleteObjects)
+							{
+								await namespaceInfo.Store.DeleteAsync(objectKey, cancellationToken);
+							}
+
+							numItemsRemoved++;
+						}
+					}
+					_ = _redisService.GetDatabase().SortedSetRemoveAsync(checkSet, value, CommandFlags.FireAndForget);
+				}
+			}
+		}
+
+		static void SyncNamespaceList(GcState state, List<NamespaceConfig> namespaces)
+		{
+			HashSet<NamespaceId> validNamespaceIds = new HashSet<NamespaceId>(namespaces.Select(x => x.Id));
+			state.Namespaces.RemoveAll(x => !validNamespaceIds.Contains(x.Id));
+
+			HashSet<NamespaceId> currentNamespaceIds = new HashSet<NamespaceId>(state.Namespaces.Select(x => x.Id));
+			foreach (NamespaceConfig config in namespaces)
+			{
+				if (!currentNamespaceIds.Contains(config.Id))
+				{
+					state.Namespaces.Add(new GcNamespaceState { Id = config.Id });
+				}
+			}
+
+			state.Namespaces.SortBy(x => x.Id.Text.Text);
+		}
+
+		void AddGcCheckRecord(NamespaceId namespaceId, ObjectId id)
+		{
+			double score = GetGcTimestamp();
+			_ = _redisService.GetDatabase().SortedSetAddAsync(GetGcCheckSet(namespaceId), id.ToByteArray(), score, flags: CommandFlags.FireAndForget);
+		}
+
+		#endregion
+	}
+}
diff --git a/Engine/Source/Programs/Horde/Plugins/Storage/HordeServer.Storage/StorageConfig.cs b/Engine/Source/Programs/Horde/Plugins/Storage/HordeServer.Storage/StorageConfig.cs
new file mode 100644
--- /dev/null
+++ b/Engine/Source/Programs/Horde/Plugins/Storage/HordeServer.Storage/StorageConfig.cs
@@ -0,0 +1,369 @@
+// Copyright Epic Games, Inc. All Rights Reserved.
+
+using System;
+using System.Collections.Generic;
+using System.ComponentModel.DataAnnotations;
+using System.Diagnostics;
+using System.Diagnostics.CodeAnalysis;
+using System.IO;
+using System.Reflection;
+using System.Security.Claims;
+using System.Text.Json;
+using System.Text.Json.Serialization;
+using EpicGames.Core;
+using EpicGames.Horde.Acls;
+using EpicGames.Horde.Storage;
+using HordeServer.Acls;
+using HordeServer.Plugins;
+using HordeServer.Server;
+using HordeServer.Storage.ObjectStores;
+using HordeServer.Utilities;
+
+#pragma warning disable CA1716 // Rename type Namespace to it no longer conflicts with the reserved language keyword 'Namespace'
+#pragma warning disable CA2227 // Change x to be read-only by removing the property setter
+
+namespace HordeServer.Storage
+{
+	/// <summary>
+	/// Well-known namespace identifiers
+	/// </summary>
+	public static class Namespace
+	{
+		/// <summary>
+		/// Storage of artifacts 
+		/// </summary>
+		public static NamespaceId Artifacts { get; } = new NamespaceId("horde-artifacts");
+
+		/// <summary>
+		/// Replicated Perforce data
+		/// </summary>
+		public static NamespaceId Perforce { get; } = new NamespaceId("horde-perforce");
+
+		/// <summary>
+		/// Log data
+		/// </summary>
+		public static NamespaceId Logs { get; } = new NamespaceId("horde-logs");
+
+		/// <summary>
+		/// Storage of tool data
+		/// </summary>
+		public static NamespaceId Tools { get; } = new NamespaceId("horde-tools");
+	}
+
+	/// <summary>
+	/// Configuration for storage
+	/// </summary>
+	public class StorageConfig : IPluginConfig
+	{
+		/// <summary>
+		/// Hash of the current object revision to allow detecting changes
+		/// </summary>
+		[JsonIgnore]
+		public string Revision { get; set; } = String.Empty;
+
+		/// <summary>
+		/// Whether to enable garbage collection
+		/// </summary>
+		public bool EnableGc { get; set; } = true;
+
+		/// <summary>
+		/// Whether to enable garbage collection in verification mode (nothing deleted, just logging on access to deleted blobs)
+		/// </summary>
+		public bool EnableGcVerification { get; set; }
+
+		/// <summary>
+		/// List of storage backends
+		/// </summary>
+		public List<BackendConfig> Backends { get; set; } = new List<BackendConfig>();
+
+		/// <summary>
+		/// List of namespaces for storage
+		/// </summary>
+		public List<NamespaceConfig> Namespaces { get; set; } = new List<NamespaceConfig>();
+
+		private readonly Dictionary<BackendId, BackendConfig> _backendLookup = new Dictionary<BackendId, BackendConfig>();
+		private readonly Dictionary<NamespaceId, NamespaceConfig> _namespaceLookup = new Dictionary<NamespaceId, NamespaceConfig>();
+
+		/// <summary>
+		/// Called after the config has been read from disk
+		/// </summary>
+		public void PostLoad(AclConfig parentAcl)
+		{
+			// Create a lookup for backend configs, allowing later entries to override previous ones
+			_backendLookup.Clear();
+			foreach (BackendConfig backendConfig in Backends)
+			{
+				_backendLookup[backendConfig.Id] = backendConfig;
+			}
+
+			// Fixup all the inherited properties
+			Dictionary<BackendId, BackendConfig> mergedBackendConfigs = new Dictionary<BackendId, BackendConfig>(Backends.Count);
+			foreach (BackendConfig backendConfig in Backends)
+			{
+				MergeBackendConfigs(backendConfig.Id, _backendLookup, mergedBackendConfigs);
+			}
+
+			// Compute the hash for each backend
+			foreach (BackendConfig backendConfig in Backends)
+			{
+				using (MemoryStream stream = new MemoryStream())
+				{
+					JsonSerializerOptions options = new JsonSerializerOptions();
+					JsonUtils.ConfigureJsonSerializer(options);
+
+					JsonSerializer.Serialize(stream, backendConfig, options: options);
+
+					backendConfig.Hash = IoHash.Compute(stream.ToArray());
+				}
+			}
+
+			// Validate the backend config for each namespace
+			_namespaceLookup.Clear();
+			foreach (NamespaceConfig namespaceConfig in Namespaces)
+			{
+				BackendConfig? backendConfig;
+				if (!_backendLookup.TryGetValue(namespaceConfig.Backend, out backendConfig))
+				{
+					throw new StorageException($"Missing or invalid backend identifier for namespace {namespaceConfig.Id}");
+				}
+
+				namespaceConfig.PostLoad(parentAcl);
+				_namespaceLookup[namespaceConfig.Id] = namespaceConfig;
+			}
+
+			// Create a unique id to allow detect changes
+			byte[] data = JsonSerializer.SerializeToUtf8Bytes<StorageConfig>(this, JsonUtils.DefaultSerializerOptions);
+			Revision = IoHash.Compute(data).ToString();
+		}
+
+		static BackendConfig MergeBackendConfigs(BackendId backendId, Dictionary<BackendId, BackendConfig> baseIdToConfig, Dictionary<BackendId, BackendConfig> mergedIdToConfig)
+		{
+			BackendConfig? config;
+			if (mergedIdToConfig.TryGetValue(backendId, out config))
+			{
+				if (config == null)
+				{
+					throw new InvalidDataException($"Configuration for storage backend '{backendId}' is recursive.");
+				}
+			}
+			else
+			{
+				mergedIdToConfig.Add(backendId, null!);
+
+				config = baseIdToConfig[backendId];
+				if (!config.Base.IsEmpty)
+				{
+					BackendConfig baseConfig = MergeBackendConfigs(config.Base, baseIdToConfig, mergedIdToConfig);
+					config.MergeDefaults(baseConfig);
+				}
+
+				mergedIdToConfig[backendId] = config;
+			}
+			return config;
+		}
+
+		/// <summary>
+		/// Gets a backend with the given id
+		/// </summary>
+		/// <param name="backendId">Identifier for the backend</param>
+		/// <param name="backendConfig">Receives the backend config on success</param>
+		/// <returns>True on success</returns>
+		public bool TryGetBackend(BackendId backendId, [NotNullWhen(true)] out BackendConfig? backendConfig) => _backendLookup.TryGetValue(backendId, out backendConfig);
+
+		/// <summary>
+		/// Gets a namespace with the given id
+		/// </summary>
+		/// <param name="namespaceId">Identifier for the backend</param>
+		/// <param name="namespaceConfig">Receives the backend config on success</param>
+		/// <returns>True on success</returns>
+		public bool TryGetNamespace(NamespaceId namespaceId, [NotNullWhen(true)] out NamespaceConfig? namespaceConfig) => _namespaceLookup.TryGetValue(namespaceId, out namespaceConfig);
+	}
+
+	/// <summary>
+	/// Static settings for the storage system
+	/// </summary>
+	public class StaticStorageConfig
+	{
+		/// <summary>
+		/// Directory to use for the coarse-grained backend cache. This caches full bundles downloaded from the upstream object store.
+		/// </summary>
+		public string? BundleCacheDir { get; set; }
+
+		/// <summary>
+		/// Maximum size of the storage cache on disk. Accepts standard binary suffixes (kb, mb, gb, tb, etc...)
+		/// </summary>
+		public string BundleCacheSize { get; set; } = "1gb";
+
+		/// <summary>
+		/// Accessor for the bundle cache size in bytes
+		/// </summary>
+		public long BundleCacheSizeBytes => StringUtils.ParseBytesString(BundleCacheSize);
+
+		/// <summary>
+		/// Overridden settings for storage backends. Useful for running against a production server with custom backends.
+		/// </summary>
+		public List<BackendConfig> Backends { get; set; } = new List<BackendConfig>();
+	}
+
+	/// <summary>
+	/// Types of storage backend to use
+	/// </summary>
+	public enum StorageBackendType
+	{
+		/// <summary>
+		/// Local filesystem
+		/// </summary>
+		FileSystem,
+
+		/// <summary>
+		/// AWS S3
+		/// </summary>
+		Aws,
+
+		/// <summary>
+		/// Azure blob store
+		/// </summary>
+		Azure,
+
+		/// <summary>
+		/// In-memory only (for testing)
+		/// </summary>
+		Memory,
+	};
+
+	/// <summary>
+	/// Common settings object for different providers
+	/// </summary>
+	[DebuggerDisplay("{Id}")]
+	public class BackendConfig : IAwsStorageOptions
+	{
+		/// <summary>
+		/// The storage backend ID
+		/// </summary>
+		[Required]
+		public BackendId Id { get; set; }
+
+		/// <summary>
+		/// Base backend to copy default settings from
+		/// </summary>
+		public BackendId Base { get; set; }
+
+		/// <inheritdoc/>
+		public StorageBackendType? Type { get; set; }
+
+		/// <inheritdoc/>
+		public string? BaseDir { get; set; }
+
+		/// <inheritdoc/>
+		public string? AwsBucketName { get; set; }
+
+		/// <inheritdoc/>
+		public string? AwsBucketPath { get; set; }
+
+		/// <inheritdoc/>
+		public AwsCredentialsType? AwsCredentials { get; set; }
+
+		/// <inheritdoc/>
+		public string? AwsRole { get; set; }
+
+		/// <inheritdoc/>
+		public string? AwsProfile { get; set; }
+
+		/// <inheritdoc/>
+		public string? AwsRegion { get; set; }
+
+		/// <inheritdoc/>
+		public string? AzureConnectionString { get; set; }
+
+		/// <inheritdoc/>
+		public string? AzureContainerName { get; set; }
+
+		/// <inheritdoc/>
+		public string? RelayServer { get; set; }
+
+		/// <inheritdoc/>
+		public string? RelayToken { get; set; }
+
+		/// <summary>
+		/// Hash of this backend config. Used for caching backend instances.
+		/// </summary>
+		[JsonIgnore]
+		internal IoHash Hash { get; set; }
+
+		/// <summary>
+		/// Merge default values from another object
+		/// </summary>
+		/// <param name="other">Object to merge defaults from</param>
+		public void MergeDefaults(BackendConfig other)
+		{
+			foreach (PropertyInfo propertyInfo in GetType().GetProperties(BindingFlags.Instance | BindingFlags.Public))
+			{
+				if (propertyInfo.GetValue(this) == null)
+				{
+					propertyInfo.SetValue(this, propertyInfo.GetValue(other));
+				}
+			}
+		}
+	}
+
+	/// <summary>
+	/// Configuration of a particular namespace
+	/// </summary>
+	[DebuggerDisplay("{Id}")]
+	public class NamespaceConfig
+	{
+		/// <summary>
+		/// Identifier for this namespace
+		/// </summary>
+		[Required]
+		public NamespaceId Id { get; set; }
+
+		/// <summary>
+		/// Backend to use for this namespace
+		/// </summary>
+		[Required]
+		public BackendId Backend { get; set; }
+
+		/// <summary>
+		/// Prefix for items within this namespace
+		/// </summary>
+		public string Prefix { get; set; } = String.Empty;
+
+		/// <summary>
+		/// How frequently to run garbage collection, in hours.
+		/// </summary>
+		public double GcFrequencyHrs { get; set; } = 2.0;
+
+		/// <summary>
+		/// How long to keep newly uploaded orphanned objects before allowing them to be deleted, in hours.
+		/// </summary>
+		public double GcDelayHrs { get; set; } = 6.0;
+
+		/// <summary>
+		/// Support querying exports by their aliases
+		/// </summary>
+		public bool EnableAliases { get; set; }
+
+		/// <summary>
+		/// Access list for this namespace
+		/// </summary>
+		public AclConfig Acl { get; set; } = new AclConfig();
+
+		/// <summary>
+		/// Callback once the configuration has been read from disk
+		/// </summary>
+		/// <param name="backendConfig"></param>
+		public void PostLoad(AclConfig parentAcl)
+		{
+			Acl.PostLoad(parentAcl, $"namespace:{Id}");
+		}
+
+		/// <summary>
+		/// Authorizes a user to perform a given action
+		/// </summary>
+		/// <param name="action">The action being performed</param>
+		/// <param name="user">The principal to validate</param>
+		public bool Authorize(AclAction action, ClaimsPrincipal user)
+			=> Acl.Authorize(action, user);
+	}
+}
diff --git a/Engine/Source/Programs/Horde/Plugins/Storage/HordeServer.Storage/StoragePlugin.cs b/Engine/Source/Programs/Horde/Plugins/Storage/HordeServer.Storage/StoragePlugin.cs
new file mode 100644
--- /dev/null
+++ b/Engine/Source/Programs/Horde/Plugins/Storage/HordeServer.Storage/StoragePlugin.cs
@@ -0,0 +1,55 @@
+// Copyright Epic Games, Inc. All Rights Reserved.
+
+using System;
+using System.Collections.Generic;
+using System.Linq;
+using System.Text;
+using System.Threading.Tasks;
+using EpicGames.Core;
+using EpicGames.Horde.Storage;
+using EpicGames.Horde.Storage.Bundles;
+using EpicGames.Horde.Storage.ObjectStores;
+using HordeServer.Plugins;
+using HordeServer.Storage.ObjectStores;
+using Microsoft.Extensions.DependencyInjection;
+using Microsoft.Extensions.Logging;
+using Microsoft.Extensions.Options;
+
+namespace HordeServer.Storage
+{
+	/// <summary>
+	/// Entry point for the storage plugin
+	/// </summary>
+	[Plugin("Storage", GlobalConfigType = typeof(StorageConfig), ServerConfigType = typeof(StaticStorageConfig))]
+	public class StoragePlugin : IPluginStartup
+	{
+		readonly StaticStorageConfig _staticConfig;
+
+		public StoragePlugin(StaticStorageConfig staticConfig)
+		{
+			_staticConfig = staticConfig;
+		}
+
+		/// <inheritdoc/>
+		public void ConfigureServices(IServiceCollection services)
+		{
+			services.AddSingleton<StorageService>();
+			services.AddSingleton<IStorageService>(sp => sp.GetRequiredService<StorageService>());
+			services.AddScoped(sp => sp.GetRequiredService<StorageService>().CreateStorageClientFactory(sp.GetRequiredService<IOptionsSnapshot<StorageConfig>>().Value));
+
+			services.AddSingleton<IObjectStoreFactory, ObjectStoreFactory>();
+			services.AddSingleton<AwsObjectStoreFactory>();
+			services.AddSingleton<FileObjectStoreFactory>();
+
+			services.AddSingleton<BundleCache>();
+			services.AddSingleton<StorageBackendCache>(CreateStorageBackendCache);
+		}
+
+		StorageBackendCache CreateStorageBackendCache(IServiceProvider serviceProvider)
+		{
+			IServerInfo serverInfo = serviceProvider.GetRequiredService<IServerInfo>();
+			DirectoryReference cacheDir = DirectoryReference.Combine(serverInfo.DataDir, String.IsNullOrEmpty(_staticConfig.BundleCacheDir) ? "Cache" : _staticConfig.BundleCacheDir);
+			return new StorageBackendCache(cacheDir, _staticConfig.BundleCacheSizeBytes, serviceProvider.GetRequiredService<ILogger<StorageBackendCache>>());
+		}
+	}
+}
diff --git a/Engine/Source/Programs/Horde/Scripts/test.sh b/Engine/Source/Programs/Horde/Scripts/test.sh
--- a/Engine/Source/Programs/Horde/Scripts/test.sh
+++ b/Engine/Source/Programs/Horde/Scripts/test.sh
@@ -16,6 +16,7 @@ testProjects=(
 	"Source/Programs/Shared/EpicGames.Serialization.Tests/EpicGames.Serialization.Tests.csproj"
 	"Source/Programs/Horde/HordeServer.Tests/HordeServer.Tests.csproj"
 	"Source/Programs/Horde/Plugins/Analytics/HordeServer.Analytics.Tests/HordeServer.Analytics.Tests.csproj"
+	"Source/Programs/Horde/Plugins/Storage/HordeServer.Storage.Tests/HordeServer.Storage.Tests.csproj"
 )
 
 for csProj in "${testProjects[@]}"; do
