Horde: Handle no general category for tools view
#rnx
#skipci

diff --git a/Engine/Source/Programs/Horde/HordeDashboard/src/components/ToolView.tsx b/Engine/Source/Programs/Horde/HordeDashboard/src/components/ToolView.tsx
--- a/Engine/Source/Programs/Horde/HordeDashboard/src/components/ToolView.tsx
+++ b/Engine/Source/Programs/Horde/HordeDashboard/src/components/ToolView.tsx
@@ -338,7 +338,15 @@ export const ToolViewInner: React.FC = observer(() => {
       return <PivotItem headerText={cat} itemKey={cat} key={cat} style={{ color: modeColors.text }} />;
    }).filter(p => !!p);
 
-   pivotItems.unshift(<PivotItem headerText={defaultCategory} itemKey={defaultCategory} key={defaultCategory} style={{ color: modeColors.text }} />);
+   if (handler.categories.get(defaultCategory)?.length) {
+      pivotItems.unshift(<PivotItem headerText={defaultCategory} itemKey={defaultCategory} key={defaultCategory} style={{ color: modeColors.text }} />);
+   } else {
+      if (categories.length) {
+         if (!selectedKey || selectedKey === defaultCategory) {
+            setSelectedKey(categories[0]);
+         }
+      }
+   }
 
    return <Stack styles={{ root: { width: "100%", backgroundColor: modeColors.background } }}>
       <Stack style={{ width: "100%", backgroundColor: modeColors.background }}>
