Horde: Move DDC functionality into a plugin.

#rnx

diff --git a/Engine/Source/Programs/Horde/HordeServer.Tests/Accounts/AccountControllerTests.cs b/Engine/Source/Programs/Horde/HordeServer.Tests/Accounts/AccountControllerTests.cs
--- a/Engine/Source/Programs/Horde/HordeServer.Tests/Accounts/AccountControllerTests.cs
+++ b/Engine/Source/Programs/Horde/HordeServer.Tests/Accounts/AccountControllerTests.cs
@@ -30,7 +30,7 @@ public class AccountControllerTest : IAsyncDisposable
 		{
 			{ "Horde:AuthMethod", "Horde" },
 		};
-		_app = new FakeHordeWebApp(settings);
+		_app = new FakeHordeWebApp(settings: settings);
 		_accountCollection = _app.ServiceProvider.GetRequiredService<IAccountCollection>();
 	}
 
diff --git a/Engine/Source/Programs/Horde/HordeServer.Tests/HordeServer.Tests.csproj b/Engine/Source/Programs/Horde/HordeServer.Tests/HordeServer.Tests.csproj
--- a/Engine/Source/Programs/Horde/HordeServer.Tests/HordeServer.Tests.csproj
+++ b/Engine/Source/Programs/Horde/HordeServer.Tests/HordeServer.Tests.csproj
@@ -64,52 +64,4 @@
     <Folder Include="Resources\" />
   </ItemGroup>
 
-  <ItemGroup>
-    <None Update="Ddc\Functional\ContentId\Payloads\05d7c699a2668efdecbe48f10db0d621d736f449.uecomp">
-      <CopyToOutputDirectory>Always</CopyToOutputDirectory>
-    </None>
-    <None Update="Ddc\Functional\ContentId\Payloads\0b29a878f71cc1eb6a26663a1152537e8fe7dca9.uecomp">
-      <CopyToOutputDirectory>Always</CopyToOutputDirectory>
-    </None>
-    <None Update="Ddc\Functional\ContentId\Payloads\4958689fe783e02fb35b13c14b0c3d7beb91e50c">
-      <CopyToOutputDirectory>Always</CopyToOutputDirectory>
-    </None>
-    <None Update="Ddc\Functional\ContentId\Payloads\7835c353d7dc67e8a0531c88fbc75ddfda10dee4">
-      <CopyToOutputDirectory>Always</CopyToOutputDirectory>
-    </None>
-    <None Update="Ddc\Functional\ContentId\Payloads\dce31eb416f3dcb4c8250ac545eda3930919d3ff">
-      <CopyToOutputDirectory>Always</CopyToOutputDirectory>
-    </None>
-    <None Update="Ddc\Functional\ContentId\Payloads\largerfile">
-      <CopyToOutputDirectory>Always</CopyToOutputDirectory>
-    </None>
-    <None Update="Ddc\Functional\ContentId\Payloads\Oodle-f0b9c675fe21951ca27699f9baab9f9f5040b202">
-      <CopyToOutputDirectory>Always</CopyToOutputDirectory>
-    </None>
-    <None Update="Ddc\Functional\ContentId\Payloads\Oodle-f895ea954b37217270e88d8b728bd3c09152689c">
-      <CopyToOutputDirectory>Always</CopyToOutputDirectory>
-    </None>
-    <None Update="Ddc\Functional\ContentId\Payloads\OodleTexture_CAS_dbda9040e75c4674fcec173f982fddf12b021e24.udd">
-      <CopyToOutputDirectory>Always</CopyToOutputDirectory>
-    </None>
-    <None Update="Ddc\Functional\ContentId\Payloads\smallerfile">
-      <CopyToOutputDirectory>Always</CopyToOutputDirectory>
-    </None>
-    <None Update="Ddc\Functional\ContentId\Payloads\UncompressedTexture_CAS_dea81b6c3b565bb5089695377c98ce0f1c13b0c3.udd">
-      <CopyToOutputDirectory>Always</CopyToOutputDirectory>
-    </None>
-    <None Update="Ddc\Functional\Objects\Payloads\lyra.cb">
-      <CopyToOutputDirectory>Always</CopyToOutputDirectory>
-    </None>
-    <None Update="Ddc\UnitTests\CompactBinaryObjects\build">
-      <CopyToOutputDirectory>Always</CopyToOutputDirectory>
-    </None>
-    <None Update="Ddc\UnitTests\CompactBinaryObjects\compact_binary">
-      <CopyToOutputDirectory>Always</CopyToOutputDirectory>
-    </None>
-    <None Update="Ddc\UnitTests\CompactBinaryObjects\ReferenceOutput">
-      <CopyToOutputDirectory>Always</CopyToOutputDirectory>
-    </None>
-  </ItemGroup>
-
 </Project>
diff --git a/Engine/Source/Programs/Horde/HordeServer.Tests/Server/DebugControllerTests.cs b/Engine/Source/Programs/Horde/HordeServer.Tests/Server/DebugControllerTests.cs
--- a/Engine/Source/Programs/Horde/HordeServer.Tests/Server/DebugControllerTests.cs
+++ b/Engine/Source/Programs/Horde/HordeServer.Tests/Server/DebugControllerTests.cs
@@ -45,7 +45,7 @@ public class DebugControllerTests
 		{
 			settings["Horde:EnableDebugEndpoints"] = Convert.ToString(enableDebugEndpoints.Value);
 		}
-		return new FakeHordeWebApp(settings);
+		return new FakeHordeWebApp(settings: settings);
 	}
 }
 
diff --git a/Engine/Source/Programs/Horde/HordeServer.Tests/Server/SchemaControllerTest.cs b/Engine/Source/Programs/Horde/HordeServer.Tests/Server/SchemaControllerTest.cs
--- a/Engine/Source/Programs/Horde/HordeServer.Tests/Server/SchemaControllerTest.cs
+++ b/Engine/Source/Programs/Horde/HordeServer.Tests/Server/SchemaControllerTest.cs
@@ -28,8 +28,6 @@ namespace HordeServer.Tests.Server
 		[TestMethod]
 		public async Task TestCatalogAsync()
 		{
-			await GetFixtureAsync();
-
 			// Fetch the catalog
 			SchemaCatalog catalog;
 			using (HttpResponseMessage response = await Client.GetAsync("/api/v1/schema/catalog.json"))
diff --git a/Engine/Source/Programs/Horde/HordeServer.Tests/TestSetup.cs b/Engine/Source/Programs/Horde/HordeServer.Tests/TestSetup.cs
--- a/Engine/Source/Programs/Horde/HordeServer.Tests/TestSetup.cs
+++ b/Engine/Source/Programs/Horde/HordeServer.Tests/TestSetup.cs
@@ -2,7 +2,6 @@
 
 using System;
 using System.Collections.Generic;
-using System.Diagnostics.Metrics;
 using System.Net;
 using System.Net.Sockets;
 using System.Security.Claims;
@@ -12,9 +11,6 @@ using Amazon.CloudWatch;
 using Amazon.EC2;
 using EpicGames.Horde.Agents;
 using EpicGames.Horde.Agents.Pools;
-using EpicGames.Horde.Storage;
-using EpicGames.Horde.Storage.Bundles;
-using EpicGames.Horde.Storage.ObjectStores;
 using HordeServer.Accounts;
 using HordeServer.Agents;
 using HordeServer.Agents.Enrollment;
@@ -58,7 +54,6 @@ using Microsoft.Extensions.DependencyInjection;
 using Microsoft.Extensions.Logging;
 using Microsoft.VisualStudio.TestTools.UnitTesting;
 using Moq;
-using OpenTelemetry.Trace;
 
 namespace HordeServer.Tests
 {
@@ -330,7 +325,7 @@ namespace HordeServer.Tests
 			{
 				await agent.TryUpdateWorkspacesAsync(workspaces, false);
 			}
-			
+
 			if (requestShutdown)
 			{
 				await agent.TryUpdateAsync(new UpdateAgentOptions { RequestShutdown = true });
diff --git a/Engine/Source/Programs/Horde/HordeServer.Tests/Tools/ToolsControllerTests.cs b/Engine/Source/Programs/Horde/HordeServer.Tests/Tools/ToolsControllerTests.cs
--- a/Engine/Source/Programs/Horde/HordeServer.Tests/Tools/ToolsControllerTests.cs
+++ b/Engine/Source/Programs/Horde/HordeServer.Tests/Tools/ToolsControllerTests.cs
@@ -42,14 +42,14 @@ public class ToolsControllerTests
 		storageConfig.Namespaces.Add(new NamespaceConfig { Id = Namespace.Tools, Backend = new BackendId("tools-backend") });
 
 		GlobalConfig globalConfig = new();
-		globalConfig.Plugins.AddStorageTestConfig(storageConfig);
+		globalConfig.Plugins.AddStorageConfig(storageConfig);
 		globalConfig.Tools.Add(new ToolConfig(toolId) { Name = "Foo", Description = "This is foo", Acl = new AclConfig() { Entries = [aclEntryConfig] }, Public = false });
 		
 		ServerSettings serverSettings = new() { AuthMethod = AuthMethod.Horde };
 		globalConfig.PostLoad(serverSettings, pluginCollection.LoadedPlugins);
 
 		Dictionary<string, string> settings = new() { { "Horde:AuthMethod", AuthMethod.Horde.ToString() } };
-		await using FakeHordeWebApp app = new(settings);
+		await using FakeHordeWebApp app = new(settings: settings);
 		
 		ConfigService configService = app.ServiceProvider.GetRequiredService<ConfigService>();
 		IToolCollection tools = app.ServiceProvider.GetRequiredService<IToolCollection>();
diff --git a/Engine/Source/Programs/Horde/HordeServer/HordeServer.csproj b/Engine/Source/Programs/Horde/HordeServer/HordeServer.csproj
--- a/Engine/Source/Programs/Horde/HordeServer/HordeServer.csproj
+++ b/Engine/Source/Programs/Horde/HordeServer/HordeServer.csproj
@@ -135,6 +135,7 @@
     <ProjectReference Include="..\..\Shared\EpicGames.Slack\EpicGames.Slack.csproj" />
     <ProjectReference Include="..\HordeServer.Shared\HordeServer.Shared.csproj" />
     <ProjectReference Include="..\Plugins\Analytics\HordeServer.Analytics\HordeServer.Analytics.csproj" />
+    <ProjectReference Include="..\Plugins\Ddc\HordeServer.Ddc\HordeServer.Ddc.csproj" />
     <ProjectReference Include="..\Plugins\Storage\HordeServer.Storage\HordeServer.Storage.csproj" />
   </ItemGroup>
 
diff --git a/Engine/Source/Programs/Horde/HordeServer/Startup.cs b/Engine/Source/Programs/Horde/HordeServer/Startup.cs
--- a/Engine/Source/Programs/Horde/HordeServer/Startup.cs
+++ b/Engine/Source/Programs/Horde/HordeServer/Startup.cs
@@ -52,7 +52,6 @@ using HordeServer.Commits;
 using HordeServer.Compute;
 using HordeServer.Configuration;
 using HordeServer.Dashboard;
-using HordeServer.Ddc;
 using HordeServer.Devices;
 using HordeServer.Issues;
 using HordeServer.Issues.External;
@@ -124,7 +123,6 @@ namespace HordeServer
 {
 	using ContentHash = EpicGames.Core.ContentHash;
 	using ILogger = Microsoft.Extensions.Logging.ILogger;
-	using ReferenceResolver = HordeServer.Ddc.ReferenceResolver;
 
 	class Startup
 	{
@@ -317,6 +315,7 @@ namespace HordeServer
 			PluginCollection pluginCollection = new PluginCollection();
 			pluginCollection.Add<AnalyticsPlugin>();
 			pluginCollection.Add<StoragePlugin>();
+			pluginCollection.Add<DdcPlugin>();
 			services.AddSingleton<IPluginCollection>(pluginCollection);
 
 			// Register all the plugin services
@@ -516,17 +515,6 @@ namespace HordeServer
 				services.AddSingleton<IExternalIssueService, DefaultExternalIssueService>();
 			}
 
-			// DDC
-			services.AddScoped<IRequestHelper, RequestHelper>();
-			services.AddScoped<IBlobService, BlobService>();
-			services.AddScoped<IRefService, RefService>();
-			services.AddScoped<IReferenceResolver, ReferenceResolver>();
-			services.AddScoped<IContentIdStore, ContentIdStore>();
-			services.AddSingleton<BufferedPayloadFactory>();
-			services.AddSingleton<NginxRedirectHelper>();
-			services.AddSingleton<FormatResolver>();
-			services.AddSingleton<CompressedBufferUtils>();
-
 			// Secret providers
 			services.AddSingleton<SecretCollectionInternal>();
 			services.AddScoped<ISecretCollection, SecretCollection>();
diff --git a/Engine/Source/Programs/Horde/Plugins/Ddc/HordeServer.Ddc.Tests/HordeServer.Ddc.Tests.csproj b/Engine/Source/Programs/Horde/Plugins/Ddc/HordeServer.Ddc.Tests/HordeServer.Ddc.Tests.csproj
new file mode 100644
--- /dev/null
+++ b/Engine/Source/Programs/Horde/Plugins/Ddc/HordeServer.Ddc.Tests/HordeServer.Ddc.Tests.csproj
@@ -0,0 +1,77 @@
+<Project Sdk="Microsoft.NET.Sdk">
+
+  <PropertyGroup>
+    <TargetFramework>net8.0</TargetFramework>
+    <ImplicitUsings>enable</ImplicitUsings>
+    <Nullable>enable</Nullable>
+
+    <IsPackable>false</IsPackable>
+    <IsTestProject>true</IsTestProject>
+    <RootNamespace>HordeServer.Tests</RootNamespace>
+  </PropertyGroup>
+
+  <ItemGroup>
+    <PackageReference Include="coverlet.collector" Version="6.0.0" />
+    <PackageReference Include="Microsoft.NET.Test.Sdk" Version="17.8.0" />
+    <PackageReference Include="MSTest.TestAdapter" Version="3.1.1" />
+    <PackageReference Include="MSTest.TestFramework" Version="3.1.1" />
+  </ItemGroup>
+
+  <ItemGroup>
+    <ProjectReference Include="..\HordeServer.Ddc\HordeServer.Ddc.csproj" />
+    <ProjectReference Include="..\..\..\HordeServer.Tests.Shared\HordeServer.Tests.Shared.csproj" />
+  </ItemGroup>
+
+  <ItemGroup>
+    <Using Include="Microsoft.VisualStudio.TestTools.UnitTesting" />
+  </ItemGroup>
+
+  <ItemGroup>
+    <None Update="Ddc\Functional\ContentId\Payloads\05d7c699a2668efdecbe48f10db0d621d736f449.uecomp">
+      <CopyToOutputDirectory>Always</CopyToOutputDirectory>
+    </None>
+    <None Update="Ddc\Functional\ContentId\Payloads\0b29a878f71cc1eb6a26663a1152537e8fe7dca9.uecomp">
+      <CopyToOutputDirectory>Always</CopyToOutputDirectory>
+    </None>
+    <None Update="Ddc\Functional\ContentId\Payloads\4958689fe783e02fb35b13c14b0c3d7beb91e50c">
+      <CopyToOutputDirectory>Always</CopyToOutputDirectory>
+    </None>
+    <None Update="Ddc\Functional\ContentId\Payloads\7835c353d7dc67e8a0531c88fbc75ddfda10dee4">
+      <CopyToOutputDirectory>Always</CopyToOutputDirectory>
+    </None>
+    <None Update="Ddc\Functional\ContentId\Payloads\dce31eb416f3dcb4c8250ac545eda3930919d3ff">
+      <CopyToOutputDirectory>Always</CopyToOutputDirectory>
+    </None>
+    <None Update="Ddc\Functional\ContentId\Payloads\largerfile">
+      <CopyToOutputDirectory>Always</CopyToOutputDirectory>
+    </None>
+    <None Update="Ddc\Functional\ContentId\Payloads\Oodle-f0b9c675fe21951ca27699f9baab9f9f5040b202">
+      <CopyToOutputDirectory>Always</CopyToOutputDirectory>
+    </None>
+    <None Update="Ddc\Functional\ContentId\Payloads\Oodle-f895ea954b37217270e88d8b728bd3c09152689c">
+      <CopyToOutputDirectory>Always</CopyToOutputDirectory>
+    </None>
+    <None Update="Ddc\Functional\ContentId\Payloads\OodleTexture_CAS_dbda9040e75c4674fcec173f982fddf12b021e24.udd">
+      <CopyToOutputDirectory>Always</CopyToOutputDirectory>
+    </None>
+    <None Update="Ddc\Functional\ContentId\Payloads\smallerfile">
+      <CopyToOutputDirectory>Always</CopyToOutputDirectory>
+    </None>
+    <None Update="Ddc\Functional\ContentId\Payloads\UncompressedTexture_CAS_dea81b6c3b565bb5089695377c98ce0f1c13b0c3.udd">
+      <CopyToOutputDirectory>Always</CopyToOutputDirectory>
+    </None>
+    <None Update="Ddc\Functional\Objects\Payloads\lyra.cb">
+      <CopyToOutputDirectory>Always</CopyToOutputDirectory>
+    </None>
+    <None Update="Ddc\UnitTests\CompactBinaryObjects\build">
+      <CopyToOutputDirectory>Always</CopyToOutputDirectory>
+    </None>
+    <None Update="Ddc\UnitTests\CompactBinaryObjects\compact_binary">
+      <CopyToOutputDirectory>Always</CopyToOutputDirectory>
+    </None>
+    <None Update="Ddc\UnitTests\CompactBinaryObjects\ReferenceOutput">
+      <CopyToOutputDirectory>Always</CopyToOutputDirectory>
+    </None>
+  </ItemGroup>
+
+</Project>
diff --git a/Engine/Source/Programs/Horde/Plugins/Ddc/HordeServer.Ddc/DdcPlugin.cs b/Engine/Source/Programs/Horde/Plugins/Ddc/HordeServer.Ddc/DdcPlugin.cs
new file mode 100644
--- /dev/null
+++ b/Engine/Source/Programs/Horde/Plugins/Ddc/HordeServer.Ddc/DdcPlugin.cs
@@ -0,0 +1,29 @@
+// Copyright Epic Games, Inc. All Rights Reserved.
+
+using HordeServer.Ddc;
+using HordeServer.Plugins;
+using Microsoft.Extensions.DependencyInjection;
+
+namespace HordeServer
+{
+	/// <summary>
+	/// Entry point for the ddc plugin
+	/// </summary>
+	[Plugin("Ddc")]
+	public class DdcPlugin : IPluginStartup
+	{
+		/// <inheritdoc/>
+		public void ConfigureServices(IServiceCollection services)
+		{
+			services.AddScoped<IRequestHelper, RequestHelper>();
+			services.AddScoped<IBlobService, BlobService>();
+			services.AddScoped<IRefService, RefService>();
+			services.AddScoped<IReferenceResolver, ReferenceResolver>();
+			services.AddScoped<IContentIdStore, ContentIdStore>();
+			services.AddSingleton<BufferedPayloadFactory>();
+			services.AddSingleton<NginxRedirectHelper>();
+			services.AddSingleton<FormatResolver>();
+			services.AddSingleton<CompressedBufferUtils>();
+		}
+	}
+}
diff --git a/Engine/Source/Programs/Horde/Plugins/Ddc/HordeServer.Ddc/HordeServer.Ddc.csproj b/Engine/Source/Programs/Horde/Plugins/Ddc/HordeServer.Ddc/HordeServer.Ddc.csproj
new file mode 100644
--- /dev/null
+++ b/Engine/Source/Programs/Horde/Plugins/Ddc/HordeServer.Ddc/HordeServer.Ddc.csproj
@@ -0,0 +1,29 @@
+﻿<Project Sdk="Microsoft.NET.Sdk">
+
+    <PropertyGroup>
+    <TargetFramework>net8.0</TargetFramework>
+    <ImplicitUsings>enable</ImplicitUsings>
+    <Nullable>enable</Nullable>
+  </PropertyGroup>
+
+  <ItemGroup>
+    <FrameworkReference Include="Microsoft.AspNetCore.App" />
+  </ItemGroup>
+
+  <ItemGroup>
+    <PackageReference Include="Crc32.NET" Version="1.2.0" />
+    <PackageReference Include="Serilog.Extensions.Hosting" Version="8.0.0" />
+    <PackageReference Include="System.Linq.Async" Version="6.0.1" />
+  </ItemGroup>
+
+  <ItemGroup>
+    <ProjectReference Include="..\..\..\..\Shared\EpicGames.AspNet\EpicGames.AspNet.csproj" />
+    <ProjectReference Include="..\..\..\..\Shared\EpicGames.Core\EpicGames.Core.csproj" />
+    <ProjectReference Include="..\..\..\..\Shared\EpicGames.Horde\EpicGames.Horde.csproj" />
+    <ProjectReference Include="..\..\..\..\Shared\EpicGames.IoHash\EpicGames.IoHash.csproj" />
+    <ProjectReference Include="..\..\..\..\Shared\EpicGames.Redis\EpicGames.Redis.csproj" />
+    <ProjectReference Include="..\..\Storage\HordeServer.Storage\HordeServer.Storage.csproj" />
+    <ProjectReference Include="..\..\..\HordeServer.Shared\HordeServer.Shared.csproj" />
+  </ItemGroup>
+
+</Project>
diff --git a/Engine/Source/Programs/Horde/Plugins/Storage/HordeServer.Storage/StoragePlugin.cs b/Engine/Source/Programs/Horde/Plugins/Storage/HordeServer.Storage/StoragePlugin.cs
--- a/Engine/Source/Programs/Horde/Plugins/Storage/HordeServer.Storage/StoragePlugin.cs
+++ b/Engine/Source/Programs/Horde/Plugins/Storage/HordeServer.Storage/StoragePlugin.cs
@@ -52,4 +52,16 @@ namespace HordeServer.Storage
 			return new StorageBackendCache(cacheDir, _staticConfig.BundleCacheSizeBytes, serviceProvider.GetRequiredService<ILogger<StorageBackendCache>>());
 		}
 	}
+
+	/// <summary>
+	/// Helper methods for storage config
+	/// </summary>
+	public static class StoragePluginExtensions
+	{
+		/// <summary>
+		/// Configures the storage plugin
+		/// </summary>
+		public static void AddStorageConfig(this IDictionary<PluginName, IPluginConfig> dictionary, StorageConfig storageConfig)
+			=> dictionary[new PluginName("Storage")] = storageConfig;
+	}
 }
diff --git a/Engine/Source/Programs/Horde/Scripts/test.sh b/Engine/Source/Programs/Horde/Scripts/test.sh
--- a/Engine/Source/Programs/Horde/Scripts/test.sh
+++ b/Engine/Source/Programs/Horde/Scripts/test.sh
@@ -16,6 +16,7 @@ testProjects=(
 	"Source/Programs/Shared/EpicGames.Serialization.Tests/EpicGames.Serialization.Tests.csproj"
 	"Source/Programs/Horde/HordeServer.Tests/HordeServer.Tests.csproj"
 	"Source/Programs/Horde/Plugins/Analytics/HordeServer.Analytics.Tests/HordeServer.Analytics.Tests.csproj"
+	"Source/Programs/Horde/Plugins/Ddc/HordeServer.Ddc.Tests/HordeServer.Ddc.Tests.csproj"
 	"Source/Programs/Horde/Plugins/Storage/HordeServer.Storage.Tests/HordeServer.Storage.Tests.csproj"
 )
 
