Horde: Add custom user name for preflights being cancelled by Horde.

#rnx

diff --git a/Engine/Source/Programs/Horde/HordeServer/Notifications/Sinks/SlackNotificationSink.cs b/Engine/Source/Programs/Horde/HordeServer/Notifications/Sinks/SlackNotificationSink.cs
--- a/Engine/Source/Programs/Horde/HordeServer/Notifications/Sinks/SlackNotificationSink.cs
+++ b/Engine/Source/Programs/Horde/HordeServer/Notifications/Sinks/SlackNotificationSink.cs
@@ -1611,7 +1611,11 @@ namespace HordeServer.Notifications.Sinks
 
 		static string GetDefaultUserName(UserId userId)
 		{
-			if (userId == IIssue.ResolvedByUnknownId)
+			if (userId == KnownUsers.System)
+			{
+				return "Horde";
+			}
+			else if (userId == IIssue.ResolvedByUnknownId)
 			{
 				return "Horde (Unknown)";
 			}
