Horde: Fix Slack message for cancelled preflights having a stray asterisk.

#rnx

diff --git a/Engine/Source/Programs/Horde/HordeServer/Notifications/Sinks/SlackNotificationSink.cs b/Engine/Source/Programs/Horde/HordeServer/Notifications/Sinks/SlackNotificationSink.cs
--- a/Engine/Source/Programs/Horde/HordeServer/Notifications/Sinks/SlackNotificationSink.cs
+++ b/Engine/Source/Programs/Horde/HordeServer/Notifications/Sinks/SlackNotificationSink.cs
@@ -528,7 +528,7 @@ namespace HordeServer.Notifications.Sinks
 
 			if (job.AbortedByUserId != null)
 			{
-				attachment.AddSection($"*Job cancelled by {await FormatMentionAsync(job.AbortedByUserId.Value, true, cancellationToken)}");
+				attachment.AddSection($"*Job cancelled by {await FormatMentionAsync(job.AbortedByUserId.Value, true, cancellationToken)}*");
 			}
 			else if (jobOutcome == JobStepOutcome.Success)
 			{
