Horde: Move tools and secrets into plugins.

diff --git a/Engine/Source/Programs/Horde/HordeServer.Shared/Utilities/PropertyFilter.cs b/Engine/Source/Programs/Horde/HordeServer.Shared/Utilities/PropertyFilter.cs
new file mode 100644
--- /dev/null
+++ b/Engine/Source/Programs/Horde/HordeServer.Shared/Utilities/PropertyFilter.cs
@@ -0,0 +1,263 @@
+// Copyright Epic Games, Inc. All Rights Reserved.
+
+using System;
+using System.Collections;
+using System.Collections.Generic;
+using System.ComponentModel;
+using System.Globalization;
+using System.Reflection;
+using System.Text.Json.Nodes;
+
+namespace HordeServer.Utilities
+{
+	/// <summary>
+	/// Utility class to filter a response to include a subset of fields
+	/// </summary>
+	[TypeConverter(typeof(PropertyFilterTypeConverter))]
+	public class PropertyFilter
+	{
+		/// <summary>
+		/// Specific filters for properties of this object. A key of '*' will match any property not matched by anything else, and a value of null will include the whole object.
+		/// </summary>
+		readonly Dictionary<string, PropertyFilter?> _nameToFilter = new Dictionary<string, PropertyFilter?>(StringComparer.OrdinalIgnoreCase);
+
+		/// <summary>
+		/// Parse a list of fields as a property filter
+		/// </summary>
+		/// <param name="fields">List of fields to parse, separated by commas</param>
+		/// <returns>New property filter</returns>
+		public static PropertyFilter Parse(string? fields)
+		{
+			PropertyFilter rootFilter = new PropertyFilter();
+			if (fields != null)
+			{
+				foreach (string field in fields.Split(','))
+				{
+					// Split the field name into segments
+					string[] segments = field.Split('.');
+
+					// Create the tree of segments
+					Dictionary<string, PropertyFilter?> nameToFilter = rootFilter._nameToFilter;
+					for (int idx = 0; idx < segments.Length - 1; idx++)
+					{
+						PropertyFilter? nextFilter;
+						if (!nameToFilter.TryGetValue(segments[idx], out nextFilter))
+						{
+							nextFilter = new PropertyFilter();
+							nameToFilter[segments[idx]] = nextFilter;
+						}
+						else if (nextFilter == null)
+						{
+							break;
+						}
+						nameToFilter = nextFilter._nameToFilter;
+					}
+					nameToFilter[segments[^1]] = null;
+				}
+			}
+			return rootFilter;
+		}
+
+		/// <summary>
+		/// Checks whether the filter contains the given field
+		/// </summary>
+		/// <param name="filter">The filter to check</param>
+		/// <param name="field">Name of the field to check for</param>
+		/// <returns>True if the filter includes the given field</returns>
+		public static bool Includes(PropertyFilter? filter, string field)
+		{
+			return filter == null || filter.Includes(field);
+		}
+
+		/// <summary>
+		/// Attempts to apply a property filter to a response object, where the property filter may be null.
+		/// </summary>
+		/// <param name="response">Response object</param>
+		/// <param name="filter">The filter to apply</param>
+		/// <returns>Filtered object</returns>
+		public static object Apply(object response, PropertyFilter? filter)
+		{
+			return (filter == null) ? response : filter.ApplyTo(response);
+		}
+
+		/// <summary>
+		/// Checks whether the filter contains the given field
+		/// </summary>
+		/// <param name="field">Name of the field to check for</param>
+		/// <returns>True if the filter includes the given field</returns>
+		public bool Includes(string field)
+		{
+			return _nameToFilter.Count == 0 || _nameToFilter.ContainsKey(field) || _nameToFilter.ContainsKey("*");
+		}
+
+		/// <summary>
+		/// Extract the given fields from the response
+		/// </summary>
+		/// <param name="response">The response to filter</param>
+		/// <returns>Filtered response</returns>
+		public object ApplyTo(object response)
+		{
+			// If there are no filters at this depth, return the whole object
+			if (_nameToFilter.Count == 0)
+			{
+				return response;
+			}
+
+			// Check if this object is a list. If it is, we'll apply the filter to each element in the list.
+			IList? list = response as IList;
+			if (list != null)
+			{
+				List<object?> newResponse = new List<object?>();
+				for (int index = 0; index < list.Count; index++)
+				{
+					object? value = list[index];
+					if (value == null)
+					{
+						newResponse.Add(value);
+					}
+					else
+					{
+						newResponse.Add(ApplyTo(value));
+					}
+				}
+				return newResponse;
+			}
+			else
+			{
+				Dictionary<string, object?> newResponse = new Dictionary<string, object?>(StringComparer.Ordinal);
+
+				// Check if this object is a dictionary.
+				IDictionary? dictionary = response as IDictionary;
+				if (dictionary != null)
+				{
+					foreach (DictionaryEntry? entry in dictionary)
+					{
+						if (entry.HasValue)
+						{
+							string? key = entry.Value.Key.ToString();
+							if (key != null)
+							{
+								key = ConvertPascalToCamelCase(key);
+
+								PropertyFilter? filter;
+								if (_nameToFilter.TryGetValue(key, out filter))
+								{
+									object? value = entry.Value.Value;
+									if (filter == null)
+									{
+										newResponse[key] = value;
+									}
+									else if (value != null)
+									{
+										newResponse[key] = filter.ApplyTo(value);
+									}
+								}
+							}
+						}
+					}
+				}
+				else if (response is JsonObject jsonObject)
+				{
+					// Iterate the keys in this object
+					foreach ((string name, JsonNode? node) in jsonObject)
+					{
+						PropertyFilter? filter;
+						if (_nameToFilter.TryGetValue(name, out filter))
+						{
+							string key = ConvertPascalToCamelCase(name);
+							if (filter == null)
+							{
+								newResponse[key] = node;
+							}
+							else if (node != null)
+							{
+								newResponse[key] = filter.ApplyTo(node);
+							}
+						}
+					}
+				}
+				else
+				{
+					// Otherwise try to get the properties of this object
+					foreach (PropertyInfo property in response.GetType().GetProperties(BindingFlags.Public | BindingFlags.GetProperty | BindingFlags.Instance))
+					{
+						PropertyFilter? filter;
+						if (_nameToFilter.TryGetValue(property.Name, out filter))
+						{
+							string key = ConvertPascalToCamelCase(property.Name);
+							object? value = property.GetValue(response);
+							if (filter == null)
+							{
+								newResponse[key] = value;
+							}
+							else if (value != null)
+							{
+								newResponse[key] = filter.ApplyTo(value);
+							}
+						}
+					}
+				}
+				return newResponse;
+			}
+		}
+
+		/// <summary>
+		/// Convert a property name to camel case
+		/// </summary>
+		/// <param name="name">The name to convert to camelCase</param>
+		/// <returns>Camelcase version of the name</returns>
+		static string ConvertPascalToCamelCase(string name)
+		{
+			if (name[0] >= 'A' && name[0] <= 'Z')
+			{
+				return (char)(name[0] - 'A' + 'a') + name.Substring(1);
+			}
+			else
+			{
+				return name;
+			}
+		}
+	}
+
+	/// <summary>
+	/// Type converter from strings to PropertyFilter objects
+	/// </summary>
+	public class PropertyFilterTypeConverter : TypeConverter
+	{
+		/// <inheritdoc/>
+		public override bool CanConvertFrom(ITypeDescriptorContext? context, Type sourceType)
+		{
+			return sourceType == typeof(string);
+		}
+
+		/// <inheritdoc/>
+		public override object ConvertFrom(ITypeDescriptorContext? context, CultureInfo? culture, object value)
+		{
+			return PropertyFilter.Parse((string)value);
+		}
+	}
+
+	/// <summary>
+	/// Extension methods for property filters
+	/// </summary>
+	public static class PropertyFilterExtensions
+	{
+		/// <summary>
+		/// Apply a filter to a response object
+		/// </summary>
+		/// <param name="obj">The object to filter</param>
+		/// <param name="filter">The filter to apply</param>
+		/// <returns>The filtered object</returns>
+		public static object ApplyFilter(this object obj, PropertyFilter? filter)
+		{
+			if (filter == null)
+			{
+				return obj;
+			}
+			else
+			{
+				return filter.ApplyTo(obj);
+			}
+		}
+	}
+}
diff --git a/Engine/Source/Programs/Horde/HordeServer.Shared/Utilities/TryAuthorizeAttribute.cs b/Engine/Source/Programs/Horde/HordeServer.Shared/Utilities/TryAuthorizeAttribute.cs
new file mode 100644
--- /dev/null
+++ b/Engine/Source/Programs/Horde/HordeServer.Shared/Utilities/TryAuthorizeAttribute.cs
@@ -0,0 +1,28 @@
+// Copyright Epic Games, Inc. All Rights Reserved.
+
+using System;
+using System.Threading.Tasks;
+using Microsoft.AspNetCore.Authorization;
+using Microsoft.AspNetCore.Authorization.Policy;
+using Microsoft.AspNetCore.Mvc.Filters;
+using Microsoft.Extensions.DependencyInjection;
+
+namespace HordeServer.Utilities
+{
+	/// <summary>
+	/// Attempts to authorize the user, but does not challenge or forbid access if authorization fails.
+	/// </summary>
+	[AttributeUsage(AttributeTargets.Method | AttributeTargets.Class)]
+	public sealed class TryAuthorizeAttribute : Attribute, IAsyncAuthorizationFilter
+	{
+		/// <inheritdoc/>
+		public async Task OnAuthorizationAsync(AuthorizationFilterContext context)
+		{
+			IAuthorizationPolicyProvider policyProvider = context.HttpContext.RequestServices.GetRequiredService<IAuthorizationPolicyProvider>();
+			AuthorizationPolicy effectivePolicy = await policyProvider.GetDefaultPolicyAsync();
+
+			IPolicyEvaluator policyEvaluator = context.HttpContext.RequestServices.GetRequiredService<IPolicyEvaluator>();
+			await policyEvaluator.AuthenticateAsync(effectivePolicy, context.HttpContext);
+		}
+	}
+}
diff --git a/Engine/Source/Programs/Horde/HordeServer.Tests/Tools/ToolTests.cs b/Engine/Source/Programs/Horde/HordeServer.Tests/Tools/ToolTests.cs
deleted file mode 100644
diff --git a/Engine/Source/Programs/Horde/HordeServer.Tests/Tools/ToolsControllerTests.cs b/Engine/Source/Programs/Horde/HordeServer.Tests/Tools/ToolsControllerTests.cs
deleted file mode 100644
diff --git a/Engine/Source/Programs/Horde/HordeServer/Commands/Generate/ToolDataCommand.cs b/Engine/Source/Programs/Horde/HordeServer/Commands/Generate/ToolDataCommand.cs
--- a/Engine/Source/Programs/Horde/HordeServer/Commands/Generate/ToolDataCommand.cs
+++ b/Engine/Source/Programs/Horde/HordeServer/Commands/Generate/ToolDataCommand.cs
@@ -10,6 +10,7 @@ using EpicGames.Horde.Storage;
 using EpicGames.Horde.Storage.Backends;
 using EpicGames.Horde.Storage.Bundles;
 using EpicGames.Horde.Storage.Nodes;
+using HordeServer.Tools;
 using HordeServer.Utilities;
 using Microsoft.Extensions.Logging;
 
@@ -98,7 +99,9 @@ namespace HordeServer.Commands.Generate
 				JsonConfigFile serverConfig = await JsonConfigFile.ReadAsync(serverConfigFile);
 
 				JsonObject hordeConfig = JsonConfigFile.FindOrAddNode(serverConfig.Root, "Horde", () => new JsonObject());
-				JsonArray bundledTools = JsonConfigFile.FindOrAddNode(hordeConfig, nameof(ServerSettings.BundledTools), () => new JsonArray());
+				JsonObject pluginsConfig = JsonConfigFile.FindOrAddNode(hordeConfig, "Plugins", () => new JsonObject());
+				JsonObject toolsConfig = JsonConfigFile.FindOrAddNode(serverConfig.Root, "Tools", () => new JsonObject());
+				JsonArray bundledTools = JsonConfigFile.FindOrAddNode(toolsConfig, nameof(ToolsServerConfig.BundledTools), () => new JsonArray());
 
 				JsonObject bundledTool = JsonConfigFile.FindOrAddElementByKey(bundledTools, nameof(BundledToolConfig.Id), Id);
 				bundledTool[nameof(BundledToolConfig.Name)] = Name ?? Id.ToString();
diff --git a/Engine/Source/Programs/Horde/HordeServer/Configuration/ConfigService.cs b/Engine/Source/Programs/Horde/HordeServer/Configuration/ConfigService.cs
--- a/Engine/Source/Programs/Horde/HordeServer/Configuration/ConfigService.cs
+++ b/Engine/Source/Programs/Horde/HordeServer/Configuration/ConfigService.cs
@@ -588,7 +588,9 @@ namespace HordeServer.Configuration
 
 				JsonObject configObj = await context.PreprocessFileAsync(GetGlobalConfigUri(), configNode, cancellationToken);
 				CopyJsonNode(configObj, "TelemetryStores", "Plugins.Analytics.Stores");
+				CopyJsonNode(configObj, "Secrets", "Plugins.Secrets.Secrets");
 				CopyJsonNode(configObj, "Storage", "Plugins.Storage");
+				CopyJsonNode(configObj, "Tools", "Plugins.Tools.Tools");
 
 				GlobalConfig globalConfig = JsonSerializer.Deserialize<GlobalConfig>(configObj, _jsonOptions)!;
 				if (globalConfig.VersionEnum < GlobalVersion.Latest)
diff --git a/Engine/Source/Programs/Horde/HordeServer/HordeServer.csproj b/Engine/Source/Programs/Horde/HordeServer/HordeServer.csproj
--- a/Engine/Source/Programs/Horde/HordeServer/HordeServer.csproj
+++ b/Engine/Source/Programs/Horde/HordeServer/HordeServer.csproj
@@ -136,7 +136,9 @@
     <ProjectReference Include="..\HordeServer.Shared\HordeServer.Shared.csproj" />
     <ProjectReference Include="..\Plugins\Analytics\HordeServer.Analytics\HordeServer.Analytics.csproj" />
     <ProjectReference Include="..\Plugins\Ddc\HordeServer.Ddc\HordeServer.Ddc.csproj" />
+    <ProjectReference Include="..\Plugins\Secrets\HordeServer.Secrets\HordeServer.Secrets.csproj" />
     <ProjectReference Include="..\Plugins\Storage\HordeServer.Storage\HordeServer.Storage.csproj" />
+    <ProjectReference Include="..\Plugins\Tools\HordeServer.Tools\HordeServer.Tools.csproj" />
   </ItemGroup>
 
   <ItemGroup>
diff --git a/Engine/Source/Programs/Horde/HordeServer/Plugins/PluginCollection.cs b/Engine/Source/Programs/Horde/HordeServer/Plugins/PluginCollection.cs
--- a/Engine/Source/Programs/Horde/HordeServer/Plugins/PluginCollection.cs
+++ b/Engine/Source/Programs/Horde/HordeServer/Plugins/PluginCollection.cs
@@ -11,7 +11,7 @@ namespace HordeServer.Plugins
 	/// <summary>
 	/// Concrete implementation of <see cref="IPluginCollection"/>
 	/// </summary>
-	class PluginCollection : IPluginCollection
+	public class PluginCollection : IPluginCollection
 	{
 		// Plugin which is compiled into the application or already loaded
 		class StaticPlugin : IPlugin
diff --git a/Engine/Source/Programs/Horde/HordeServer/Secrets/ISecret.cs b/Engine/Source/Programs/Horde/HordeServer/Secrets/ISecret.cs
deleted file mode 100644
diff --git a/Engine/Source/Programs/Horde/HordeServer/Secrets/ISecretCollection.cs b/Engine/Source/Programs/Horde/HordeServer/Secrets/ISecretCollection.cs
deleted file mode 100644
diff --git a/Engine/Source/Programs/Horde/HordeServer/Secrets/ISecretProvider.cs b/Engine/Source/Programs/Horde/HordeServer/Secrets/ISecretProvider.cs
deleted file mode 100644
diff --git a/Engine/Source/Programs/Horde/HordeServer/Secrets/Providers/AwsParameterStoreSecretProvider.cs b/Engine/Source/Programs/Horde/HordeServer/Secrets/Providers/AwsParameterStoreSecretProvider.cs
deleted file mode 100644
diff --git a/Engine/Source/Programs/Horde/HordeServer/Secrets/SecretAclAction.cs b/Engine/Source/Programs/Horde/HordeServer/Secrets/SecretAclAction.cs
deleted file mode 100644
diff --git a/Engine/Source/Programs/Horde/HordeServer/Secrets/SecretCollection.cs b/Engine/Source/Programs/Horde/HordeServer/Secrets/SecretCollection.cs
deleted file mode 100644
diff --git a/Engine/Source/Programs/Horde/HordeServer/Secrets/SecretConfig.cs b/Engine/Source/Programs/Horde/HordeServer/Secrets/SecretConfig.cs
deleted file mode 100644
diff --git a/Engine/Source/Programs/Horde/HordeServer/Secrets/SecretsController.cs b/Engine/Source/Programs/Horde/HordeServer/Secrets/SecretsController.cs
deleted file mode 100644
diff --git a/Engine/Source/Programs/Horde/HordeServer/Server/GlobalConfig.cs b/Engine/Source/Programs/Horde/HordeServer/Server/GlobalConfig.cs
--- a/Engine/Source/Programs/Horde/HordeServer/Server/GlobalConfig.cs
+++ b/Engine/Source/Programs/Horde/HordeServer/Server/GlobalConfig.cs
@@ -17,7 +17,6 @@ using EpicGames.Horde.Common;
 using EpicGames.Horde.Compute;
 using EpicGames.Horde.Jobs.Templates;
 using EpicGames.Horde.Projects;
-using EpicGames.Horde.Secrets;
 using EpicGames.Horde.Streams;
 using EpicGames.Horde.Tools;
 using EpicGames.Horde.Users;
@@ -39,7 +38,6 @@ using HordeServer.Notifications;
 using HordeServer.Perforce;
 using HordeServer.Plugins;
 using HordeServer.Projects;
-using HordeServer.Secrets;
 using HordeServer.Streams;
 using HordeServer.Tools;
 using HordeServer.Utilities;
@@ -202,21 +200,11 @@ namespace HordeServer.Server
 		/// </summary>
 		public List<ComputeClusterConfig> Compute { get; set; } = new List<ComputeClusterConfig>();
 
-		/// <summary>
-		/// List of secrets
-		/// </summary>
-		public List<SecretConfig> Secrets { get; set; } = new List<SecretConfig>();
-
 		/// <summary>
 		/// Device configuration
 		/// </summary>
 		public DeviceConfig? Devices { get; set; }
 
-		/// <summary>
-		/// List of tools hosted by the server
-		/// </summary>
-		public List<ToolConfig> Tools { get; set; } = new List<ToolConfig>();
-
 		/// <summary>
 		/// Maximum number of conforms to run at once
 		/// </summary>
@@ -273,11 +261,9 @@ namespace HordeServer.Server
 
 		private readonly Dictionary<ProjectId, ProjectConfig> _projectLookup = new Dictionary<ProjectId, ProjectConfig>();
 		private readonly Dictionary<StreamId, StreamConfig> _streamLookup = new Dictionary<StreamId, StreamConfig>();
-		private readonly Dictionary<ToolId, ToolConfig> _toolLookup = new Dictionary<ToolId, ToolConfig>();
 		private readonly Dictionary<ClusterId, ComputeClusterConfig> _computeClusterLookup = new Dictionary<ClusterId, ComputeClusterConfig>();
 		private readonly Dictionary<AclScopeName, AclConfig> _aclLookup = new Dictionary<AclScopeName, AclConfig>();
 		private readonly Dictionary<ArtifactType, ArtifactTypeConfig> _artifactTypeLookup = new Dictionary<ArtifactType, ArtifactTypeConfig>();
-		private readonly Dictionary<SecretId, SecretConfig> _secretLookup = new Dictionary<SecretId, SecretConfig>();
 		private readonly Dictionary<PoolId, PoolConfig> _poolLookup = new Dictionary<PoolId, PoolConfig>();
 
 		/// <inheritdoc cref="AclConfig.Authorize(AclAction, ClaimsPrincipal)"/>
@@ -309,13 +295,6 @@ namespace HordeServer.Server
 				}
 			}
 
-			_toolLookup.Clear();
-			foreach (ToolConfig tool in Tools)
-			{
-				_toolLookup.Add(tool.Id, tool);
-				tool.PostLoad(this);
-			}
-
 			_computeClusterLookup.Clear();
 			foreach (ComputeClusterConfig computeCluster in Compute)
 			{
@@ -329,13 +308,6 @@ namespace HordeServer.Server
 				_artifactTypeLookup.Add(artifactType.Type, artifactType);
 			}
 
-			_secretLookup.Clear();
-			foreach (SecretConfig secret in Secrets)
-			{
-				_secretLookup.Add(secret.Id, secret);
-				secret.PostLoad(this);
-			}
-
 			_poolLookup.Clear();
 			foreach (PoolConfig pool in Pools)
 			{
@@ -558,14 +530,6 @@ namespace HordeServer.Server
 			return streamConfig.TryGetTemplate(templateId, out config);
 		}
 
-		/// <summary>
-		/// Attempts to get configuration for a tool from this object
-		/// </summary>
-		/// <param name="toolId">The tool identifier</param>
-		/// <param name="config">Configuration for the tool</param>
-		/// <returns>True if the tool configuration was found</returns>
-		public bool TryGetTool(ToolId toolId, [NotNullWhen(true)] out ToolConfig? config) => _toolLookup.TryGetValue(toolId, out config);
-
 		/// <summary>
 		/// Attempts to get configuration for a pool from this object
 		/// </summary>
@@ -637,21 +601,6 @@ namespace HordeServer.Server
 		public bool TryGetComputeCluster(ClusterId clusterId, [NotNullWhen(true)] out ComputeClusterConfig? config)
 			=> _computeClusterLookup.TryGetValue(clusterId, out config);
 
-		/// <summary>
-		/// Attempts to get compute cluster configuration from this object
-		/// </summary>
-		/// <param name="secretId">Secret id</param>
-		/// <param name="config">Receives the secret configuration on success</param>
-		/// <returns>True on success</returns>
-		public bool TryGetSecret(SecretId secretId, [NotNullWhen(true)] out SecretConfig? config)
-			=> _secretLookup.TryGetValue(secretId, out config);
-
-		/// <summary>
-		/// Authorize access to a secret
-		/// </summary>
-		public bool Authorize(SecretId secretId, AclAction action, ClaimsPrincipal user)
-			=> TryGetSecret(secretId, out SecretConfig? secretConfig) && secretConfig.Authorize(action, user);
-
 		/// <summary>
 		/// Authorizes a user to perform a given action
 		/// </summary>
diff --git a/Engine/Source/Programs/Horde/HordeServer/ServerSettings.cs b/Engine/Source/Programs/Horde/HordeServer/ServerSettings.cs
--- a/Engine/Source/Programs/Horde/HordeServer/ServerSettings.cs
+++ b/Engine/Source/Programs/Horde/HordeServer/ServerSettings.cs
@@ -3,20 +3,16 @@
 using System;
 using System.Collections.Generic;
 using System.ComponentModel;
-using System.Diagnostics.CodeAnalysis;
 using System.Linq;
 using System.Net;
 using EpicGames.Core;
 using EpicGames.Horde;
 using EpicGames.Horde.Server;
-using EpicGames.Horde.Storage;
 using EpicGames.Horde.Storage.Bundles;
 using EpicGames.Horde.Storage.Nodes;
-using EpicGames.Horde.Tools;
 using EpicGames.Perforce;
 using HordeServer.Agents.Fleet;
 using HordeServer.Server;
-using HordeServer.Tools;
 
 namespace HordeServer
 {
@@ -85,35 +81,6 @@ namespace HordeServer
 		public ChunkingOptions Chunking { get; set; } = new ChunkingOptions();
 	}
 
-	/// <summary>
-	/// Configuration for a tool bundled alongsize the server
-	/// </summary>
-	public class BundledToolConfig : ToolConfig
-	{
-		/// <summary>
-		/// Version string for the current tool data
-		/// </summary>
-		public string Version { get; set; } = "1.0";
-
-		/// <summary>
-		/// Ref name in the tools directory
-		/// </summary>
-		public RefName RefName { get; set; } = new RefName("default-ref");
-
-		/// <summary>
-		/// Directory containing blob data for this tool. If empty, the tools/{id} folder next to the server will be used.
-		/// </summary>
-		public string? DataDir { get; set; }
-
-		/// <summary>
-		/// Constructor
-		/// </summary>
-		public BundledToolConfig()
-		{
-			Public = true;
-		}
-	}
-
 	/// <summary>
 	/// OpenTelemetry configuration for collection and sending of traces and metrics.
 	/// </summary>
@@ -737,11 +704,6 @@ namespace HordeServer
 		/// </summary>
 		public CommitSettings Commits { get; set; } = new CommitSettings();
 
-		/// <summary>
-		/// Tools bundled along with the server. Data for each tool can be produced using the 'bundle create' command, and should be stored in the Tools directory.
-		/// </summary>
-		public List<BundledToolConfig> BundledTools { get; set; } = new List<BundledToolConfig>();
-
 		/// <summary>
 		/// Options for OpenTelemetry
 		/// </summary>
@@ -772,18 +734,6 @@ namespace HordeServer
 				throw new ArgumentException($"Settings key '{nameof(RunModes)}' contains one or more invalid entries");
 			}
 		}
-
-		/// <summary>
-		/// Attempts to get a bundled tool with the given id
-		/// </summary>
-		/// <param name="toolId">The tool id</param>
-		/// <param name="bundledToolConfig">Configuration for the bundled tool</param>
-		/// <returns>True if the tool was found</returns>
-		public bool TryGetBundledTool(ToolId toolId, [NotNullWhen(true)] out BundledToolConfig? bundledToolConfig)
-		{
-			bundledToolConfig = BundledTools.FirstOrDefault(x => x.Id == toolId);
-			return bundledToolConfig != null;
-		}
 	}
 
 	/// <summary>
diff --git a/Engine/Source/Programs/Horde/HordeServer/Startup.cs b/Engine/Source/Programs/Horde/HordeServer/Startup.cs
--- a/Engine/Source/Programs/Horde/HordeServer/Startup.cs
+++ b/Engine/Source/Programs/Horde/HordeServer/Startup.cs
@@ -68,15 +68,12 @@ using HordeServer.Notifications.Sinks;
 using HordeServer.Perforce;
 using HordeServer.Plugins;
 using HordeServer.Replicators;
-using HordeServer.Secrets;
-using HordeServer.Secrets.Providers;
 using HordeServer.Server;
 using HordeServer.Server.Notices;
 using HordeServer.ServiceAccounts;
 using HordeServer.Storage;
 using HordeServer.Streams;
 using HordeServer.Tasks;
-using HordeServer.Tools;
 using HordeServer.Ugs;
 using HordeServer.Users;
 using HordeServer.Utilities;
@@ -316,6 +313,8 @@ namespace HordeServer
 			pluginCollection.Add<AnalyticsPlugin>();
 			pluginCollection.Add<StoragePlugin>();
 			pluginCollection.Add<DdcPlugin>();
+			pluginCollection.Add<SecretsPlugin>();
+			pluginCollection.Add<ToolsPlugin>();
 			services.AddSingleton<IPluginCollection>(pluginCollection);
 
 			// Register all the plugin services
@@ -406,8 +405,6 @@ namespace HordeServer
 			services.AddSingleton<INoticeCollection, NoticeCollection>();
 			services.AddSingleton<IDashboardPreviewCollection, DashboardPreviewCollection>();
 
-			services.AddSingleton<IToolCollection, ToolCollection>();
-
 			services.AddSingleton<IConfigSource, InMemoryConfigSource>();
 			services.AddSingleton<IConfigSource, FileConfigSource>();
 			services.AddSingleton<IConfigSource, PerforceConfigSource>();
@@ -515,11 +512,6 @@ namespace HordeServer
 				services.AddSingleton<IExternalIssueService, DefaultExternalIssueService>();
 			}
 
-			// Secret providers
-			services.AddSingleton<SecretCollectionInternal>();
-			services.AddScoped<ISecretCollection, SecretCollection>();
-			services.AddSingleton<ISecretProvider, AwsParameterStoreSecretProvider>();
-
 			if (settings.WithAws)
 			{
 				AWSOptions awsOptions = Configuration.GetAWSOptions();
diff --git a/Engine/Source/Programs/Horde/HordeServer/Tools/ITool.cs b/Engine/Source/Programs/Horde/HordeServer/Tools/ITool.cs
deleted file mode 100644
diff --git a/Engine/Source/Programs/Horde/HordeServer/Tools/IToolCollection.cs b/Engine/Source/Programs/Horde/HordeServer/Tools/IToolCollection.cs
deleted file mode 100644
diff --git a/Engine/Source/Programs/Horde/HordeServer/Tools/ToolAclAction.cs b/Engine/Source/Programs/Horde/HordeServer/Tools/ToolAclAction.cs
deleted file mode 100644
diff --git a/Engine/Source/Programs/Horde/HordeServer/Tools/ToolCollection.cs b/Engine/Source/Programs/Horde/HordeServer/Tools/ToolCollection.cs
deleted file mode 100644
diff --git a/Engine/Source/Programs/Horde/HordeServer/Tools/ToolConfig.cs b/Engine/Source/Programs/Horde/HordeServer/Tools/ToolConfig.cs
deleted file mode 100644
diff --git a/Engine/Source/Programs/Horde/HordeServer/Tools/ToolsController.cs b/Engine/Source/Programs/Horde/HordeServer/Tools/ToolsController.cs
deleted file mode 100644
diff --git a/Engine/Source/Programs/Horde/HordeServer/Utilities/PropertyFilter.cs b/Engine/Source/Programs/Horde/HordeServer/Utilities/PropertyFilter.cs
deleted file mode 100644
diff --git a/Engine/Source/Programs/Horde/HordeServer/Utilities/TryAuthorizeAttribute.cs b/Engine/Source/Programs/Horde/HordeServer/Utilities/TryAuthorizeAttribute.cs
deleted file mode 100644
diff --git a/Engine/Source/Programs/Horde/Plugins/Secrets/HordeServer.Secrets/HordeServer.Secrets.csproj b/Engine/Source/Programs/Horde/Plugins/Secrets/HordeServer.Secrets/HordeServer.Secrets.csproj
new file mode 100644
--- /dev/null
+++ b/Engine/Source/Programs/Horde/Plugins/Secrets/HordeServer.Secrets/HordeServer.Secrets.csproj
@@ -0,0 +1,27 @@
+<Project Sdk="Microsoft.NET.Sdk">
+
+  <PropertyGroup>
+    <TargetFramework>net8.0</TargetFramework>
+    <ImplicitUsings>enable</ImplicitUsings>
+    <Nullable>enable</Nullable>
+    <RootNamespace>HordeServer</RootNamespace>
+  </PropertyGroup>
+
+  <ItemGroup>
+    <FrameworkReference Include="Microsoft.AspNetCore.App" />
+  </ItemGroup>
+
+  <ItemGroup>
+    <PackageReference Include="AWSSDK.SimpleSystemsManagement" Version="3.7.304.24" />
+  </ItemGroup>
+
+  <ItemGroup>
+    <ProjectReference Include="..\..\..\..\Shared\EpicGames.Core\EpicGames.Core.csproj" />
+    <ProjectReference Include="..\..\..\..\Shared\EpicGames.Horde\EpicGames.Horde.csproj" />
+    <ProjectReference Include="..\..\..\..\Shared\EpicGames.IoHash\EpicGames.IoHash.csproj" />
+    <ProjectReference Include="..\..\..\..\Shared\EpicGames.Redis\EpicGames.Redis.csproj" />
+    <ProjectReference Include="..\..\Storage\HordeServer.Storage\HordeServer.Storage.csproj" />
+    <ProjectReference Include="..\..\..\HordeServer.Shared\HordeServer.Shared.csproj" />
+  </ItemGroup>
+
+</Project>
diff --git a/Engine/Source/Programs/Horde/Plugins/Secrets/HordeServer.Secrets/Secrets/ISecret.cs b/Engine/Source/Programs/Horde/Plugins/Secrets/HordeServer.Secrets/Secrets/ISecret.cs
new file mode 100644
--- /dev/null
+++ b/Engine/Source/Programs/Horde/Plugins/Secrets/HordeServer.Secrets/Secrets/ISecret.cs
@@ -0,0 +1,23 @@
+// Copyright Epic Games, Inc. All Rights Reserved.
+
+using System.Collections.Generic;
+using EpicGames.Horde.Secrets;
+
+namespace HordeServer.Secrets
+{
+	/// <summary>
+	/// Information about a secret
+	/// </summary>
+	public interface ISecret
+	{
+		/// <summary>
+		/// Identifier for the secret
+		/// </summary>
+		SecretId Id { get; }
+
+		/// <summary>
+		/// The secret values
+		/// </summary>
+		IReadOnlyDictionary<string, string> Data { get; }
+	}
+}
diff --git a/Engine/Source/Programs/Horde/Plugins/Secrets/HordeServer.Secrets/Secrets/ISecretCollection.cs b/Engine/Source/Programs/Horde/Plugins/Secrets/HordeServer.Secrets/Secrets/ISecretCollection.cs
new file mode 100644
--- /dev/null
+++ b/Engine/Source/Programs/Horde/Plugins/Secrets/HordeServer.Secrets/Secrets/ISecretCollection.cs
@@ -0,0 +1,21 @@
+﻿// Copyright Epic Games, Inc. All Rights Reserved.
+
+using System.Threading;
+using System.Threading.Tasks;
+using EpicGames.Horde.Secrets;
+
+namespace HordeServer.Secrets
+{
+	/// <summary>
+	/// Collection of secrets
+	/// </summary>
+	public interface ISecretCollection
+	{
+		/// <summary>
+		/// Resolve a secret to concrete values
+		/// </summary>
+		/// <param name="secretId">Identifier for the secret</param>
+		/// <param name="cancellationToken">Cancellation token for the operation</param>
+		Task<ISecret?> GetAsync(SecretId secretId, CancellationToken cancellationToken);
+	}
+}
diff --git a/Engine/Source/Programs/Horde/Plugins/Secrets/HordeServer.Secrets/Secrets/ISecretProvider.cs b/Engine/Source/Programs/Horde/Plugins/Secrets/HordeServer.Secrets/Secrets/ISecretProvider.cs
new file mode 100644
--- /dev/null
+++ b/Engine/Source/Programs/Horde/Plugins/Secrets/HordeServer.Secrets/Secrets/ISecretProvider.cs
@@ -0,0 +1,25 @@
+// Copyright Epic Games, Inc. All Rights Reserved.
+
+using System.Threading;
+using System.Threading.Tasks;
+
+namespace HordeServer.Secrets
+{
+	/// <summary>
+	/// Interface for a service that can provide secret data
+	/// </summary>
+	public interface ISecretProvider
+	{
+		/// <summary>
+		/// Name of this provider
+		/// </summary>
+		string Name { get; }
+
+		/// <summary>
+		/// Gets values for a particular secret
+		/// </summary>
+		/// <param name="path">Path to the secret to fetch</param>
+		/// <param name="cancellationToken">Cancellation token for the operation</param>
+		Task<string> GetSecretAsync(string path, CancellationToken cancellationToken);
+	}
+}
diff --git a/Engine/Source/Programs/Horde/Plugins/Secrets/HordeServer.Secrets/Secrets/Providers/AwsParameterStoreSecretProvider.cs b/Engine/Source/Programs/Horde/Plugins/Secrets/HordeServer.Secrets/Secrets/Providers/AwsParameterStoreSecretProvider.cs
new file mode 100644
--- /dev/null
+++ b/Engine/Source/Programs/Horde/Plugins/Secrets/HordeServer.Secrets/Secrets/Providers/AwsParameterStoreSecretProvider.cs
@@ -0,0 +1,42 @@
+// Copyright Epic Games, Inc. All Rights Reserved.
+
+using System;
+using System.Threading;
+using System.Threading.Tasks;
+using Amazon.Extensions.NETCore.Setup;
+using Amazon.SimpleSystemsManagement;
+using Amazon.SimpleSystemsManagement.Model;
+
+namespace HordeServer.Secrets.Providers
+{
+	/// <summary>
+	/// Fetches secrets from the AWS parameter store
+	/// </summary>
+	public class AwsParameterStoreSecretProvider : ISecretProvider
+	{
+		/// <inheritdoc/>
+		public string Name => "AwsParameterStore";
+
+		readonly IAmazonSimpleSystemsManagement _systemsManagement;
+
+		/// <summary>
+		/// Constructor
+		/// </summary>
+		public AwsParameterStoreSecretProvider()
+		{
+			AWSOptions awsOptions = new AWSOptions();
+			_systemsManagement = awsOptions.CreateServiceClient<IAmazonSimpleSystemsManagement>();
+		}
+
+		/// <inheritdoc/>
+		public async Task<string> GetSecretAsync(string path, CancellationToken cancellationToken)
+		{
+			GetParameterResponse response = await _systemsManagement.GetParameterAsync(new GetParameterRequest { Name = path, WithDecryption = true }, cancellationToken);
+			if (response.Parameter == null)
+			{
+				throw new InvalidOperationException($"Unable to fetch secret '{path}' from AWS parameter store ({response.HttpStatusCode})");
+			}
+			return response.Parameter.Value;
+		}
+	}
+}
diff --git a/Engine/Source/Programs/Horde/Plugins/Secrets/HordeServer.Secrets/Secrets/SecretAclAction.cs b/Engine/Source/Programs/Horde/Plugins/Secrets/HordeServer.Secrets/Secrets/SecretAclAction.cs
new file mode 100644
--- /dev/null
+++ b/Engine/Source/Programs/Horde/Plugins/Secrets/HordeServer.Secrets/Secrets/SecretAclAction.cs
@@ -0,0 +1,17 @@
+// Copyright Epic Games, Inc. All Rights Reserved.
+
+using EpicGames.Horde.Acls;
+
+namespace HordeServer.Secrets
+{
+	/// <summary>
+	/// ACL actions for secrets
+	/// </summary>
+	public static class SecretAclAction
+	{
+		/// <summary>
+		/// View a credential
+		/// </summary>
+		public static AclAction ViewSecret { get; } = new AclAction("ViewSecret");
+	}
+}
diff --git a/Engine/Source/Programs/Horde/Plugins/Secrets/HordeServer.Secrets/Secrets/SecretCollection.cs b/Engine/Source/Programs/Horde/Plugins/Secrets/HordeServer.Secrets/Secrets/SecretCollection.cs
new file mode 100644
--- /dev/null
+++ b/Engine/Source/Programs/Horde/Plugins/Secrets/HordeServer.Secrets/Secrets/SecretCollection.cs
@@ -0,0 +1,115 @@
+// Copyright Epic Games, Inc. All Rights Reserved.
+
+using System;
+using System.Collections.Generic;
+using System.Linq;
+using System.Text.Json;
+using System.Threading;
+using System.Threading.Tasks;
+using EpicGames.Horde.Secrets;
+using HordeServer.Server;
+using Microsoft.Extensions.Options;
+
+namespace HordeServer.Secrets
+{
+	/// <summary>
+	/// Implementation of <see cref="ISecretCollection"/>
+	/// </summary>
+	public class SecretCollection : ISecretCollection
+	{
+		readonly SecretCollectionInternal _secretCollectionInternal;
+		readonly IOptionsSnapshot<SecretsConfig> _secretsConfig;
+
+		/// <summary>
+		/// Constructor
+		/// </summary>
+		public SecretCollection(SecretCollectionInternal secretCollectionInternal, IOptionsSnapshot<SecretsConfig> secretsConfig)
+		{
+			_secretCollectionInternal = secretCollectionInternal;
+			_secretsConfig = secretsConfig;
+		}
+
+		/// <inheritdoc/>
+		public async Task<ISecret?> GetAsync(SecretId secretId, CancellationToken cancellationToken)
+		{
+			SecretConfig? secretConfig;
+			if (!_secretsConfig.Value.TryGetSecret(secretId, out secretConfig))
+			{
+				return null;
+			}
+			return await _secretCollectionInternal.GetAsync(secretConfig, cancellationToken);
+		}
+	}
+
+	/// <summary>
+	/// Core implementation of <see cref="ISecretCollection"/> which exists as a global service.
+	/// </summary>
+	public class SecretCollectionInternal
+	{
+		class Secret : ISecret
+		{
+			public SecretId Id { get; }
+			public IReadOnlyDictionary<string, string> Data { get; }
+
+			public Secret(SecretId id, IReadOnlyDictionary<string, string> data)
+			{
+				Id = id;
+				Data = data;
+			}
+		}
+
+		readonly Dictionary<string, ISecretProvider> _providers;
+
+		/// <summary>
+		/// Constructor
+		/// </summary>
+		public SecretCollectionInternal(IEnumerable<ISecretProvider> providers)
+		{
+			_providers = providers.ToDictionary(x => x.Name, x => x);
+		}
+
+		/// <summary>
+		/// Resolve a secret to concrete values
+		/// </summary>
+		/// <param name="config">Configuration for the secret</param>
+		/// <param name="cancellationToken">Cancellation token for the operation</param>
+		public async Task<ISecret> GetAsync(SecretConfig config, CancellationToken cancellationToken)
+		{
+			// Add the hard-coded secrets
+			Dictionary<string, string> data = new Dictionary<string, string>();
+			foreach ((string key, string value) in config.Data)
+			{
+				data.Add(key, value);
+			}
+
+			// Fetch all the secrets from external providers
+			foreach (ExternalSecretConfig externalConfig in config.Sources)
+			{
+				ISecretProvider provider = _providers[externalConfig.Provider];
+
+				string secret = await provider.GetSecretAsync(externalConfig.Path, cancellationToken);
+				if (externalConfig.Format == ExternalSecretFormat.Text)
+				{
+					data[externalConfig.Key ?? "default"] = secret;
+				}
+				else if (externalConfig.Format == ExternalSecretFormat.Json)
+				{
+					Dictionary<string, string>? values = JsonSerializer.Deserialize<Dictionary<string, string>>(secret, new JsonSerializerOptions { AllowTrailingCommas = true });
+					if (values != null)
+					{
+						foreach ((string key, string value) in values)
+						{
+							data[key] = value;
+						}
+					}
+				}
+				else
+				{
+					throw new NotImplementedException();
+				}
+			}
+
+			return new Secret(config.Id, data);
+		}
+	}
+}
diff --git a/Engine/Source/Programs/Horde/Plugins/Secrets/HordeServer.Secrets/Secrets/SecretConfig.cs b/Engine/Source/Programs/Horde/Plugins/Secrets/HordeServer.Secrets/Secrets/SecretConfig.cs
new file mode 100644
--- /dev/null
+++ b/Engine/Source/Programs/Horde/Plugins/Secrets/HordeServer.Secrets/Secrets/SecretConfig.cs
@@ -0,0 +1,99 @@
+// Copyright Epic Games, Inc. All Rights Reserved.
+
+using System;
+using System.Collections.Generic;
+using System.Security.Claims;
+using EpicGames.Horde.Acls;
+using EpicGames.Horde.Secrets;
+using HordeServer.Acls;
+using HordeServer.Server;
+
+#pragma warning disable CA2227 // Change x to be read-only by removing the property setter
+
+namespace HordeServer.Secrets
+{
+	/// <summary>
+	/// Configuration for a secret value
+	/// </summary>
+	public class SecretConfig
+	{
+		/// <summary>
+		/// Identifier for this secret
+		/// </summary>
+		public SecretId Id { get; set; }
+
+		/// <summary>
+		/// Key/value pairs associated with this secret
+		/// </summary>
+		public Dictionary<string, string> Data { get; set; } = new Dictionary<string, string>();
+
+		/// <summary>
+		/// Providers to source key/value pairs from
+		/// </summary>
+		public List<ExternalSecretConfig> Sources { get; set; } = new List<ExternalSecretConfig>();
+
+		/// <summary>
+		/// Defines access to this particular secret
+		/// </summary>
+		public AclConfig Acl { get; set; } = new AclConfig();
+
+		/// <summary>
+		/// Called after the config has been read
+		/// </summary>
+		/// <param name="parentAcl">Parent acl object</param>
+		public void PostLoad(AclConfig parentAcl)
+		{
+			Acl.PostLoad(parentAcl, $"secret:{Id}");
+		}
+
+		/// <summary>
+		/// Authorizes a user to perform a given action
+		/// </summary>
+		/// <param name="action">The action being performed</param>
+		/// <param name="user">The principal to validate</param>
+		public bool Authorize(AclAction action, ClaimsPrincipal user)
+			=> Acl.Authorize(action, user);
+	}
+
+	/// <summary>
+	/// Format describing how to parse external secret values
+	/// </summary>
+	public enum ExternalSecretFormat
+	{
+		/// <summary>
+		/// Secret is a plain text value which will be stored using the external secret key
+		/// </summary>
+		Text,
+
+		/// <summary>
+		/// Secret is a JSON formatted string containing key/value pairs
+		/// </summary>
+		Json,
+	}
+
+	/// <summary>
+	/// Configuration for an external secret provider
+	/// </summary>
+	public class ExternalSecretConfig
+	{
+		/// <summary>
+		/// Name of the provider to use
+		/// </summary>
+		public string Provider { get; set; } = String.Empty;
+
+		/// <summary>
+		/// Format of the secret
+		/// </summary>
+		public ExternalSecretFormat Format { get; set; }
+
+		/// <summary>
+		/// Optional key indicating the parameter to set in the resulting data array. Required if if <see cref="Format"/> is <see cref="ExternalSecretFormat.Text"/>.
+		/// </summary>
+		public string? Key { get; set; }
+
+		/// <summary>
+		/// Optional value indicating what to fetch from the provider
+		/// </summary>
+		public string Path { get; set; } = "default";
+	}
+}
diff --git a/Engine/Source/Programs/Horde/Plugins/Secrets/HordeServer.Secrets/Secrets/SecretsController.cs b/Engine/Source/Programs/Horde/Plugins/Secrets/HordeServer.Secrets/Secrets/SecretsController.cs
new file mode 100644
--- /dev/null
+++ b/Engine/Source/Programs/Horde/Plugins/Secrets/HordeServer.Secrets/Secrets/SecretsController.cs
@@ -0,0 +1,79 @@
+// Copyright Epic Games, Inc. All Rights Reserved.
+
+using System.Collections.Generic;
+using System.Linq;
+using System.Threading;
+using System.Threading.Tasks;
+using EpicGames.Horde.Secrets;
+using HordeServer.Server;
+using HordeServer.Utilities;
+using Microsoft.AspNetCore.Authorization;
+using Microsoft.AspNetCore.Mvc;
+using Microsoft.Extensions.Options;
+
+namespace HordeServer.Secrets
+{
+	/// <summary>
+	/// Controller for the /api/v1/secrets endpoint
+	/// </summary>
+	[ApiController]
+	[Authorize]
+	[Route("[controller]")]
+	public class SecretsController : HordeControllerBase
+	{
+		readonly ISecretCollection _secretCollection;
+		readonly IOptionsSnapshot<SecretsConfig> _secretsConfig;
+
+		/// <summary>
+		/// Constructor
+		/// </summary>
+		public SecretsController(ISecretCollection secretCollection, IOptionsSnapshot<SecretsConfig> secretsConfig)
+		{
+			_secretCollection = secretCollection;
+			_secretsConfig = secretsConfig;
+		}
+
+		/// <summary>
+		/// Query all the secrets available for the current user
+		/// </summary>
+		[HttpGet]
+		[Route("/api/v1/secrets")]
+		public ActionResult<GetSecretsResponse> GetSecrets()
+		{
+			List<SecretId> secretIds = new List<SecretId>();
+			foreach (SecretConfig secret in _secretsConfig.Value.Secrets)
+			{
+				if (secret.Authorize(SecretAclAction.ViewSecret, User))
+				{
+					secretIds.Add(secret.Id);
+				}
+			}
+			return new GetSecretsResponse(secretIds);
+		}
+
+		/// <summary>
+		/// Retrieve information about a specific secret
+		/// </summary>
+		/// <param name="secretId">Id of the secret to retrieve</param>
+		/// <param name="filter">Filter for properties to return</param>
+		/// <param name="cancellationToken">Cancellation token for the operation</param>
+		/// <returns>Information about the requested secret</returns>
+		[HttpGet]
+		[Route("/api/v1/secrets/{secretId}")]
+		[ProducesResponseType(typeof(GetSecretResponse), 200)]
+		public async Task<ActionResult<object>> GetSecretAsync(SecretId secretId, [FromQuery] PropertyFilter? filter = null, CancellationToken cancellationToken = default)
+		{
+			ISecret? secret = await _secretCollection.GetAsync(secretId, cancellationToken);
+			if (secret == null)
+			{
+				return NotFound(secretId);
+			}
+			if (!_secretsConfig.Value.Authorize(secret.Id, SecretAclAction.ViewSecret, User))
+			{
+				return Forbid(SecretAclAction.ViewSecret, secretId);
+			}
+
+			return new GetSecretResponse(secret.Id, secret.Data.ToDictionary(x => x.Key, x => x.Value)).ApplyFilter(filter);
+		}
+	}
+}
diff --git a/Engine/Source/Programs/Horde/Plugins/Secrets/HordeServer.Secrets/SecretsConfig.cs b/Engine/Source/Programs/Horde/Plugins/Secrets/HordeServer.Secrets/SecretsConfig.cs
new file mode 100644
--- /dev/null
+++ b/Engine/Source/Programs/Horde/Plugins/Secrets/HordeServer.Secrets/SecretsConfig.cs
@@ -0,0 +1,53 @@
+// Copyright Epic Games, Inc. All Rights Reserved.
+
+using System.Diagnostics.CodeAnalysis;
+using System.Security.Claims;
+using EpicGames.Horde.Acls;
+using EpicGames.Horde.Secrets;
+using HordeServer.Acls;
+using HordeServer.Plugins;
+using HordeServer.Secrets;
+
+#pragma warning disable CA2227 // Change x to be read-only by removing the property setter
+
+namespace HordeServer
+{
+	/// <summary>
+	/// Configuration for the secrets system
+	/// </summary>
+	public class SecretsConfig : IPluginConfig
+	{
+		/// <summary>
+		/// List of secrets
+		/// </summary>
+		public List<SecretConfig> Secrets { get; set; } = new List<SecretConfig>();
+
+		private readonly Dictionary<SecretId, SecretConfig> _secretLookup = new Dictionary<SecretId, SecretConfig>();
+
+		/// <inheritdoc/>
+		public void PostLoad(AclConfig parentAcl)
+		{
+			_secretLookup.Clear();
+			foreach (SecretConfig secret in Secrets)
+			{
+				_secretLookup.Add(secret.Id, secret);
+				secret.PostLoad(parentAcl);
+			}
+		}
+
+		/// <summary>
+		/// Attempts to get a secret configuration from this object
+		/// </summary>
+		/// <param name="secretId">Secret id</param>
+		/// <param name="config">Receives the secret configuration on success</param>
+		/// <returns>True on success</returns>
+		public bool TryGetSecret(SecretId secretId, [NotNullWhen(true)] out SecretConfig? config)
+			=> _secretLookup.TryGetValue(secretId, out config);
+
+		/// <summary>
+		/// Authorize access to a secret
+		/// </summary>
+		public bool Authorize(SecretId secretId, AclAction action, ClaimsPrincipal user)
+			=> TryGetSecret(secretId, out SecretConfig? secretConfig) && secretConfig.Authorize(action, user);
+	}
+}
diff --git a/Engine/Source/Programs/Horde/Plugins/Secrets/HordeServer.Secrets/SecretsPlugin.cs b/Engine/Source/Programs/Horde/Plugins/Secrets/HordeServer.Secrets/SecretsPlugin.cs
new file mode 100644
--- /dev/null
+++ b/Engine/Source/Programs/Horde/Plugins/Secrets/HordeServer.Secrets/SecretsPlugin.cs
@@ -0,0 +1,36 @@
+// Copyright Epic Games, Inc. All Rights Reserved.
+
+using HordeServer.Plugins;
+using HordeServer.Secrets;
+using HordeServer.Secrets.Providers;
+using Microsoft.Extensions.DependencyInjection;
+
+namespace HordeServer
+{
+	/// <summary>
+	/// Entry point for the storage plugin
+	/// </summary>
+	[Plugin("Secrets", GlobalConfigType = typeof(SecretsConfig))]
+	public class SecretsPlugin : IPluginStartup
+	{
+		/// <inheritdoc/>
+		public void ConfigureServices(IServiceCollection services)
+		{
+			services.AddSingleton<SecretCollectionInternal>();
+			services.AddScoped<ISecretCollection, SecretCollection>();
+			services.AddSingleton<ISecretProvider, AwsParameterStoreSecretProvider>();
+		}
+	}
+
+	/// <summary>
+	/// Helper methods for secrets config
+	/// </summary>
+	public static class SecretsPluginExtensions
+	{
+		/// <summary>
+		/// Configures the secrets plugin
+		/// </summary>
+		public static void AddSecretsConfig(this IDictionary<PluginName, IPluginConfig> dictionary, SecretsConfig secretsConfig)
+			=> dictionary[new PluginName("Secrets")] = secretsConfig;
+	}
+}
diff --git a/Engine/Source/Programs/Horde/Plugins/Tools/HordeServer.Tools.Tests/HordeServer.Tools.Tests.csproj b/Engine/Source/Programs/Horde/Plugins/Tools/HordeServer.Tools.Tests/HordeServer.Tools.Tests.csproj
new file mode 100644
--- /dev/null
+++ b/Engine/Source/Programs/Horde/Plugins/Tools/HordeServer.Tools.Tests/HordeServer.Tools.Tests.csproj
@@ -0,0 +1,29 @@
+<Project Sdk="Microsoft.NET.Sdk">
+
+  <PropertyGroup>
+    <TargetFramework>net8.0</TargetFramework>
+    <ImplicitUsings>enable</ImplicitUsings>
+    <Nullable>enable</Nullable>
+
+    <IsPackable>false</IsPackable>
+    <IsTestProject>true</IsTestProject>
+    <RootNamespace>HordeServer.Tests</RootNamespace>
+  </PropertyGroup>
+
+  <ItemGroup>
+    <PackageReference Include="coverlet.collector" Version="6.0.0" />
+    <PackageReference Include="Microsoft.NET.Test.Sdk" Version="17.8.0" />
+    <PackageReference Include="MSTest.TestAdapter" Version="3.1.1" />
+    <PackageReference Include="MSTest.TestFramework" Version="3.1.1" />
+  </ItemGroup>
+
+  <ItemGroup>
+    <ProjectReference Include="..\HordeServer.Tools\HordeServer.Tools.csproj" />
+    <ProjectReference Include="..\..\..\HordeServer.Tests.Shared\HordeServer.Tests.Shared.csproj" />
+  </ItemGroup>
+
+  <ItemGroup>
+    <Using Include="Microsoft.VisualStudio.TestTools.UnitTesting" />
+  </ItemGroup>
+
+</Project>
diff --git a/Engine/Source/Programs/Horde/Plugins/Tools/HordeServer.Tools.Tests/Tools/ToolTests.cs b/Engine/Source/Programs/Horde/Plugins/Tools/HordeServer.Tools.Tests/Tools/ToolTests.cs
new file mode 100644
--- /dev/null
+++ b/Engine/Source/Programs/Horde/Plugins/Tools/HordeServer.Tools.Tests/Tools/ToolTests.cs
@@ -0,0 +1,163 @@
+// Copyright Epic Games, Inc. All Rights Reserved.
+
+using System;
+using System.IO;
+using System.IO.Compression;
+using System.Linq;
+using System.Text;
+using System.Threading;
+using System.Threading.Tasks;
+using EpicGames.Horde.Tools;
+using HordeServer.Server;
+using HordeServer.Storage;
+using HordeServer.Tools;
+using HordeCommon;
+using Microsoft.Extensions.DependencyInjection;
+using Microsoft.VisualStudio.TestTools.UnitTesting;
+using HordeServer.Plugins;
+
+namespace HordeServer.Tests.Tools
+{
+	[TestClass]
+	public class ToolTests : ServerServiceTest
+	{
+		readonly ToolId _toolId = new ToolId("ugs");
+
+		public ToolTests()
+		{
+			AddPlugin<StoragePlugin>();
+			AddPlugin<ToolsPlugin>();
+
+			StorageConfig storageConfig = new StorageConfig();
+			storageConfig.Backends.Clear();
+			storageConfig.Backends.Add(new BackendConfig { Id = new BackendId("tools-backend"), Type = StorageBackendType.Memory });
+			storageConfig.Namespaces.Clear();
+			storageConfig.Namespaces.Add(new NamespaceConfig { Id = Namespace.Tools, Backend = new BackendId("tools-backend") });
+
+			ToolsConfig toolsConfig = new ToolsConfig();
+			toolsConfig.Tools.Add(new ToolConfig(_toolId) { Name = "UnrealGameSync", Description = "Tool for syncing content from source control" });
+
+			GlobalConfig globalConfig = new GlobalConfig();
+			globalConfig.Plugins.AddStorageConfig(storageConfig);
+			globalConfig.Plugins.AddToolsConfig(toolsConfig);
+			SetConfig(globalConfig);
+		}
+
+		[TestMethod]
+		public async Task AddToolAsync()
+		{
+			IToolCollection collection = ServiceProvider.GetRequiredService<IToolCollection>();
+
+			ITool tool = Deref(await collection.GetAsync(_toolId));
+			Assert.AreEqual(tool.Id, new ToolId("ugs"));
+			Assert.AreEqual(tool.Name, "UnrealGameSync");
+			Assert.AreEqual(tool.Description, "Tool for syncing content from source control");
+
+			ITool tool2 = Deref(await collection.GetAsync(tool.Id));
+			Assert.AreEqual(tool.Id, tool2.Id);
+			Assert.AreEqual(tool.Name, tool2.Name);
+			Assert.AreEqual(tool.Description, tool2.Description);
+		}
+
+		[TestMethod]
+		public async Task AddDeploymentAsync()
+		{
+			IToolCollection collection = ServiceProvider.GetRequiredService<IToolCollection>();
+
+			ITool tool = Deref(await collection.GetAsync(_toolId));
+			Assert.AreEqual(new ToolId("ugs"), tool.Id);
+			Assert.AreEqual("UnrealGameSync", tool.Name);
+			Assert.AreEqual("Tool for syncing content from source control", tool.Description);
+			Assert.AreEqual(0, tool.Deployments.Count);
+
+			const string FileName = "test.txt";
+			byte[] fileData = Encoding.UTF8.GetBytes("hello world");
+
+			byte[] zipData;
+			using (MemoryStream stream = new MemoryStream())
+			{
+				using (ZipArchive archive = new ZipArchive(stream, ZipArchiveMode.Create))
+				{
+					ZipArchiveEntry entry = archive.CreateEntry(FileName);
+					using (Stream entryStream = entry.Open())
+					{
+						await entryStream.WriteAsync(fileData);
+					}
+				}
+				zipData = stream.ToArray();
+			}
+
+			ToolDeploymentId deploymentId;
+			using (MemoryStream stream = new MemoryStream(zipData))
+			{
+				tool = Deref(await tool.CreateDeploymentAsync(new ToolDeploymentConfig { Version = "1.0", Duration = TimeSpan.FromMinutes(5.0), CreatePaused = true }, stream, CancellationToken.None));
+				Assert.AreEqual(1, tool.Deployments.Count);
+				Assert.IsNull(tool.Deployments[0].StartedAt);
+				deploymentId = tool.Deployments[^1].Id;
+			}
+
+			// Check that the deployment doesn't do anything until started
+			FakeClock clock = ServiceProvider.GetRequiredService<FakeClock>();
+			await clock.AdvanceAsync(TimeSpan.FromHours(1.0));
+
+			tool = Deref(await collection.GetAsync(tool.Id));
+			Assert.AreEqual(1, tool.Deployments.Count);
+			Assert.IsNull(tool.Deployments[0].StartedAt);
+
+			// Start the deployment
+			IToolDeployment deployment = tool.Deployments[0];
+			deployment = Deref(await deployment.UpdateAsync(ToolDeploymentState.Active));
+			tool = Deref(await collection.GetAsync(tool.Id));
+			Assert.AreEqual(1, tool.Deployments.Count);
+			Assert.IsNotNull(tool.Deployments[0].StartedAt);
+			Assert.IsTrue(Math.Abs((tool.Deployments[0].StartedAt!.Value - clock.UtcNow).TotalSeconds) < 1.0);
+
+			// Check it updates
+			await clock.AdvanceAsync(TimeSpan.FromMinutes(2.5));
+			deployment = Deref(await deployment.UpdateAsync(ToolDeploymentState.Paused));
+			tool = Deref(await collection.GetAsync(tool.Id));
+			Assert.AreEqual(1, tool.Deployments.Count);
+			Assert.IsNull(deployment.StartedAt);
+			Assert.IsTrue((deployment.Progress - 0.5) < 0.1);
+
+			// Check it stays paused
+			await clock.AdvanceAsync(TimeSpan.FromHours(1.0));
+			tool = Deref(await collection.GetAsync(tool.Id));
+			Assert.AreEqual(1, tool.Deployments.Count);
+			Assert.IsNull(tool.Deployments[0].StartedAt);
+			Assert.IsTrue((tool.Deployments[0].Progress - 0.5) < 0.1);
+
+			// Get the deployment data
+			using Stream dataStream = await tool.Deployments[0].OpenZipStreamAsync(CancellationToken.None);
+			using (ZipArchive archive = new ZipArchive(dataStream, ZipArchiveMode.Read))
+			{
+				ZipArchiveEntry entry = archive.Entries.First();
+				Assert.AreEqual(entry.Name, FileName);
+
+				byte[] outputFileData;
+				using (MemoryStream stream = new MemoryStream())
+				{
+					using Stream entryStream = entry.Open();
+					await entryStream.CopyToAsync(stream);
+					outputFileData = stream.ToArray();
+				}
+
+				Assert.IsTrue(fileData.AsSpan().SequenceEqual(outputFileData.AsSpan()));
+			}
+		}
+
+		public static async Task<byte[]> CreateZipFileDataAsync(string fileName, string fileData)
+		{
+			using MemoryStream stream = new ();
+			using (ZipArchive archive = new (stream, ZipArchiveMode.Create))
+			{
+				ZipArchiveEntry entry = archive.CreateEntry(fileName);
+				await using (Stream entryStream = entry.Open())
+				{
+					await entryStream.WriteAsync(Encoding.UTF8.GetBytes(fileData));
+				}
+			}
+			return stream.ToArray();
+		}
+	}
+}
\ No newline at end of file
diff --git a/Engine/Source/Programs/Horde/Plugins/Tools/HordeServer.Tools.Tests/Tools/ToolsControllerTests.cs b/Engine/Source/Programs/Horde/Plugins/Tools/HordeServer.Tools.Tests/Tools/ToolsControllerTests.cs
new file mode 100644
--- /dev/null
+++ b/Engine/Source/Programs/Horde/Plugins/Tools/HordeServer.Tools.Tests/Tools/ToolsControllerTests.cs
@@ -0,0 +1,78 @@
+﻿// Copyright Epic Games, Inc. All Rights Reserved.
+
+using System.Collections.Generic;
+using System.IO;
+using System.Net;
+using System.Net.Http;
+using System.Net.Http.Headers;
+using System.Threading.Tasks;
+using EpicGames.Horde.Server;
+using EpicGames.Horde.Tools;
+using HordeServer.Acls;
+using HordeServer.Configuration;
+using HordeServer.Plugins;
+using HordeServer.Server;
+using HordeServer.ServiceAccounts;
+using HordeServer.Storage;
+using HordeServer.Tools;
+using HordeServer.Users;
+using HordeServer.Utilities;
+using Microsoft.Extensions.DependencyInjection;
+using Microsoft.VisualStudio.TestTools.UnitTesting;
+
+namespace HordeServer.Tests.Tools;
+
+[TestClass]
+public class ToolsControllerTests
+{
+	[TestMethod]
+	public async Task DownloadToolWithAclAsync()
+	{
+		AclEntryConfig aclEntryConfig = new(HordeClaims.AgentRoleClaim, [ToolAclAction.DownloadTool]);
+		ToolId toolId = new("foo");
+
+		PluginCollection pluginCollection = new PluginCollection();
+		pluginCollection.Add<AnalyticsPlugin>();
+		pluginCollection.Add<StoragePlugin>();
+
+		StorageConfig storageConfig = new StorageConfig();
+		storageConfig.Backends.Clear();
+		storageConfig.Backends.Add(new BackendConfig { Id = new BackendId("tools-backend"), Type = StorageBackendType.Memory });
+		storageConfig.Namespaces.Clear();
+		storageConfig.Namespaces.Add(new NamespaceConfig { Id = Namespace.Tools, Backend = new BackendId("tools-backend") });
+
+		ToolsConfig toolsConfig = new ToolsConfig();
+		toolsConfig.Tools.Add(new ToolConfig(toolId) { Name = "Foo", Description = "This is foo", Acl = new AclConfig() { Entries = [aclEntryConfig] }, Public = false });
+
+		GlobalConfig globalConfig = new();
+		globalConfig.Plugins.AddStorageConfig(storageConfig);
+		globalConfig.Plugins.AddToolsConfig(toolsConfig);
+		
+		ServerSettings serverSettings = new() { AuthMethod = AuthMethod.Horde };
+		globalConfig.PostLoad(serverSettings, pluginCollection.LoadedPlugins);
+
+		Dictionary<string, string> settings = new() { { "Horde:AuthMethod", AuthMethod.Horde.ToString() } };
+		await using FakeHordeWebApp app = new(settings: settings);
+		
+		ConfigService configService = app.ServiceProvider.GetRequiredService<ConfigService>();
+		IToolCollection tools = app.ServiceProvider.GetRequiredService<IToolCollection>();
+		IServiceAccountCollection serviceAccounts = app.ServiceProvider.GetRequiredService<IServiceAccountCollection>();
+
+		configService.OverrideConfig(globalConfig);
+
+		List<IUserClaim> claims = [new UserClaim("http://epicgames.com/ue/horde/role", "agent")];
+		(IServiceAccount _, string token) = await serviceAccounts.CreateAsync(new CreateServiceAccountOptions("myDesc", claims));
+		
+		// Create tool and deployment
+		using MemoryStream ms = new(await ToolTests.CreateZipFileDataAsync("foo.txt", "foo content"));
+		ITool? tool = await tools.GetAsync(toolId);
+		Assert.IsNotNull(tool);
+		tool = await tool.CreateDeploymentAsync(new ToolDeploymentConfig() { Version = "1" }, ms);
+
+		HttpClient client = app.CreateHttpClient();
+		using HttpRequestMessage req = new(HttpMethod.Get, $"/api/v1/tools/{toolId.Id}?action=download");
+		req.Headers.Authorization = new AuthenticationHeaderValue("ServiceAccount", token);
+		HttpResponseMessage res = await client.SendAsync(req);
+		Assert.AreEqual(HttpStatusCode.OK, res.StatusCode);
+	}
+}
diff --git a/Engine/Source/Programs/Horde/Plugins/Tools/HordeServer.Tools/HordeServer.Tools.csproj b/Engine/Source/Programs/Horde/Plugins/Tools/HordeServer.Tools/HordeServer.Tools.csproj
new file mode 100644
--- /dev/null
+++ b/Engine/Source/Programs/Horde/Plugins/Tools/HordeServer.Tools/HordeServer.Tools.csproj
@@ -0,0 +1,26 @@
+﻿<Project Sdk="Microsoft.NET.Sdk">
+
+  <PropertyGroup>
+    <TargetFramework>net8.0</TargetFramework>
+    <ImplicitUsings>enable</ImplicitUsings>
+    <Nullable>enable</Nullable>
+    <RootNamespace>HordeServer</RootNamespace>
+  </PropertyGroup>
+
+  <ItemGroup>
+    <FrameworkReference Include="Microsoft.AspNetCore.App" />
+  </ItemGroup>
+
+  <ItemGroup>
+  </ItemGroup>
+
+  <ItemGroup>
+    <ProjectReference Include="..\..\..\..\Shared\EpicGames.Core\EpicGames.Core.csproj" />
+    <ProjectReference Include="..\..\..\..\Shared\EpicGames.Horde\EpicGames.Horde.csproj" />
+    <ProjectReference Include="..\..\..\..\Shared\EpicGames.IoHash\EpicGames.IoHash.csproj" />
+    <ProjectReference Include="..\..\..\..\Shared\EpicGames.Redis\EpicGames.Redis.csproj" />
+    <ProjectReference Include="..\..\Storage\HordeServer.Storage\HordeServer.Storage.csproj" />
+    <ProjectReference Include="..\..\..\HordeServer.Shared\HordeServer.Shared.csproj" />
+  </ItemGroup>
+
+</Project>
diff --git a/Engine/Source/Programs/Horde/Plugins/Tools/HordeServer.Tools/Properties/AssemblyInfo.cs b/Engine/Source/Programs/Horde/Plugins/Tools/HordeServer.Tools/Properties/AssemblyInfo.cs
new file mode 100644
--- /dev/null
+++ b/Engine/Source/Programs/Horde/Plugins/Tools/HordeServer.Tools/Properties/AssemblyInfo.cs
@@ -0,0 +1,7 @@
+// Copyright Epic Games, Inc. All Rights Reserved.
+
+using System.Resources;
+using System.Runtime.CompilerServices;
+
+[assembly: InternalsVisibleTo("HordeServer.Tools.Tests")]
+[assembly: NeutralResourcesLanguage("en-US")]
diff --git a/Engine/Source/Programs/Horde/Plugins/Tools/HordeServer.Tools/Tools/ITool.cs b/Engine/Source/Programs/Horde/Plugins/Tools/HordeServer.Tools/Tools/ITool.cs
new file mode 100644
--- /dev/null
+++ b/Engine/Source/Programs/Horde/Plugins/Tools/HordeServer.Tools/Tools/ITool.cs
@@ -0,0 +1,219 @@
+// Copyright Epic Games, Inc. All Rights Reserved.
+
+using System;
+using System.Collections.Generic;
+using System.IO;
+using System.Security.Claims;
+using System.Threading;
+using System.Threading.Tasks;
+using EpicGames.Horde.Acls;
+using EpicGames.Horde.Storage;
+using EpicGames.Horde.Tools;
+
+#pragma warning disable CA1716 // Rename virtual/interface member ITool.Public so that it no longer conflicts with the reserved language keyword 'Public'.
+
+namespace HordeServer.Tools
+{
+	/// <summary>
+	/// Describes a standalone, external tool hosted and deployed by Horde. Provides basic functionality for performing
+	/// gradual roll-out, versioning, etc...
+	/// </summary>
+	public interface ITool
+	{
+		/// <summary>
+		/// Identifier for the tool
+		/// </summary>
+		ToolId Id { get; }
+
+		/// <summary>
+		/// Name of the tool
+		/// </summary>
+		string Name { get; }
+
+		/// <summary>
+		/// Long-form description of the tool
+		/// </summary>
+		string Description { get; }
+
+		/// <summary>
+		/// Category for the tool on the dashboard
+		/// </summary>
+		string? Category { get; }
+
+		/// <summary>
+		/// Grouping key for merging tool versions together on the dashboard
+		/// </summary>
+		string? Group { get; }
+
+		/// <summary>
+		/// Supported platforms, as a NET runtime identifiers
+		/// </summary>
+		IReadOnlyList<string>? Platforms { get; }
+
+		/// <summary>
+		/// Whether the tool is available to authenticated users
+		/// </summary>
+		bool Public { get; }
+
+		/// <summary>
+		/// Whether to show the tool for download in UGS
+		/// </summary>
+		bool ShowInUgs { get; }
+
+		/// <summary>
+		/// Whether to show the tool for download in the dashboard
+		/// </summary>
+		bool ShowInDashboard { get; }
+
+		/// <summary>
+		/// Current deployments of this tool, sorted by time.
+		/// </summary>
+		IReadOnlyList<IToolDeployment> Deployments { get; }
+
+		/// <summary>
+		/// Authorize a user to perform a particular action
+		/// </summary>
+		/// <param name="action">Action the user is trying to perform</param>
+		/// <param name="principal">Identity of the user trying to perform the action</param>
+		bool Authorize(AclAction action, ClaimsPrincipal principal);
+
+		/// <summary>
+		/// Adds a new deployment to the given tool. The new deployment will replace the current active deployment.
+		/// </summary>
+		/// <param name="options">Options for the new deployment</param>
+		/// <param name="stream">Stream containing the tool data</param>
+		/// <param name="cancellationToken">Cancellation token for the operation</param>
+		/// <returns>Updated tool document, or null if it does not exist</returns>
+		Task<ITool?> CreateDeploymentAsync(ToolDeploymentConfig options, Stream stream, CancellationToken cancellationToken = default);
+
+		/// <summary>
+		/// Adds a new deployment to the given tool. The new deployment will replace the current active deployment.
+		/// </summary>
+		/// <param name="options">Options for the new deployment</param>
+		/// <param name="target">Path to the root node containing the tool data</param>
+		/// <param name="cancellationToken">Cancellation token for the operation</param>
+		/// <returns>Updated tool document, or null if it does not exist</returns>
+		Task<ITool?> CreateDeploymentAsync(ToolDeploymentConfig options, BlobRefValue target, CancellationToken cancellationToken = default);
+
+		/// <summary>
+		/// Gets the storage backend for a particular tool
+		/// </summary>
+		/// <returns>Instance of the backend client</returns>
+		IStorageBackend CreateStorageBackend();
+
+		/// <summary>
+		/// Gets the storage backend for a particular tool
+		/// </summary>
+		/// <returns>Instance of the storage client</returns>
+		IStorageClient CreateStorageClient();
+	}
+
+	/// <summary>
+	/// Deployment of a tool
+	/// </summary>
+	public interface IToolDeployment
+	{
+		/// <summary>
+		/// Identifier for this deployment. A new identifier will be assigned to each created instance, so an identifier corresponds to a unique deployment.
+		/// </summary>
+		ToolDeploymentId Id { get; }
+
+		/// <summary>
+		/// Descriptive version string for this tool revision
+		/// </summary>
+		string Version { get; }
+
+		/// <summary>
+		/// Current state of this deployment
+		/// </summary>
+		ToolDeploymentState State { get; }
+
+		/// <summary>
+		/// Current progress of the deployment
+		/// </summary>
+		double Progress { get; }
+
+		/// <summary>
+		/// Last time at which the progress started. Set to null if the deployment was paused.
+		/// </summary>
+		DateTime? StartedAt { get; }
+
+		/// <summary>
+		/// Length of time over which to make the deployment
+		/// </summary>
+		TimeSpan Duration { get; }
+
+		/// <summary>
+		/// Namespace containing the tool
+		/// </summary>
+		NamespaceId NamespaceId { get; }
+
+		/// <summary>
+		/// Reference to this tool in Horde Storage.
+		/// </summary>
+		RefName RefName { get; }
+
+		/// <summary>
+		/// Updates the state of the current deployment
+		/// </summary>
+		/// <param name="state">New state of the deployment</param>
+		/// <param name="cancellationToken">Cancellation token for the operation</param>
+		/// <returns></returns>
+		Task<IToolDeployment?> UpdateAsync(ToolDeploymentState state, CancellationToken cancellationToken = default);
+
+		/// <summary>
+		/// Opens a stream to the data for a particular deployment
+		/// </summary>
+		/// <param name="cancellationToken">Cancellation token for the operation</param>
+		/// <returns>Stream for the data</returns>
+		Task<Stream> OpenZipStreamAsync(CancellationToken cancellationToken = default);
+	}
+
+	/// <summary>
+	/// Extension methods for tools
+	/// </summary>
+	public static class ToolExtensions
+	{
+		/// <summary>
+		/// Gets the current deployment
+		/// </summary>
+		/// <param name="tool">Tool to query</param>
+		/// <param name="phase">Adoption phase for the caller. 0 which the </param>
+		/// <param name="utcNow">Current time</param>
+		/// <returns></returns>
+		public static IToolDeployment? GetCurrentDeployment(this ITool tool, double phase, DateTime utcNow)
+		{
+			int idx = tool.Deployments.Count - 1;
+			for (; idx >= 0; idx--)
+			{
+				if (phase <= tool.Deployments[idx].GetProgressValue(utcNow) || idx == 0)
+				{
+					return tool.Deployments[idx];
+				}
+			}
+			return null;
+		}
+
+		/// <summary>
+		/// Get the progress fraction for a particular deployment and time
+		/// </summary>
+		/// <param name="deployment"></param>
+		/// <param name="utcNow"></param>
+		/// <returns></returns>
+		public static double GetProgressValue(this IToolDeployment deployment, DateTime utcNow)
+		{
+			if (deployment.StartedAt == null)
+			{
+				return deployment.Progress;
+			}
+			else if (deployment.Duration > TimeSpan.Zero)
+			{
+				return Math.Clamp((utcNow - deployment.StartedAt.Value) / deployment.Duration, 0.0, 1.0);
+			}
+			else
+			{
+				return 1.0;
+			}
+		}
+	}
+}
diff --git a/Engine/Source/Programs/Horde/Plugins/Tools/HordeServer.Tools/Tools/IToolCollection.cs b/Engine/Source/Programs/Horde/Plugins/Tools/HordeServer.Tools/Tools/IToolCollection.cs
new file mode 100644
--- /dev/null
+++ b/Engine/Source/Programs/Horde/Plugins/Tools/HordeServer.Tools/Tools/IToolCollection.cs
@@ -0,0 +1,30 @@
+// Copyright Epic Games, Inc. All Rights Reserved.
+
+using System.Collections.Generic;
+using System.Threading;
+using System.Threading.Tasks;
+using EpicGames.Horde.Tools;
+
+namespace HordeServer.Tools
+{
+	/// <summary>
+	/// Collection of tools
+	/// </summary>
+	public interface IToolCollection
+	{
+		/// <summary>
+		/// Gets a tool with the given identifier
+		/// </summary>
+		/// <param name="id">The tool identifier</param>
+		/// <param name="cancellationToken">Cancellation token for the operation</param>
+		/// <returns>The requested tool, or null if it does not exist</returns>
+		Task<ITool?> GetAsync(ToolId id, CancellationToken cancellationToken = default);
+
+		/// <summary>
+		/// Gets all the available tools
+		/// </summary>
+		/// <param name="cancellationToken"></param>
+		/// <returns>List of the available tools</returns>
+		Task<IReadOnlyList<ITool>> GetAllAsync(CancellationToken cancellationToken = default);
+	}
+}
diff --git a/Engine/Source/Programs/Horde/Plugins/Tools/HordeServer.Tools/Tools/ToolAclAction.cs b/Engine/Source/Programs/Horde/Plugins/Tools/HordeServer.Tools/Tools/ToolAclAction.cs
new file mode 100644
--- /dev/null
+++ b/Engine/Source/Programs/Horde/Plugins/Tools/HordeServer.Tools/Tools/ToolAclAction.cs
@@ -0,0 +1,22 @@
+// Copyright Epic Games, Inc. All Rights Reserved.
+
+using EpicGames.Horde.Acls;
+
+namespace HordeServer.Tools
+{
+	/// <summary>
+	/// ACL actions which apply to tools
+	/// </summary>
+	public static class ToolAclAction
+	{
+		/// <summary>
+		/// Ability to download a tool
+		/// </summary>
+		public static AclAction DownloadTool { get; } = new AclAction("DownloadTool");
+
+		/// <summary>
+		/// Ability to upload new tool versions
+		/// </summary>
+		public static AclAction UploadTool { get; } = new AclAction("UploadTool");
+	}
+}
diff --git a/Engine/Source/Programs/Horde/Plugins/Tools/HordeServer.Tools/Tools/ToolCollection.cs b/Engine/Source/Programs/Horde/Plugins/Tools/HordeServer.Tools/Tools/ToolCollection.cs
new file mode 100644
--- /dev/null
+++ b/Engine/Source/Programs/Horde/Plugins/Tools/HordeServer.Tools/Tools/ToolCollection.cs
@@ -0,0 +1,527 @@
+// Copyright Epic Games, Inc. All Rights Reserved.
+
+using System;
+using System.Buffers.Binary;
+using System.Collections.Generic;
+using System.IO;
+using System.Linq;
+using System.Security.Claims;
+using System.Text;
+using System.Threading;
+using System.Threading.Tasks;
+using EpicGames.Core;
+using EpicGames.Horde;
+using EpicGames.Horde.Acls;
+using EpicGames.Horde.Storage;
+using EpicGames.Horde.Storage.Backends;
+using EpicGames.Horde.Storage.Bundles;
+using EpicGames.Horde.Storage.Nodes;
+using EpicGames.Horde.Storage.ObjectStores;
+using EpicGames.Horde.Tools;
+using HordeServer.Server;
+using HordeServer.Storage;
+using HordeServer.Utilities;
+using Microsoft.Extensions.Logging;
+using Microsoft.Extensions.Options;
+using MongoDB.Bson.Serialization.Attributes;
+using MongoDB.Driver;
+
+namespace HordeServer.Tools
+{
+	/// <summary>
+	/// Collection of tool documents
+	/// </summary>
+	public class ToolCollection : IToolCollection
+	{
+		class ToolDocument
+		{
+			public ToolId Id { get; set; }
+
+			[BsonElement("dep")]
+			public List<ToolDeploymentDocument> Deployments { get; set; } = new List<ToolDeploymentDocument>();
+
+			// Last time that the document was updated. This field is checked and updated as part of updates to ensure atomicity.
+			[BsonElement("_u")]
+			public DateTime LastUpdateTime { get; set; }
+
+			[BsonConstructor]
+			public ToolDocument(ToolId id)
+			{
+				Id = id;
+			}
+		}
+
+		class ToolDeploymentDocument
+		{
+			public ToolDeploymentId Id { get; set; }
+
+			[BsonElement("ver")]
+			public string Version { get; set; }
+
+			[BsonElement("bpr")]
+			public double BaseProgress { get; set; }
+
+			[BsonElement("stm")]
+			public DateTime? StartedAt { get; set; }
+
+			[BsonElement("dur")]
+			public TimeSpan Duration { get; set; }
+
+			[BsonElement("ns")]
+			public NamespaceId NamespaceId { get; set; } = Namespace.Tools;
+
+			[BsonElement("ref")]
+			public RefName RefName { get; set; }
+
+			[BsonConstructor]
+			public ToolDeploymentDocument(ToolDeploymentId id)
+			{
+				Id = id;
+				Version = String.Empty;
+			}
+
+			public ToolDeploymentDocument(ToolDeploymentId id, ToolDeploymentConfig options, NamespaceId namespaceId, RefName refName)
+			{
+				Id = id;
+				Version = options.Version;
+				Duration = options.Duration;
+				NamespaceId = namespaceId;
+				RefName = refName;
+			}
+		}
+
+		class Tool : ITool
+		{
+			readonly ToolCollection _collection;
+			readonly ToolDocument _document;
+			readonly ToolConfig _config;
+			readonly List<ToolDeployment> _deployments;
+
+			public ToolConfig Config => _config;
+			public ToolDocument Document => _document;
+
+			public ToolId Id => _config.Id;
+			public string Name => _config.Name;
+			public string Description => _config.Description;
+			public string? Category => _config.Category;
+			public string? Group => _config.Group;
+			public IReadOnlyList<string>? Platforms => _config.Platforms;
+			public bool Public => _config.Public;
+			public bool ShowInUgs => _config.ShowInUgs;
+			public bool ShowInDashboard => _config.ShowInDashboard;
+			public IReadOnlyList<IToolDeployment> Deployments => _deployments;
+
+			public Tool(ToolCollection collection, ToolDocument document, ToolConfig config, DateTime utcNow)
+			{
+				_collection = collection;
+				_document = document;
+				_config = config;
+				_deployments = document.Deployments.ConvertAll(x => new ToolDeployment(this, _collection, x, utcNow));
+			}
+
+			public bool Authorize(AclAction action, ClaimsPrincipal principal)
+				=> _config.Authorize(action, principal);
+
+			public async Task<ITool?> CreateDeploymentAsync(ToolDeploymentConfig options, Stream stream, CancellationToken cancellationToken = default)
+			{
+				ToolDocument? document = await _collection.CreateDeploymentAsync(_document, _config, options, stream, cancellationToken);
+				return _collection.CreateToolObject(document);
+			}
+
+			public async Task<ITool?> CreateDeploymentAsync(ToolDeploymentConfig options, BlobRefValue target, CancellationToken cancellationToken = default)
+			{
+				ToolDocument? document = await _collection.CreateDeploymentAsync(_document, _config, options, target, cancellationToken);
+				return _collection.CreateToolObject(document);
+			}
+
+			public IStorageBackend CreateStorageBackend()
+				=> _collection.CreateStorageBackend(_config);
+
+			public IStorageClient CreateStorageClient()
+				=> _collection.CreateStorageClient(_config);
+		}
+
+		class ToolDeployment : IToolDeployment
+		{
+			readonly Tool _tool;
+			readonly ToolCollection _collection;
+			readonly ToolDeploymentDocument _document;
+
+			public ToolDeploymentId Id => _document.Id;
+			public string Version => _document.Version;
+			public ToolDeploymentState State { get; }
+			public double Progress { get; }
+			public DateTime? StartedAt => _document.StartedAt;
+			public TimeSpan Duration => _document.Duration;
+			public NamespaceId NamespaceId => _document.NamespaceId;
+			public RefName RefName => _document.RefName;
+
+			public ToolDeployment(Tool tool, ToolCollection collection, ToolDeploymentDocument document, DateTime utcNow)
+			{
+				_tool = tool;
+				_collection = collection;
+				_document = document;
+
+				if (document.BaseProgress >= 1.0)
+				{
+					State = ToolDeploymentState.Complete;
+					Progress = 1.0;
+				}
+				else if (StartedAt == null)
+				{
+					State = ToolDeploymentState.Paused;
+					Progress = document.BaseProgress;
+				}
+				else if (Duration > TimeSpan.Zero)
+				{
+					State = ToolDeploymentState.Active;
+					Progress = Math.Clamp((utcNow - StartedAt.Value) / Duration, 0.0, 1.0);
+				}
+				else
+				{
+					State = ToolDeploymentState.Complete;
+					Progress = 1.0;
+				}
+			}
+
+			public Task<Stream> OpenZipStreamAsync(CancellationToken cancellationToken = default)
+				=> _collection.GetDeploymentZipAsync(_tool.Config, _document, cancellationToken);
+
+			public async Task<IToolDeployment?> UpdateAsync(ToolDeploymentState action, CancellationToken cancellationToken = default)
+			{
+				ToolDocument? newDocument = await _collection.UpdateDeploymentAsync(_tool.Document, Id, action, Progress, cancellationToken);
+				Tool? newTool = _collection.CreateToolObject(newDocument);
+				return newTool?.Deployments.FirstOrDefault(x => x.Id == Id);
+			}
+		}
+
+		private readonly IServerInfo _serverInfo;
+		private readonly IMongoCollection<ToolDocument> _tools;
+		private readonly StorageService _storageService;
+		private readonly FileObjectStoreFactory _fileObjectStoreFactory;
+		private readonly IClock _clock;
+		private readonly BundleCache _cache;
+		private readonly IOptionsMonitor<ToolsConfig> _toolsConfig;
+		private readonly IOptions<ToolsServerConfig> _toolsServerConfig;
+		private readonly ILogger _logger;
+
+		/// <summary>
+		/// Constructor
+		/// </summary>
+		public ToolCollection(IMongoService mongoService, IServerInfo serverInfo, StorageService storageService, BundleCache cache, FileObjectStoreFactory fileObjectStoreFactory, IClock clock, IOptionsMonitor<ToolsConfig> toolsConfig, IOptions<ToolsServerConfig> toolsServerConfig, ILogger<ToolCollection> logger)
+		{
+			_serverInfo = serverInfo;
+			_tools = mongoService.GetCollection<ToolDocument>("Tools");
+			_storageService = storageService;
+			_fileObjectStoreFactory = fileObjectStoreFactory;
+			_clock = clock;
+			_toolsConfig = toolsConfig;
+			_toolsServerConfig = toolsServerConfig;
+			_cache = cache;
+			_logger = logger;
+		}
+
+		Tool? CreateToolObject(ToolDocument? document)
+		{
+			if (document == null)
+			{
+				return null;
+			}
+
+			ToolConfig? toolConfig;
+			if (!_toolsConfig.CurrentValue.TryGetTool(document.Id, out toolConfig))
+			{
+				return null;
+			}
+
+			return new Tool(this, document, toolConfig, _clock.UtcNow);
+		}
+
+		/// <inheritdoc/>
+		public async Task<ITool?> GetAsync(ToolId id, CancellationToken cancellationToken)
+		{
+			ToolsConfig toolsConfig = _toolsConfig.CurrentValue;
+
+			ToolConfig? toolConfig;
+			if (toolsConfig.TryGetTool(id, out toolConfig))
+			{
+				ToolDocument document = await FindOrAddDocumentAsync(toolConfig, cancellationToken);
+				return new Tool(this, document, toolConfig, _clock.UtcNow);
+			}
+
+			BundledToolConfig? bundledToolConfig;
+			if (_toolsServerConfig.Value.TryGetBundledTool(id, out bundledToolConfig))
+			{
+				ToolDocument document = CreateBundledToolDocument(bundledToolConfig);
+				return new Tool(this, document, bundledToolConfig, _clock.UtcNow);
+			}
+
+			return null;
+		}
+
+		/// <inheritdoc/>
+		public async Task<IReadOnlyList<ITool>> GetAllAsync(CancellationToken cancellationToken)
+		{
+			DateTime utcNow = _clock.UtcNow;
+			ToolsConfig toolsConfig = _toolsConfig.CurrentValue;
+
+			List<Tool> tools = new List<Tool>();
+			foreach (ToolConfig toolConfig in toolsConfig.Tools)
+			{
+				ToolDocument document = await FindOrAddDocumentAsync(toolConfig, cancellationToken);
+				tools.Add(new Tool(this, document, toolConfig, utcNow));
+			}
+			foreach (BundledToolConfig bundledToolConfig in _toolsServerConfig.Value.BundledTools)
+			{
+				ToolDocument document = CreateBundledToolDocument(bundledToolConfig);
+				tools.Add(new Tool(this, document, bundledToolConfig, utcNow));
+			}
+
+			return tools;
+		}
+
+		async Task<ToolDocument> FindOrAddDocumentAsync(ToolConfig toolConfig, CancellationToken cancellationToken)
+		{
+			ToolDocument? tool;
+			for (; ; )
+			{
+				tool = await _tools.Find(x => x.Id == toolConfig.Id).FirstOrDefaultAsync(cancellationToken);
+				if (tool != null)
+				{
+					break;
+				}
+
+				tool = new ToolDocument(toolConfig.Id);
+				if (await _tools.InsertOneIgnoreDuplicatesAsync(tool, cancellationToken))
+				{
+					break;
+				}
+			}
+			return tool;
+		}
+
+		static ToolDocument CreateBundledToolDocument(BundledToolConfig bundledToolConfig)
+		{
+			ToolDocument tool = new ToolDocument(bundledToolConfig.Id);
+
+			ToolDeploymentId deploymentId = GetDeploymentId(bundledToolConfig);
+
+			ToolDeploymentDocument deployment = new ToolDeploymentDocument(deploymentId);
+			deployment.StartedAt = DateTime.MinValue;
+			deployment.Version = bundledToolConfig.Version;
+			deployment.RefName = bundledToolConfig.RefName;
+			tool.Deployments.Add(deployment);
+
+			return tool;
+		}
+
+		static ToolDeploymentId GetDeploymentId(BundledToolConfig bundledToolConfig)
+		{
+			// Create a pseudo-random deployment id from the tool id and version
+			IoHash hash = IoHash.Compute(Encoding.UTF8.GetBytes($"{bundledToolConfig.Id}/{bundledToolConfig.Version}"));
+
+			Span<byte> bytes = stackalloc byte[IoHash.NumBytes];
+			hash.CopyTo(bytes);
+
+			// Set a valid timestamp to allow it to appear as an ObjectId
+			BinaryPrimitives.WriteInt32BigEndian(bytes, 1446492960);
+			return new ToolDeploymentId(new BinaryId(bytes));
+		}
+
+		async Task<ToolDocument?> CreateDeploymentAsync(ToolDocument tool, ToolConfig toolConfig, ToolDeploymentConfig options, Stream stream, CancellationToken cancellationToken)
+		{
+			ToolDeploymentId deploymentId = new ToolDeploymentId(BinaryIdUtils.CreateNew());
+
+			using IStorageClient client = _storageService.CreateClient(toolConfig.NamespaceId);
+
+			IBlobRef<DirectoryNode> nodeRef;
+			await using (IBlobWriter writer = client.CreateBlobWriter($"{tool.Id}/{deploymentId}"))
+			{
+				DirectoryNode directoryNode = new DirectoryNode();
+				await directoryNode.CopyFromZipStreamAsync(stream, writer, new ChunkingOptions(), cancellationToken: cancellationToken);
+				nodeRef = await writer.WriteBlobAsync(directoryNode, cancellationToken: cancellationToken);
+			}
+
+			return await CreateDeploymentInternalAsync(tool, toolConfig, deploymentId, options, client, nodeRef, cancellationToken);
+		}
+
+		async Task<ToolDocument?> CreateDeploymentAsync(ToolDocument tool, ToolConfig toolConfig, ToolDeploymentConfig options, BlobRefValue target, CancellationToken cancellationToken)
+		{
+			ToolDeploymentId deploymentId = new ToolDeploymentId(BinaryIdUtils.CreateNew());
+
+			using IStorageClient client = _storageService.CreateClient(toolConfig.NamespaceId);
+			return await CreateDeploymentInternalAsync(tool, toolConfig, deploymentId, options, client, client.CreateBlobRef(target), cancellationToken);
+		}
+
+		async Task<ToolDocument?> CreateDeploymentInternalAsync(ToolDocument tool, ToolConfig toolConfig, ToolDeploymentId deploymentId, ToolDeploymentConfig options, IStorageClient storageClient, IBlobRef content, CancellationToken cancellationToken)
+		{
+			if (toolConfig is BundledToolConfig)
+			{
+				throw new InvalidOperationException("Cannot update the state of bundled tools.");
+			}
+
+			// Write a ref for the deployment so the blobs aren't GC'd
+			RefName refName = new RefName($"{tool.Id}/{deploymentId}");
+			await storageClient.WriteRefAsync(refName, content, cancellationToken: cancellationToken);
+
+			// Create the new deployment object
+			ToolDeploymentDocument deployment = new ToolDeploymentDocument(deploymentId, options, toolConfig.NamespaceId, refName);
+
+			// Start the deployment
+			DateTime utcNow = _clock.UtcNow;
+			if (!options.CreatePaused)
+			{
+				deployment.StartedAt = utcNow;
+			}
+
+			// Create the deployment
+			ToolDocument? newTool = (ToolDocument)tool;
+			for (; ; )
+			{
+				newTool = await TryAddDeploymentAsync(newTool, deployment, cancellationToken);
+				if (newTool != null)
+				{
+					return newTool;
+				}
+
+				newTool = await FindOrAddDocumentAsync(toolConfig, cancellationToken);
+				if (newTool == null)
+				{
+					return null;
+				}
+			}
+		}
+
+		async ValueTask<ToolDocument?> TryAddDeploymentAsync(ToolDocument tool, ToolDeploymentDocument deployment, CancellationToken cancellationToken)
+		{
+			ToolDocument? newTool = tool;
+
+			// If there are already a maximum number of deployments, remove the oldest one
+			const int MaxDeploymentCount = 5;
+			while (newTool.Deployments.Count >= MaxDeploymentCount)
+			{
+				newTool = await UpdateAsync(newTool, Builders<ToolDocument>.Update.PopFirst(x => x.Deployments), cancellationToken);
+				if (newTool == null)
+				{
+					return null;
+				}
+
+				ToolDeploymentDocument removeDeployment = tool.Deployments[0];
+				using IStorageClient client = _storageService.CreateClient(removeDeployment.NamespaceId);
+				await client.DeleteRefAsync(removeDeployment.RefName, cancellationToken);
+			}
+
+			// Add the new deployment
+			return await UpdateAsync(newTool, Builders<ToolDocument>.Update.Push(x => x.Deployments, deployment), cancellationToken);
+		}
+
+		async Task<ToolDocument?> UpdateDeploymentAsync(ToolDocument tool, ToolDeploymentId deploymentId, ToolDeploymentState action, double currentProgress, CancellationToken cancellationToken)
+		{
+			int idx = tool.Deployments.FindIndex(x => x.Id == deploymentId);
+			if (idx == -1)
+			{
+				return null;
+			}
+
+			ToolDeploymentDocument deployment = tool.Deployments[idx];
+			switch (action)
+			{
+				case ToolDeploymentState.Complete:
+					return await UpdateAsync(tool, Builders<ToolDocument>.Update.Set(x => x.Deployments[idx].BaseProgress, 1.0).Unset(x => x.Deployments[idx].StartedAt), cancellationToken);
+
+				case ToolDeploymentState.Cancelled:
+					List<ToolDeploymentDocument> newDeployments = tool.Deployments.Where(x => x != deployment).ToList();
+					return await UpdateAsync(tool, Builders<ToolDocument>.Update.Set(x => x.Deployments, newDeployments), cancellationToken);
+
+				case ToolDeploymentState.Paused:
+					if (deployment.StartedAt == null)
+					{
+						return tool;
+					}
+					else
+					{
+						return await UpdateAsync(tool, Builders<ToolDocument>.Update.Set(x => x.Deployments[idx].BaseProgress, currentProgress).Set(x => x.Deployments[idx].StartedAt, null), cancellationToken);
+					}
+
+				case ToolDeploymentState.Active:
+					if (deployment.StartedAt != null)
+					{
+						return tool;
+					}
+					else
+					{
+						return await UpdateAsync(tool, Builders<ToolDocument>.Update.Set(x => x.Deployments[idx].StartedAt, _clock.UtcNow), cancellationToken);
+					}
+
+				default:
+					throw new ArgumentException("Invalid action for deployment", nameof(action));
+			}
+		}
+
+		/// <summary>
+		/// Gets the storage client containing data for a particular tool
+		/// </summary>
+		/// <param name="toolConfig">Identifier for the tool</param>
+		/// <returns>Storage client for the data</returns>
+		IStorageClient CreateStorageClient(ToolConfig toolConfig)
+		{
+			if (toolConfig is BundledToolConfig bundledConfig)
+			{
+				return BundleStorageClient.CreateFromDirectory(DirectoryReference.Combine(_serverInfo.AppDir, bundledConfig.DataDir ?? $"Tools"), _cache, _logger);
+			}
+			else
+			{
+				return _storageService.CreateClient(Namespace.Tools);
+			}
+		}
+
+		/// <summary>
+		/// Gets the storage client containing data for a particular tool
+		/// </summary>
+		/// <param name="toolConfig">Identifier for the tool</param>
+		/// <returns>Storage client for the data</returns>
+		IStorageBackend CreateStorageBackend(ToolConfig toolConfig)
+		{
+			if (toolConfig is BundledToolConfig bundledConfig)
+			{
+				return new FileStorageBackend(_fileObjectStoreFactory.CreateStore(DirectoryReference.Combine(_serverInfo.AppDir, bundledConfig.DataDir ?? $"tools/{toolConfig.Id}")), _logger);
+			}
+			else
+			{
+				return _storageService.CreateBackend(Namespace.Tools);
+			}
+		}
+
+		/// <summary>
+		/// Opens a stream to the data for a particular deployment
+		/// </summary>
+		/// <param name="tool">Identifier for the tool</param>
+		/// <param name="deployment">The deployment</param>
+		/// <param name="cancellationToken">Cancellation token for the operation</param>
+		/// <returns>Stream for the data</returns>
+		async Task<Stream> GetDeploymentZipAsync(ToolConfig tool, ToolDeploymentDocument deployment, CancellationToken cancellationToken)
+		{
+#pragma warning disable CA2000
+			IStorageClient client = CreateStorageClient(tool);
+			try
+			{
+				IBlobRef<DirectoryNode> nodeRef = await client.ReadRefAsync<DirectoryNode>(deployment.RefName, DateTime.UtcNow - TimeSpan.FromDays(2.0), cancellationToken: cancellationToken);
+				return nodeRef.AsZipStream().WrapOwnership(client);
+			}
+			catch
+			{
+				client.Dispose();
+				throw;
+			}
+#pragma warning restore CA2000
+		}
+
+		async Task<ToolDocument> UpdateAsync(ToolDocument tool, UpdateDefinition<ToolDocument> update, CancellationToken cancellationToken)
+		{
+			update = update.Set(x => x.LastUpdateTime, new DateTime(Math.Max(tool.LastUpdateTime.Ticks + 1, DateTime.UtcNow.Ticks)));
+
+			FilterDefinition<ToolDocument> filter = Builders<ToolDocument>.Filter.Eq(x => x.Id, tool.Id) & Builders<ToolDocument>.Filter.Eq(x => x.LastUpdateTime, tool.LastUpdateTime);
+			return await _tools.FindOneAndUpdateAsync(filter, update, new FindOneAndUpdateOptions<ToolDocument> { ReturnDocument = ReturnDocument.After }, cancellationToken);
+		}
+	}
+}
diff --git a/Engine/Source/Programs/Horde/Plugins/Tools/HordeServer.Tools/Tools/ToolConfig.cs b/Engine/Source/Programs/Horde/Plugins/Tools/HordeServer.Tools/Tools/ToolConfig.cs
new file mode 100644
--- /dev/null
+++ b/Engine/Source/Programs/Horde/Plugins/Tools/HordeServer.Tools/Tools/ToolConfig.cs
@@ -0,0 +1,162 @@
+// Copyright Epic Games, Inc. All Rights Reserved.
+
+using System;
+using System.Collections.Generic;
+using System.ComponentModel.DataAnnotations;
+using System.Diagnostics;
+using System.Security.Claims;
+using System.Text.Json.Serialization;
+using EpicGames.Horde.Acls;
+using EpicGames.Horde.Storage;
+using EpicGames.Horde.Tools;
+using HordeServer.Acls;
+using HordeServer.Plugins;
+using HordeServer.Server;
+using HordeServer.Storage;
+
+#pragma warning disable CA2227 // Change x to be read-only by removing the property setter
+
+namespace HordeServer.Tools
+{
+	/// <summary>
+	/// Options for configuring a tool
+	/// </summary>
+	[DebuggerDisplay("{Id}")]
+	public class ToolConfig
+	{
+		/// <summary>
+		/// Unique identifier for the tool
+		/// </summary>
+		[Required]
+		public ToolId Id { get; set; }
+
+		/// <summary>
+		/// Name of the tool
+		/// </summary>
+		[Required]
+		public string Name { get; set; }
+
+		/// <summary>
+		/// Description for the tool
+		/// </summary>
+		public string Description { get; set; }
+
+		/// <summary>
+		/// Category for the tool. Will cause the tool to be shown in a different tab in the dashboard.
+		/// </summary>
+		public string? Category { get; set; }
+
+		/// <summary>
+		/// Grouping key for different variations of the same tool. The dashboard will show these together.
+		/// </summary>
+		public string? Group { get; set; }
+
+		/// <summary>
+		/// Platforms for this tool. Takes the form of a NET RID (https://learn.microsoft.com/en-us/dotnet/core/rid-catalog).
+		/// </summary>
+		public List<string>? Platforms { get; set; }
+
+		/// <summary>
+		/// Whether this tool should be exposed for download on a public endpoint without authentication
+		/// </summary>
+		public bool Public { get; set; }
+
+		/// <summary>
+		/// Whether to show this tool for download in the UGS tools menu
+		/// </summary>
+		public bool ShowInUgs { get; set; }
+
+		/// <summary>
+		/// Whether to show this tool for download in the dashboard
+		/// </summary>
+		public bool ShowInDashboard { get; set; } = true;
+
+		/// <summary>
+		/// Default namespace for new deployments of this tool
+		/// </summary>
+		public NamespaceId NamespaceId { get; set; } = Namespace.Tools;
+
+		/// <summary>
+		/// Permissions for the tool
+		/// </summary>
+		public AclConfig Acl { get; set; } = new AclConfig();
+
+		/// <summary>
+		/// Default constructor for serialization
+		/// </summary>
+		public ToolConfig()
+		{
+			Name = String.Empty;
+			Description = String.Empty;
+		}
+
+		/// <summary>
+		/// Constructor
+		/// </summary>
+		public ToolConfig(ToolId id)
+		{
+			Id = id;
+			Name = id.ToString();
+			Description = String.Empty;
+		}
+
+		/// <summary>
+		/// Called after the config has been read
+		/// </summary>
+		/// <param name="parentAcl">Parent ACL object</param>
+		public void PostLoad(AclConfig parentAcl)
+		{
+			Acl.PostLoad(parentAcl, $"tool:{Id}");
+		}
+
+		/// <inheritdoc cref="AclConfig.Authorize(AclAction, ClaimsPrincipal)"/>
+		public bool Authorize(AclAction action, ClaimsPrincipal user)
+			=> Acl.Authorize(action, user);
+	}
+
+	/// <summary>
+	/// Configuration for a tool bundled alongsize the server
+	/// </summary>
+	public class BundledToolConfig : ToolConfig
+	{
+		/// <summary>
+		/// Version string for the current tool data
+		/// </summary>
+		public string Version { get; set; } = "1.0";
+
+		/// <summary>
+		/// Ref name in the tools directory
+		/// </summary>
+		public RefName RefName { get; set; } = new RefName("default-ref");
+
+		/// <summary>
+		/// Directory containing blob data for this tool. If empty, the tools/{id} folder next to the server will be used.
+		/// </summary>
+		public string? DataDir { get; set; }
+
+		/// <summary>
+		/// Constructor
+		/// </summary>
+		public BundledToolConfig()
+		{
+			Public = true;
+		}
+	}
+
+	/// <summary>
+	/// Options for a new deployment
+	/// </summary>
+	public class ToolDeploymentConfig
+	{
+		/// <inheritdoc cref="IToolDeployment.Version"/>
+		public string Version { get; set; } = "Unknown";
+
+		/// <inheritdoc cref="IToolDeployment.Duration"/>
+		public TimeSpan Duration { get; set; }
+
+		/// <summary>
+		/// Whether to create the deployment in a paused state
+		/// </summary>
+		public bool CreatePaused { get; set; }
+	}
+}
diff --git a/Engine/Source/Programs/Horde/Plugins/Tools/HordeServer.Tools/Tools/ToolsController.cs b/Engine/Source/Programs/Horde/Plugins/Tools/HordeServer.Tools/Tools/ToolsController.cs
new file mode 100644
--- /dev/null
+++ b/Engine/Source/Programs/Horde/Plugins/Tools/HordeServer.Tools/Tools/ToolsController.cs
@@ -0,0 +1,416 @@
+﻿// Copyright Epic Games, Inc. All Rights Reserved.
+
+using System;
+using System.Collections.Generic;
+using System.IO;
+using System.Linq;
+using System.Threading;
+using System.Threading.Tasks;
+using EpicGames.Core;
+using EpicGames.Horde.Storage;
+using EpicGames.Horde.Storage.Nodes;
+using EpicGames.Horde.Tools;
+using HordeServer.Storage;
+using HordeServer.Utilities;
+using Microsoft.AspNetCore.Authorization;
+using Microsoft.AspNetCore.Http;
+using Microsoft.AspNetCore.Mvc;
+using Microsoft.AspNetCore.StaticFiles;
+using Microsoft.Extensions.Logging;
+
+namespace HordeServer.Tools
+{
+	/// <summary>
+	/// Controller for the /api/v1/tools endpoint
+	/// </summary>
+	[ApiController]
+	[Authorize]
+	[Route("[controller]")]
+	public class ToolsController : HordeControllerBase
+	{
+		readonly IToolCollection _toolCollection;
+
+		/// <summary>
+		/// Constructor
+		/// </summary>
+		public ToolsController(IToolCollection toolCollection)
+		{
+			_toolCollection = toolCollection;
+		}
+
+		/// <summary>
+		/// Uploads blob data for a new tool deployment.
+		/// </summary>
+		/// <param name="id">Identifier of the tool to upload</param>
+		/// <param name="request">Upload request</param>
+		/// <param name="cancellationToken">Cancellation token for the operation</param>
+		[HttpPost]
+		[Route("/api/v1/tools/{id}/blobs")]
+		public async Task<ActionResult<WriteBlobResponse>> WriteBlobAsync(ToolId id, WriteBlobRequest request, CancellationToken cancellationToken = default)
+		{
+			ITool? tool = await _toolCollection.GetAsync(id, cancellationToken);
+			if (tool == null)
+			{
+				return NotFound(id);
+			}
+			if (!tool.Authorize(ToolAclAction.UploadTool, User))
+			{
+				return Forbid(ToolAclAction.UploadTool, id);
+			}
+
+			if (String.IsNullOrEmpty(request.Prefix))
+			{
+				request.Prefix = id.ToString();
+			}
+			else
+			{
+				request.Prefix = $"{id}/{request.Prefix}";
+			}
+
+			IStorageBackend storageBackend = tool.CreateStorageBackend();
+			return await StorageController.WriteBlobAsync(storageBackend, request, cancellationToken);
+		}
+
+		/// <summary>
+		/// Create a new deployment of the given tool.
+		/// </summary>
+		/// <returns>Information about the registered agent</returns>
+		[HttpPost]
+		[Route("/api/v1/tools/{id}/deployments")]
+		public async Task<ActionResult<CreateToolDeploymentResponse>> CreateDeploymentAsync(ToolId id, [FromForm] ToolDeploymentConfig options, [FromForm] IFormFile file, CancellationToken cancellationToken)
+		{
+			ITool? tool = await _toolCollection.GetAsync(id, cancellationToken);
+			if (tool == null)
+			{
+				return NotFound(id);
+			}
+			if (!tool.Authorize(ToolAclAction.UploadTool, User))
+			{
+				return Forbid(ToolAclAction.UploadTool, id);
+			}
+
+			using (Stream stream = file.OpenReadStream())
+			{
+				tool = await tool.CreateDeploymentAsync(options, stream, cancellationToken);
+				if (tool == null)
+				{
+					return NotFound(id);
+				}
+			}
+
+			return new CreateToolDeploymentResponse(tool.Deployments[^1].Id);
+		}
+
+		/// <summary>
+		/// Create a new deployment of the given tool.
+		/// </summary>
+		[HttpPost]
+		[Route("/api/v2/tools/{id}/deployments")]
+		public async Task<ActionResult<CreateToolDeploymentResponse>> CreateDeploymentAsync(ToolId id, CreateToolDeploymentRequest request, CancellationToken cancellationToken)
+		{
+			ITool? tool = await _toolCollection.GetAsync(id, cancellationToken);
+
+			if (tool == null)
+			{
+				return NotFound(id);
+			}
+			if (!tool.Authorize(ToolAclAction.UploadTool, User))
+			{
+				return Forbid(ToolAclAction.UploadTool, id);
+			}
+
+			ToolDeploymentConfig options = new ToolDeploymentConfig { Version = request.Version, Duration = TimeSpan.FromMinutes(request.Duration ?? 0.0), CreatePaused = request.CreatePaused ?? false };
+
+			tool = await tool.CreateDeploymentAsync(options, request.Content, cancellationToken);
+			if (tool == null)
+			{
+				return NotFound(id);
+			}
+
+			return new CreateToolDeploymentResponse(tool.Deployments[^1].Id);
+		}
+
+		/// <summary>
+		/// Updates the state of an active deployment.
+		/// </summary>
+		[HttpPatch]
+		[Route("/api/v1/tools/{id}/deployments/{deploymentId}")]
+		public async Task<ActionResult> UpdateDeploymentAsync(ToolId id, ToolDeploymentId deploymentId, [FromBody] UpdateDeploymentRequest request)
+		{
+			ITool? tool = await _toolCollection.GetAsync(id, HttpContext.RequestAborted);
+			if (tool == null)
+			{
+				return NotFound(id);
+			}
+			if (!tool.Authorize(ToolAclAction.UploadTool, User))
+			{
+				return Forbid(ToolAclAction.UploadTool, id);
+			}
+
+			IToolDeployment? deployment = tool.Deployments.FirstOrDefault(x => x.Id == deploymentId);
+			if (deployment == null)
+			{
+				return NotFound(deploymentId);
+			}
+
+			if (request.State != null)
+			{
+				deployment = await deployment.UpdateAsync(request.State.Value, HttpContext.RequestAborted);
+				if (request.State != ToolDeploymentState.Cancelled && deployment == null)
+				{
+					return NotFound(id, deploymentId);
+				}
+			}
+			return Ok();
+		}
+	}
+
+	/// <summary>
+	/// Public methods available without authorization (or with very custom authorization)
+	/// </summary>
+	[ApiController]
+	[TryAuthorize]
+	[Tags("Tools")]
+	public class PublicToolsController : HordeControllerBase
+	{
+		readonly IToolCollection _toolCollection;
+		readonly IClock _clock;
+
+		/// <summary>
+		/// Constructor
+		/// </summary>
+		public PublicToolsController(IToolCollection toolCollection, IClock clock)
+		{
+			_toolCollection = toolCollection;
+			_clock = clock;
+		}
+
+		/// <summary>
+		/// Enumerates all the available tools.
+		/// </summary>
+		[HttpGet]
+		[TryAuthorize]
+		[Route("/api/v1/tools")]
+		public async Task<ActionResult<GetToolsSummaryResponse>> GetToolsAsync()
+		{
+			IReadOnlyList<ITool> tools = await _toolCollection.GetAllAsync(HttpContext.RequestAborted);
+
+			List<GetToolSummaryResponse> toolSummaryList = new List<GetToolSummaryResponse>();
+			foreach (ITool tool in tools.OrderBy(x => x.Name, StringComparer.Ordinal))
+			{
+				if (AuthorizeDownload(tool))
+				{
+					toolSummaryList.Add(CreateGetToolSummaryResponse(tool));
+				}
+			}
+
+			return new GetToolsSummaryResponse(toolSummaryList);
+		}
+
+		static GetToolSummaryResponse CreateGetToolSummaryResponse(ITool tool)
+		{
+			IToolDeployment? deployment = (tool.Deployments.Count == 0) ? null : tool.Deployments[^1];
+			return new GetToolSummaryResponse(tool.Id, tool.Name, tool.Description, tool.Category, tool.Group, tool.Platforms?.ToList(), deployment?.Version, deployment?.Id, deployment?.State, deployment?.Progress, tool.ShowInUgs, tool.ShowInDashboard);
+		}
+
+		/// <summary>
+		/// Gets information about a particular tool
+		/// </summary>
+		/// <returns>Information about the registered agent</returns>
+		[HttpGet]
+		[Route("/api/v1/tools/{id}")]
+		public async Task<ActionResult> GetToolAsync(ToolId id, GetToolAction action = GetToolAction.Info, CancellationToken cancellationToken = default)
+		{
+			ITool? tool = await _toolCollection.GetAsync(id, cancellationToken);
+			if (tool == null)
+			{
+				return NotFound(id);
+			}
+			if (!AuthorizeDownload(tool))
+			{
+				return Forbid(ToolAclAction.DownloadTool, id);
+			}
+
+			if (action == GetToolAction.Info)
+			{
+				List<GetToolDeploymentResponse> deploymentResponses = new List<GetToolDeploymentResponse>();
+				foreach (IToolDeployment deployment in tool.Deployments)
+				{
+					GetToolDeploymentResponse deploymentResponse = await GetDeploymentInfoResponseAsync(tool, deployment, cancellationToken);
+					deploymentResponses.Add(deploymentResponse);
+				}
+				return Ok(CreateGetToolResponse(tool, deploymentResponses));
+			}
+			else
+			{
+				if (tool.Deployments.Count == 0)
+				{
+					return NotFound(LogEvent.Create(LogLevel.Error, "Tool {ToolId} does not currently have any deployments", id));
+				}
+
+				return await GetDeploymentResponseAsync(tool, tool.Deployments[^1], action, cancellationToken);
+			}
+		}
+
+		static GetToolResponse CreateGetToolResponse(ITool tool, List<GetToolDeploymentResponse> deployments)
+		{
+			return new GetToolResponse(tool.Id, tool.Name, tool.Description, tool.Category, tool.Group, tool.Platforms?.ToList(), deployments, tool.Public, tool.ShowInUgs, tool.ShowInDashboard);
+		}
+
+		/// <summary>
+		/// Finds deployments of a particular tool.
+		/// </summary>
+		/// <param name="id">The tool identifier</param>
+		/// <param name="phase">Value indicating the client's preference for deployment to receive.</param>
+		/// <param name="action">Information about the returned deployment</param>
+		/// <param name="cancellationToken">Cancellation token for the operation</param>
+		/// <returns>Information about the registered agent</returns>
+		[HttpGet]
+		[Route("/api/v1/tools/{id}/deployments")]
+		public async Task<ActionResult> FindDeploymentAsync(ToolId id, [FromQuery] double phase = 0.0, [FromQuery] GetToolAction action = GetToolAction.Info, CancellationToken cancellationToken = default)
+		{
+			ITool? tool = await _toolCollection.GetAsync(id, cancellationToken);
+			if (tool == null)
+			{
+				return NotFound(id);
+			}
+			if (!AuthorizeDownload(tool))
+			{
+				return Forbid(ToolAclAction.DownloadTool, id);
+			}
+
+			IToolDeployment? deployment = tool.GetCurrentDeployment(phase, _clock.UtcNow);
+			if (deployment == null)
+			{
+				return NotFound(LogEvent.Create(LogLevel.Error, "Tool {ToolId} does not currently have any deployments", id));
+			}
+
+			return await GetDeploymentResponseAsync(tool, deployment, action, cancellationToken);
+		}
+
+		/// <summary>
+		/// Gets information about a specific tool deployment.
+		/// </summary>
+		/// <returns>Information about the registered agent</returns>
+		[HttpGet]
+		[Route("/api/v1/tools/{id}/deployments/{deploymentId}")]
+		public async Task<ActionResult> GetDeploymentAsync(ToolId id, ToolDeploymentId deploymentId, [FromQuery] GetToolAction action = GetToolAction.Info, CancellationToken cancellationToken = default)
+		{
+			ITool? tool = await _toolCollection.GetAsync(id, cancellationToken);
+			if (tool == null)
+			{
+				return NotFound(id);
+			}
+			if (!AuthorizeDownload(tool))
+			{
+				return Forbid(ToolAclAction.DownloadTool, id);
+			}
+
+			IToolDeployment? deployment = tool.Deployments.FirstOrDefault(x => x.Id == deploymentId);
+			if (deployment == null)
+			{
+				return NotFound(id, deploymentId);
+			}
+
+			return await GetDeploymentResponseAsync(tool, deployment, action, cancellationToken);
+		}
+
+		private async Task<ActionResult> GetDeploymentResponseAsync(ITool tool, IToolDeployment deployment, GetToolAction action, CancellationToken cancellationToken)
+		{
+			if (action == GetToolAction.Info)
+			{
+				GetToolDeploymentResponse response = await GetDeploymentInfoResponseAsync(tool, deployment, cancellationToken);
+				return Ok(response);
+			}
+
+			IStorageClient client = tool.CreateStorageClient();
+			try
+			{
+				IBlobRef<DirectoryNode> nodeRef = await client.ReadRefAsync<DirectoryNode>(deployment.RefName, DateTime.UtcNow - TimeSpan.FromDays(2.0), cancellationToken: cancellationToken);
+
+				// If we weren't specifically asked for a zip, see if this download is a single file. If it is, allow downloading it directory.
+				if (action != GetToolAction.Zip)
+				{
+					DirectoryNode node = await nodeRef.ReadBlobAsync(cancellationToken);
+					if (node.Directories.Count == 0 && node.Files.Count == 1)
+					{
+						FileEntry entry = node.Files.First();
+
+						string? contentType;
+						if (!new FileExtensionContentTypeProvider().TryGetContentType(entry.Name.ToString(), out contentType))
+						{
+							contentType = "application/octet-stream";
+						}
+
+						Response.Headers.ContentLength = entry.Length;
+
+						Stream fileStream = entry.OpenAsStream().WrapOwnership(client);
+						return new FileStreamResult(fileStream, contentType) { FileDownloadName = entry.Name.ToString() };
+					}
+				}
+
+				Stream stream = nodeRef.AsZipStream().WrapOwnership(client);
+				return new FileStreamResult(stream, "application/zip") { FileDownloadName = $"{tool.Id}-{deployment.Version}.zip" };
+			}
+			catch
+			{
+				client.Dispose();
+				throw;
+			}
+		}
+
+		static async Task<GetToolDeploymentResponse> GetDeploymentInfoResponseAsync(ITool tool, IToolDeployment deployment, CancellationToken cancellationToken)
+		{
+			using IStorageClient client = tool.CreateStorageClient();
+			IBlobHandle rootHandle = await client.ReadRefAsync(deployment.RefName, cancellationToken: cancellationToken);
+
+			return new GetToolDeploymentResponse(deployment.Id, deployment.Version, deployment.State, deployment.Progress, deployment.StartedAt, deployment.Duration, deployment.RefName, rootHandle.GetLocator());
+		}
+
+		/// <summary>
+		/// Retrieves blobs for a particular tool
+		/// </summary>
+		/// <param name="id">Identifier of the tool to retrieve</param>
+		/// <param name="locator">The blob locator</param>
+		/// <param name="cancellationToken">Cancellation token for the operation</param>
+		/// <returns>Information about all the artifacts</returns>
+		[HttpGet]
+		[Route("/api/v1/tools/{id}/blobs/{*locator}")]
+		public async Task<ActionResult> ReadToolBlobAsync(ToolId id, BlobLocator locator, CancellationToken cancellationToken = default)
+		{
+			ITool? tool = await _toolCollection.GetAsync(id, cancellationToken);
+			if (tool == null)
+			{
+				return NotFound(id);
+			}
+			if (!AuthorizeDownload(tool))
+			{
+				return Forbid(ToolAclAction.DownloadTool, id);
+			}
+
+			if (!locator.WithinFolder(tool.Id.Id.Text))
+			{
+				return BadRequest("Invalid blob id for tool");
+			}
+
+			IStorageBackend storageBackend = tool.CreateStorageBackend();
+			return await StorageController.ReadBlobInternalAsync(storageBackend, locator, Request.Headers, cancellationToken);
+		}
+
+		bool AuthorizeDownload(ITool tool)
+		{
+			if (!tool.Public)
+			{
+				if (User.Identity == null || !User.Identity.IsAuthenticated)
+				{
+					return false;
+				}
+				if (!tool.Authorize(ToolAclAction.DownloadTool, User))
+				{
+					return false;
+				}
+			}
+			return true;
+		}
+	}
+}
diff --git a/Engine/Source/Programs/Horde/Plugins/Tools/HordeServer.Tools/ToolsConfig.cs b/Engine/Source/Programs/Horde/Plugins/Tools/HordeServer.Tools/ToolsConfig.cs
new file mode 100644
--- /dev/null
+++ b/Engine/Source/Programs/Horde/Plugins/Tools/HordeServer.Tools/ToolsConfig.cs
@@ -0,0 +1,43 @@
+// Copyright Epic Games, Inc. All Rights Reserved.
+
+using System.Diagnostics.CodeAnalysis;
+using EpicGames.Horde.Tools;
+using HordeServer.Acls;
+using HordeServer.Plugins;
+
+#pragma warning disable CA2227 // Change x to be read-only by removing the property setter
+
+namespace HordeServer.Tools
+{
+	/// <summary>
+	/// Configuration for the tools system
+	/// </summary>
+	public class ToolsConfig : IPluginConfig
+	{
+		/// <summary>
+		/// Tool configurations
+		/// </summary>
+		public List<ToolConfig> Tools { get; set; } = new List<ToolConfig>();
+
+		private readonly Dictionary<ToolId, ToolConfig> _toolLookup = new Dictionary<ToolId, ToolConfig>();
+
+		/// <inheritdoc/>
+		public void PostLoad(AclConfig parentAcl)
+		{
+			_toolLookup.Clear();
+			foreach (ToolConfig tool in Tools)
+			{
+				_toolLookup.Add(tool.Id, tool);
+				tool.PostLoad(parentAcl);
+			}
+		}
+
+		/// <summary>
+		/// Attempts to get configuration for a tool from this object
+		/// </summary>
+		/// <param name="toolId">The tool identifier</param>
+		/// <param name="config">Configuration for the tool</param>
+		/// <returns>True if the tool configuration was found</returns>
+		public bool TryGetTool(ToolId toolId, [NotNullWhen(true)] out ToolConfig? config) => _toolLookup.TryGetValue(toolId, out config);
+	}
+}
diff --git a/Engine/Source/Programs/Horde/Plugins/Tools/HordeServer.Tools/ToolsPlugin.cs b/Engine/Source/Programs/Horde/Plugins/Tools/HordeServer.Tools/ToolsPlugin.cs
new file mode 100644
--- /dev/null
+++ b/Engine/Source/Programs/Horde/Plugins/Tools/HordeServer.Tools/ToolsPlugin.cs
@@ -0,0 +1,33 @@
+// Copyright Epic Games, Inc. All Rights Reserved.
+
+using HordeServer.Plugins;
+using HordeServer.Tools;
+using Microsoft.Extensions.DependencyInjection;
+
+namespace HordeServer
+{
+	/// <summary>
+	/// Entry point for the storage plugin
+	/// </summary>
+	[Plugin("Tools", GlobalConfigType = typeof(ToolsConfig), ServerConfigType = typeof(ToolsServerConfig))]
+	public class ToolsPlugin : IPluginStartup
+	{
+		/// <inheritdoc/>
+		public void ConfigureServices(IServiceCollection services)
+		{
+			services.AddSingleton<IToolCollection, ToolCollection>();
+		}
+	}
+
+	/// <summary>
+	/// Helper methods for tools config
+	/// </summary>
+	public static class ToolsPluginExtensions
+	{
+		/// <summary>
+		/// Configures the tools plugin
+		/// </summary>
+		public static void AddToolsConfig(this IDictionary<PluginName, IPluginConfig> dictionary, ToolsConfig toolsConfig)
+			=> dictionary[new PluginName("Tools")] = toolsConfig;
+	}
+}
diff --git a/Engine/Source/Programs/Horde/Plugins/Tools/HordeServer.Tools/ToolsServerConfig.cs b/Engine/Source/Programs/Horde/Plugins/Tools/HordeServer.Tools/ToolsServerConfig.cs
new file mode 100644
--- /dev/null
+++ b/Engine/Source/Programs/Horde/Plugins/Tools/HordeServer.Tools/ToolsServerConfig.cs
@@ -0,0 +1,33 @@
+// Copyright Epic Games, Inc. All Rights Reserved.
+
+using System.Diagnostics.CodeAnalysis;
+using EpicGames.Horde.Tools;
+using HordeServer.Tools;
+
+#pragma warning disable CA2227 // Change x to be read-only by removing the property setter
+
+namespace HordeServer
+{
+	/// <summary>
+	/// Server configuration for bundled tools
+	/// </summary>
+	public class ToolsServerConfig
+	{
+		/// <summary>
+		/// Tools bundled along with the server. Data for each tool can be produced using the 'bundle create' command, and should be stored in the Tools directory.
+		/// </summary>
+		public List<BundledToolConfig> BundledTools { get; set; } = new List<BundledToolConfig>();
+
+		/// <summary>
+		/// Attempts to get a bundled tool with the given id
+		/// </summary>
+		/// <param name="toolId">The tool id</param>
+		/// <param name="bundledToolConfig">Configuration for the bundled tool</param>
+		/// <returns>True if the tool was found</returns>
+		public bool TryGetBundledTool(ToolId toolId, [NotNullWhen(true)] out BundledToolConfig? bundledToolConfig)
+		{
+			bundledToolConfig = BundledTools.FirstOrDefault(x => x.Id == toolId);
+			return bundledToolConfig != null;
+		}
+	}
+}
diff --git a/Engine/Source/Programs/Horde/Scripts/test.sh b/Engine/Source/Programs/Horde/Scripts/test.sh
--- a/Engine/Source/Programs/Horde/Scripts/test.sh
+++ b/Engine/Source/Programs/Horde/Scripts/test.sh
@@ -18,6 +18,7 @@ testProjects=(
 	"Source/Programs/Horde/Plugins/Analytics/HordeServer.Analytics.Tests/HordeServer.Analytics.Tests.csproj"
 	"Source/Programs/Horde/Plugins/Ddc/HordeServer.Ddc.Tests/HordeServer.Ddc.Tests.csproj"
 	"Source/Programs/Horde/Plugins/Storage/HordeServer.Storage.Tests/HordeServer.Storage.Tests.csproj"
+	"Source/Programs/Horde/Plugins/Tools/HordeServer.Tools.Tests/HordeServer.Tools.Tests.csproj"
 )
 
 for csProj in "${testProjects[@]}"; do
