Horde: Allow querying for the server runmode through IServerInfo, so the log service doesn't need to access the global server settings.

#rnx

diff --git a/Engine/Source/Programs/Horde/HordeServer.Shared/IServerInfo.cs b/Engine/Source/Programs/Horde/HordeServer.Shared/IServerInfo.cs
--- a/Engine/Source/Programs/Horde/HordeServer.Shared/IServerInfo.cs
+++ b/Engine/Source/Programs/Horde/HordeServer.Shared/IServerInfo.cs
@@ -5,6 +5,34 @@ using Microsoft.Extensions.Configuration;
 
 namespace HordeServer
 {
+	/// <summary>
+	/// Type of run mode this process should use. Each carry different types of workloads. 
+	/// More than one mode can be active. But not all modes are not guaranteed to be compatible with each other and will
+	/// raise an error if combined in such a way.
+	/// </summary>
+	public enum RunMode
+	{
+		/// <summary>
+		/// Default no-op value (ASP.NET config will default to this for enums that cannot be parsed)
+		/// </summary> 
+		None,
+
+		/// <summary>
+		/// Handle and respond to incoming external requests, such as HTTP REST and gRPC calls.
+		/// These requests are time-sensitive and short-lived, typically less than 5 secs.
+		/// If processes handling requests are unavailable, it will be very visible for users.
+		/// </summary>
+		Server,
+
+		/// <summary>
+		/// Run non-request facing workloads. Such as background services, processing queues, running work
+		/// based on timers etc. Short periods of downtime or high CPU usage due to bursts are fine for this mode.
+		/// No user requests will be impacted directly. If auto-scaling is used, a much more aggressive policy can be
+		/// applied (tighter process packing, higher avg CPU usage).
+		/// </summary>
+		Worker
+	}
+
 	/// <summary>
 	/// Provides access to server deployment information
 	/// </summary>
@@ -34,5 +62,12 @@ namespace HordeServer
 		/// Directory to store server data files
 		/// </summary>
 		DirectoryReference DataDir { get; }
+
+		/// <summary>
+		/// Helper method to check if this process has activated the given mode
+		/// </summary>
+		/// <param name="mode">Run mode</param>
+		/// <returns>True if mode is active</returns>
+		bool IsRunModeActive(RunMode mode);
 	}
 }
diff --git a/Engine/Source/Programs/Horde/HordeServer/Logs/LogTailService.cs b/Engine/Source/Programs/Horde/HordeServer/Logs/LogTailService.cs
--- a/Engine/Source/Programs/Horde/HordeServer/Logs/LogTailService.cs
+++ b/Engine/Source/Programs/Horde/HordeServer/Logs/LogTailService.cs
@@ -57,7 +57,7 @@ namespace HordeServer.Logs
 		readonly ITicker _expireTailsTicker;
 
 		readonly ILogger _logger;
-		readonly ServerSettings _settings;
+		readonly IServerInfo _serverInfo;
 
 		IAsyncDisposable? _tailStartSubscription;
 
@@ -72,15 +72,15 @@ namespace HordeServer.Logs
 		/// <summary>
 		/// Constructor
 		/// </summary>
-		public LogTailService(IRedisService redisService, IClock clock, IOptions<ServerSettings> settings, ILogger<LogTailService> logger)
-			: this(redisService, clock, 32, settings, logger)
+		public LogTailService(IRedisService redisService, IClock clock, IServerInfo serverInfo, ILogger<LogTailService> logger)
+			: this(redisService, clock, 32, serverInfo, logger)
 		{
 		}
 
 		/// <summary>
 		/// Constructor
 		/// </summary>
-		public LogTailService(IRedisService redisService, IClock clock, int chunkLineCount, IOptions<ServerSettings> settings, ILogger<LogTailService> logger)
+		public LogTailService(IRedisService redisService, IClock clock, int chunkLineCount, IServerInfo serverInfo, ILogger<LogTailService> logger)
 		{
 			if ((chunkLineCount & (chunkLineCount - 1)) != 0)
 			{
@@ -94,7 +94,7 @@ namespace HordeServer.Logs
 
 			_clock = clock;
 			_expireTailsTicker = clock.AddSharedTicker<LogTailService>(TimeSpan.FromSeconds(30.0), ExpireTailsAsync, logger);
-			_settings = settings.Value;
+			_serverInfo = serverInfo;
 			_logger = logger;
 		}
 
@@ -102,7 +102,7 @@ namespace HordeServer.Logs
 		public async Task StartAsync(CancellationToken cancellationToken)
 		{
 			_tailStartSubscription = await _redisService.SubscribeAsync(_tailStartChannel, OnTailStart);
-			if (_settings.IsRunModeActive(RunMode.Worker))
+			if (_serverInfo.IsRunModeActive(RunMode.Worker))
 			{
 				await _expireTailsTicker.StartAsync();
 			}
@@ -111,7 +111,7 @@ namespace HordeServer.Logs
 		/// <inheritdoc/>
 		public async Task StopAsync(CancellationToken cancellationToken)
 		{
-			if (_settings.IsRunModeActive(RunMode.Worker))
+			if (_serverInfo.IsRunModeActive(RunMode.Worker))
 			{
 				await _expireTailsTicker.StopAsync();
 			}
diff --git a/Engine/Source/Programs/Horde/HordeServer/ServerInfo.cs b/Engine/Source/Programs/Horde/HordeServer/ServerInfo.cs
--- a/Engine/Source/Programs/Horde/HordeServer/ServerInfo.cs
+++ b/Engine/Source/Programs/Horde/HordeServer/ServerInfo.cs
@@ -1,6 +1,7 @@
 // Copyright Epic Games, Inc. All Rights Reserved.
 
 using EpicGames.Core;
+using Microsoft.Extensions.Options;
 
 namespace HordeServer
 {
@@ -23,5 +24,17 @@ namespace HordeServer
 
 		/// <inheritdoc/>
 		public DirectoryReference DataDir => ServerApp.DataDir;
+
+		readonly IOptions<ServerSettings> _serverSettings;
+
+		/// <summary>
+		/// Constructor
+		/// </summary>
+		public ServerInfo(IOptions<ServerSettings> serverSettings)
+			=> _serverSettings = serverSettings;
+
+		/// <inheritdoc/>
+		public bool IsRunModeActive(RunMode mode)
+			=> _serverSettings.Value.IsRunModeActive(mode);
 	}
 }
diff --git a/Engine/Source/Programs/Horde/HordeServer/ServerSettings.cs b/Engine/Source/Programs/Horde/HordeServer/ServerSettings.cs
--- a/Engine/Source/Programs/Horde/HordeServer/ServerSettings.cs
+++ b/Engine/Source/Programs/Horde/HordeServer/ServerSettings.cs
@@ -16,34 +16,6 @@ using HordeServer.Server;
 
 namespace HordeServer
 {
-	/// <summary>
-	/// Type of run mode this process should use. Each carry different types of workloads. 
-	/// More than one mode can be active. But not all modes are not guaranteed to be compatible with each other and will
-	/// raise an error if combined in such a way.
-	/// </summary>
-	public enum RunMode
-	{
-		/// <summary>
-		/// Default no-op value (ASP.NET config will default to this for enums that cannot be parsed)
-		/// </summary> 
-		None,
-
-		/// <summary>
-		/// Handle and respond to incoming external requests, such as HTTP REST and gRPC calls.
-		/// These requests are time-sensitive and short-lived, typically less than 5 secs.
-		/// If processes handling requests are unavailable, it will be very visible for users.
-		/// </summary>
-		Server,
-
-		/// <summary>
-		/// Run non-request facing workloads. Such as background services, processing queues, running work
-		/// based on timers etc. Short periods of downtime or high CPU usage due to bursts are fine for this mode.
-		/// No user requests will be impacted directly. If auto-scaling is used, a much more aggressive policy can be
-		/// applied (tighter process packing, higher avg CPU usage).
-		/// </summary>
-		Worker
-	}
-
 	/// <summary>
 	/// Feature flags to aid rollout of new features.
 	///
diff --git a/Engine/Source/Programs/Horde/HordeServer/Startup.cs b/Engine/Source/Programs/Horde/HordeServer/Startup.cs
--- a/Engine/Source/Programs/Horde/HordeServer/Startup.cs
+++ b/Engine/Source/Programs/Horde/HordeServer/Startup.cs
@@ -292,7 +292,7 @@ namespace HordeServer
 		// This method gets called *multiple times* by the runtime. Use this method to add services to the container.
 		public void ConfigureServices(IServiceCollection services)
 		{
-			services.AddSingleton<IServerInfo>(new ServerInfo());
+			services.AddSingleton<IServerInfo, ServerInfo>();
 			services.AddSingleton<JsonSchemaCache>();
 
 			// IOptionsMonitor pattern for live updating of configuration settings
