Horde: Fix static analysis warning.

#rnx

diff --git a/Engine/Source/Programs/Horde/HordeServer/Logs/LogTailService.cs b/Engine/Source/Programs/Horde/HordeServer/Logs/LogTailService.cs
--- a/Engine/Source/Programs/Horde/HordeServer/Logs/LogTailService.cs
+++ b/Engine/Source/Programs/Horde/HordeServer/Logs/LogTailService.cs
@@ -12,7 +12,6 @@ using EpicGames.Redis;
 using HordeServer.Server;
 using Microsoft.Extensions.Hosting;
 using Microsoft.Extensions.Logging;
-using Microsoft.Extensions.Options;
 using StackExchange.Redis;
 
 namespace HordeServer.Logs
