Horde: Remove unused file.

#rnx

diff --git a/Engine/Source/Programs/Horde/HordeServer.Tests.Shared/ServerServiceTest.cs b/Engine/Source/Programs/Horde/HordeServer.Tests.Shared/ServerServiceTest.cs
--- a/Engine/Source/Programs/Horde/HordeServer.Tests.Shared/ServerServiceTest.cs
+++ b/Engine/Source/Programs/Horde/HordeServer.Tests.Shared/ServerServiceTest.cs
@@ -91,7 +91,7 @@ namespace HordeServer.Tests
 		{
 			base.ConfigureServices(services);
 
-			services.AddSingleton<IServerInfo>(new ServerInfo());
+			services.AddSingleton<IServerInfo, ServerInfo>();
 			services.AddSingleton<IPluginCollection>(_pluginCollection);
 
 			services.AddLogging(builder => { builder.AddConsole().SetMinimumLevel(LogLevel.Debug); });
diff --git a/Engine/Source/Programs/Horde/HordeServer.Tests/Logs/LogServiceTests.cs b/Engine/Source/Programs/Horde/HordeServer.Tests/Logs/LogServiceTests.cs
--- a/Engine/Source/Programs/Horde/HordeServer.Tests/Logs/LogServiceTests.cs
+++ b/Engine/Source/Programs/Horde/HordeServer.Tests/Logs/LogServiceTests.cs
@@ -15,7 +15,6 @@ using EpicGames.Horde.Logs;
 using HordeServer.Agents.Sessions;
 using HordeServer.Jobs;
 using HordeServer.Logs;
-using HordeServer.Logs.Data;
 using HordeServer.Utilities;
 using Microsoft.VisualStudio.TestTools.UnitTesting;
 
@@ -31,7 +30,7 @@ namespace HordeServer.Tests.Logs
 			byte[] jsonData = Encoding.UTF8.GetBytes(jsonText);
 
 			byte[] outputData = new byte[jsonData.Length];
-			int outputLength = LogText.ConvertToPlainText(jsonData, outputData, 0);
+			int outputLength = HordeServer.Logs.LogCollection.ConvertToPlainText(jsonData, outputData, 0);
 			string outputText = Encoding.UTF8.GetString(outputData, 0, outputLength);
 
 			string expectText = "       \tdepends on: /mnt/horde/U5+R5.0+Inc+Min/Sync/Engine/Binaries/Linux/libUnrealEditor-Core.so\n";
diff --git a/Engine/Source/Programs/Horde/HordeServer.Tests/Logs/LogTailServiceTests.cs b/Engine/Source/Programs/Horde/HordeServer.Tests/Logs/LogTailServiceTests.cs
--- a/Engine/Source/Programs/Horde/HordeServer.Tests/Logs/LogTailServiceTests.cs
+++ b/Engine/Source/Programs/Horde/HordeServer.Tests/Logs/LogTailServiceTests.cs
@@ -73,7 +73,7 @@ namespace HordeServer.Tests.Logs
 		{
 			IRedisService redisService = ServiceProvider.GetRequiredService<IRedisService>();
 
-			LogTailService tailService = new LogTailService(redisService, ServiceProvider.GetRequiredService<IClock>(), 4, ServiceProvider.GetRequiredService<IOptions<ServerSettings>>(), ServiceProvider.GetRequiredService<ILogger<LogTailService>>());
+			LogTailService tailService = new LogTailService(redisService, ServiceProvider.GetRequiredService<IClock>(), 4, ServiceProvider.GetRequiredService<IServerInfo>(), ServiceProvider.GetRequiredService<ILogger<LogTailService>>());
 			await tailService.StartAsync(CancellationToken.None);
 
 			await tailService.EnableTailingAsync(_logId, 30);
@@ -173,7 +173,7 @@ namespace HordeServer.Tests.Logs
 			LogTailService tailService = ServiceProvider.GetRequiredService<LogTailService>();
 			await tailService.StartAsync(CancellationToken.None);
 
-			LogTailService tailService2 = new LogTailService(ServiceProvider.GetRequiredService<IRedisService>(), ServiceProvider.GetRequiredService<IClock>(), 4, ServiceProvider.GetRequiredService<IOptions<ServerSettings>>(), ServiceProvider.GetRequiredService<ILogger<LogTailService>>());
+			LogTailService tailService2 = new LogTailService(ServiceProvider.GetRequiredService<IRedisService>(), ServiceProvider.GetRequiredService<IClock>(), 4, ServiceProvider.GetRequiredService<IServerInfo>(), ServiceProvider.GetRequiredService<ILogger<LogTailService>>());
 
 			Task<int> task = tailService.WaitForTailNextAsync(_logId, CancellationToken.None);
 			Assert.IsTrue(!task.IsCompleted);
diff --git a/Engine/Source/Programs/Horde/HordeServer/Logs/Data/LogText.cs b/Engine/Source/Programs/Horde/HordeServer/Logs/Data/LogText.cs
deleted file mode 100644
