Horde: Move some classes into HordeServer.Shared to make them accessible outside the core server.

#rnx

diff --git a/Engine/Source/Programs/Horde/HordeServer.Shared/Acls/IAclService.cs b/Engine/Source/Programs/Horde/HordeServer.Shared/Acls/IAclService.cs
new file mode 100644
--- /dev/null
+++ b/Engine/Source/Programs/Horde/HordeServer.Shared/Acls/IAclService.cs
@@ -0,0 +1,25 @@
+// Copyright Epic Games, Inc. All Rights Reserved.
+
+using System;
+using System.Collections.Generic;
+using System.Linq;
+using System.Text;
+using System.Threading.Tasks;
+
+namespace HordeServer.Acls
+{
+	/// <summary>
+	/// Functionality for manipulating ACLs
+	/// </summary>
+	public interface IAclService
+	{
+		/// <summary>
+		/// Issues a bearer token with the given roles
+		/// </summary>
+		/// <param name="claims">List of claims to include</param>
+		/// <param name="expiry">Time that the token expires</param>
+		/// <param name="cancellationToken">Cancellation token for the operation</param>
+		/// <returns>JWT security token with a claim for creating new agents</returns>
+		ValueTask<string> IssueBearerTokenAsync(IEnumerable<AclClaimConfig> claims, TimeSpan? expiry, CancellationToken cancellationToken = default);
+	}
+}
diff --git a/Engine/Source/Programs/Horde/HordeServer.Shared/Server/IDowntimeService.cs b/Engine/Source/Programs/Horde/HordeServer.Shared/Server/IDowntimeService.cs
new file mode 100644
--- /dev/null
+++ b/Engine/Source/Programs/Horde/HordeServer.Shared/Server/IDowntimeService.cs
@@ -0,0 +1,24 @@
+// Copyright Epic Games, Inc. All Rights Reserved.
+
+using System;
+using System.Collections.Generic;
+using System.Linq;
+using System.Text;
+using System.Threading.Tasks;
+
+namespace HordeServer.Server
+{
+	/// <summary>
+	/// Interface for a service which keeps track of whether we're during downtime
+	/// </summary>
+	public interface IDowntimeService
+	{
+		/// <summary>
+		/// Returns true if downtime is currently active
+		/// </summary>
+		public bool IsDowntimeActive
+		{
+			get;
+		}
+	}
+}
diff --git a/Engine/Source/Programs/Horde/HordeServer.Shared/Utilities/AsyncTaskQueue.cs b/Engine/Source/Programs/Horde/HordeServer.Shared/Utilities/AsyncTaskQueue.cs
new file mode 100644
--- /dev/null
+++ b/Engine/Source/Programs/Horde/HordeServer.Shared/Utilities/AsyncTaskQueue.cs
@@ -0,0 +1,77 @@
+﻿// Copyright Epic Games, Inc. All Rights Reserved.
+
+using System;
+using System.Threading;
+using System.Threading.Tasks;
+using Microsoft.Extensions.Logging;
+
+namespace HordeServer.Utilities
+{
+	/// <summary>
+	/// Allows spawning async tasks that will run sequentially, and allows waiting for them to complete
+	/// </summary>
+	public sealed class AsyncTaskQueue : IAsyncDisposable
+	{
+		readonly object _lockObject = new object();
+		Task _task;
+		readonly ILogger _logger;
+		CancellationTokenSource _cancellationSource;
+
+		/// <summary>
+		/// Constructor
+		/// </summary>
+		/// <param name="logger">Logger for any exceptions thrown by tasks</param>
+		public AsyncTaskQueue(ILogger logger)
+		{
+			_task = Task.CompletedTask;
+			_logger = logger;
+			_cancellationSource = new CancellationTokenSource();
+		}
+
+		/// <inheritdoc/>
+		public async ValueTask DisposeAsync()
+		{
+			if (_cancellationSource != null)
+			{
+				await _cancellationSource.CancelAsync();
+				await _task;
+
+				_cancellationSource.Dispose();
+				_cancellationSource = null!;
+			}
+		}
+
+		/// <summary>
+		/// Adds a new task to the queue
+		/// </summary>
+		/// <param name="func">Method to create the new task</param>
+		public void Enqueue(Func<CancellationToken, Task> func)
+		{
+			lock (_lockObject)
+			{
+				Task prevTask = _task;
+				_task = Task.Run(() => RunAsync(prevTask, func));
+			}
+		}
+
+		/// <summary>
+		/// Waits for any queued tasks to complete
+		/// </summary>
+		public Task FlushAsync(CancellationToken cancellationToken = default)
+			=> _task.WaitAsync(cancellationToken);
+
+		async Task RunAsync(Task prevTask, Func<CancellationToken, Task> func)
+		{
+			try
+			{
+				await prevTask;
+			}
+			catch (Exception ex)
+			{
+				_logger.LogError(ex, "Unhandled exception while running task in AsyncTaskQueue: {Message}", ex.Message);
+			}
+
+			await func(_cancellationSource.Token);
+		}
+	}
+}
diff --git a/Engine/Source/Programs/Horde/HordeServer.Shared/Utilities/OpenTelemetryExtensions.cs b/Engine/Source/Programs/Horde/HordeServer.Shared/Utilities/OpenTelemetryExtensions.cs
new file mode 100644
--- /dev/null
+++ b/Engine/Source/Programs/Horde/HordeServer.Shared/Utilities/OpenTelemetryExtensions.cs
@@ -0,0 +1,34 @@
+// Copyright Epic Games, Inc. All Rights Reserved.
+
+using OpenTelemetry.Trace;
+
+namespace HordeServer.Utilities
+{
+	/// <summary>
+	/// Extensions to handle Horde specific data types in the OpenTelemetry library
+	/// </summary>
+	public static class OpenTelemetryExtensions
+	{
+		/// <summary>Set a key:value tag on the span</summary>
+		/// <returns>This span instance, for chaining</returns>
+		public static TelemetrySpan SetAttribute(this TelemetrySpan span, string key, int? value)
+		{
+			if (value != null)
+			{
+				span.SetAttribute(key, value.Value);
+			}
+			return span;
+		}
+
+		/// <summary>Set a key:value tag on the span</summary>
+		/// <returns>This span instance, for chaining</returns>
+		public static TelemetrySpan SetAttribute(this TelemetrySpan span, string key, DateTimeOffset? value)
+		{
+			if (value != null)
+			{
+				span.SetAttribute(key, value.ToString());
+			}
+			return span;
+		}
+	}
+}
diff --git a/Engine/Source/Programs/Horde/HordeServer.Shared/Utilities/StructuredRpcException.cs b/Engine/Source/Programs/Horde/HordeServer.Shared/Utilities/StructuredRpcException.cs
new file mode 100644
--- /dev/null
+++ b/Engine/Source/Programs/Horde/HordeServer.Shared/Utilities/StructuredRpcException.cs
@@ -0,0 +1,86 @@
+// Copyright Epic Games, Inc. All Rights Reserved.
+
+using System;
+using System.Globalization;
+using System.Text;
+using Grpc.Core;
+
+namespace HordeServer.Utilities
+{
+	/// <summary>
+	/// Exception class designed to allow logging structured log messages
+	/// </summary>
+	public class StructuredRpcException : RpcException
+	{
+		/// <summary>
+		/// The format string with named parameters
+		/// </summary>
+		public string Format { get; }
+
+		/// <summary>
+		/// The argument list
+		/// </summary>
+		public object[] Args { get; }
+
+		/// <summary>
+		/// Constructor
+		/// </summary>
+		/// <param name="statusCode">Status code to return</param>
+		/// <param name="format">The format string</param>
+		/// <param name="args">Arguments for the format string</param>
+		public StructuredRpcException(StatusCode statusCode, string format, params object[] args)
+			: base(new Status(statusCode, FormatMessage(format, args)))
+		{
+			Format = format;
+			Args = args;
+		}
+
+		/// <summary>
+		/// Replace named arguments in the format message with their values
+		/// </summary>
+		/// <param name="format"></param>
+		/// <param name="args"></param>
+		/// <returns></returns>
+		static string FormatMessage(string format, params object[] args)
+		{
+			string newFormat = ConvertToFormatString(format);
+			return String.Format(CultureInfo.CurrentCulture, newFormat, args);
+		}
+
+		/// <summary>
+		/// Converts a named parameter format string to a String.Format style string
+		/// </summary>
+		/// <param name="format"></param>
+		/// <returns></returns>
+		static string ConvertToFormatString(string format)
+		{
+			int argIdx = 0;
+
+			StringBuilder newFormat = new StringBuilder();
+			for (int idx = 0; idx < format.Length; idx++)
+			{
+				char character = format[idx];
+				newFormat.Append(character);
+
+				if (character == '{' && idx + 1 < format.Length)
+				{
+					char nextCharacter = format[idx + 1];
+					if ((nextCharacter >= 'a' && nextCharacter <= 'z') || (nextCharacter >= 'A' && nextCharacter <= 'Z') || nextCharacter == '_')
+					{
+						for (int endIdx = idx + 2; endIdx < format.Length; endIdx++)
+						{
+							if (format[endIdx] == ':' || format[endIdx] == '}')
+							{
+								newFormat.Append(argIdx++);
+								idx = endIdx - 1;
+								break;
+							}
+						}
+					}
+				}
+			}
+
+			return newFormat.ToString();
+		}
+	}
+}
diff --git a/Engine/Source/Programs/Horde/HordeServer/Acls/AclService.cs b/Engine/Source/Programs/Horde/HordeServer/Acls/AclService.cs
--- a/Engine/Source/Programs/Horde/HordeServer/Acls/AclService.cs
+++ b/Engine/Source/Programs/Horde/HordeServer/Acls/AclService.cs
@@ -20,7 +20,7 @@ namespace HordeServer.Acls
 	/// <summary>
 	/// Wraps functionality for manipulating permissions
 	/// </summary>
-	public class AclService
+	public class AclService : IAclService
 	{
 		private readonly GlobalsService _globalsService;
 
@@ -78,55 +78,4 @@ namespace HordeServer.Acls
 			}
 		}
 	}
-
-	internal static class ClaimExtensions
-	{
-		public static bool HasAgentClaim(this ClaimsPrincipal user, AgentId agentId)
-		{
-			return user.HasClaim(HordeClaimTypes.Agent, agentId.ToString());
-		}
-
-		public static bool HasLeaseClaim(this ClaimsPrincipal user, LeaseId leaseId)
-		{
-			return user.HasClaim(HordeClaimTypes.Lease, leaseId.ToString());
-		}
-
-		public static bool HasSessionClaim(this ClaimsPrincipal user, SessionId sessionId)
-		{
-			return user.HasClaim(HordeClaimTypes.AgentSessionId, sessionId.ToString());
-		}
-
-		public static LeaseId? GetLeaseClaim(this ClaimsPrincipal user)
-		{
-			Claim? claim = user.FindFirst(HordeClaimTypes.Lease);
-			if (claim == null || !LeaseId.TryParse(claim.Value, out LeaseId leaseIdValue))
-			{
-				return null;
-			}
-			else
-			{
-				return leaseIdValue;
-			}
-		}
-
-		public static SessionId? GetSessionClaim(this ClaimsPrincipal user)
-		{
-			Claim? claim = user.FindFirst(HordeClaimTypes.AgentSessionId);
-			if (claim == null || !SessionId.TryParse(claim.Value, out SessionId sessionIdValue))
-			{
-				return null;
-			}
-			else
-			{
-				return sessionIdValue;
-			}
-		}
-
-		public static string GetSessionClaimsAsString(this ClaimsPrincipal user)
-		{
-			return String.Join(",", user.Claims
-				.Where(c => c.Type == HordeClaimTypes.AgentSessionId)
-				.Select(c => c.Value));
-		}
-	}
 }
diff --git a/Engine/Source/Programs/Horde/HordeServer/Agents/AgentClaimExtensions.cs b/Engine/Source/Programs/Horde/HordeServer/Agents/AgentClaimExtensions.cs
new file mode 100644
--- /dev/null
+++ b/Engine/Source/Programs/Horde/HordeServer/Agents/AgentClaimExtensions.cs
@@ -0,0 +1,77 @@
+// Copyright Epic Games, Inc. All Rights Reserved.
+
+using System;
+using System.Collections.Generic;
+using System.Linq;
+using System.Security.Claims;
+using System.Text;
+using System.Threading.Tasks;
+using EpicGames.Horde.Agents;
+using EpicGames.Horde.Agents.Leases;
+using EpicGames.Horde.Agents.Sessions;
+using HordeServer.Utilities;
+
+namespace HordeServer.Agents
+{
+	/// <summary>
+	/// Extension methods for getting agent-related claims from principals
+	/// </summary>
+	public static class ClaimExtensions
+	{
+		/// <summary>
+		/// Test whether a principal is the given agent
+		/// </summary>
+		public static bool HasAgentClaim(this ClaimsPrincipal user, AgentId agentId)
+		{
+			return user.HasClaim(HordeClaimTypes.Agent, agentId.ToString());
+		}
+
+		/// <summary>
+		/// Test whether a principal is the given lease
+		/// </summary>
+		public static bool HasLeaseClaim(this ClaimsPrincipal user, LeaseId leaseId)
+		{
+			return user.HasClaim(HordeClaimTypes.Lease, leaseId.ToString());
+		}
+
+		/// <summary>
+		/// Test whether a principal is the given session
+		/// </summary>
+		public static bool HasSessionClaim(this ClaimsPrincipal user, SessionId sessionId)
+		{
+			return user.HasClaim(HordeClaimTypes.AgentSessionId, sessionId.ToString());
+		}
+
+		/// <summary>
+		/// Attempts to get the lease id from the given principal
+		/// </summary>
+		public static LeaseId? GetLeaseClaim(this ClaimsPrincipal user)
+		{
+			Claim? claim = user.FindFirst(HordeClaimTypes.Lease);
+			if (claim == null || !LeaseId.TryParse(claim.Value, out LeaseId leaseIdValue))
+			{
+				return null;
+			}
+			else
+			{
+				return leaseIdValue;
+			}
+		}
+
+		/// <summary>
+		/// Attempts to get the session id from the given principal
+		/// </summary>
+		public static SessionId? GetSessionClaim(this ClaimsPrincipal user)
+		{
+			Claim? claim = user.FindFirst(HordeClaimTypes.AgentSessionId);
+			if (claim == null || !SessionId.TryParse(claim.Value, out SessionId sessionIdValue))
+			{
+				return null;
+			}
+			else
+			{
+				return sessionIdValue;
+			}
+		}
+	}
+}
diff --git a/Engine/Source/Programs/Horde/HordeServer/Agents/Fleet/JobQueueSettings.cs b/Engine/Source/Programs/Horde/HordeServer/Agents/Fleet/JobQueueSettings.cs
new file mode 100644
--- /dev/null
+++ b/Engine/Source/Programs/Horde/HordeServer/Agents/Fleet/JobQueueSettings.cs
@@ -0,0 +1,74 @@
+﻿// Copyright Epic Games, Inc. All Rights Reserved.
+
+using System;
+using System.Collections.Generic;
+using System.Linq;
+using System.Text;
+using System.Text.Json.Serialization;
+using System.Threading.Tasks;
+
+namespace HordeServer.Agents.Fleet
+{
+	/// <summary>
+	/// Job queue sizing settings for a pool
+	/// </summary>
+	public class JobQueueSettings
+	{
+		/// <summary>
+		/// Factor translating queue size to additional agents to grow the pool with
+		/// The result is always rounded up to nearest integer. 
+		/// Example: if there are 20 jobs in queue, a factor 0.25 will result in 5 new agents being added (20 * 0.25)
+		/// </summary>
+		public double ScaleOutFactor { get; set; } = 0.25;
+
+		/// <summary>
+		/// Factor by which to shrink the pool size with when queue is empty
+		/// The result is always rounded up to nearest integer.
+		/// Example: when the queue size is zero, a default value of 0.9 will shrink the pool by 10% (current agent count * 0.9)
+		/// </summary>
+		public double ScaleInFactor { get; set; } = 0.9;
+
+		/// <summary>
+		/// How far back in time to look for job batches (that potentially are in the queue)
+		/// </summary>
+		public int SamplePeriodMin { get; set; } = 60 * 24 * 1; // 1 day
+
+		/// <summary>
+		/// Time spent in ready state before considered truly waiting for an agent
+		///
+		/// A job batch can be in ready state before getting picked up and executed.
+		/// This threshold will help ensure only batches that have been waiting longer than this value will be considered.
+		/// </summary>
+		public int ReadyTimeThresholdSec { get; set; } = 45;
+
+		/// <summary>
+		/// Constructor used for JSON serialization
+		/// </summary>
+		[JsonConstructor]
+		public JobQueueSettings()
+		{
+		}
+
+		/// <summary>
+		/// Constructor
+		/// </summary>
+		/// <param name="scaleOutFactor"></param>
+		/// <param name="scaleInFactor"></param>
+		public JobQueueSettings(double? scaleOutFactor = null, double? scaleInFactor = null)
+		{
+			ScaleOutFactor = scaleOutFactor.GetValueOrDefault(ScaleOutFactor);
+			ScaleInFactor = scaleInFactor.GetValueOrDefault(ScaleInFactor);
+		}
+
+		/// <inheritdoc />
+		public override string ToString()
+		{
+			StringBuilder sb = new(50);
+			sb.AppendFormat("{0}={1} ", nameof(ScaleOutFactor), ScaleOutFactor);
+			sb.AppendFormat("{0}={1} ", nameof(ScaleInFactor), ScaleInFactor);
+			sb.AppendFormat("{0}={1} ", nameof(SamplePeriodMin), SamplePeriodMin);
+			sb.AppendFormat("{0}={1} ", nameof(ReadyTimeThresholdSec), ReadyTimeThresholdSec);
+			return sb.ToString();
+		}
+	}
+}
diff --git a/Engine/Source/Programs/Horde/HordeServer/Agents/Fleet/JobQueueStrategy.cs b/Engine/Source/Programs/Horde/HordeServer/Agents/Fleet/JobQueueStrategy.cs
--- a/Engine/Source/Programs/Horde/HordeServer/Agents/Fleet/JobQueueStrategy.cs
+++ b/Engine/Source/Programs/Horde/HordeServer/Agents/Fleet/JobQueueStrategy.cs
@@ -23,69 +23,6 @@ using OpenTelemetry.Trace;
 
 namespace HordeServer.Agents.Fleet
 {
-	/// <summary>
-	/// Job queue sizing settings for a pool
-	/// </summary>
-	public class JobQueueSettings
-	{
-		/// <summary>
-		/// Factor translating queue size to additional agents to grow the pool with
-		/// The result is always rounded up to nearest integer. 
-		/// Example: if there are 20 jobs in queue, a factor 0.25 will result in 5 new agents being added (20 * 0.25)
-		/// </summary>
-		public double ScaleOutFactor { get; set; } = 0.25;
-
-		/// <summary>
-		/// Factor by which to shrink the pool size with when queue is empty
-		/// The result is always rounded up to nearest integer.
-		/// Example: when the queue size is zero, a default value of 0.9 will shrink the pool by 10% (current agent count * 0.9)
-		/// </summary>
-		public double ScaleInFactor { get; set; } = 0.9;
-
-		/// <summary>
-		/// How far back in time to look for job batches (that potentially are in the queue)
-		/// </summary>
-		public int SamplePeriodMin { get; set; } = 60 * 24 * 1; // 1 day
-
-		/// <summary>
-		/// Time spent in ready state before considered truly waiting for an agent
-		///
-		/// A job batch can be in ready state before getting picked up and executed.
-		/// This threshold will help ensure only batches that have been waiting longer than this value will be considered.
-		/// </summary>
-		public int ReadyTimeThresholdSec { get; set; } = 45;
-
-		/// <summary>
-		/// Constructor used for JSON serialization
-		/// </summary>
-		[JsonConstructor]
-		public JobQueueSettings()
-		{
-		}
-
-		/// <summary>
-		/// Constructor
-		/// </summary>
-		/// <param name="scaleOutFactor"></param>
-		/// <param name="scaleInFactor"></param>
-		public JobQueueSettings(double? scaleOutFactor = null, double? scaleInFactor = null)
-		{
-			ScaleOutFactor = scaleOutFactor.GetValueOrDefault(ScaleOutFactor);
-			ScaleInFactor = scaleInFactor.GetValueOrDefault(ScaleInFactor);
-		}
-
-		/// <inheritdoc />
-		public override string ToString()
-		{
-			StringBuilder sb = new(50);
-			sb.AppendFormat("{0}={1} ", nameof(ScaleOutFactor), ScaleOutFactor);
-			sb.AppendFormat("{0}={1} ", nameof(ScaleInFactor), ScaleInFactor);
-			sb.AppendFormat("{0}={1} ", nameof(SamplePeriodMin), SamplePeriodMin);
-			sb.AppendFormat("{0}={1} ", nameof(ReadyTimeThresholdSec), ReadyTimeThresholdSec);
-			return sb.ToString();
-		}
-	}
-
 	/// <summary>
 	/// Calculate pool size by observing the number of jobs in waiting state
 	///
diff --git a/Engine/Source/Programs/Horde/HordeServer/Artifacts/ArtifactsController.cs b/Engine/Source/Programs/Horde/HordeServer/Artifacts/ArtifactsController.cs
--- a/Engine/Source/Programs/Horde/HordeServer/Artifacts/ArtifactsController.cs
+++ b/Engine/Source/Programs/Horde/HordeServer/Artifacts/ArtifactsController.cs
@@ -37,6 +37,7 @@ using Microsoft.AspNetCore.StaticFiles;
 using Microsoft.Extensions.Logging;
 using Microsoft.Extensions.Options;
 using EpicGames.Horde.Acls;
+using HordeServer.Agents;
 
 namespace HordeServer.Artifacts
 {
diff --git a/Engine/Source/Programs/Horde/HordeServer/Compute/ComputeController.cs b/Engine/Source/Programs/Horde/HordeServer/Compute/ComputeController.cs
--- a/Engine/Source/Programs/Horde/HordeServer/Compute/ComputeController.cs
+++ b/Engine/Source/Programs/Horde/HordeServer/Compute/ComputeController.cs
@@ -10,6 +10,7 @@ using EpicGames.Horde.Agents;
 using EpicGames.Horde.Agents.Pools;
 using EpicGames.Horde.Compute;
 using HordeServer.Acls;
+using HordeServer.Agents;
 using HordeServer.Server;
 using HordeServer.Utilities;
 using Microsoft.AspNetCore.Authorization;
diff --git a/Engine/Source/Programs/Horde/HordeServer/Logs/ILog.cs b/Engine/Source/Programs/Horde/HordeServer/Logs/ILog.cs
--- a/Engine/Source/Programs/Horde/HordeServer/Logs/ILog.cs
+++ b/Engine/Source/Programs/Horde/HordeServer/Logs/ILog.cs
@@ -14,6 +14,7 @@ using EpicGames.Horde.Jobs;
 using EpicGames.Horde.Logs;
 using EpicGames.Horde.Storage;
 using HordeServer.Acls;
+using HordeServer.Agents;
 using Microsoft.Extensions.Logging;
 using MongoDB.Bson;
 
diff --git a/Engine/Source/Programs/Horde/HordeServer/Logs/LogsController.cs b/Engine/Source/Programs/Horde/HordeServer/Logs/LogsController.cs
--- a/Engine/Source/Programs/Horde/HordeServer/Logs/LogsController.cs
+++ b/Engine/Source/Programs/Horde/HordeServer/Logs/LogsController.cs
@@ -15,6 +15,7 @@ using EpicGames.Horde.Jobs;
 using EpicGames.Horde.Logs;
 using EpicGames.Horde.Storage;
 using HordeServer.Acls;
+using HordeServer.Agents;
 using HordeServer.Agents.Sessions;
 using HordeServer.Issues;
 using HordeServer.Jobs;
diff --git a/Engine/Source/Programs/Horde/HordeServer/Server/DowntimeService.cs b/Engine/Source/Programs/Horde/HordeServer/Server/DowntimeService.cs
--- a/Engine/Source/Programs/Horde/HordeServer/Server/DowntimeService.cs
+++ b/Engine/Source/Programs/Horde/HordeServer/Server/DowntimeService.cs
@@ -11,20 +11,6 @@ using Microsoft.Extensions.Options;
 
 namespace HordeServer.Server
 {
-	/// <summary>
-	/// Interface for a service which keeps track of whether we're during downtime
-	/// </summary>
-	public interface IDowntimeService
-	{
-		/// <summary>
-		/// Returns true if downtime is currently active
-		/// </summary>
-		public bool IsDowntimeActive
-		{
-			get;
-		}
-	}
-
 	/// <summary>
 	/// Service which manages the downtime schedule
 	/// </summary>
diff --git a/Engine/Source/Programs/Horde/HordeServer/Startup.cs b/Engine/Source/Programs/Horde/HordeServer/Startup.cs
--- a/Engine/Source/Programs/Horde/HordeServer/Startup.cs
+++ b/Engine/Source/Programs/Horde/HordeServer/Startup.cs
@@ -442,6 +442,7 @@ namespace HordeServer
 			services.AddHostedService(provider => provider.GetRequiredService<AgentRelayService>());
 
 			services.AddSingleton<AclService>();
+			services.AddSingleton<IAclService>(sp => sp.GetRequiredService<AclService>());
 			services.AddSingleton<AgentService>();
 			services.AddSingleton(provider => new Lazy<AgentService>(provider.GetRequiredService<AgentService>));
 			services.AddHostedService<ArtifactExpirationService>();
diff --git a/Engine/Source/Programs/Horde/HordeServer/Utilities/AsyncTaskQueue.cs b/Engine/Source/Programs/Horde/HordeServer/Utilities/AsyncTaskQueue.cs
deleted file mode 100644
diff --git a/Engine/Source/Programs/Horde/HordeServer/Utilities/OpenTelemetry.cs b/Engine/Source/Programs/Horde/HordeServer/Utilities/OpenTelemetry.cs
--- a/Engine/Source/Programs/Horde/HordeServer/Utilities/OpenTelemetry.cs
+++ b/Engine/Source/Programs/Horde/HordeServer/Utilities/OpenTelemetry.cs
@@ -260,28 +260,6 @@ public static class OpenTelemetrySpanExtensions
 		return span;
 	}
 
-	/// <summary>Set a key:value tag on the span</summary>
-	/// <returns>This span instance, for chaining</returns>
-	public static TelemetrySpan SetAttribute(this TelemetrySpan span, string key, int? value)
-	{
-		if (value != null)
-		{
-			span.SetAttribute(key, value.Value);
-		}
-		return span;
-	}
-
-	/// <summary>Set a key:value tag on the span</summary>
-	/// <returns>This span instance, for chaining</returns>
-	public static TelemetrySpan SetAttribute(this TelemetrySpan span, string key, DateTimeOffset? value)
-	{
-		if (value != null)
-		{
-			span.SetAttribute(key, value.ToString());
-		}
-		return span;
-	}
-
 	/// <inheritdoc cref="TelemetrySpan.SetAttribute(System.String, System.String)"/>
 	public static TelemetrySpan SetAttribute(this TelemetrySpan span, string key, StreamId? value) => span.SetAttribute(key, value?.ToString());
 
diff --git a/Engine/Source/Programs/Horde/HordeServer/Utilities/StructuredRpcException.cs b/Engine/Source/Programs/Horde/HordeServer/Utilities/StructuredRpcException.cs
deleted file mode 100644
