Horde: Allow authorizing against ACL scope names from IAclService.

#rnx

diff --git a/Engine/Source/Programs/Horde/HordeServer.Shared/Acls/IAclService.cs b/Engine/Source/Programs/Horde/HordeServer.Shared/Acls/IAclService.cs
--- a/Engine/Source/Programs/Horde/HordeServer.Shared/Acls/IAclService.cs
+++ b/Engine/Source/Programs/Horde/HordeServer.Shared/Acls/IAclService.cs
@@ -2,9 +2,12 @@
 
 using System;
 using System.Collections.Generic;
+using System.Diagnostics.CodeAnalysis;
 using System.Linq;
+using System.Security.Claims;
 using System.Text;
 using System.Threading.Tasks;
+using EpicGames.Horde.Acls;
 
 namespace HordeServer.Acls
 {
@@ -21,5 +24,27 @@ namespace HordeServer.Acls
 		/// <param name="cancellationToken">Cancellation token for the operation</param>
 		/// <returns>JWT security token with a claim for creating new agents</returns>
 		ValueTask<string> IssueBearerTokenAsync(IEnumerable<AclClaimConfig> claims, TimeSpan? expiry, CancellationToken cancellationToken = default);
+
+		/// <summary>
+		/// Finds an ACL scope by name
+		/// </summary>
+		/// <param name="scopeName">Name of the scope to auth against</param>
+		/// <param name="scopeConfig">Configuration for the scope</param>
+		bool TryGetAclScope(AclScopeName scopeName, [NotNullWhen(true)] out AclConfig? scopeConfig);
+	}
+
+	/// <summary>
+	/// Extension methods for <see cref="IAclService"/>
+	/// </summary>
+	public static class AclServiceExtensions
+	{
+		/// <summary>
+		/// Authorizes a user to perform a given action
+		/// </summary>
+		/// <param name="scopeName">Name of the scope to auth against</param>
+		/// <param name="action">The action being performed</param>
+		/// <param name="user">The principal to validate</param>
+		public static bool Authorize(this IAclService aclService, AclScopeName scopeName, AclAction action, ClaimsPrincipal user)
+			=> aclService.TryGetAclScope(scopeName, out AclConfig? scopeConfig) && scopeConfig.Authorize(action, user);
 	}
 }
diff --git a/Engine/Source/Programs/Horde/HordeServer/Acls/AclService.cs b/Engine/Source/Programs/Horde/HordeServer/Acls/AclService.cs
--- a/Engine/Source/Programs/Horde/HordeServer/Acls/AclService.cs
+++ b/Engine/Source/Programs/Horde/HordeServer/Acls/AclService.cs
@@ -2,16 +2,17 @@
 
 using System;
 using System.Collections.Generic;
+using System.Diagnostics.CodeAnalysis;
 using System.IdentityModel.Tokens.Jwt;
 using System.Linq;
 using System.Security.Claims;
 using System.Threading;
 using System.Threading.Tasks;
+using EpicGames.Horde.Acls;
 using EpicGames.Horde.Agents;
-using EpicGames.Horde.Agents.Leases;
-using EpicGames.Horde.Agents.Sessions;
 using HordeServer.Server;
 using HordeServer.Utilities;
+using Microsoft.Extensions.Options;
 using Microsoft.IdentityModel.Tokens;
 using MongoDB.Driver;
 
@@ -23,13 +24,15 @@ namespace HordeServer.Acls
 	public class AclService : IAclService
 	{
 		private readonly GlobalsService _globalsService;
+		private readonly IOptionsMonitor<GlobalConfig> _globalConfig;
 
 		/// <summary>
 		/// Constructor
 		/// </summary>
-		public AclService(GlobalsService globalsService)
+		public AclService(GlobalsService globalsService, IOptionsMonitor<GlobalConfig> globalConfig)
 		{
 			_globalsService = globalsService;
+			_globalConfig = globalConfig;
 		}
 
 		/// <summary>
@@ -77,5 +80,9 @@ namespace HordeServer.Acls
 				return new AgentId(claim.Value);
 			}
 		}
+
+		/// <inheritdoc/>
+		public bool TryGetAclScope(AclScopeName scopeName, [NotNullWhen(true)] out AclConfig? scopeConfig)
+			=> _globalConfig.CurrentValue.TryGetAclScope(scopeName, out scopeConfig);
 	}
 }
diff --git a/Engine/Source/Programs/Horde/HordeServer/Agents/AgentClaimExtensions.cs b/Engine/Source/Programs/Horde/HordeServer/Agents/AgentClaimExtensions.cs
--- a/Engine/Source/Programs/Horde/HordeServer/Agents/AgentClaimExtensions.cs
+++ b/Engine/Source/Programs/Horde/HordeServer/Agents/AgentClaimExtensions.cs
@@ -1,11 +1,6 @@
 // Copyright Epic Games, Inc. All Rights Reserved.
 
-using System;
-using System.Collections.Generic;
-using System.Linq;
 using System.Security.Claims;
-using System.Text;
-using System.Threading.Tasks;
 using EpicGames.Horde.Agents;
 using EpicGames.Horde.Agents.Leases;
 using EpicGames.Horde.Agents.Sessions;
diff --git a/Engine/Source/Programs/Horde/HordeServer/Agents/Fleet/JobQueueSettings.cs b/Engine/Source/Programs/Horde/HordeServer/Agents/Fleet/JobQueueSettings.cs
--- a/Engine/Source/Programs/Horde/HordeServer/Agents/Fleet/JobQueueSettings.cs
+++ b/Engine/Source/Programs/Horde/HordeServer/Agents/Fleet/JobQueueSettings.cs
@@ -1,11 +1,7 @@
 ﻿// Copyright Epic Games, Inc. All Rights Reserved.
 
-using System;
-using System.Collections.Generic;
-using System.Linq;
 using System.Text;
 using System.Text.Json.Serialization;
-using System.Threading.Tasks;
 
 namespace HordeServer.Agents.Fleet
 {
diff --git a/Engine/Source/Programs/Horde/HordeServer/Agents/Fleet/JobQueueStrategy.cs b/Engine/Source/Programs/Horde/HordeServer/Agents/Fleet/JobQueueStrategy.cs
--- a/Engine/Source/Programs/Horde/HordeServer/Agents/Fleet/JobQueueStrategy.cs
+++ b/Engine/Source/Programs/Horde/HordeServer/Agents/Fleet/JobQueueStrategy.cs
@@ -3,8 +3,6 @@
 using System;
 using System.Collections.Generic;
 using System.Linq;
-using System.Text;
-using System.Text.Json.Serialization;
 using System.Threading;
 using System.Threading.Tasks;
 using EpicGames.Core;
diff --git a/Engine/Source/Programs/Horde/HordeServer/Compute/ComputeController.cs b/Engine/Source/Programs/Horde/HordeServer/Compute/ComputeController.cs
--- a/Engine/Source/Programs/Horde/HordeServer/Compute/ComputeController.cs
+++ b/Engine/Source/Programs/Horde/HordeServer/Compute/ComputeController.cs
@@ -9,7 +9,6 @@ using EpicGames.Core;
 using EpicGames.Horde.Agents;
 using EpicGames.Horde.Agents.Pools;
 using EpicGames.Horde.Compute;
-using HordeServer.Acls;
 using HordeServer.Agents;
 using HordeServer.Server;
 using HordeServer.Utilities;
@@ -38,7 +37,7 @@ namespace HordeServer.Compute
 			_computeService = computeService;
 			_globalConfig = globalConfig;
 		}
-		
+
 		/// <summary>
 		/// Add tasks to be executed remotely and auto-select appropriate compute cluster to use
 		/// </summary>
@@ -54,7 +53,7 @@ namespace HordeServer.Compute
 			ClusterId clusterId = ComputeService.FindBestComputeClusterId(_globalConfig.Value, requesterIp);
 			return await AssignComputeResourceInClusterAsync(clusterId, request, cancellationToken);
 		}
-		
+
 		/// <summary>
 		/// Add tasks to be executed remotely
 		/// </summary>
diff --git a/Engine/Source/Programs/Horde/HordeServer/Logs/ILog.cs b/Engine/Source/Programs/Horde/HordeServer/Logs/ILog.cs
--- a/Engine/Source/Programs/Horde/HordeServer/Logs/ILog.cs
+++ b/Engine/Source/Programs/Horde/HordeServer/Logs/ILog.cs
@@ -13,7 +13,6 @@ using EpicGames.Horde.Agents.Sessions;
 using EpicGames.Horde.Jobs;
 using EpicGames.Horde.Logs;
 using EpicGames.Horde.Storage;
-using HordeServer.Acls;
 using HordeServer.Agents;
 using Microsoft.Extensions.Logging;
 using MongoDB.Bson;
diff --git a/Engine/Source/Programs/Horde/HordeServer/Logs/LogsController.cs b/Engine/Source/Programs/Horde/HordeServer/Logs/LogsController.cs
--- a/Engine/Source/Programs/Horde/HordeServer/Logs/LogsController.cs
+++ b/Engine/Source/Programs/Horde/HordeServer/Logs/LogsController.cs
@@ -14,7 +14,6 @@ using EpicGames.Horde.Acls;
 using EpicGames.Horde.Jobs;
 using EpicGames.Horde.Logs;
 using EpicGames.Horde.Storage;
-using HordeServer.Acls;
 using HordeServer.Agents;
 using HordeServer.Agents.Sessions;
 using HordeServer.Issues;
