Horde: Include an ACL scope name in log objects so we can (in the future) explicitly authorize access to them without requiring hooks into the job system.

#rnx

diff --git a/Engine/Source/Programs/Horde/HordeServer.Tests/Logs/LogServiceTests.cs b/Engine/Source/Programs/Horde/HordeServer.Tests/Logs/LogServiceTests.cs
--- a/Engine/Source/Programs/Horde/HordeServer.Tests/Logs/LogServiceTests.cs
+++ b/Engine/Source/Programs/Horde/HordeServer.Tests/Logs/LogServiceTests.cs
@@ -178,7 +178,7 @@ namespace HordeServer.Tests.Logs
 			// Will implicitly test GetLogFileAsync(), AddCachedLogFile()
 			JobId jobId = JobIdUtils.GenerateNewId();
 			SessionId sessionId = SessionIdUtils.GenerateNewId();
-			ILog a = await LogCollection.AddAsync(jobId, null, sessionId, LogType.Text, logId: null, CancellationToken.None);
+			ILog a = await LogCollection.AddAsync(jobId, null, sessionId, LogType.Text, logId: null, cancellationToken: CancellationToken.None);
 			ILog b = (await LogCollection.GetAsync(a.Id, CancellationToken.None))!;
 			Assert.AreEqual(a.JobId, b.JobId);
 			Assert.AreEqual(a.SessionId, b.SessionId);
diff --git a/Engine/Source/Programs/Horde/HordeServer.Tests/Logs/LogTailServiceTests.cs b/Engine/Source/Programs/Horde/HordeServer.Tests/Logs/LogTailServiceTests.cs
--- a/Engine/Source/Programs/Horde/HordeServer.Tests/Logs/LogTailServiceTests.cs
+++ b/Engine/Source/Programs/Horde/HordeServer.Tests/Logs/LogTailServiceTests.cs
@@ -12,7 +12,6 @@ using HordeServer.Logs;
 using HordeServer.Server;
 using Microsoft.Extensions.DependencyInjection;
 using Microsoft.Extensions.Logging;
-using Microsoft.Extensions.Options;
 using Microsoft.VisualStudio.TestTools.UnitTesting;
 
 namespace HordeServer.Tests.Logs
diff --git a/Engine/Source/Programs/Horde/HordeServer/Jobs/JobRpcService.cs b/Engine/Source/Programs/Horde/HordeServer/Jobs/JobRpcService.cs
--- a/Engine/Source/Programs/Horde/HordeServer/Jobs/JobRpcService.cs
+++ b/Engine/Source/Programs/Horde/HordeServer/Jobs/JobRpcService.cs
@@ -312,7 +312,7 @@ namespace HordeServer.Jobs
 			}
 
 			// Create a log file if necessary
-			log.Value ??= await _logCollection.AddAsync(job.Id, batch.LeaseId, batch.SessionId, LogType.Json);
+			log.Value ??= await _logCollection.AddAsync(job.Id, batch.LeaseId, batch.SessionId, LogType.Json, aclScopeName: streamConfig.Acl.ScopeName);
 
 			// Get the node for this step
 			IGraph graph = await _jobService.GetGraphAsync(job);
diff --git a/Engine/Source/Programs/Horde/HordeServer/Jobs/JobTaskSource.cs b/Engine/Source/Programs/Horde/HordeServer/Jobs/JobTaskSource.cs
--- a/Engine/Source/Programs/Horde/HordeServer/Jobs/JobTaskSource.cs
+++ b/Engine/Source/Programs/Horde/HordeServer/Jobs/JobTaskSource.cs
@@ -831,7 +831,7 @@ namespace HordeServer.Jobs
 					if (waiter.LeaseSource.TrySetResult(lease))
 					{
 						_logger.LogInformation("Assigned lease to agent");
-						await _logCollection.AddAsync(job.Id, leaseId, agent.SessionId, LogType.Json, logId, cancellationToken);
+						await _logCollection.AddAsync(job.Id, leaseId, agent.SessionId, LogType.Json, logId, streamConfig.Acl.ScopeName, cancellationToken);
 						return lease;
 					}
 				}
diff --git a/Engine/Source/Programs/Horde/HordeServer/Logs/ILog.cs b/Engine/Source/Programs/Horde/HordeServer/Logs/ILog.cs
--- a/Engine/Source/Programs/Horde/HordeServer/Logs/ILog.cs
+++ b/Engine/Source/Programs/Horde/HordeServer/Logs/ILog.cs
@@ -8,6 +8,7 @@ using System.Text.Json;
 using System.Threading;
 using System.Threading.Tasks;
 using EpicGames.Core;
+using EpicGames.Horde.Acls;
 using EpicGames.Horde.Agents.Leases;
 using EpicGames.Horde.Agents.Sessions;
 using EpicGames.Horde.Jobs;
@@ -27,47 +28,42 @@ namespace HordeServer.Logs
 		/// <summary>
 		/// Identifier for the log. Randomly generated.
 		/// </summary>
-		public LogId Id { get; }
+		LogId Id { get; }
 
 		/// <summary>
 		/// Unique id of the job containing this log
 		/// </summary>
-		public JobId JobId { get; }
+		JobId JobId { get; }
 
 		/// <summary>
 		/// The lease allowed to write to this log
 		/// </summary>
-		public LeaseId? LeaseId { get; }
+		LeaseId? LeaseId { get; }
 
 		/// <summary>
 		/// The session allowed to write to this log
 		/// </summary>
-		public SessionId? SessionId { get; }
+		SessionId? SessionId { get; }
 
 		/// <summary>
 		/// Type of data stored in this log 
 		/// </summary>
-		public LogType Type { get; }
+		LogType Type { get; }
 
 		/// <summary>
 		/// Namespace containing the log data
 		/// </summary>
-		public NamespaceId NamespaceId { get; }
+		NamespaceId NamespaceId { get; }
 
 		/// <summary>
 		/// Name of the ref used to store data for this log
 		/// </summary>
-		public RefName RefName { get; }
+		RefName RefName { get; }
 
-		//		/// <summary>
-		//		/// Number of lines (V2 storage backend)
-		//		/// </summary>
-		//		public int LineCount { get; }
-
-		//		/// <summary>
-		//		/// Whether the log is complete (V2 storage backend)
-		//		/// </summary>
-		//		public bool Complete { get; }
+		/// <summary>
+		/// Acl scope that owns this log
+		/// </summary>
+		AclScopeName AclScopeName { get; }
 
 		/// <summary>
 		/// Delete this log
diff --git a/Engine/Source/Programs/Horde/HordeServer/Logs/ILogCollection.cs b/Engine/Source/Programs/Horde/HordeServer/Logs/ILogCollection.cs
--- a/Engine/Source/Programs/Horde/HordeServer/Logs/ILogCollection.cs
+++ b/Engine/Source/Programs/Horde/HordeServer/Logs/ILogCollection.cs
@@ -3,6 +3,7 @@
 using System.Collections.Generic;
 using System.Threading;
 using System.Threading.Tasks;
+using EpicGames.Horde.Acls;
 using EpicGames.Horde.Agents.Leases;
 using EpicGames.Horde.Agents.Sessions;
 using EpicGames.Horde.Jobs;
@@ -24,9 +25,10 @@ namespace HordeServer.Logs
 		/// <param name="sessionId">Agent session allowed to update the log</param>
 		/// <param name="type">Type of events to be stored in the log</param>
 		/// <param name="logId">ID of the log file (optional)</param>
+		/// <param name="aclScopeName">Name of the acl scope to use for authorizing access to this log</param>
 		/// <param name="cancellationToken">Cancellation token for the call</param>
 		/// <returns>The new log file document</returns>
-		Task<ILog> AddAsync(JobId jobId, LeaseId? leaseId, SessionId? sessionId, LogType type, LogId? logId = null, CancellationToken cancellationToken = default);
+		Task<ILog> AddAsync(JobId jobId, LeaseId? leaseId, SessionId? sessionId, LogType type, LogId? logId = null, AclScopeName aclScopeName = default, CancellationToken cancellationToken = default);
 
 		/// <summary>
 		/// Gets a log by ID
diff --git a/Engine/Source/Programs/Horde/HordeServer/Logs/LogCollection.cs b/Engine/Source/Programs/Horde/HordeServer/Logs/LogCollection.cs
--- a/Engine/Source/Programs/Horde/HordeServer/Logs/LogCollection.cs
+++ b/Engine/Source/Programs/Horde/HordeServer/Logs/LogCollection.cs
@@ -11,6 +11,7 @@ using System.Text.Json;
 using System.Threading;
 using System.Threading.Tasks;
 using EpicGames.Core;
+using EpicGames.Horde.Acls;
 using EpicGames.Horde.Agents.Leases;
 using EpicGames.Horde.Agents.Sessions;
 using EpicGames.Horde.Jobs;
@@ -45,6 +46,7 @@ namespace HordeServer.Logs
 			LogType ILog.Type => _document.Type;
 			NamespaceId ILog.NamespaceId => _document.NamespaceId;
 			RefName ILog.RefName => _document.RefName;
+			AclScopeName ILog.AclScopeName => _document.AclScopeName;
 
 			public Log(LogCollection collection, LogDocument document)
 			{
@@ -115,6 +117,7 @@ namespace HordeServer.Logs
 
 			public NamespaceId NamespaceId { get; set; } = Namespace.Logs;
 			public RefName RefName { get; set; }
+			public AclScopeName AclScopeName { get; set; }
 
 			[BsonIgnoreIfDefault]
 			public bool Complete { get; set; }
@@ -127,7 +130,7 @@ namespace HordeServer.Logs
 			{
 			}
 
-			public LogDocument(JobId jobId, LeaseId? leaseId, SessionId? sessionId, LogType type, LogId? logId, NamespaceId namespaceId)
+			public LogDocument(JobId jobId, LeaseId? leaseId, SessionId? sessionId, LogType type, LogId? logId, NamespaceId namespaceId, AclScopeName aclScopeName)
 			{
 				Id = logId ?? LogIdUtils.GenerateNewId();
 				JobId = jobId;
@@ -138,6 +141,7 @@ namespace HordeServer.Logs
 				MaxLineIndex = 0;
 				NamespaceId = namespaceId;
 				RefName = new RefName(Id.ToString());
+				AclScopeName = aclScopeName;
 			}
 		}
 
@@ -292,9 +296,9 @@ namespace HordeServer.Logs
 		}
 
 		/// <inheritdoc/>
-		public async Task<ILog> AddAsync(JobId jobId, LeaseId? leaseId, SessionId? sessionId, LogType type, LogId? logId, CancellationToken cancellationToken)
+		public async Task<ILog> AddAsync(JobId jobId, LeaseId? leaseId, SessionId? sessionId, LogType type, LogId? logId, AclScopeName aclScopeName, CancellationToken cancellationToken)
 		{
-			LogDocument newLog = new LogDocument(jobId, leaseId, sessionId, type, logId, Namespace.Logs);
+			LogDocument newLog = new LogDocument(jobId, leaseId, sessionId, type, logId, Namespace.Logs, aclScopeName);
 			await _logCollection.InsertOneAsync(newLog, null, cancellationToken);
 			return new Log(this, newLog);
 		}
