Horde: Add endpoint for resolving a suitable compute cluster ID

diff --git a/Engine/Source/Programs/Horde/HordeServer/Compute/ComputeController.cs b/Engine/Source/Programs/Horde/HordeServer/Compute/ComputeController.cs
--- a/Engine/Source/Programs/Horde/HordeServer/Compute/ComputeController.cs
+++ b/Engine/Source/Programs/Horde/HordeServer/Compute/ComputeController.cs
@@ -1,5 +1,6 @@
 // Copyright Epic Games, Inc. All Rights Reserved.
 
+using System;
 using System.Collections.Generic;
 using System.Linq;
 using System.Net;
@@ -15,6 +16,7 @@ using HordeServer.Utilities;
 using Microsoft.AspNetCore.Authorization;
 using Microsoft.AspNetCore.Mvc;
 using Microsoft.Extensions.Options;
+using OpenTelemetry.Trace;
 
 namespace HordeServer.Compute
 {
@@ -28,14 +30,34 @@ namespace HordeServer.Compute
 	{
 		readonly ComputeService _computeService;
 		readonly IOptionsSnapshot<GlobalConfig> _globalConfig;
+		readonly Tracer _tracer;
 
 		/// <summary>
 		/// Constructor
 		/// </summary>
-		public ComputeController(ComputeService computeService, IOptionsSnapshot<GlobalConfig> globalConfig)
+		public ComputeController(ComputeService computeService, IOptionsSnapshot<GlobalConfig> globalConfig, Tracer tracer)
 		{
 			_computeService = computeService;
 			_globalConfig = globalConfig;
+			_tracer = tracer;
+		}
+		
+		/// <summary>
+		/// Find the most suitable cluster given a compute assignment request
+		/// </summary>
+		/// <param name="request">The request parameters</param>
+		/// <param name="cancellationToken">Cancellation token for the operation</param>
+		/// <returns></returns>
+		[HttpPost]
+		[Authorize]
+		[Route("/api/v2/compute/_cluster")] // Underscore to avoid clashing with endpoint for clusters below
+		public Task<ActionResult<GetClusterResponse>> GetClusterAsync([FromBody] AssignComputeRequest request, CancellationToken cancellationToken)
+		{
+			using TelemetrySpan span = _tracer.StartActiveSpan($"{nameof(ComputeController)}.{nameof(GetClusterAsync)}");
+			IPAddress requesterIp = ComputeService.ResolveRequesterIp(HttpContext.Connection.RemoteIpAddress, request.Connection?.PreferPublicIp, request.Connection?.ClientPublicIp);
+			ClusterId clusterId = ComputeService.FindBestComputeClusterId(_globalConfig.Value, requesterIp);
+			GetClusterResponse response = new () { ClusterId = clusterId };
+			return Task.FromResult(new ActionResult<GetClusterResponse>(response));
 		}
 
 		/// <summary>
@@ -47,8 +69,10 @@ namespace HordeServer.Compute
 		[HttpPost]
 		[Authorize]
 		[Route("/api/v2/compute")]
+		[Obsolete("Resolve cluster with get cluster endpoint instead")]
 		public async Task<ActionResult<AssignComputeResponse>> AssignComputeResourceAsync([FromBody] AssignComputeRequest request, CancellationToken cancellationToken)
 		{
+			using TelemetrySpan span = _tracer.StartActiveSpan($"{nameof(ComputeController)}.{nameof(AssignComputeResourceAsync)}");
 			IPAddress requesterIp = ComputeService.ResolveRequesterIp(HttpContext.Connection.RemoteIpAddress, request.Connection?.PreferPublicIp, request.Connection?.ClientPublicIp);
 			ClusterId clusterId = ComputeService.FindBestComputeClusterId(_globalConfig.Value, requesterIp);
 			return await AssignComputeResourceInClusterAsync(clusterId, request, cancellationToken);
diff --git a/Engine/Source/Programs/Shared/EpicGames.Horde/Compute/ComputeMessages.cs b/Engine/Source/Programs/Shared/EpicGames.Horde/Compute/ComputeMessages.cs
--- a/Engine/Source/Programs/Shared/EpicGames.Horde/Compute/ComputeMessages.cs
+++ b/Engine/Source/Programs/Shared/EpicGames.Horde/Compute/ComputeMessages.cs
@@ -108,7 +108,18 @@ namespace EpicGames.Horde.Compute
 	}
 
 	/// <summary>
-	/// Request a machine to execute compute requests
+	/// Response to a cluster lookup request
+	/// </summary>
+	public class GetClusterResponse
+	{
+		/// <summary>
+		/// Compute cluster ID
+		/// </summary>
+		public ClusterId ClusterId { get; set; }
+	}
+	
+	/// <summary>
+	/// Response to compute allocation request
 	/// </summary>
 	public class AssignComputeResponse
 	{
