Horde: Fix serialization of null scope names.

#rnx

diff --git a/Engine/Source/Programs/Horde/HordeServer/Startup.cs b/Engine/Source/Programs/Horde/HordeServer/Startup.cs
--- a/Engine/Source/Programs/Horde/HordeServer/Startup.cs
+++ b/Engine/Source/Programs/Horde/HordeServer/Startup.cs
@@ -977,19 +977,23 @@ namespace HordeServer
 		public sealed class AclActionBsonSerializer : SerializerBase<AclAction>
 		{
 			/// <inheritdoc/>
-			public override AclAction Deserialize(BsonDeserializationContext context, BsonDeserializationArgs args) => new AclAction(context.Reader.ReadString());
+			public override AclAction Deserialize(BsonDeserializationContext context, BsonDeserializationArgs args) 
+				=> new AclAction(context.Reader.ReadString());
 
 			/// <inheritdoc/>
-			public override void Serialize(BsonSerializationContext context, BsonSerializationArgs args, AclAction value) => context.Writer.WriteString(value.Name);
+			public override void Serialize(BsonSerializationContext context, BsonSerializationArgs args, AclAction value) 
+				=> context.Writer.WriteString(value.Name);
 		}
 
 		public sealed class AclScopeNameBsonSerializer : SerializerBase<AclScopeName>
 		{
 			/// <inheritdoc/>
-			public override AclScopeName Deserialize(BsonDeserializationContext context, BsonDeserializationArgs args) => new AclScopeName(context.Reader.ReadString());
+			public override AclScopeName Deserialize(BsonDeserializationContext context, BsonDeserializationArgs args) 
+				=> new AclScopeName(context.Reader.ReadString());
 
 			/// <inheritdoc/>
-			public override void Serialize(BsonSerializationContext context, BsonSerializationArgs args, AclScopeName value) => context.Writer.WriteString(value.Text);
+			public override void Serialize(BsonSerializationContext context, BsonSerializationArgs args, AclScopeName value)
+				=> context.Writer.WriteString(value.Text ?? String.Empty);
 		}
 
 		sealed class SubResourceIdBsonSerializer<TValue, TConverter> : SerializerBase<TValue> where TValue : struct where TConverter : SubResourceIdConverter<TValue>, new()
@@ -997,10 +1001,12 @@ namespace HordeServer
 			readonly TConverter _converter = new TConverter();
 
 			/// <inheritdoc/>
-			public override TValue Deserialize(BsonDeserializationContext context, BsonDeserializationArgs args) => _converter.FromSubResourceId(new SubResourceId((ushort)context.Reader.ReadInt32()));
+			public override TValue Deserialize(BsonDeserializationContext context, BsonDeserializationArgs args) 
+				=> _converter.FromSubResourceId(new SubResourceId((ushort)context.Reader.ReadInt32()));
 
 			/// <inheritdoc/>
-			public override void Serialize(BsonSerializationContext context, BsonSerializationArgs args, TValue value) => context.Writer.WriteInt32(_converter.ToSubResourceId(value).Value);
+			public override void Serialize(BsonSerializationContext context, BsonSerializationArgs args, TValue value) 
+				=> context.Writer.WriteInt32(_converter.ToSubResourceId(value).Value);
 		}
 
 		/// <summary>
diff --git a/Engine/Source/Programs/Shared/EpicGames.Horde/Acls/AclScopeName.cs b/Engine/Source/Programs/Shared/EpicGames.Horde/Acls/AclScopeName.cs
--- a/Engine/Source/Programs/Shared/EpicGames.Horde/Acls/AclScopeName.cs
+++ b/Engine/Source/Programs/Shared/EpicGames.Horde/Acls/AclScopeName.cs
@@ -28,13 +28,15 @@ namespace EpicGames.Horde.Acls
 		/// </summary>
 		/// <param name="name">Name to append</param>
 		/// <returns>New scope name</returns>
-		public AclScopeName Append(string name) => new AclScopeName($"{Text}/{name}");
+		public AclScopeName Append(string name) 
+			=> String.IsNullOrEmpty(Text)? new AclScopeName(name) : new AclScopeName($"{Text}/{name}");
 
 		/// <inheritdoc/>
-		public bool Equals(AclScopeName other) => Text.Equals(other.Text, StringComparison.Ordinal);
+		public bool Equals(AclScopeName other) 
+			=> String.Equals(Text ?? String.Empty, other.Text ?? String.Empty, StringComparison.Ordinal);
 
 		/// <inheritdoc/>
-		public override int GetHashCode() => Text.GetHashCode(StringComparison.Ordinal);
+		public override int GetHashCode() => (Text ?? String.Empty).GetHashCode(StringComparison.Ordinal);
 
 		/// <inheritdoc/>
 		public override string ToString() => Text;
