Horde: Remove v1 artifacts from dashboard
#rnx
#skipci

diff --git a/Engine/Source/Programs/Horde/HordeDashboard/src/backend/Api.ts b/Engine/Source/Programs/Horde/HordeDashboard/src/backend/Api.ts
--- a/Engine/Source/Programs/Horde/HordeDashboard/src/backend/Api.ts
+++ b/Engine/Source/Programs/Horde/HordeDashboard/src/backend/Api.ts
@@ -213,7 +213,6 @@ export type PoolData = GetPoolResponse
 export type BatchData = GetBatchResponse
 export type EventData = GetLogEventResponse
 export type LogData = GetLogFileResponse
-export type ArtifactData = GetArtifactResponse
 export type AclData = GetAclResponse
 export type AclEntryData = GetAclEntryResponse
 export type SoftwareData = GetSoftwareResponse
@@ -1307,9 +1306,6 @@ export type GetJobResponse = {
 	/** Whether to update issues based on the outcome of this job */
 	updateIssues?: boolean;
 
-	/** Whether to use the V2 artifacts endpoint */
-	useArtifactsV2?: boolean;
-
 	artifacts?: GetJobArtifactResponse[];
 
 	/**  Custom permissions for this object */
@@ -1626,35 +1622,6 @@ export type GetLogFileResponse = {
 
 }
 
-/**Response describing an artifact */
-export type GetArtifactResponse = {
-
-	/**Unique id of the artifact */
-	id: string;
-
-	/** Unique id of the job for this artifact */
-	jobId: string;
-
-	/** Unique id of the job for this artifact */
-	stepId?: string;
-
-	/** Download code for this artifact */
-	code?: string;
-
-	/** Name of the artifact */
-	name: string;
-
-	/** MimeType of the artifact	*/
-	mimeType: string;
-
-	/** Length of the artifact, in bytes */
-	length: number;
-
-	/**Per-object permissions */
-	acl?: GetAclResponse;
-
-}
-
 /**Parameters request a zip file for artifacts */
 export type GetArtifactZipRequest = {
 
diff --git a/Engine/Source/Programs/Horde/HordeDashboard/src/backend/Backend.ts b/Engine/Source/Programs/Horde/HordeDashboard/src/backend/Backend.ts
--- a/Engine/Source/Programs/Horde/HordeDashboard/src/backend/Backend.ts
+++ b/Engine/Source/Programs/Horde/HordeDashboard/src/backend/Backend.ts
@@ -1,7 +1,7 @@
 // Copyright Epic Games, Inc. All Rights Reserved.
 
 import templateCache from '../backend/TemplateCache';
-import { AccountClaimMessage, AgentData, AgentQuery, ApproveAgentsRequest, ArtifactData, AuditLogEntry, AuditLogQuery, BatchUpdatePoolRequest, ChangeSummaryData, CreateAccountRequest, CreateBisectTaskRequest, CreateBisectTaskResponse, CreateDeviceRequest, CreateDeviceResponse, CreateExternalIssueRequest, CreateExternalIssueResponse, CreateJobRequest, CreateJobResponse, CreateNoticeRequest, CreatePoolRequest, CreateServiceAccountRequest, CreateServiceAccountResponse, CreateSoftwareResponse, CreateSubscriptionRequest, CreateSubscriptionResponse, CreateZipRequest, DashboardPreference, DevicePoolTelemetryQuery, DeviceTelemetryQuery, EventData, FindArtifactsResponse, FindIssueResponse, FindJobTimingsResponse, GetAccountResponse, GetAgentLeaseResponse, GetAgentSessionResponse, GetAgentTelemetrySampleResponse, GetArtifactDirectoryResponse, GetArtifactZipRequest, GetBisectTaskResponse, GetDashboardConfigResponse, GetDevicePlatformResponse, GetDevicePoolResponse, GetDevicePoolTelemetryResponse, GetDeviceReservationResponse, GetDeviceResponse, GetDeviceTelemetryResponse, GetExternalIssueProjectResponse, GetExternalIssueResponse, GetIssueStreamResponse, GetJobsTabResponse, GetJobStepRefResponse, GetJobStepTraceResponse, GetJobTimingResponse, GetLogEventResponse, GetLogFileResponse, GetNoticeResponse, GetNotificationResponse, GetPendingAgentsResponse, GetPerforceServerStatusResponse, GetPoolResponse, GetPoolSummaryResponse, GetServerInfoResponse, GetServerSettingsResponse, GetServiceAccountResponse, GetSoftwareResponse, GetSubscriptionResponse, GetTelemetryMetricsResponse, GetTestDataDetailsResponse, GetTestDataRefResponse, GetTestMetaResponse, GetTestResponse, GetTestsRequest, GetTestStreamResponse, GetToolResponse, GetToolSummaryResponse, GetUserResponse, GetUtilizationTelemetryResponse, GlobalConfig, IssueData, IssueQuery, IssueQueryV2, JobData, JobQuery, JobsTabColumnType, JobStepOutcome, JobStreamQuery, JobTimingsQuery, LeaseData, LogData, LogEventQuery, LogLineData, MetricsQuery, PoolData, PoolQuery, PreflightConfigResponse, ProjectData, ScheduleData, ScheduleQuery, SearchLogFileResponse, ServerStatusResponse, ServerUpdateResponse, SessionData, StreamData, TabType, TestData, UpdateAccountRequest, UpdateAgentRequest, UpdateBisectTaskRequest, UpdateCurrentAccountRequest, UpdateDeploymentRequest, UpdateDeviceRequest, UpdateGlobalConfigRequest, UpdateIssueRequest, UpdateJobRequest, UpdateLeaseRequest, UpdateNoticeRequest, UpdateNotificationsRequest, UpdatePoolRequest, UpdateServerSettingsRequest, UpdateServiceAccountRequest, UpdateServiceAccountResponse, UpdateStepRequest, UpdateStepResponse, UpdateTemplateRefRequest, UpdateUserRequest, UsersQuery } from './Api';
+import { AccountClaimMessage, AgentData, AgentQuery, ApproveAgentsRequest, AuditLogEntry, AuditLogQuery, BatchUpdatePoolRequest, ChangeSummaryData, CreateAccountRequest, CreateBisectTaskRequest, CreateBisectTaskResponse, CreateDeviceRequest, CreateDeviceResponse, CreateExternalIssueRequest, CreateExternalIssueResponse, CreateJobRequest, CreateJobResponse, CreateNoticeRequest, CreatePoolRequest, CreateServiceAccountRequest, CreateServiceAccountResponse, CreateSoftwareResponse, CreateSubscriptionRequest, CreateSubscriptionResponse, CreateZipRequest, DashboardPreference, DevicePoolTelemetryQuery, DeviceTelemetryQuery, EventData, FindArtifactsResponse, FindIssueResponse, FindJobTimingsResponse, GetAccountResponse, GetAgentLeaseResponse, GetAgentSessionResponse, GetAgentTelemetrySampleResponse, GetArtifactDirectoryResponse, GetArtifactZipRequest, GetBisectTaskResponse, GetDashboardConfigResponse, GetDevicePlatformResponse, GetDevicePoolResponse, GetDevicePoolTelemetryResponse, GetDeviceReservationResponse, GetDeviceResponse, GetDeviceTelemetryResponse, GetExternalIssueProjectResponse, GetExternalIssueResponse, GetIssueStreamResponse, GetJobsTabResponse, GetJobStepRefResponse, GetJobStepTraceResponse, GetJobTimingResponse, GetLogEventResponse, GetLogFileResponse, GetNoticeResponse, GetNotificationResponse, GetPendingAgentsResponse, GetPerforceServerStatusResponse, GetPoolResponse, GetPoolSummaryResponse, GetServerInfoResponse, GetServerSettingsResponse, GetServiceAccountResponse, GetSoftwareResponse, GetSubscriptionResponse, GetTelemetryMetricsResponse, GetTestDataDetailsResponse, GetTestDataRefResponse, GetTestMetaResponse, GetTestResponse, GetTestsRequest, GetTestStreamResponse, GetToolResponse, GetToolSummaryResponse, GetUserResponse, GetUtilizationTelemetryResponse, GlobalConfig, IssueData, IssueQuery, IssueQueryV2, JobData, JobQuery, JobsTabColumnType, JobStepOutcome, JobStreamQuery, JobTimingsQuery, LeaseData, LogData, LogEventQuery, LogLineData, MetricsQuery, PoolData, PoolQuery, PreflightConfigResponse, ProjectData, ScheduleData, ScheduleQuery, SearchLogFileResponse, ServerStatusResponse, ServerUpdateResponse, SessionData, StreamData, TabType, TestData, UpdateAccountRequest, UpdateAgentRequest, UpdateBisectTaskRequest, UpdateCurrentAccountRequest, UpdateDeploymentRequest, UpdateDeviceRequest, UpdateGlobalConfigRequest, UpdateIssueRequest, UpdateJobRequest, UpdateLeaseRequest, UpdateNoticeRequest, UpdateNotificationsRequest, UpdatePoolRequest, UpdateServerSettingsRequest, UpdateServiceAccountRequest, UpdateServiceAccountResponse, UpdateStepRequest, UpdateStepResponse, UpdateTemplateRefRequest, UpdateUserRequest, UsersQuery } from './Api';
 import dashboard, { Dashboard } from './Dashboard';
 import { ChallengeStatus, Fetch } from './Fetch';
 import { projectStore } from './ProjectStore';
@@ -542,29 +542,6 @@ export class Backend {
     }
 
 
-    getJobArtifacts(jobId: string, stepId?: string): Promise<ArtifactData[]> {
-
-        const params: any = {
-            jobId: jobId,
-            code: true,
-        };
-
-        if (stepId) {
-            params.stepId = stepId;
-        }
-
-        return new Promise<ArtifactData[]>((resolve, reject) => {
-            this.backend.get(`/api/v1/artifacts`, { params: params, suppress404: true }).then((value) => {
-                resolve(value.data as ArtifactData[]);
-            }).catch(reason => {
-                resolve([]);
-                if (reason !== "Not Found") {
-                    console.error(reason);
-                }
-            });
-        });
-    }
-
     getArtifactV2(artifactId: string, path: string): Promise<object> {
         const url = `/api/v2/artifacts/${artifactId}/file?path=${encodeURIComponent(path)}`;
         return new Promise<object>((resolve, reject) => {
@@ -580,7 +557,6 @@ export class Backend {
 
     }
 
-
     downloadArtifactV2(artifactId: string, path: string) {
         const url = `/api/v2/artifacts/${artifactId}/file?path=${encodeURIComponent(path)}`;
         window.location.assign(url);
@@ -599,38 +575,6 @@ export class Backend {
 
     }
 
-
-    downloadJobArtifacts(request: GetArtifactZipRequest): Promise<boolean> {
-        return new Promise<any>((resolve, reject) => {
-            this.backend.post(`/api/v1/artifacts/zip`, request, { responseBlob: true }).then(response => {
-                const url = window.URL.createObjectURL(response.data);
-                const link = document.createElement('a');
-                link.href = url;
-                let name = request.jobId + '-' + request.stepId + ".zip";
-                link.setAttribute('download', name);
-                document.body.appendChild(link);
-                link.click();
-                document.body.removeChild(link);
-                setTimeout(() => {
-                    window.URL.revokeObjectURL(url);
-                }, 100);
-                resolve(true);
-            }).catch(reason => {
-                reject(reason);
-            });
-        });
-    }
-
-    getArtifactDataById(id: string): Promise<object> {
-        return new Promise<object>((resolve, reject) => {
-            this.backend.get(`/api/v1/artifacts/${id}/data`).then((value) => {
-                resolve(value.data as object);
-            }).catch((reason) => {
-                reject(reason);
-            });
-        });
-    }
-
     getLogData(logId: string): Promise<LogData> {
 
         return new Promise<LogData>((resolve, reject) => {
diff --git a/Engine/Source/Programs/Horde/HordeDashboard/src/backend/JobDetails.ts b/Engine/Source/Programs/Horde/HordeDashboard/src/backend/JobDetails.ts
--- a/Engine/Source/Programs/Horde/HordeDashboard/src/backend/JobDetails.ts
+++ b/Engine/Source/Programs/Horde/HordeDashboard/src/backend/JobDetails.ts
@@ -6,7 +6,7 @@ import moment from 'moment';
 import backend from '.';
 import { getBatchInitElapsed, getNiceTime, getStepETA, getStepElapsed, getStepFinishTime, getStepTimingDelta } from '../base/utilities/timeUtils';
 import { getBatchText } from '../components/JobDetailCommon';
-import { AgentData, ArtifactData, BatchData, EventData, GetJobTimingResponse, GetLabelStateResponse, GetLabelTimingInfoResponse, IssueData, JobData, JobState, JobStepBatchState, JobStepError, JobStepOutcome, JobStepState, LabelState, ReportPlacement, StepData, StreamData, TestData } from './Api';
+import { AgentData, BatchData, EventData, GetJobTimingResponse, GetLabelStateResponse, GetLabelTimingInfoResponse, IssueData, JobData, JobState, JobStepBatchState, JobStepError, JobStepOutcome, JobStepState, LabelState, ReportPlacement, StepData, StreamData, TestData } from './Api';
 import { projectStore } from './ProjectStore';
 
 
@@ -766,12 +766,6 @@ export class JobDetails {
                 requests.push(this.getLogEvents(logId));
             }
 
-            let artifactsIdx = -1;
-            if (!jobdata.useArtifactsV2 && !this.artifacts?.length) {
-                artifactsIdx = requests.length;
-                requests.push(backend.getJobArtifacts(this.id!));
-            }
-
             Promise.all(requests as any).then(async (values) => {
 
                 if (this.canceled.has(cancelID)) {
@@ -782,10 +776,6 @@ export class JobDetails {
                     this.testdata = values[1] as any;
                 }
 
-                if (artifactsIdx !== -1) {
-                    this.artifacts = values[artifactsIdx] as any;
-                }
-
                 this.process();
 
                 const agentIds: string[] = [];
@@ -850,8 +840,7 @@ export class JobDetails {
     stream?: StreamData
     labels: JobLabel[] = []
     batches: BatchData[] = []
-    events: EventData[] = []
-    artifacts: ArtifactData[] = []
+    events: EventData[] = []    
     testdata: TestData[] = []
     issues: IssueData[] = [];
     private timing?: GetJobTimingResponse;
diff --git a/Engine/Source/Programs/Horde/HordeDashboard/src/backend/TestDataHandler.ts b/Engine/Source/Programs/Horde/HordeDashboard/src/backend/TestDataHandler.ts
--- a/Engine/Source/Programs/Horde/HordeDashboard/src/backend/TestDataHandler.ts
+++ b/Engine/Source/Programs/Horde/HordeDashboard/src/backend/TestDataHandler.ts
@@ -2,7 +2,7 @@
 
 import { observable, action, when, makeObservable } from "mobx";
 import backend from '../backend';
-import { JobData, ArtifactData, TestData } from './Api';
+import { JobData, TestData } from './Api';
 
 export type DataWrapper = {
     postPropertiesUpdate(): void
@@ -32,8 +32,6 @@ export class TestDataWrapper implements TestData {
 
     get type(): string { return this.key.split('::', 2)[0]; }
 
-    private jobArtifacts?: ArtifactData[];
-    private artifactMap?: Map<string, ArtifactData | undefined>;
     private artifactV2Map: Map<string, object | undefined> = new Map();
     private artifactV2Id?: string;
     private jobdata?: JobData;
@@ -93,7 +91,7 @@ export class TestDataWrapper implements TestData {
 
     async getArtifacts() {
 
-        if (this.jobArtifacts?.length || this.artifactV2Id) {
+        if (this.artifactV2Id) {
             return;
         }
 
@@ -101,38 +99,21 @@ export class TestDataWrapper implements TestData {
             this.jobdata = await backend.getJob(this.jobId);
         }
 
-        if (this.jobdata?.useArtifactsV2) {
-
-            const v = await backend.getJobArtifactsV2(undefined, [`job:${this.jobId}/step:${this.stepId}`]);
-            const av2 = v?.artifacts.find(a => a.type === "step-saved")
-            if (av2?.id) {
-               this.artifactV2Id = av2.id;
-            } else {
-                console.error("Unable to get step-saved artifacts v2 for report");
-            }
-
+        const v = await backend.getJobArtifactsV2(undefined, [`job:${this.jobId}/step:${this.stepId}`]);
+        const av2 = v?.artifacts.find(a => a.type === "step-saved")
+        if (av2?.id) {
+            this.artifactV2Id = av2.id;
         } else {
-
-            this.jobArtifacts = await backend.getJobArtifacts(this.jobId, this.stepId);
-
+            console.error("Unable to get step-saved artifacts v2 for report");
         }
+
     }
 
     async getArtifactImageLink(referencePath: string) {
 
         await this.getArtifacts();
 
-        const artifactName = referencePath.replace(/\\/g, '/');
-
-        if (this.artifactV2Id) {
-           return `${backend.serverUrl}/api/v2/artifacts/${this.artifactV2Id}/file?path=Engine/Programs/AutomationTool/Saved/Logs/RunUnreal/${encodeURIComponent(referencePath)}&inline=true`;
-        }
-
-        const artifact = this.jobArtifacts!.find(a => a.name.indexOf(artifactName) > -1);
-        if (artifact) {
-           return `${backend.serverUrl}/api/v1/artifacts/${artifact.id}/download?Code=${artifact.code}`;
-        }
-        return undefined;
+        return `${backend.serverUrl}/api/v2/artifacts/${this.artifactV2Id}/file?path=Engine/Programs/AutomationTool/Saved/Logs/RunUnreal/${encodeURIComponent(referencePath)}&inline=true`;
 
     }
 
@@ -158,31 +139,10 @@ export class TestDataWrapper implements TestData {
                 this.artifactV2Map.set(path, result);
                 return data;
             }
-
-            return undefined;
-        }
-
-        if (!this.artifactMap) {
-            this.artifactMap = new Map();
-        }
-
-        if (this.artifactMap.has(artifactName)) {
-            return this.artifactMap.get(artifactName);
-        }
-
-        const found = this.jobArtifacts?.find((value) => value.name.indexOf(artifactName) > -1);
-        this.artifactMap.set(artifactName, found);
-
-
-        if (found?.id) {
-            try {
-                return await backend.getArtifactDataById(found.id);
-            } catch (ex) {
-                console.error(ex);
-            }
         }
 
         return undefined;
+        
     }
 
     async getJobStepName() {
diff --git a/Engine/Source/Programs/Horde/HordeDashboard/src/components/AutomationSuiteTest.tsx b/Engine/Source/Programs/Horde/HordeDashboard/src/components/AutomationSuiteTest.tsx
--- a/Engine/Source/Programs/Horde/HordeDashboard/src/components/AutomationSuiteTest.tsx
+++ b/Engine/Source/Programs/Horde/HordeDashboard/src/components/AutomationSuiteTest.tsx
@@ -5,7 +5,7 @@ import { observer } from 'mobx-react-lite';
 import React, { useState } from "react";
 import { Link } from "react-router-dom";
 import backend from "../backend";
-import { ArtifactData, GetSuiteTestDataResponse, GetTestDataDetailsResponse, GetTestDataRefResponse, GetTestMetaResponse, GetTestResponse, GetTestSuiteResponse, TestData, TestOutcome } from "../backend/Api";
+import { GetSuiteTestDataResponse, GetTestDataDetailsResponse, GetTestDataRefResponse, GetTestMetaResponse, GetTestResponse, GetTestSuiteResponse, TestData, TestOutcome } from "../backend/Api";
 import { TestDataHandler } from "../backend/AutomationTestData";
 import dashboard, { StatusColor } from "../backend/Dashboard";
 import { getShortNiceTime } from "../base/utilities/timeUtils";
@@ -59,29 +59,10 @@ class SuiteTestHandler {
       this.artifactV2Id = undefined;
    }
 
-   get jobArtifacts(): ArtifactData[] | undefined {
-      if (!this.jobId || !this.stepId) {
-         return undefined;
-      }
-
-      const key = this.jobId + this.stepId;
-      return this.artifacts.get(key);
-   }
-
    getArtifactImagePath(referencePath: string) {
-
-      const artifactName = referencePath.replace(/\\/g, '/');
-
-      if (this.artifactV2Id) {
-         return `${backend.serverUrl}/api/v2/artifacts/${this.artifactV2Id}/file?path=Engine/Programs/AutomationTool/Saved/Logs/RunUnreal/${encodeURIComponent(referencePath)}&inline=true`;
-      } 
-
-      const artifact = this.jobArtifacts!.find(a => a.name.indexOf(artifactName) > -1);
-      if (artifact) {
-         return `${backend.serverUrl}/api/v1/artifacts/${artifact.id}/download?Code=${artifact.code}`;         
-      } 
-      return undefined;
-  
+      
+      return `${backend.serverUrl}/api/v2/artifacts/${this.artifactV2Id}/file?path=Engine/Programs/AutomationTool/Saved/Logs/RunUnreal/${encodeURIComponent(referencePath)}&inline=true`;
+      
    }
 
    getSuiteTest(metaId: string) {
@@ -187,14 +168,6 @@ class SuiteTestHandler {
                   this.artifactV2Id = av2.id;
                }
             }
-
-            if (!this.artifactV2Id) {
-               const key = this.jobId + this.stepId;
-               if (!this.artifacts.get(key)) {                                    
-                  const artifacts = await backend.getJobArtifacts(this.jobId, this.stepId);                  
-                  this.artifacts.set(key, artifacts);                     
-               }   
-            }
          }
       }
 
@@ -271,7 +244,6 @@ class SuiteTestHandler {
 
    jobId?: string;
    stepId?: string;
-   artifacts: Map<string, ArtifactData[]> = new Map();
    artifactV2Id?: string;
    events: TestEvent[] = [];
 
@@ -309,7 +281,7 @@ const EventPanel: React.FC<{ handler: SuiteTestHandler }> = observer(({ handler
          return oa - ob;
       });      
 
-      if (e.Tag === EventTag.ImageComparison && (handler.artifactV2Id || handler.jobArtifacts?.length) && imageArtifacts?.length) {         
+      if (e.Tag === EventTag.ImageComparison && (handler.artifactV2Id) && imageArtifacts?.length) {         
 
          imageArtifacts.forEach(ta => {
 
diff --git a/Engine/Source/Programs/Horde/HordeDashboard/src/components/JobDetailArtifacts.tsx b/Engine/Source/Programs/Horde/HordeDashboard/src/components/JobDetailArtifacts.tsx
deleted file mode 100644
diff --git a/Engine/Source/Programs/Horde/HordeDashboard/src/components/LogSource.tsx b/Engine/Source/Programs/Horde/HordeDashboard/src/components/LogSource.tsx
--- a/Engine/Source/Programs/Horde/HordeDashboard/src/components/LogSource.tsx
+++ b/Engine/Source/Programs/Horde/HordeDashboard/src/components/LogSource.tsx
@@ -443,7 +443,7 @@ export class JobLogSource extends LogSource {
    private async updateArtifacts() {
       const details = this.jobDetails;
 
-      if (!details.jobdata?.useArtifactsV2 || this.artifactsV2 !== undefined || !this.logData?.id || details.getLogActive(this.logData.id)) {
+      if (this.artifactsV2 !== undefined || !this.logData?.id || details.getLogActive(this.logData.id)) {
          return;
       }
 
diff --git a/Engine/Source/Programs/Horde/HordeDashboard/src/components/LogView.tsx b/Engine/Source/Programs/Horde/HordeDashboard/src/components/LogView.tsx
--- a/Engine/Source/Programs/Horde/HordeDashboard/src/components/LogView.tsx
+++ b/Engine/Source/Programs/Horde/HordeDashboard/src/components/LogView.tsx
@@ -7,7 +7,7 @@ import moment from 'moment-timezone';
 import React, { useEffect, useId, useState } from 'react';
 import { Link, useLocation, useNavigate, useParams } from 'react-router-dom';
 import backend from '../backend';
-import { ArtifactContextType, ArtifactData, EventSeverity, GetChangeSummaryResponse, GetJobStepRefResponse, GetLogEventResponse, LogLevel } from '../backend/Api';
+import { ArtifactContextType, EventSeverity, GetChangeSummaryResponse, GetJobStepRefResponse, GetLogEventResponse, LogLevel } from '../backend/Api';
 import { CommitCache } from '../backend/CommitCache';
 import dashboard from '../backend/Dashboard';
 import { JobDetails } from '../backend/JobDetails';
@@ -22,7 +22,6 @@ import { Breadcrumbs } from './Breadcrumbs';
 import { ChangeContextMenu, ChangeContextMenuTarget } from './ChangeButton';
 import { HistoryModal } from './HistoryModal';
 import { IssueModalV2 } from './IssueViewV2';
-import { JobDetailArtifacts } from './JobDetailArtifacts';
 import { useQuery } from './JobDetailCommon';
 import { LogItem, renderLine } from './LogRender';
 import { JobLogSource, LogSource } from './LogSource';
@@ -412,41 +411,6 @@ const StepHistoryModal: React.FC<{ jobDetails: JobDetails, stepId: string | unde
    </Modal>);
 });
 
-const StepArtifactsModal: React.FC<{ jobDetails: JobDetails, stepId: string | undefined, onClose: () => void }> = observer(({ jobDetails, stepId, onClose }) => {
-
-   const { hordeClasses } = getHordeStyling();
-
-   let artifacts: ArtifactData[] = jobDetails.artifacts;
-   if (stepId) {
-      artifacts = artifacts.filter(artifact => artifact.stepId === stepId);
-   }
-
-   let height = Math.min(36 * artifacts.length + 60, 500) + 200;
-
-   const hordeTheme = getHordeTheme();
-
-   return (<Modal isOpen={true} styles={{ main: { padding: 8, width: 1084, height: height, backgroundColor: hordeTheme.horde.contentBackground } }} className={hordeClasses.modal} onDismiss={() => { onClose() }}>
-
-      <Stack styles={{ root: { paddingTop: 8, paddingLeft: 24, paddingRight: 12, paddingBottom: 16 } }}>
-         <Stack tokens={{ childrenGap: 12 }}>
-            <Stack horizontal styles={{ root: { padding: 8 } }}>
-               <Stack grow horizontalAlign="end">
-                  <IconButton
-                     iconProps={{ iconName: 'Cancel' }}
-                     onClick={() => { onClose(); }}
-                  />
-               </Stack>
-            </Stack>
-
-            <Stack styles={{ root: { paddingLeft: 4, paddingRight: 0, paddingBottom: 4 } }}>
-               <JobDetailArtifacts jobDetails={jobDetails} stepId={stepId} topPadding={0} />
-            </Stack>
-         </Stack>
-      </Stack>
-   </Modal>);
-
-});
-
 const LogProgressIndicator: React.FC<{ logSource: LogSource }> = observer(({ logSource }) => {
 
    // subscribe
@@ -516,7 +480,7 @@ export const LogList: React.FC<{ logId: string }> = observer(({ logId }) => {
 
       return () => {
          globalHandler = globalSearchState = undefined;
-         window.removeEventListener("keydown", handler);         
+         window.removeEventListener("keydown", handler);
       };
 
    }, []);
@@ -924,110 +888,92 @@ export const LogList: React.FC<{ logId: string }> = observer(({ logId }) => {
 
    if (fixme) {
 
-      let artifacts: ArtifactData[] = [];
 
-      if (fixme.artifacts) {
-         const artifactStepId = fixme!.stepByLogId(logId)?.id
-         if (artifactStepId) {
-            artifacts = fixme.artifacts.filter(artifact => artifact.stepId === artifactStepId);
-         }
-      }
-
-      if (!fixme.jobdata?.useArtifactsV2) {
-         menuProps.items.push({
-            key: 'jobstep_artifacts',
-            disabled: artifacts.length === 0,
-            text: `Step Artifacts`,
-            onClick: () => setLogArtifacts("legacy")
-         })
-      } else {
+      const stepArtifacts = (logSource as JobLogSource).artifactsV2;
 
-         const stepArtifacts = (logSource as JobLogSource).artifactsV2;
+      const atypes = new Map<ArtifactContextType, number>();
+      const knownTypes = new Set<string>(["step-saved", "step-output", "step-trace"]);
 
-         const atypes = new Map<ArtifactContextType, number>();
-         const knownTypes = new Set<string>(["step-saved", "step-output", "step-trace"]);
+      stepArtifacts?.forEach(a => {
+         let c = atypes.get(a.type) ?? 0;
+         c++;
+         atypes.set(a.type, c);
+      });
 
-         stepArtifacts?.forEach(a => {
-            let c = atypes.get(a.type) ?? 0;
-            c++;
-            atypes.set(a.type, c);
-         });
+      const opsList: IContextualMenuItem[] = [];
 
-         const opsList: IContextualMenuItem[] = [];
+      const navigateToArtifacts = (context: string) => {
+         const search = new URLSearchParams(window.location.search);
+         search.set("artifactContext", encodeURIComponent(context));
+         const url = `${window.location.pathname}?` + search.toString();
+         navigate(url, { replace: true })
+      }
 
-         const navigateToArtifacts = (context: string) => {
-            const search = new URLSearchParams(window.location.search);
-            search.set("artifactContext", encodeURIComponent(context));
-            const url = `${window.location.pathname}?` + search.toString();
-            navigate(url, { replace: true })
+      opsList.push({
+         key: 'stepops_artifacts_step',
+         text: "Logs",
+         iconProps: { iconName: "Folder" },
+         disabled: !atypes.get("step-saved"),
+         onClick: () => {
+            navigateToArtifacts("step-saved");
          }
+      });
 
-         opsList.push({
-            key: 'stepops_artifacts_step',
-            text: "Logs",
-            iconProps: { iconName: "Folder" },
-            disabled: !atypes.get("step-saved"),
-            onClick: () => {
-               navigateToArtifacts("step-saved");
-            }
-         });
+      opsList.push({
+         key: 'stepops_artifacts_output',
+         text: "Temp Storage",
+         iconProps: { iconName: "MenuOpen" },
+         disabled: !atypes.get("step-output"),
+         onClick: () => {
+            navigateToArtifacts("step-output");
+         }
+      });
 
-         opsList.push({
-            key: 'stepops_artifacts_output',
-            text: "Temp Storage",
-            iconProps: { iconName: "MenuOpen" },
-            disabled: !atypes.get("step-output"),
-            onClick: () => {
-               navigateToArtifacts("step-output");
-            }
-         });
+      opsList.push({
+         key: 'stepops_artifacts_trace',
+         text: "Traces",
+         iconProps: { iconName: "SearchTemplate" },
+         disabled: !atypes.get("step-trace"),
+         onClick: () => {
+            navigateToArtifacts("step-trace");
+         }
+      });
 
+      const custom = stepArtifacts?.filter(a => !knownTypes.has(a.type)).sort((a, b) => a.type.localeCompare(b.type));
+      custom?.forEach(c => {
          opsList.push({
-            key: 'stepops_artifacts_trace',
-            text: "Traces",
-            iconProps: { iconName: "SearchTemplate" },
-            disabled: !atypes.get("step-trace"),
-            onClick: () => {
-               navigateToArtifacts("step-trace");
-            }
+            key: `stepops_artifacts_${c.type}`,
+            text: c.description ?? c.name,
+            iconProps: { iconName: "Clean" },
+            onClick: () => { navigateToArtifacts(c.type) }
          });
+      })
 
-         const custom = stepArtifacts?.filter(a => !knownTypes.has(a.type)).sort((a, b) => a.type.localeCompare(b.type));
-         custom?.forEach(c => {
-            opsList.push({
-               key: `stepops_artifacts_${c.type}`,
-               text: c.description ?? c.name,
-               iconProps: { iconName: "Clean" },
-               onClick: () => { navigateToArtifacts(c.type) }
-            });
-         })
-
-
-         menuProps.items.push({
-            key: 'jobstep_artifacts',
-            text: `Artifacts`,
-            subMenuProps: {
-               items: opsList
-            }
-         })
-      }
 
       menuProps.items.push({
-         key: 'jobstep_history',
-         text: 'Step History',
-         onClick: () => setLogHistory(true)
+         key: 'jobstep_artifacts',
+         text: `Artifacts`,
+         subMenuProps: {
+            items: opsList
+         }
       })
+   }
 
-      if (logSource.agentTelemetry) {
-         menuProps.items.push({
-            key: 'jobstep_agent_telemetry',
-            text: logTelemetry ? 'Hide Telemetry' : 'Show Telemetry',
-            onClick: () => {
-               logSource.agentTelemetry?.show(!logTelemetry);
-               setLogTelemetry(!logTelemetry)               
-            }
-         })
-      }
+   menuProps.items.push({
+      key: 'jobstep_history',
+      text: 'Step History',
+      onClick: () => setLogHistory(true)
+   })
+
+   if (logSource.agentTelemetry) {
+      menuProps.items.push({
+         key: 'jobstep_agent_telemetry',
+         text: logTelemetry ? 'Hide Telemetry' : 'Show Telemetry',
+         onClick: () => {
+            logSource.agentTelemetry?.show(!logTelemetry);
+            setLogTelemetry(!logTelemetry)
+         }
+      })
    }
 
    function updateError() {
@@ -1156,7 +1102,6 @@ export const LogList: React.FC<{ logId: string }> = observer(({ logId }) => {
    return <Stack>
       {!!fixme && <IssueModalV2 issueId={query.get("issue")} popHistoryOnClose={issueHistory} />}
       {!!fixme && logHistory && <StepHistoryModal jobDetails={fixme!} stepId={fixme!.stepByLogId(logId)?.id} onClose={() => setLogHistory(false)} />}
-      {!!fixme && logArtifacts === "legacy" && <StepArtifactsModal jobDetails={fixme!} stepId={fixme!.stepByLogId(logId)?.id} onClose={() => setLogArtifacts("")} />}
       {!!fixme && !!artifactContext && logArtifacts !== "legacy" && <JobArtifactsModal jobId={fixme!.jobdata!.id} stepId={fixme!.stepByLogId(logId)?.id!} artifacts={(logSource as JobLogSource).artifactsV2} contextType={artifactContext} artifactPath={artifactPath} onClose={() => { navigate(window.location.pathname, { replace: true }) }} />}
       {!!historyAgentId && <HistoryModal agentId={historyAgentId} onDismiss={() => { navigate(baseUrl, { replace: true }); setHistoryAgentId(undefined) }} />}
       <Breadcrumbs items={logSource?.crumbs ?? []} title={logSource?.crumbTitle} />
@@ -1304,7 +1249,7 @@ export const LogList: React.FC<{ logId: string }> = observer(({ logId }) => {
                         <Markdown>{summaryText}</Markdown>
                      </Stack>
                   </Stack>
-                  {logTelemetry && !! logSource.agentTelemetry && <AgentTelemetrySparkline handler={logSource.agentTelemetry} />}
+                  {logTelemetry && !!logSource.agentTelemetry && <AgentTelemetrySparkline handler={logSource.agentTelemetry} />}
                   <Stack horizontalAlign="center" style={{ paddingBottom: 12 }}>
                      <Separator styles={{ root: { fontSize: 0, width: "100%", padding: 0, selectors: { '::before': { background: dashboard.darktheme ? '#313638' : '#D3D2D1' } } } }} />
                   </Stack>
diff --git a/Engine/Source/Programs/Horde/HordeDashboard/src/components/jobDetailsV2/JobDetailArtifactsV2.tsx b/Engine/Source/Programs/Horde/HordeDashboard/src/components/jobDetailsV2/JobDetailArtifactsV2.tsx
deleted file mode 100644
diff --git a/Engine/Source/Programs/Horde/HordeDashboard/src/components/jobDetailsV2/JobOperationsBar.tsx b/Engine/Source/Programs/Horde/HordeDashboard/src/components/jobDetailsV2/JobOperationsBar.tsx
--- a/Engine/Source/Programs/Horde/HordeDashboard/src/components/jobDetailsV2/JobOperationsBar.tsx
+++ b/Engine/Source/Programs/Horde/HordeDashboard/src/components/jobDetailsV2/JobOperationsBar.tsx
@@ -238,10 +238,6 @@ const StepArtifactsOperations: React.FC<{ jobDetails: JobDetailsV2, stepId: stri
 
    const jobData = jobDetails.jobData;
 
-   if (!jobData?.useArtifactsV2) {
-      return null;
-   }
-
    // subscribe
    if (dashboard.updated) { }
    if (jobDetails.updated) { }
diff --git a/Engine/Source/Programs/Horde/HordeDashboard/src/components/jobDetailsV2/StepDetailView.tsx b/Engine/Source/Programs/Horde/HordeDashboard/src/components/jobDetailsV2/StepDetailView.tsx
--- a/Engine/Source/Programs/Horde/HordeDashboard/src/components/jobDetailsV2/StepDetailView.tsx
+++ b/Engine/Source/Programs/Horde/HordeDashboard/src/components/jobDetailsV2/StepDetailView.tsx
@@ -9,7 +9,6 @@ import { GetArtifactResponseV2 } from "../../backend/Api";
 import { getHordeStyling } from "../../styles/Styles";
 import { HistoryModal } from "../HistoryModal";
 import { useQuery } from "../JobDetailCommon";
-import { JobDetailArtifactsV2 } from "./JobDetailArtifactsV2";
 import { BisectionPanel } from "./JobDetailBisection";
 import { HealthPanel } from "./JobDetailHealthV2";
 import { PreflightPanel } from "./JobDetailPreflight";
@@ -50,11 +49,6 @@ class StepDetailDataView extends JobDataView {
          return;
       }
 
-      if (!this.details.jobData.useArtifactsV2) {
-         this.artifacts = [];
-         return;
-      }
-
       const stepId = this.stepId!;
 
       let artifacts = this.details.stepArtifacts.get(stepId);
@@ -163,9 +157,6 @@ const StepDetailViewInner: React.FC<{ jobDetails: JobDetailsV2, stepId: string }
       {<Stack>
          <StepHistoryPanel jobDetails={jobDetails} stepId={stepId} />
       </Stack>}
-      {<Stack>
-         <JobDetailArtifactsV2 jobDetails={jobDetails} stepId={stepId} />
-      </Stack>}
       {!!step && <Stack>
          <TimelinePanel jobDetails={jobDetails} stepId={stepId} />
       </Stack>}
