Horde: Decouple the log system from jobs and issues.

#rnx

diff --git a/Engine/Source/Programs/Horde/HordeServer/Issues/IssueCollection.cs b/Engine/Source/Programs/Horde/HordeServer/Issues/IssueCollection.cs
--- a/Engine/Source/Programs/Horde/HordeServer/Issues/IssueCollection.cs
+++ b/Engine/Source/Programs/Horde/HordeServer/Issues/IssueCollection.cs
@@ -20,6 +20,7 @@ using EpicGames.Redis.Utility;
 using HordeServer.Auditing;
 using HordeServer.Jobs;
 using HordeServer.Jobs.Graphs;
+using HordeServer.Logs;
 using HordeServer.Server;
 using HordeServer.Streams;
 using HordeServer.Telemetry;
@@ -36,7 +37,7 @@ using OpenTelemetry.Trace;
 
 namespace HordeServer.Issues
 {
-	class IssueCollection : IIssueCollection
+	class IssueCollection : IIssueCollection, ILogExtIssueProvider
 	{
 		[SingletonDocument("issue-ledger", "5e4c226440ce25fa3207a9af")]
 		class IssueLedger : SingletonBase
@@ -1698,5 +1699,12 @@ namespace HordeServer.Issues
 		{
 			return _auditLog[issueId];
 		}
+
+		/// <inheritdoc/>
+		public async ValueTask<int?> GetIssueIdAsync(ObjectId spanId, CancellationToken cancellationToken = default)
+		{
+			IIssueSpan? span = await GetSpanAsync(spanId, cancellationToken);
+			return span?.IssueId;
+		}
 	}
 }
diff --git a/Engine/Source/Programs/Horde/HordeServer/Jobs/JobService.cs b/Engine/Source/Programs/Horde/HordeServer/Jobs/JobService.cs
--- a/Engine/Source/Programs/Horde/HordeServer/Jobs/JobService.cs
+++ b/Engine/Source/Programs/Horde/HordeServer/Jobs/JobService.cs
@@ -26,6 +26,7 @@ using HordeServer.Jobs.Bisect;
 using HordeServer.Jobs.Graphs;
 using HordeServer.Jobs.Templates;
 using HordeServer.Jobs.Timing;
+using HordeServer.Logs;
 using HordeServer.Notifications;
 using HordeServer.Perforce;
 using HordeServer.Server;
@@ -33,6 +34,7 @@ using HordeServer.Streams;
 using HordeServer.Users;
 using HordeServer.Utilities;
 using Microsoft.Extensions.Logging;
+using Microsoft.Extensions.Options;
 using MongoDB.Bson;
 using OpenTelemetry.Trace;
 
@@ -73,7 +75,7 @@ namespace HordeServer.Jobs
 	/// <summary>
 	/// Wraps funtionality for manipulating jobs, jobsteps, and jobstep runs
 	/// </summary>
-	public class JobService
+	public class JobService : ILogExtAuthProvider
 	{
 		readonly IJobCollection _jobs;
 		readonly IGraphCollection _graphs;
@@ -88,6 +90,7 @@ namespace HordeServer.Jobs
 		readonly ITemplateCollection _templateCollection;
 		readonly IssueService? _issueService;
 		readonly IPerforceService _perforceService;
+		readonly IOptionsMonitor<GlobalConfig> _globalConfig;
 		readonly Tracer _tracer;
 		readonly ILogger _logger;
 
@@ -121,7 +124,7 @@ namespace HordeServer.Jobs
 		/// <summary>
 		/// Constructor
 		/// </summary>
-		public JobService(IJobCollection jobs, IGraphCollection graphs, IAgentCollection agents, IJobStepRefCollection jobStepRefs, IBisectTaskCollection bisectTasks, IJobTimingCollection jobTimings, IUserCollection userCollection, INotificationTriggerCollection triggerCollection, JobTaskSource jobTaskSource, IStreamCollection streamCollection, ITemplateCollection templateCollection, IssueService issueService, IPerforceService perforceService, Tracer tracer, ILogger<JobService> logger)
+		public JobService(IJobCollection jobs, IGraphCollection graphs, IAgentCollection agents, IJobStepRefCollection jobStepRefs, IBisectTaskCollection bisectTasks, IJobTimingCollection jobTimings, IUserCollection userCollection, INotificationTriggerCollection triggerCollection, JobTaskSource jobTaskSource, IStreamCollection streamCollection, ITemplateCollection templateCollection, IssueService issueService, IPerforceService perforceService, Tracer tracer, IOptionsMonitor<GlobalConfig> globalConfig, ILogger<JobService> logger)
 		{
 			_jobs = jobs;
 			_graphs = graphs;
@@ -136,6 +139,7 @@ namespace HordeServer.Jobs
 			_templateCollection = templateCollection;
 			_issueService = issueService;
 			_perforceService = perforceService;
+			_globalConfig = globalConfig;
 			_tracer = tracer;
 			_logger = logger;
 
@@ -1409,6 +1413,12 @@ namespace HordeServer.Jobs
 			return globalConfig.Authorize(job, action, user);
 		}
 
+		/// <inheritdoc/>
+		public async ValueTask<bool> AuthorizeAsync(ILog log, AclAction action, ClaimsPrincipal principal, CancellationToken cancellationToken)
+		{
+			return log.JobId != JobId.Empty && await AuthorizeAsync(log.JobId, action, principal, _globalConfig.CurrentValue, cancellationToken);
+		}
+
 		/// <summary>
 		/// Determines if the user is authorized to perform an action on a particular stream
 		/// </summary>
diff --git a/Engine/Source/Programs/Horde/HordeServer/Logs/ILogExtAuthProvider.cs b/Engine/Source/Programs/Horde/HordeServer/Logs/ILogExtAuthProvider.cs
new file mode 100644
--- /dev/null
+++ b/Engine/Source/Programs/Horde/HordeServer/Logs/ILogExtAuthProvider.cs
@@ -0,0 +1,24 @@
+// Copyright Epic Games, Inc. All Rights Reserved.
+
+using System.Security.Claims;
+using System.Threading;
+using System.Threading.Tasks;
+using EpicGames.Horde.Acls;
+
+namespace HordeServer.Logs
+{
+	/// <summary>
+	/// Interface which allows 
+	/// </summary>
+	public interface ILogExtAuthProvider
+	{
+		/// <summary>
+		/// Determine if the given principal can access a particular log
+		/// </summary>
+		/// <param name="log">The log to query</param>
+		/// <param name="action">Action to perform</param>
+		/// <param name="principal">The principal to authorize</param>
+		/// <param name="cancellationToken">Cancellation token for the operation</param>
+		public ValueTask<bool> AuthorizeAsync(ILog log, AclAction action, ClaimsPrincipal principal, CancellationToken cancellationToken = default);
+	}
+}
diff --git a/Engine/Source/Programs/Horde/HordeServer/Logs/ILogExtIssueProvider.cs b/Engine/Source/Programs/Horde/HordeServer/Logs/ILogExtIssueProvider.cs
new file mode 100644
--- /dev/null
+++ b/Engine/Source/Programs/Horde/HordeServer/Logs/ILogExtIssueProvider.cs
@@ -0,0 +1,22 @@
+// Copyright Epic Games, Inc. All Rights Reserved.
+
+using System.Threading;
+using System.Threading.Tasks;
+using MongoDB.Bson;
+
+namespace HordeServer.Logs
+{
+	/// <summary>
+	/// Interface which provides issue identifiers for event ids in log documents
+	/// </summary>
+	public interface ILogExtIssueProvider
+	{
+		/// <summary>
+		/// Gets an issue it for the given span
+		/// </summary>
+		/// <param name="spanId">Identifier for the span</param>
+		/// <param name="cancellationToken">Cancellation token for the operation</param>
+		/// <returns>Issue identifier</returns>
+		public ValueTask<int?> GetIssueIdAsync(ObjectId spanId, CancellationToken cancellationToken = default);
+	}
+}
diff --git a/Engine/Source/Programs/Horde/HordeServer/Logs/LogsController.cs b/Engine/Source/Programs/Horde/HordeServer/Logs/LogsController.cs
--- a/Engine/Source/Programs/Horde/HordeServer/Logs/LogsController.cs
+++ b/Engine/Source/Programs/Horde/HordeServer/Logs/LogsController.cs
@@ -11,13 +11,10 @@ using System.Threading;
 using System.Threading.Tasks;
 using EpicGames.Core;
 using EpicGames.Horde.Acls;
-using EpicGames.Horde.Jobs;
 using EpicGames.Horde.Logs;
 using EpicGames.Horde.Storage;
 using HordeServer.Agents;
 using HordeServer.Agents.Sessions;
-using HordeServer.Issues;
-using HordeServer.Jobs;
 using HordeServer.Server;
 using HordeServer.Storage;
 using HordeServer.Utilities;
@@ -53,20 +50,20 @@ namespace HordeServer.Logs
 	public class LogsController : ControllerBase
 	{
 		private readonly ILogCollection _logCollection;
-		private readonly IIssueCollection _issueCollection;
-		private readonly JobService _jobService;
 		private readonly StorageService _storageService;
+		private readonly IEnumerable<ILogExtAuthProvider> _authProviders;
+		private readonly IEnumerable<ILogExtIssueProvider> _issueProviders;
 		private readonly IOptionsSnapshot<GlobalConfig> _globalConfig;
 
 		/// <summary>
 		/// Constructor
 		/// </summary>
-		public LogsController(ILogCollection logCollection, IIssueCollection issueCollection, JobService jobService, StorageService storageService, IOptionsSnapshot<GlobalConfig> globalConfig)
+		public LogsController(ILogCollection logCollection, StorageService storageService, IEnumerable<ILogExtAuthProvider> authProviders, IEnumerable<ILogExtIssueProvider> issueProviders, IOptionsSnapshot<GlobalConfig> globalConfig)
 		{
 			_logCollection = logCollection;
-			_issueCollection = issueCollection;
-			_jobService = jobService;
 			_storageService = storageService;
+			_authProviders = authProviders;
+			_issueProviders = issueProviders;
 			_globalConfig = globalConfig;
 		}
 
@@ -339,8 +336,10 @@ namespace HordeServer.Logs
 				int? issueId = null;
 				if (logEvent.SpanId != null && !spanIdToIssueId.TryGetValue(logEvent.SpanId.Value, out issueId))
 				{
-					IIssueSpan? span = await _issueCollection.GetSpanAsync(logEvent.SpanId.Value, cancellationToken);
-					issueId = span?.IssueId;
+					foreach (ILogExtIssueProvider issueProvider in _issueProviders)
+					{
+						issueId ??= await issueProvider.GetIssueIdAsync(logEvent.SpanId.Value, cancellationToken);
+					}
 					spanIdToIssueId[logEvent.SpanId.Value] = issueId;
 				}
 
@@ -390,13 +389,16 @@ namespace HordeServer.Logs
 			{
 				return true;
 			}
-			if (log.JobId != JobId.Empty && await _jobService.AuthorizeAsync(log.JobId, action, user, globalConfig, cancellationToken))
+			if (action == LogAclAction.ViewLog && log.SessionId != null && globalConfig.Authorize(SessionAclAction.ViewSession, user))
 			{
 				return true;
 			}
-			if (action == LogAclAction.ViewLog && log.SessionId != null && globalConfig.Authorize(SessionAclAction.ViewSession, user))
+			foreach (ILogExtAuthProvider authProvider in _authProviders)
 			{
-				return true;
+				if (await authProvider.AuthorizeAsync(log, action, user, cancellationToken))
+				{
+					return true;
+				}
 			}
 			if (globalConfig.Authorize(action, user))
 			{
diff --git a/Engine/Source/Programs/Horde/HordeServer/Startup.cs b/Engine/Source/Programs/Horde/HordeServer/Startup.cs
--- a/Engine/Source/Programs/Horde/HordeServer/Startup.cs
+++ b/Engine/Source/Programs/Horde/HordeServer/Startup.cs
@@ -381,7 +381,9 @@ namespace HordeServer
 			services.AddSingleton<IAgentCollection, AgentCollection>();
 			services.AddSingleton<IArtifactCollection, ArtifactCollection>();
 			services.AddSingleton<IGraphCollection, GraphCollection>();
-			services.AddSingleton<IIssueCollection, IssueCollection>();
+			services.AddSingleton<IssueCollection>();
+			services.AddSingleton<IIssueCollection>(sp => sp.GetRequiredService<IssueCollection>());
+			services.AddSingleton<ILogExtIssueProvider>(sp => sp.GetRequiredService<IssueCollection>());
 			services.AddSingleton<IJobCollection, JobCollection>();
 			services.AddSingleton<IJobStepRefCollection, JobStepRefCollection>();
 			services.AddSingleton<IJobTimingCollection, JobTimingCollection>();
@@ -458,6 +460,7 @@ namespace HordeServer
 			services.AddSingleton<IDowntimeService, DowntimeService>();
 			services.AddSingleton<IssueService>();
 			services.AddSingleton<JobService>();
+			services.AddSingleton<ILogExtAuthProvider>(sp => sp.GetRequiredService<JobService>());
 			services.AddSingleton<LifetimeService>();
 			services.AddHostedService(provider => provider.GetRequiredService<LifetimeService>());
 			services.AddSingleton(typeof(IHealthMonitor<>), typeof(HealthMonitor<>));
