Horde: Rename V2 artifacts to just be artifacts, add a find by id to the artifact search
#rnx
#skipci

diff --git a/Engine/Source/Programs/Horde/HordeDashboard/src/backend/Api.ts b/Engine/Source/Programs/Horde/HordeDashboard/src/backend/Api.ts
--- a/Engine/Source/Programs/Horde/HordeDashboard/src/backend/Api.ts
+++ b/Engine/Source/Programs/Horde/HordeDashboard/src/backend/Api.ts
@@ -1657,7 +1657,7 @@ export type CreateZipRequest = {
 
 
 /** Describes an artifact */
-export type GetArtifactResponseV2 = {
+export type GetArtifactResponse = {
 	id: string;
 	type: ArtifactContextType;
 	keys: string[]
@@ -1671,7 +1671,7 @@ export type GetArtifactResponseV2 = {
 /** Result of an artifact search */
 export type FindArtifactsResponse = {
 	/** List of artifacts matching the search criteria*/
-	artifacts: GetArtifactResponseV2[];
+	artifacts: GetArtifactResponse[];
 }
 
 /** Describes a file within an artifact */
diff --git a/Engine/Source/Programs/Horde/HordeDashboard/src/backend/Backend.ts b/Engine/Source/Programs/Horde/HordeDashboard/src/backend/Backend.ts
--- a/Engine/Source/Programs/Horde/HordeDashboard/src/backend/Backend.ts
+++ b/Engine/Source/Programs/Horde/HordeDashboard/src/backend/Backend.ts
@@ -1,7 +1,7 @@
 // Copyright Epic Games, Inc. All Rights Reserved.
 
 import templateCache from '../backend/TemplateCache';
-import { AccountClaimMessage, AgentData, AgentQuery, ApproveAgentsRequest, AuditLogEntry, AuditLogQuery, BatchUpdatePoolRequest, ChangeSummaryData, CreateAccountRequest, CreateBisectTaskRequest, CreateBisectTaskResponse, CreateDeviceRequest, CreateDeviceResponse, CreateExternalIssueRequest, CreateExternalIssueResponse, CreateJobRequest, CreateJobResponse, CreateNoticeRequest, CreatePoolRequest, CreateServiceAccountRequest, CreateServiceAccountResponse, CreateSoftwareResponse, CreateSubscriptionRequest, CreateSubscriptionResponse, CreateZipRequest, DashboardPreference, DevicePoolTelemetryQuery, DeviceTelemetryQuery, EventData, FindArtifactsResponse, FindIssueResponse, FindJobTimingsResponse, GetAccountResponse, GetAgentLeaseResponse, GetAgentSessionResponse, GetAgentTelemetrySampleResponse, GetArtifactDirectoryResponse, GetArtifactZipRequest, GetBisectTaskResponse, GetDashboardConfigResponse, GetDevicePlatformResponse, GetDevicePoolResponse, GetDevicePoolTelemetryResponse, GetDeviceReservationResponse, GetDeviceResponse, GetDeviceTelemetryResponse, GetExternalIssueProjectResponse, GetExternalIssueResponse, GetIssueStreamResponse, GetJobsTabResponse, GetJobStepRefResponse, GetJobStepTraceResponse, GetJobTimingResponse, GetLogEventResponse, GetLogFileResponse, GetNoticeResponse, GetNotificationResponse, GetPendingAgentsResponse, GetPerforceServerStatusResponse, GetPoolResponse, GetPoolSummaryResponse, GetServerInfoResponse, GetServerSettingsResponse, GetServiceAccountResponse, GetSoftwareResponse, GetSubscriptionResponse, GetTelemetryMetricsResponse, GetTestDataDetailsResponse, GetTestDataRefResponse, GetTestMetaResponse, GetTestResponse, GetTestsRequest, GetTestStreamResponse, GetToolResponse, GetToolSummaryResponse, GetUserResponse, GetUtilizationTelemetryResponse, GlobalConfig, IssueData, IssueQuery, IssueQueryV2, JobData, JobQuery, JobsTabColumnType, JobStepOutcome, JobStreamQuery, JobTimingsQuery, LeaseData, LogData, LogEventQuery, LogLineData, MetricsQuery, PoolData, PoolQuery, PreflightConfigResponse, ProjectData, ScheduleData, ScheduleQuery, SearchLogFileResponse, ServerStatusResponse, ServerUpdateResponse, SessionData, StreamData, TabType, TestData, UpdateAccountRequest, UpdateAgentRequest, UpdateBisectTaskRequest, UpdateCurrentAccountRequest, UpdateDeploymentRequest, UpdateDeviceRequest, UpdateGlobalConfigRequest, UpdateIssueRequest, UpdateJobRequest, UpdateLeaseRequest, UpdateNoticeRequest, UpdateNotificationsRequest, UpdatePoolRequest, UpdateServerSettingsRequest, UpdateServiceAccountRequest, UpdateServiceAccountResponse, UpdateStepRequest, UpdateStepResponse, UpdateTemplateRefRequest, UpdateUserRequest, UsersQuery } from './Api';
+import { AccountClaimMessage, AgentData, AgentQuery, ApproveAgentsRequest, AuditLogEntry, AuditLogQuery, BatchUpdatePoolRequest, ChangeSummaryData, CreateAccountRequest, CreateBisectTaskRequest, CreateBisectTaskResponse, CreateDeviceRequest, CreateDeviceResponse, CreateExternalIssueRequest, CreateExternalIssueResponse, CreateJobRequest, CreateJobResponse, CreateNoticeRequest, CreatePoolRequest, CreateServiceAccountRequest, CreateServiceAccountResponse, CreateSoftwareResponse, CreateSubscriptionRequest, CreateSubscriptionResponse, CreateZipRequest, DashboardPreference, DevicePoolTelemetryQuery, DeviceTelemetryQuery, EventData, FindArtifactsResponse, FindIssueResponse, FindJobTimingsResponse, GetAccountResponse, GetAgentLeaseResponse, GetAgentSessionResponse, GetAgentTelemetrySampleResponse, GetArtifactDirectoryResponse, GetArtifactResponse, GetArtifactZipRequest, GetBisectTaskResponse, GetDashboardConfigResponse, GetDevicePlatformResponse, GetDevicePoolResponse, GetDevicePoolTelemetryResponse, GetDeviceReservationResponse, GetDeviceResponse, GetDeviceTelemetryResponse, GetExternalIssueProjectResponse, GetExternalIssueResponse, GetIssueStreamResponse, GetJobsTabResponse, GetJobStepRefResponse, GetJobStepTraceResponse, GetJobTimingResponse, GetLogEventResponse, GetLogFileResponse, GetNoticeResponse, GetNotificationResponse, GetPendingAgentsResponse, GetPerforceServerStatusResponse, GetPoolResponse, GetPoolSummaryResponse, GetServerInfoResponse, GetServerSettingsResponse, GetServiceAccountResponse, GetSoftwareResponse, GetSubscriptionResponse, GetTelemetryMetricsResponse, GetTestDataDetailsResponse, GetTestDataRefResponse, GetTestMetaResponse, GetTestResponse, GetTestsRequest, GetTestStreamResponse, GetToolResponse, GetToolSummaryResponse, GetUserResponse, GetUtilizationTelemetryResponse, GlobalConfig, IssueData, IssueQuery, IssueQueryV2, JobData, JobQuery, JobsTabColumnType, JobStepOutcome, JobStreamQuery, JobTimingsQuery, LeaseData, LogData, LogEventQuery, LogLineData, MetricsQuery, PoolData, PoolQuery, PreflightConfigResponse, ProjectData, ScheduleData, ScheduleQuery, SearchLogFileResponse, ServerStatusResponse, ServerUpdateResponse, SessionData, StreamData, TabType, TestData, UpdateAccountRequest, UpdateAgentRequest, UpdateBisectTaskRequest, UpdateCurrentAccountRequest, UpdateDeploymentRequest, UpdateDeviceRequest, UpdateGlobalConfigRequest, UpdateIssueRequest, UpdateJobRequest, UpdateLeaseRequest, UpdateNoticeRequest, UpdateNotificationsRequest, UpdatePoolRequest, UpdateServerSettingsRequest, UpdateServiceAccountRequest, UpdateServiceAccountResponse, UpdateStepRequest, UpdateStepResponse, UpdateTemplateRefRequest, UpdateUserRequest, UsersQuery } from './Api';
 import dashboard, { Dashboard } from './Dashboard';
 import { ChallengeStatus, Fetch } from './Fetch';
 import { projectStore } from './ProjectStore';
@@ -541,6 +541,16 @@ export class Backend {
         })
     }
 
+    getArtifactData(artifactId: string): Promise<GetArtifactResponse> {
+        const url = `/api/v2/artifacts/${artifactId}`;
+        return new Promise<GetArtifactResponse>((resolve, reject) => {
+            this.backend.get(url).then((value) => {
+                resolve(value.data as GetArtifactResponse);
+            }).catch(reason => {
+                reject(reason);
+            });
+        });
+    }
 
     getArtifactV2(artifactId: string, path: string): Promise<object> {
         const url = `/api/v2/artifacts/${artifactId}/file?path=${encodeURIComponent(path)}`;
diff --git a/Engine/Source/Programs/Horde/HordeDashboard/src/components/LogSource.tsx b/Engine/Source/Programs/Horde/HordeDashboard/src/components/LogSource.tsx
--- a/Engine/Source/Programs/Horde/HordeDashboard/src/components/LogSource.tsx
+++ b/Engine/Source/Programs/Horde/HordeDashboard/src/components/LogSource.tsx
@@ -3,7 +3,7 @@
 import { action, makeObservable, observable } from "mobx";
 import moment, { Moment } from 'moment-timezone';
 import backend from '../backend';
-import { AgentData, BatchData, EventSeverity, GetArtifactResponseV2, GetJobsTabResponse, GetLogEventResponse, GetLogFileResponse, IssueData, LeaseData, LogData, StepData, StreamData } from "../backend/Api";
+import { AgentData, BatchData, EventSeverity, GetArtifactResponse, GetJobsTabResponse, GetLogEventResponse, GetLogFileResponse, IssueData, LeaseData, LogData, StepData, StreamData } from "../backend/Api";
 import { getBatchSummaryMarkdown, getStepSummaryMarkdown, JobDetails } from "../backend/JobDetails";
 import { getLeaseElapsed, getStepPercent } from '../base/utilities/timeUtils';
 import { BreadcrumbItem } from './Breadcrumbs';
@@ -608,7 +608,7 @@ export class JobLogSource extends LogSource {
    batch?: BatchData;
    step?: StepData;
    
-   artifactsV2?: GetArtifactResponseV2[];
+   artifactsV2?: GetArtifactResponse[];
 
    _agentTelemety: AgentTelemetryHandler = new AgentTelemetryHandler();
 
diff --git a/Engine/Source/Programs/Horde/HordeDashboard/src/components/artifacts/ArtifactsModal.tsx b/Engine/Source/Programs/Horde/HordeDashboard/src/components/artifacts/ArtifactsModal.tsx
--- a/Engine/Source/Programs/Horde/HordeDashboard/src/components/artifacts/ArtifactsModal.tsx
+++ b/Engine/Source/Programs/Horde/HordeDashboard/src/components/artifacts/ArtifactsModal.tsx
@@ -6,7 +6,7 @@ import { observer } from "mobx-react-lite";
 import { useEffect, useState } from "react";
 import { NavigateFunction, useNavigate } from "react-router-dom";
 import backend from "../../backend";
-import { ArtifactContextType, GetArtifactDirectoryEntryResponse, GetArtifactDirectoryResponse, GetArtifactFileEntryResponse, GetArtifactResponseV2 } from "../../backend/Api";
+import { ArtifactContextType, GetArtifactDirectoryEntryResponse, GetArtifactDirectoryResponse, GetArtifactFileEntryResponse, GetArtifactResponse } from "../../backend/Api";
 import dashboard, { StatusColor } from "../../backend/Dashboard";
 import { getHordeStyling } from "../../styles/Styles";
 
@@ -29,7 +29,7 @@ type BrowserItem = {
 
 class ArtifactsHandler {
 
-   constructor(jobId: string, stepId: string, contextType: ArtifactContextType, artifactPath?: string, artifacts?: GetArtifactResponseV2[], artifactId?: string) {
+   constructor(jobId: string, stepId: string, contextType: ArtifactContextType, artifactPath?: string, artifacts?: GetArtifactResponse[], artifactId?: string) {
       makeObservable(this);
       this.jobId = jobId;
       this.stepId = stepId;
@@ -57,7 +57,7 @@ class ArtifactsHandler {
 
    private async set(artifactPath?: string) {
 
-      let artifacts: GetArtifactResponseV2[] | undefined = this.artifacts;
+      let artifacts: GetArtifactResponse[] | undefined = this.artifacts;
 
       if (!artifacts) {
 
@@ -90,7 +90,7 @@ class ArtifactsHandler {
          return;
       }
 
-      let a: GetArtifactResponseV2 | undefined;
+      let a: GetArtifactResponse | undefined;
 
       if (this.artifactId) {
          a = artifacts.find(a => a.id === this.artifactId)
@@ -262,8 +262,8 @@ class ArtifactsHandler {
 
    browse?: GetArtifactDirectoryResponse;
 
-   artifact?: GetArtifactResponseV2;
-   artifacts?: GetArtifactResponseV2[];
+   artifact?: GetArtifactResponse;
+   artifacts?: GetArtifactResponse[];
 
    artifactId?: string;
 
@@ -552,7 +552,7 @@ const DownloadButton: React.FC<{ handler: ArtifactsHandler, openArtifactInfo: ()
 
 let idcounter = 0;
 
-const JobDetailArtifactsInner: React.FC<{ jobId: string; stepId: string, artifacts?: GetArtifactResponseV2[], contextType: ArtifactContextType, artifactPath?: string, artifactId?: string }> = observer(({ jobId, stepId, artifacts, contextType, artifactPath, artifactId }) => {
+const JobDetailArtifactsInner: React.FC<{ jobId: string; stepId: string, artifacts?: GetArtifactResponse[], contextType: ArtifactContextType, artifactPath?: string, artifactId?: string }> = observer(({ jobId, stepId, artifacts, contextType, artifactPath, artifactId }) => {
 
    // eslint-disable-next-line
    const handler = ArtifactsHandler.current ?? new ArtifactsHandler(jobId, stepId, contextType, artifactPath, artifacts, artifactId);
@@ -806,7 +806,7 @@ const JobDetailArtifactsInner: React.FC<{ jobId: string; stepId: string, artifac
 
 
 
-export const JobArtifactsModal: React.FC<{ jobId: string; stepId: string, artifacts?: GetArtifactResponseV2[], contextType: ArtifactContextType, artifactPath?: string, artifactId?: string, onClose: () => void }> = ({ jobId, stepId, artifacts, contextType, artifactPath, artifactId, onClose }) => {
+export const JobArtifactsModal: React.FC<{ jobId: string; stepId: string, artifacts?: GetArtifactResponse[], contextType: ArtifactContextType, artifactPath?: string, artifactId?: string, onClose: () => void }> = ({ jobId, stepId, artifacts, contextType, artifactPath, artifactId, onClose }) => {
 
    const { hordeClasses } = getHordeStyling();
 
diff --git a/Engine/Source/Programs/Horde/HordeDashboard/src/components/artifacts/ArtifactsSearch.tsx b/Engine/Source/Programs/Horde/HordeDashboard/src/components/artifacts/ArtifactsSearch.tsx
--- a/Engine/Source/Programs/Horde/HordeDashboard/src/components/artifacts/ArtifactsSearch.tsx
+++ b/Engine/Source/Programs/Horde/HordeDashboard/src/components/artifacts/ArtifactsSearch.tsx
@@ -3,15 +3,15 @@ import { getHordeStyling } from "../../styles/Styles";
 import { IStreamChooser, StreamChooser } from "../projects/StreamChooser";
 import React, { useState } from "react";
 import backend from "../../backend";
-import { GetArtifactResponseV2 } from "../../backend/Api";
+import { GetArtifactResponse } from "../../backend/Api";
 import dashboard from "../../backend/Dashboard";
 import { JobArtifactsModal } from "./ArtifactsModal";
 import { projectStore } from "../../backend/ProjectStore";
 
 
-const ArtifactsList: React.FC<{ artifacts?: GetArtifactResponseV2[] }> = ({ artifacts }) => {
+const ArtifactsList: React.FC<{ artifacts?: GetArtifactResponse[] }> = ({ artifacts }) => {
 
-   const [browse, setBrowse] = useState<{ jobId?: string, stepId?: string, artifactId?: string, artifact?: GetArtifactResponseV2 }>({});
+   const [browse, setBrowse] = useState<{ jobId?: string, stepId?: string, artifactId?: string, artifact?: GetArtifactResponse }>({});
 
    if (!artifacts?.length) {
       return null;
@@ -36,7 +36,7 @@ const ArtifactsList: React.FC<{ artifacts?: GetArtifactResponseV2[] }> = ({ arti
 
       if (props) {
 
-         const item = props!.item as GetArtifactResponseV2;
+         const item = props!.item as GetArtifactResponse;
 
          let background: string | undefined;
          if (props.itemIndex % 2 === 0) {
@@ -121,7 +121,7 @@ const artifactTypes: IComboBoxOption[] = [
 
 export const FindArtifactsModal: React.FC<{ streamId?: string, onClose: () => void }> = ({ streamId, onClose }) => {
 
-   const [state, setState] = useState<{ searching?: boolean, artifacts?: GetArtifactResponseV2[] }>({});
+   const [state, setState] = useState<{ searching?: boolean, artifacts?: GetArtifactResponse[] }>({});
    const streamRef = React.useRef<IStreamChooser>(null);
    const minChangeRef = React.useRef<ITextField>(null);
    const maxChangeRef = React.useRef<ITextField>(null);
@@ -172,11 +172,37 @@ export const FindArtifactsModal: React.FC<{ streamId?: string, onClose: () => vo
          if (!streamId) {
             streamId = undefined;
          }
+         
+         const mongoId = /^[a-fA-F0-9]{24}$/i;
 
-         const artifacts = await backend.getArtifacts(streamId, minChange, maxChange, name, type);
+         let id: string | undefined;
+
+         if (name?.length) {
+            if (name.match(mongoId)?.length) {
+               id = name;
+            }
+         }
+
+         let artifacts: GetArtifactResponse[] = [];
+
+         if (id) {
+            try {
+               const artifact = await backend.getArtifactData(id);
+               if (artifact?.id) {
+                  artifacts = [artifact];
+               }
+            } catch (reason) {
+               console.error(reason)
+            }            
+
+         } else {
+            const find = await backend.getArtifacts(streamId, minChange, maxChange, name, type);
+            artifacts = find.artifacts;
+         }
+         
 
          setState({
-            searching: false, artifacts: artifacts.artifacts
+            searching: false, artifacts: artifacts
          });
 
       } catch (reason) {
@@ -230,7 +256,7 @@ export const FindArtifactsModal: React.FC<{ streamId?: string, onClose: () => vo
                         <TextField key="max_change_option" componentRef={maxChangeRef} style={{ width: 92 }} label="Max Changelist" />
                      </Stack>
                      <Stack >
-                        <TextField key="name_option" componentRef={nameRef} style={{ width: 220 }} label="Name" spellCheck={false} autoComplete="off" />
+                        <TextField key="name_option" componentRef={nameRef} style={{ width: 220 }} label="Name / Artifact Id" spellCheck={false} autoComplete="off"/>
                      </Stack>
                      <Stack>
                         <Label>Artifact Type</Label>
diff --git a/Engine/Source/Programs/Horde/HordeDashboard/src/components/jobDetailsV2/JobDetailsViewCommon.tsx b/Engine/Source/Programs/Horde/HordeDashboard/src/components/jobDetailsV2/JobDetailsViewCommon.tsx
--- a/Engine/Source/Programs/Horde/HordeDashboard/src/components/jobDetailsV2/JobDetailsViewCommon.tsx
+++ b/Engine/Source/Programs/Horde/HordeDashboard/src/components/jobDetailsV2/JobDetailsViewCommon.tsx
@@ -5,7 +5,7 @@ import moment from "moment";
 import React, { useState } from "react";
 import { useLocation, useNavigate } from "react-router-dom";
 import backend from "../../backend";
-import { BatchData, GetArtifactResponseV2, GetJobTimingResponse, GetLabelStateResponse, GetTemplateRefResponse, JobData, JobState, LabelState, ReportPlacement, StepData, StreamData } from "../../backend/Api";
+import { BatchData, GetArtifactResponse, GetJobTimingResponse, GetLabelStateResponse, GetTemplateRefResponse, JobData, JobState, LabelState, ReportPlacement, StepData, StreamData } from "../../backend/Api";
 import { JobLabel } from "../../backend/JobDetails";
 import { PollBase } from "../../backend/PollBase";
 import { projectStore } from "../../backend/ProjectStore";
@@ -652,7 +652,7 @@ export class JobDetailsV2 extends PollBase {
    template?: GetTemplateRefResponse;
 
    // stepId => artifacts
-   stepArtifacts = new Map<string, GetArtifactResponseV2[]>();
+   stepArtifacts = new Map<string, GetArtifactResponse[]>();
 
    @observable
    private rootUpdated: number = 0;
diff --git a/Engine/Source/Programs/Horde/HordeDashboard/src/components/jobDetailsV2/StepDetailView.tsx b/Engine/Source/Programs/Horde/HordeDashboard/src/components/jobDetailsV2/StepDetailView.tsx
--- a/Engine/Source/Programs/Horde/HordeDashboard/src/components/jobDetailsV2/StepDetailView.tsx
+++ b/Engine/Source/Programs/Horde/HordeDashboard/src/components/jobDetailsV2/StepDetailView.tsx
@@ -5,7 +5,7 @@ import { observer } from "mobx-react-lite";
 import { useEffect, useState } from "react";
 import { useNavigate } from "react-router-dom";
 import backend from "../../backend";
-import { GetArtifactResponseV2 } from "../../backend/Api";
+import { GetArtifactResponse } from "../../backend/Api";
 import { getHordeStyling } from "../../styles/Styles";
 import { HistoryModal } from "../HistoryModal";
 import { useQuery } from "../JobDetailCommon";
@@ -92,7 +92,7 @@ class StepDetailDataView extends JobDataView {
       super.clear();
    }
 
-   artifacts?: GetArtifactResponseV2[];
+   artifacts?: GetArtifactResponse[];
 
    stepId?: string;
 
