Horde: Fix static analysis warning about unused cancellation token.

#rnx

diff --git a/Engine/Source/Programs/Horde/HordeServer/Compute/ComputeController.cs b/Engine/Source/Programs/Horde/HordeServer/Compute/ComputeController.cs
--- a/Engine/Source/Programs/Horde/HordeServer/Compute/ComputeController.cs
+++ b/Engine/Source/Programs/Horde/HordeServer/Compute/ComputeController.cs
@@ -46,18 +46,17 @@ namespace HordeServer.Compute
 		/// Find the most suitable cluster given a compute assignment request
 		/// </summary>
 		/// <param name="request">The request parameters</param>
-		/// <param name="cancellationToken">Cancellation token for the operation</param>
 		/// <returns></returns>
 		[HttpPost]
 		[Authorize]
 		[Route("/api/v2/compute/_cluster")] // Underscore to avoid clashing with endpoint for clusters below
-		public Task<ActionResult<GetClusterResponse>> GetClusterAsync([FromBody] AssignComputeRequest request, CancellationToken cancellationToken)
+		public ActionResult<GetClusterResponse> GetCluster([FromBody] AssignComputeRequest request)
 		{
-			using TelemetrySpan span = _tracer.StartActiveSpan($"{nameof(ComputeController)}.{nameof(GetClusterAsync)}");
+			using TelemetrySpan span = _tracer.StartActiveSpan($"{nameof(ComputeController)}.{nameof(GetCluster)}");
 			IPAddress requesterIp = ComputeService.ResolveRequesterIp(HttpContext.Connection.RemoteIpAddress, request.Connection?.PreferPublicIp, request.Connection?.ClientPublicIp);
 			ClusterId clusterId = ComputeService.FindBestComputeClusterId(_globalConfig.Value, requesterIp);
 			GetClusterResponse response = new () { ClusterId = clusterId };
-			return Task.FromResult(new ActionResult<GetClusterResponse>(response));
+			return response;
 		}
 
 		/// <summary>
