Horde: Upgrade server config files by copying renamed keys into their new location.

#rnx

diff --git a/Engine/Source/Programs/Horde/HordeServer/ServerApp.cs b/Engine/Source/Programs/Horde/HordeServer/ServerApp.cs
--- a/Engine/Source/Programs/Horde/HordeServer/ServerApp.cs
+++ b/Engine/Source/Programs/Horde/HordeServer/ServerApp.cs
@@ -12,6 +12,7 @@ using EpicGames.Core;
 using HordeServer.Commands;
 using HordeServer.Utilities;
 using Microsoft.Extensions.Configuration;
+using Microsoft.Extensions.Configuration.Memory;
 using Microsoft.Extensions.DependencyInjection;
 using Microsoft.Extensions.Hosting;
 using Microsoft.Win32;
@@ -210,6 +211,20 @@ namespace HordeServer
 			return new FileReference(Assembly.GetExecutingAssembly().Location).Directory;
 		}
 
+		/// <summary>
+		/// List of config settings to rename to new locations
+		/// </summary>
+		static readonly KeyValuePair<string, string>[] s_renamedConfigKeys = new[]
+		{
+			// Storage plugin
+			KeyValuePair.Create("Horde:BundleCacheDir", "Horde:Plugins:Storage:BundleCacheDir"),
+			KeyValuePair.Create("Horde:BundleCacheSize", "Horde:Plugins:Storage:BundleCacheSize"),
+			KeyValuePair.Create("Horde:Backends", "Horde:Plugins:Storage:Backends"),
+	
+			// Tools plugin
+			KeyValuePair.Create("Horde:BundledTools", "Horde:Plugins:Tools:BundledTools")
+		};
+
 		/// <summary>
 		/// Constructs a configuration object for the current environment
 		/// </summary>
@@ -233,7 +248,23 @@ namespace HordeServer
 				builder = builder.AddJsonFile(serverConfigFile.FullName, optional: true, reloadOnChange: true);
 			}
 
-			return builder.AddEnvironmentVariables().Build();
+			builder = builder.AddEnvironmentVariables();
+
+			// Create a temporary configuration object and apply any upgrades and key renames
+			IConfiguration configuration = builder.Build();
+
+			List<KeyValuePair<string, string?>> remappedValues = new List<KeyValuePair<string, string?>>();
+			foreach ((string source, string target) in s_renamedConfigKeys)
+			{
+				string? value = configuration[source];
+				if (value != null && configuration[target] == null)
+				{
+					remappedValues.Add(new KeyValuePair<string, string?>(target, value));
+				}
+			}
+
+			builder.Add(new MemoryConfigurationSource { InitialData = remappedValues });
+			return builder.Build();
 		}
 
 		static async Task CopyDefaultConfigFilesAsync(DirectoryReference sourceDir, DirectoryReference targetDir, CancellationToken cancellationToken)
