Horde: Expose the global config version number to plugins through the PostLoad() method.

diff --git a/Engine/Source/Programs/Horde/HordeServer.Shared/ConfigVersion.cs b/Engine/Source/Programs/Horde/HordeServer.Shared/ConfigVersion.cs
new file mode 100644
--- /dev/null
+++ b/Engine/Source/Programs/Horde/HordeServer.Shared/ConfigVersion.cs
@@ -0,0 +1,40 @@
+// Copyright Epic Games, Inc. All Rights Reserved.
+
+namespace HordeServer
+{
+#pragma warning disable CA1027 // Mark enums with FlagsAttribute
+#pragma warning disable CA1069 // Enum member 'Latest' has same value as ...
+	/// <summary>
+	/// Global version number for running the server. As new features are introduced that require data migrations, this version number indicates the backwards compatibility functionality that must be enabled.
+	/// When adding a new version here, also add a message to <see cref="ConfigService.CreateSnapshotAsync"/> describing the steps that need to be taken to upgrade the deployment.
+	/// </summary>
+	public enum ConfigVersion
+	{
+		/// <summary>
+		/// Not specified
+		/// </summary>
+		None,
+
+		/// <summary>
+		/// Initial version number
+		/// </summary>
+		Initial,
+
+		/// <summary>
+		/// Ability to add/remove pools via the REST API is removed. Pools should be configured through globals.json instead.
+		/// </summary>
+		PoolsInConfigFiles,
+
+		/// <summary>
+		/// One after the l`ast defined version number
+		/// </summary>
+		LatestPlusOne,
+
+		/// <summary>
+		/// Latest version number
+		/// </summary>
+		Latest = LatestPlusOne - 1,
+	}
+#pragma warning restore CA1069
+#pragma warning restore CA1027
+}
diff --git a/Engine/Source/Programs/Horde/HordeServer.Shared/Plugins/IPluginConfig.cs b/Engine/Source/Programs/Horde/HordeServer.Shared/Plugins/IPluginConfig.cs
--- a/Engine/Source/Programs/Horde/HordeServer.Shared/Plugins/IPluginConfig.cs
+++ b/Engine/Source/Programs/Horde/HordeServer.Shared/Plugins/IPluginConfig.cs
@@ -19,17 +19,22 @@ namespace HordeServer.Plugins
 		/// <summary>
 		/// Called to fixup a plugin's configuration after deserialization
 		/// </summary>
-		/// <param name="parentAcl">The parent ACL scope</param>
-		void PostLoad(AclConfig parentAcl);
+		/// <param name="options">Options for configuring the plugin</param>
+		void PostLoad(PluginConfigOptions configOptions);
 	}
 
+	/// <summary>
+	/// Options passed to <see cref="IPluginConfig.PostLoad(AclConfig)"/>
+	/// </summary>
+	public record class PluginConfigOptions(ConfigVersion Version, AclConfig ParentAcl);
+
 	/// <summary>
 	/// Empty implementation of <see cref="IPluginConfig"/>
 	/// </summary>
 	public sealed class EmptyPluginConfig : IPluginConfig
 	{
 		/// <inheritdoc/>
-		public void PostLoad(AclConfig parentAcl)
+		public void PostLoad(PluginConfigOptions configOptions)
 		{ }
 	}
 }
diff --git a/Engine/Source/Programs/Horde/HordeServer/Agents/Pools/PoolCollection.cs b/Engine/Source/Programs/Horde/HordeServer/Agents/Pools/PoolCollection.cs
--- a/Engine/Source/Programs/Horde/HordeServer/Agents/Pools/PoolCollection.cs
+++ b/Engine/Source/Programs/Horde/HordeServer/Agents/Pools/PoolCollection.cs
@@ -284,7 +284,7 @@ namespace HordeServer.Agents.Pools
 
 			List<IPoolConfig> pools = new List<IPoolConfig>();
 			pools.AddRange(globalConfig.Pools);
-			if (globalConfig.VersionEnum < GlobalVersion.PoolsInConfigFiles)
+			if (globalConfig.VersionEnum < ConfigVersion.PoolsInConfigFiles)
 			{
 				pools.AddRange(await _poolsV1.Find(FilterDefinition<PoolDocumentV1>.Empty).ToListAsync(cancellationToken));
 			}
@@ -302,7 +302,7 @@ namespace HordeServer.Agents.Pools
 			{
 				poolConfig = globalPoolConfig;
 			}
-			else if (globalConfig.VersionEnum < GlobalVersion.PoolsInConfigFiles)
+			else if (globalConfig.VersionEnum < ConfigVersion.PoolsInConfigFiles)
 			{
 				poolConfig = await _poolsV1.Find(x => x.Id == id).FirstOrDefaultAsync(cancellationToken);
 			}
@@ -338,7 +338,7 @@ namespace HordeServer.Agents.Pools
 			List<PoolId> pools = new List<PoolId>();
 			pools.AddRange(globalConfig.Pools.Select(x => x.Id));
 
-			if (globalConfig.VersionEnum < GlobalVersion.PoolsInConfigFiles)
+			if (globalConfig.VersionEnum < ConfigVersion.PoolsInConfigFiles)
 			{
 				ProjectionDefinition<PoolDocumentV1, BsonDocument> projection = Builders<PoolDocumentV1>.Projection.Include(x => x.Id);
 				List<BsonDocument> results = await _poolsV1.Find(FilterDefinition<PoolDocumentV1>.Empty).Project(projection).ToListAsync(cancellationToken);
diff --git a/Engine/Source/Programs/Horde/HordeServer/Configuration/ConfigService.cs b/Engine/Source/Programs/Horde/HordeServer/Configuration/ConfigService.cs
--- a/Engine/Source/Programs/Horde/HordeServer/Configuration/ConfigService.cs
+++ b/Engine/Source/Programs/Horde/HordeServer/Configuration/ConfigService.cs
@@ -593,19 +593,19 @@ namespace HordeServer.Configuration
 				CopyJsonNode(configObj, "Tools", "Plugins.Tools.Tools");
 
 				GlobalConfig globalConfig = JsonSerializer.Deserialize<GlobalConfig>(configObj, _jsonOptions)!;
-				if (globalConfig.VersionEnum < GlobalVersion.Latest)
+				if (globalConfig.VersionEnum < ConfigVersion.Latest)
 				{
 					List<string> message = new List<string>();
 					message.Add($"Support for the following features will be removed in an upcoming release:");
 					message.Add("");
-					if (globalConfig.VersionEnum < GlobalVersion.PoolsInConfigFiles)
+					if (globalConfig.VersionEnum < ConfigVersion.PoolsInConfigFiles)
 					{
-						message.Add($"- v{(int)GlobalVersion.PoolsInConfigFiles}: Pools should now be configured through the globals.json file rather than REST API or database. The /api/v1/server/migrate/pool-config endpoint will transcribe your configured pools into JSON.");
+						message.Add($"- v{(int)ConfigVersion.PoolsInConfigFiles}: Pools should now be configured through the globals.json file rather than REST API or database. The /api/v1/server/migrate/pool-config endpoint will transcribe your configured pools into JSON.");
 					}
 					message.Add("");
-					message.Add($"Please migrate your installation and update the 'Version' property in globals.json to {(int)GlobalVersion.Latest}");
+					message.Add($"Please migrate your installation and update the 'Version' property in globals.json to {(int)ConfigVersion.Latest}");
 
-					_logger.LogWarning("Global config file is using old version number ({Version}<{LatestVersion})\n\n{DeprecatedFeaturesMessage}\n", globalConfig.Version, (int)GlobalVersion.Latest, String.Join("\n", message));
+					_logger.LogWarning("Global config file is using old version number ({Version}<{LatestVersion})\n\n{DeprecatedFeaturesMessage}\n", globalConfig.Version, (int)ConfigVersion.Latest, String.Join("\n", message));
 				}
 
 				// Serialize it back out to a byte array
diff --git a/Engine/Source/Programs/Horde/HordeServer/Server/GlobalConfig.cs b/Engine/Source/Programs/Horde/HordeServer/Server/GlobalConfig.cs
--- a/Engine/Source/Programs/Horde/HordeServer/Server/GlobalConfig.cs
+++ b/Engine/Source/Programs/Horde/HordeServer/Server/GlobalConfig.cs
@@ -46,42 +46,6 @@ namespace HordeServer.Server
 {
 	using JsonObject = System.Text.Json.Nodes.JsonObject;
 
-#pragma warning disable CA1027 // Mark enums with FlagsAttribute
-#pragma warning disable CA1069 // Enum member 'Latest' has same value as ...
-	/// <summary>
-	/// Global version number for running the server. As new features are introduced that require data migrations, this version number indicates the backwards compatibility functionality that must be enabled.
-	/// When adding a new version here, also add a message to <see cref="ConfigService.CreateSnapshotAsync"/> describing the steps that need to be taken to upgrade the deployment.
-	/// </summary>
-	public enum GlobalVersion
-	{
-		/// <summary>
-		/// Not specified
-		/// </summary>
-		None,
-
-		/// <summary>
-		/// Initial version number
-		/// </summary>
-		Initial,
-
-		/// <summary>
-		/// Ability to add/remove pools via the REST API is removed. Pools should be configured through globals.json instead.
-		/// </summary>
-		PoolsInConfigFiles,
-
-		/// <summary>
-		/// One after the last defined version number
-		/// </summary>
-		LatestPlusOne,
-
-		/// <summary>
-		/// Latest version number
-		/// </summary>
-		Latest = (int)LatestPlusOne - 1,
-	}
-#pragma warning restore CA1069
-#pragma warning restore CA1027
-
 	/// <summary>
 	/// Configuration for an artifact
 	/// </summary>
@@ -135,7 +99,7 @@ namespace HordeServer.Server
 		public string Revision { get; set; } = String.Empty;
 
 		/// <summary>
-		/// Version number for the server. Values are indicated by the <see cref="GlobalVersion"/>.
+		/// Version number for the server. Values are indicated by the <see cref="ConfigVersion"/>.
 		/// </summary>
 		public int Version { get; set; }
 
@@ -143,7 +107,7 @@ namespace HordeServer.Server
 		/// Version number for the server, as an enum.
 		/// </summary>
 		[JsonIgnore]
-		public GlobalVersion VersionEnum => (GlobalVersion)Version;
+		public ConfigVersion VersionEnum => (ConfigVersion)Version;
 
 		/// <summary>
 		/// Other paths to include
@@ -326,9 +290,10 @@ namespace HordeServer.Server
 				}
 			}
 
+			PluginConfigOptions pluginConfigOptions = new PluginConfigOptions(VersionEnum, Acl);
 			foreach (IPluginConfig pluginConfig in Plugins.Values)
 			{
-				pluginConfig.PostLoad(Acl);
+				pluginConfig.PostLoad(pluginConfigOptions);
 			}
 
 			_aclLookup.Clear();
diff --git a/Engine/Source/Programs/Horde/HordeServer/Users/UserController.cs b/Engine/Source/Programs/Horde/HordeServer/Users/UserController.cs
--- a/Engine/Source/Programs/Horde/HordeServer/Users/UserController.cs
+++ b/Engine/Source/Programs/Horde/HordeServer/Users/UserController.cs
@@ -79,7 +79,7 @@ namespace HordeServer.Users
 			response.ShowDeviceManager = globalConfig.Dashboard.ShowDeviceManager;
 			response.ShowTests = globalConfig.Dashboard.ShowTests;
 			response.ShowNoticeEditor = globalConfig.Authorize(NoticeAclAction.CreateNotice, principal) || globalConfig.Authorize(NoticeAclAction.UpdateNotice, principal);
-			response.ShowPoolEditor = globalConfig.VersionEnum < GlobalVersion.PoolsInConfigFiles && (globalConfig.Authorize(PoolAclAction.CreatePool, principal) || globalConfig.Authorize(PoolAclAction.UpdatePool, principal));
+			response.ShowPoolEditor = globalConfig.VersionEnum < ConfigVersion.PoolsInConfigFiles && (globalConfig.Authorize(PoolAclAction.CreatePool, principal) || globalConfig.Authorize(PoolAclAction.UpdatePool, principal));
 			response.ShowRemoteDesktop = globalConfig.Authorize(AgentAclAction.UpdateAgent, principal);
 			response.ShowAccounts = settings.AuthMethod == EpicGames.Horde.Server.AuthMethod.Horde && globalConfig.Authorize(AccountAclAction.UpdateAccount, principal);
 			return response;
diff --git a/Engine/Source/Programs/Horde/Plugins/Analytics/HordeServer.Analytics/AnalyticsConfig.cs b/Engine/Source/Programs/Horde/Plugins/Analytics/HordeServer.Analytics/AnalyticsConfig.cs
--- a/Engine/Source/Programs/Horde/Plugins/Analytics/HordeServer.Analytics/AnalyticsConfig.cs
+++ b/Engine/Source/Programs/Horde/Plugins/Analytics/HordeServer.Analytics/AnalyticsConfig.cs
@@ -30,15 +30,15 @@ namespace HordeServer
 		private AclConfig _parentAcl = null!;
 		private readonly Dictionary<TelemetryStoreId, TelemetryStoreConfig> _telemetryStoreLookup = new Dictionary<TelemetryStoreId, TelemetryStoreConfig>();
 
-		public void PostLoad(AclConfig parentAcl)
+		public void PostLoad(PluginConfigOptions configOptions)
 		{
-			_parentAcl = parentAcl;
+			_parentAcl = configOptions.ParentAcl;
 
 			_telemetryStoreLookup.Clear();
 			foreach (TelemetryStoreConfig telemetryStore in Stores)
 			{
 				_telemetryStoreLookup.Add(telemetryStore.Id, telemetryStore);
-				telemetryStore.PostLoad(parentAcl);
+				telemetryStore.PostLoad(configOptions.ParentAcl);
 			}
 		}
 
diff --git a/Engine/Source/Programs/Horde/Plugins/Secrets/HordeServer.Secrets/SecretsConfig.cs b/Engine/Source/Programs/Horde/Plugins/Secrets/HordeServer.Secrets/SecretsConfig.cs
--- a/Engine/Source/Programs/Horde/Plugins/Secrets/HordeServer.Secrets/SecretsConfig.cs
+++ b/Engine/Source/Programs/Horde/Plugins/Secrets/HordeServer.Secrets/SecretsConfig.cs
@@ -25,13 +25,13 @@ namespace HordeServer
 		private readonly Dictionary<SecretId, SecretConfig> _secretLookup = new Dictionary<SecretId, SecretConfig>();
 
 		/// <inheritdoc/>
-		public void PostLoad(AclConfig parentAcl)
+		public void PostLoad(PluginConfigOptions configOptions)
 		{
 			_secretLookup.Clear();
 			foreach (SecretConfig secret in Secrets)
 			{
 				_secretLookup.Add(secret.Id, secret);
-				secret.PostLoad(parentAcl);
+				secret.PostLoad(configOptions.ParentAcl);
 			}
 		}
 
diff --git a/Engine/Source/Programs/Horde/Plugins/Storage/HordeServer.Storage/StorageConfig.cs b/Engine/Source/Programs/Horde/Plugins/Storage/HordeServer.Storage/StorageConfig.cs
--- a/Engine/Source/Programs/Horde/Plugins/Storage/HordeServer.Storage/StorageConfig.cs
+++ b/Engine/Source/Programs/Horde/Plugins/Storage/HordeServer.Storage/StorageConfig.cs
@@ -87,7 +87,7 @@ namespace HordeServer.Storage
 		/// <summary>
 		/// Called after the config has been read from disk
 		/// </summary>
-		public void PostLoad(AclConfig parentAcl)
+		public void PostLoad(PluginConfigOptions configOptions)
 		{
 			// Create a lookup for backend configs, allowing later entries to override previous ones
 			_backendLookup.Clear();
@@ -127,7 +127,7 @@ namespace HordeServer.Storage
 					throw new StorageException($"Missing or invalid backend identifier for namespace {namespaceConfig.Id}");
 				}
 
-				namespaceConfig.PostLoad(parentAcl);
+				namespaceConfig.PostLoad(configOptions.ParentAcl);
 				_namespaceLookup[namespaceConfig.Id] = namespaceConfig;
 			}
 
diff --git a/Engine/Source/Programs/Horde/Plugins/Tools/HordeServer.Tools/ToolsConfig.cs b/Engine/Source/Programs/Horde/Plugins/Tools/HordeServer.Tools/ToolsConfig.cs
--- a/Engine/Source/Programs/Horde/Plugins/Tools/HordeServer.Tools/ToolsConfig.cs
+++ b/Engine/Source/Programs/Horde/Plugins/Tools/HordeServer.Tools/ToolsConfig.cs
@@ -22,13 +22,13 @@ namespace HordeServer.Tools
 		private readonly Dictionary<ToolId, ToolConfig> _toolLookup = new Dictionary<ToolId, ToolConfig>();
 
 		/// <inheritdoc/>
-		public void PostLoad(AclConfig parentAcl)
+		public void PostLoad(PluginConfigOptions configOptions)
 		{
 			_toolLookup.Clear();
 			foreach (ToolConfig tool in Tools)
 			{
 				_toolLookup.Add(tool.Id, tool);
-				tool.PostLoad(parentAcl);
+				tool.PostLoad(configOptions.ParentAcl);
 			}
 		}
 
