Horde: Move compute functionality into a plugin.

diff --git a/Engine/Source/Programs/Horde/Drivers/JobDriver/JobDriver.csproj b/Engine/Source/Programs/Horde/Drivers/JobDriver/JobDriver.csproj
--- a/Engine/Source/Programs/Horde/Drivers/JobDriver/JobDriver.csproj
+++ b/Engine/Source/Programs/Horde/Drivers/JobDriver/JobDriver.csproj
@@ -91,12 +91,12 @@
   </ItemGroup>
   
   <ItemGroup>
-    <Protobuf Include="..\..\HordeServer\Protos\horde\job_rpc.proto" ProtoRoot="..\..\HordeServer\Protos" GrpcServices="Client" Access="Public" AdditionalImportDirs="..\..\..\Shared\EpicGames.Horde\Protos"/>
-    <Protobuf Include="..\..\HordeServer\Protos\horde\job_rpc_messages.proto" ProtoRoot="..\..\HordeServer\Protos" Access="Public" />
+    <Protobuf Include="..\..\HordeServer\Protos\horde\job_rpc.proto" ProtoRoot="..\..\HordeServer\Protos" GrpcServices="Client" Access="Public" AdditionalImportDirs="..\..\..\Shared\EpicGames.Horde\Protos;..\..\Plugins\Compute\HordeServer.Compute\Protos"/>
+    <Protobuf Include="..\..\HordeServer\Protos\horde\job_rpc_messages.proto" ProtoRoot="..\..\HordeServer\Protos" Access="Public" AdditionalImportDirs="..\..\Plugins\Compute\HordeServer.Compute\Protos" />
     <Protobuf Include="..\..\HordeServer\Protos\horde\common\label_change.proto" ProtoRoot="..\..\HordeServer\Protos" Access="Public" />
-    <Protobuf Include="..\..\HordeServer\Protos\horde\messages\agent_workspace.proto" ProtoRoot="..\..\HordeServer\Protos" Access="Public" />
-    <Protobuf Include="..\..\HordeServer\Protos\horde\tasks\conform_task.proto" ProtoRoot="..\..\HordeServer\Protos" Access="Public" />
-    <Protobuf Include="..\..\HordeServer\Protos\horde\tasks\job_task.proto" ProtoRoot="..\..\HordeServer\Protos" Access="Public" />
+    <Protobuf Include="..\..\Plugins\Compute\HordeServer.Compute\Protos\horde\messages\agent_workspace.proto" ProtoRoot="..\..\Plugins\Compute\HordeServer.Compute\Protos" Access="Public" />
+    <Protobuf Include="..\..\Plugins\Compute\HordeServer.Compute\Protos\horde\tasks\conform_task.proto" ProtoRoot="..\..\Plugins\Compute\HordeServer.Compute\Protos" Access="Public" />
+    <Protobuf Include="..\..\Plugins\Compute\HordeServer.Compute\Protos\horde\tasks\job_task.proto" ProtoRoot="..\..\Plugins\Compute\HordeServer.Compute\Protos" Access="Public" />
   </ItemGroup>
 
   <ItemGroup>
diff --git a/Engine/Source/Programs/Horde/HordeAgent/HordeAgent.csproj b/Engine/Source/Programs/Horde/HordeAgent/HordeAgent.csproj
--- a/Engine/Source/Programs/Horde/HordeAgent/HordeAgent.csproj
+++ b/Engine/Source/Programs/Horde/HordeAgent/HordeAgent.csproj
@@ -137,30 +137,30 @@
   </ItemGroup>
 
   <ItemGroup>
-    <Protobuf Include="..\HordeServer\Protos\horde\horde_rpc.proto" ProtoRoot="..\HordeServer\Protos" GrpcServices="Client" Access="Public" />
-    <Protobuf Include="..\HordeServer\Protos\horde\horde_rpc_messages.proto" ProtoRoot="..\HordeServer\Protos" Access="Public" />
-    <Protobuf Include="..\HordeServer\Protos\horde\enrollment_rpc.proto" ProtoRoot="..\HordeServer\Protos" GrpcServices="Client" Access="Public" />
-    <Protobuf Include="..\HordeServer\Protos\horde\enrollment_rpc_messages.proto" ProtoRoot="..\HordeServer\Protos" Access="Public" />
-    <Protobuf Include="..\HordeServer\Protos\horde\common\agent_status.proto" ProtoRoot="..\HordeServer\Protos" Access="Public" />
+    <Protobuf Include="..\Plugins\Compute\HordeServer.Compute\Protos\horde\horde_rpc.proto" ProtoRoot="..\Plugins\Compute\HordeServer.Compute\Protos" GrpcServices="Client" Access="Public" AdditionalImportDirs="..\Plugins\Compute\HordeServer.Compute\Protos" />
+    <Protobuf Include="..\Plugins\Compute\HordeServer.Compute\Protos\horde\horde_rpc_messages.proto" ProtoRoot="..\Plugins\Compute\HordeServer.Compute\Protos" Access="Public" AdditionalImportDirs="..\Plugins\Compute\HordeServer.Compute\Protos"/>
+    <Protobuf Include="..\Plugins\Compute\HordeServer.Compute\Protos\horde\enrollment_rpc.proto" ProtoRoot="..\Plugins\Compute\HordeServer.Compute\Protos" GrpcServices="Client" Access="Public" AdditionalImportDirs="..\Plugins\Compute\HordeServer.Compute\Protos"/>
+    <Protobuf Include="..\Plugins\Compute\HordeServer.Compute\Protos\horde\enrollment_rpc_messages.proto" ProtoRoot="..\Plugins\Compute\HordeServer.Compute\Protos" Access="Public" AdditionalImportDirs="..\Plugins\Compute\HordeServer.Compute\Protos"/>
+    <Protobuf Include="..\Plugins\Compute\HordeServer.Compute\Protos\horde\common\agent_status.proto" ProtoRoot="..\Plugins\Compute\HordeServer.Compute\Protos" Access="Public" AdditionalImportDirs="..\Plugins\Compute\HordeServer.Compute\Protos"/>
     <Protobuf Include="..\HordeServer\Protos\horde\common\io_hash.proto" ProtoRoot="..\HordeServer\Protos" Access="Public" />
-    <Protobuf Include="..\HordeServer\Protos\horde\common\lease_outcome.proto" ProtoRoot="..\HordeServer\Protos" Access="Public" />
-    <Protobuf Include="..\HordeServer\Protos\horde\common\lease_state.proto" ProtoRoot="..\HordeServer\Protos" Access="Public" />
+    <Protobuf Include="..\Plugins\Compute\HordeServer.Compute\Protos\horde\common\lease_outcome.proto" ProtoRoot="..\Plugins\Compute\HordeServer.Compute\Protos" Access="Public" AdditionalImportDirs="..\Plugins\Compute\HordeServer.Compute\Protos"/>
+    <Protobuf Include="..\Plugins\Compute\HordeServer.Compute\Protos\horde\common\lease_state.proto" ProtoRoot="..\Plugins\Compute\HordeServer.Compute\Protos" Access="Public" AdditionalImportDirs="..\Plugins\Compute\HordeServer.Compute\Protos"/>
     <Protobuf Include="..\HordeServer\Protos\horde\common\namespace_id.proto" ProtoRoot="..\HordeServer\Protos" Access="Public" />
-    <Protobuf Include="..\HordeServer\Protos\horde\messages\agent_capabilities.proto" ProtoRoot="..\HordeServer\Protos" Access="Public" />
-    <Protobuf Include="..\HordeServer\Protos\horde\messages\agent_requirements.proto" ProtoRoot="..\HordeServer\Protos" Access="Public" />
-    <Protobuf Include="..\HordeServer\Protos\horde\messages\lease.proto" ProtoRoot="..\HordeServer\Protos" Access="Public" />
-    <Protobuf Include="..\HordeServer\Protos\horde\messages\aws_lambda.proto" ProtoRoot="..\HordeServer\Protos" Access="Public" />
-    <Protobuf Include="..\HordeServer\Protos\horde\tasks\compute_task.proto" ProtoRoot="..\HordeServer\Protos" Access="Public" />
-    <Protobuf Include="..\HordeServer\Protos\horde\tasks\restart_task.proto" ProtoRoot="..\HordeServer\Protos" Access="Public" />
-    <Protobuf Include="..\HordeServer\Protos\horde\tasks\shutdown_task.proto" ProtoRoot="..\HordeServer\Protos" Access="Public" />
-    <Protobuf Include="..\HordeServer\Protos\horde\tasks\test_task.proto" ProtoRoot="..\HordeServer\Protos" Access="Public" />
-    <Protobuf Include="..\HordeServer\Protos\horde\tasks\upgrade_task.proto" ProtoRoot="..\HordeServer\Protos" Access="Public" />
-
-    <Protobuf Include="..\HordeServer\Protos\horde\job_rpc_messages.proto" ProtoRoot="..\HordeServer\Protos" Access="Public" />
-    <Protobuf Include="..\HordeServer\Protos\horde\common\label_change.proto" ProtoRoot="..\HordeServer\Protos" Access="Public" />
-    <Protobuf Include="..\HordeServer\Protos\horde\messages\agent_workspace.proto" ProtoRoot="..\HordeServer\Protos" Access="Public" />
-    <Protobuf Include="..\HordeServer\Protos\horde\tasks\conform_task.proto" ProtoRoot="..\HordeServer\Protos" Access="Public" />
-    <Protobuf Include="..\HordeServer\Protos\horde\tasks\job_task.proto" ProtoRoot="..\HordeServer\Protos" Access="Public" />
+    <Protobuf Include="..\Plugins\Compute\HordeServer.Compute\Protos\horde\messages\agent_capabilities.proto" ProtoRoot="..\Plugins\Compute\HordeServer.Compute\Protos" Access="Public" AdditionalImportDirs="..\Plugins\Compute\HordeServer.Compute\Protos"/>
+    <Protobuf Include="..\Plugins\Compute\HordeServer.Compute\Protos\horde\messages\agent_requirements.proto" ProtoRoot="..\Plugins\Compute\HordeServer.Compute\Protos" Access="Public" AdditionalImportDirs="..\Plugins\Compute\HordeServer.Compute\Protos"/>
+    <Protobuf Include="..\Plugins\Compute\HordeServer.Compute\Protos\horde\messages\lease.proto" ProtoRoot="..\Plugins\Compute\HordeServer.Compute\Protos" Access="Public" AdditionalImportDirs="..\Plugins\Compute\HordeServer.Compute\Protos"/>
+    <Protobuf Include="..\Plugins\Compute\HordeServer.Compute\Protos\horde\messages\aws_lambda.proto" ProtoRoot="..\Plugins\Compute\HordeServer.Compute\Protos" Access="Public" AdditionalImportDirs="..\Plugins\Compute\HordeServer.Compute\Protos"/>
+    <Protobuf Include="..\Plugins\Compute\HordeServer.Compute\Protos\horde\tasks\compute_task.proto" ProtoRoot="..\Plugins\Compute\HordeServer.Compute\Protos" Access="Public" AdditionalImportDirs="..\Plugins\Compute\HordeServer.Compute\Protos"/>
+    <Protobuf Include="..\Plugins\Compute\HordeServer.Compute\Protos\horde\tasks\restart_task.proto" ProtoRoot="..\Plugins\Compute\HordeServer.Compute\Protos" Access="Public" AdditionalImportDirs="..\Plugins\Compute\HordeServer.Compute\Protos"/>
+    <Protobuf Include="..\Plugins\Compute\HordeServer.Compute\Protos\horde\tasks\shutdown_task.proto" ProtoRoot="..\Plugins\Compute\HordeServer.Compute\Protos" Access="Public" AdditionalImportDirs="..\Plugins\Compute\HordeServer.Compute\Protos"/>
+    <Protobuf Include="..\Plugins\Compute\HordeServer.Compute\Protos\horde\tasks\test_task.proto" ProtoRoot="..\Plugins\Compute\HordeServer.Compute\Protos" Access="Public" AdditionalImportDirs="..\Plugins\Compute\HordeServer.Compute\Protos"/>
+    <Protobuf Include="..\Plugins\Compute\HordeServer.Compute\Protos\horde\tasks\upgrade_task.proto" ProtoRoot="..\Plugins\Compute\HordeServer.Compute\Protos" Access="Public" AdditionalImportDirs="..\Plugins\Compute\HordeServer.Compute\Protos"/>
+
+    <Protobuf Include="..\HordeServer\Protos\horde\job_rpc_messages.proto" ProtoRoot="..\HordeServer\Protos" Access="Public" AdditionalImportDirs="..\Plugins\Compute\HordeServer.Compute\Protos" />
+    <Protobuf Include="..\HordeServer\Protos\horde\common\label_change.proto" ProtoRoot="..\HordeServer\Protos" Access="Public" AdditionalImportDirs="..\Plugins\Compute\HordeServer.Compute\Protos" />
+    <Protobuf Include="..\Plugins\Compute\HordeServer.Compute\Protos\horde\messages\agent_workspace.proto" ProtoRoot="..\Plugins\Compute\HordeServer.Compute\Protos" Access="Public" />
+    <Protobuf Include="..\Plugins\Compute\HordeServer.Compute\Protos\horde\tasks\conform_task.proto" ProtoRoot="..\Plugins\Compute\HordeServer.Compute\Protos" Access="Public" />
+    <Protobuf Include="..\Plugins\Compute\HordeServer.Compute\Protos\horde\tasks\job_task.proto" ProtoRoot="..\Plugins\Compute\HordeServer.Compute\Protos" Access="Public" />
   </ItemGroup>
 
   <ItemGroup>
diff --git a/Engine/Source/Programs/Horde/HordeServer.Shared/IServerInfo.cs b/Engine/Source/Programs/Horde/HordeServer.Shared/IServerInfo.cs
--- a/Engine/Source/Programs/Horde/HordeServer.Shared/IServerInfo.cs
+++ b/Engine/Source/Programs/Horde/HordeServer.Shared/IServerInfo.cs
@@ -63,6 +63,16 @@ namespace HordeServer
 		/// </summary>
 		DirectoryReference DataDir { get; }
 
+		/// <summary>
+		/// Global configuration settings
+		/// </summary>
+		IConfiguration Configuration { get; }
+
+		/// <summary>
+		/// Whether the server is running in read-only mode
+		/// </summary>
+		bool ReadOnlyMode { get; }
+
 		/// <summary>
 		/// Helper method to check if this process has activated the given mode
 		/// </summary>
diff --git a/Engine/Source/Programs/Horde/HordeServer.Tests.Shared/ServerServiceTest.cs b/Engine/Source/Programs/Horde/HordeServer.Tests.Shared/ServerServiceTest.cs
--- a/Engine/Source/Programs/Horde/HordeServer.Tests.Shared/ServerServiceTest.cs
+++ b/Engine/Source/Programs/Horde/HordeServer.Tests.Shared/ServerServiceTest.cs
@@ -82,8 +82,6 @@ namespace HordeServer.Tests
 			{
 			}
 
-			settings.WithAws = true;
-
 			settings.ForceConfigUpdateOnStartup = true;
 		}
 
@@ -91,7 +89,15 @@ namespace HordeServer.Tests
 		{
 			base.ConfigureServices(services);
 
-			services.AddSingleton<IServerInfo, ServerInfo>();
+			IConfiguration configuration = new ConfigurationBuilder().Build();
+
+			ServerSettings settings = new ServerSettings();
+			ConfigureSettings(settings);
+
+			services.Configure<StaticComputeConfig>(x => x.WithAws = true);
+
+			ServerInfo serverInfo = new ServerInfo(configuration, Options.Create(settings));
+			services.AddSingleton<IServerInfo>(serverInfo);
 			services.AddSingleton<IPluginCollection>(_pluginCollection);
 
 			services.AddLogging(builder => { builder.AddConsole().SetMinimumLevel(LogLevel.Debug); });
@@ -114,6 +120,7 @@ namespace HordeServer.Tests
 			services.AddSingleton<IHostEnvironment, WebHostEnvironmentStub>();
 
 			services.AddSingleton<AclService>();
+			services.AddSingleton<IAclService>(sp => sp.GetRequiredService<AclService>());
 			services.AddSingleton<IDowntimeService, DowntimeServiceStub>();
 			services.AddSingleton<LifetimeService>();
 
@@ -122,7 +129,7 @@ namespace HordeServer.Tests
 
 			foreach (ILoadedPlugin plugin in _pluginCollection.LoadedPlugins)
 			{
-				plugin.ConfigureServices(new ConfigurationBuilder().Build(), services);
+				plugin.ConfigureServices(configuration, serverInfo, services);
 			}
 		}
 
diff --git a/Engine/Source/Programs/Horde/HordeServer.Tests/Agents/AgentServiceTest.cs b/Engine/Source/Programs/Horde/HordeServer.Tests/Agents/AgentServiceTest.cs
--- a/Engine/Source/Programs/Horde/HordeServer.Tests/Agents/AgentServiceTest.cs
+++ b/Engine/Source/Programs/Horde/HordeServer.Tests/Agents/AgentServiceTest.cs
@@ -10,12 +10,11 @@ using EpicGames.Core;
 using EpicGames.Horde.Agents;
 using EpicGames.Horde.Agents.Pools;
 using EpicGames.Horde.Jobs;
+using HordeCommon.Rpc.Messages;
 using HordeServer.Agents;
 using HordeServer.Auditing;
 using HordeServer.Jobs;
-using HordeServer.Server;
 using HordeServer.Utilities;
-using HordeCommon.Rpc.Messages;
 using Microsoft.AspNetCore.Mvc;
 using Microsoft.VisualStudio.TestTools.UnitTesting;
 
diff --git a/Engine/Source/Programs/Horde/HordeServer.Tests/Compute/ComputeServiceTest.cs b/Engine/Source/Programs/Horde/HordeServer.Tests/Compute/ComputeServiceTest.cs
--- a/Engine/Source/Programs/Horde/HordeServer.Tests/Compute/ComputeServiceTest.cs
+++ b/Engine/Source/Programs/Horde/HordeServer.Tests/Compute/ComputeServiceTest.cs
@@ -9,11 +9,13 @@ using System.Threading;
 using System.Threading.Tasks;
 using EpicGames.Horde.Agents.Pools;
 using EpicGames.Horde.Compute;
+using HordeServer.Acls;
 using HordeServer.Agents;
 using HordeServer.Compute;
 using HordeServer.Plugins;
-using HordeServer.Server;
+using Microsoft.Extensions.DependencyInjection;
 using Microsoft.Extensions.Logging.Abstractions;
+using Microsoft.Extensions.Options;
 using Microsoft.VisualStudio.TestTools.UnitTesting;
 using OpenTelemetry.Trace;
 
@@ -180,7 +182,8 @@ namespace HordeServer.Tests.Compute
 		[TestMethod]
 		public async Task PoolNameTemplatingAsync()
 		{
-			GlobalConfig.CurrentValue.Networks = new List<NetworkConfig>
+			IOptionsMonitor<ComputeConfig> computeConfig = ServiceProvider.GetRequiredService<IOptionsMonitor<ComputeConfig>>();
+			computeConfig.CurrentValue.Networks = new List<NetworkConfig>
 			{
 				new() { CidrBlock = "12.0.0.0/16", Id = "myNetworkId", ComputeId = "myComputeId" }
 			};
@@ -219,14 +222,14 @@ namespace HordeServer.Tests.Compute
 			}
 
 			IAgent agent = await CreateAgentAsync(new PoolId("foo"), properties: ["ComputeIp=11.0.0.1", "ComputePort=5000"]);
-			
+
 			ComputeResource? resource1 = await ComputeService.TryAllocateResourceAsync(CreateParams("11.0.0.1"), CancellationToken.None);
 			Assert.IsNull(resource1);
-			
+
 			ComputeResource? resource2 = await ComputeService.TryAllocateResourceAsync(CreateParams("11.0.0.2"), CancellationToken.None);
 			Assert.AreEqual(agent.Id, resource2!.AgentId);
 		}
-		
+
 		[TestMethod]
 		public async Task Assignment_PoolInRequirementsAsync()
 		{
@@ -237,25 +240,26 @@ namespace HordeServer.Tests.Compute
 
 			IAgent agent1 = await CreateAgentAsync(new PoolId("foo"), properties: ["ComputeIp=11.0.0.1", "ComputePort=5000"]);
 			IAgent agent2 = await CreateAgentAsync(new PoolId("bar"), properties: ["ComputeIp=11.0.0.2", "ComputePort=5000"]);
-			
+
 			ComputeResource? resource1 = await ComputeService.TryAllocateResourceAsync(CreateParams("foo"), CancellationToken.None);
 			Assert.AreEqual(agent1.Id, resource1!.AgentId);
-			
+
 			ComputeResource? resource2 = await ComputeService.TryAllocateResourceAsync(CreateParams("bar"), CancellationToken.None);
 			Assert.AreEqual(agent2.Id, resource2!.AgentId);
-			
+
 			Assert.IsNull(await ComputeService.TryAllocateResourceAsync(CreateParams("does-not-exist"), CancellationToken.None));
 		}
-		
+
 		[TestMethod]
 		public async Task Assignment_ComputeClusterMemberConditionAsync()
 		{
 			ClusterId cluster1 = new("cluster1");
 			ClusterId cluster2 = new("cluster2");
-			PoolId poolFoo = new ("foo");
-			PoolId poolBar = new ("bar");
+			PoolId poolFoo = new("foo");
+			PoolId poolBar = new("bar");
 
-			GlobalConfig.CurrentValue.Compute = new List<ComputeClusterConfig>
+			IOptionsMonitor<ComputeConfig> computeConfig = ServiceProvider.GetRequiredService<IOptionsMonitor<ComputeConfig>>();
+			computeConfig.CurrentValue.Clusters = new List<ComputeClusterConfig>
 			{
 				new() { Id = cluster1, Condition = $"pool == '{poolFoo.ToString()}'" },
 				new() { Id = cluster2, Condition = $"pool == '{poolBar.ToString()}'" }
@@ -265,24 +269,24 @@ namespace HordeServer.Tests.Compute
 			List<string> props = ["ComputeIp=11.0.0.1", "ComputePort=5000"];
 			IAgent agent1 = await CreateAgentAsync(new PoolId("foo"), properties: props);
 			IAgent agent2 = await CreateAgentAsync(new PoolId("bar"), properties: props);
-			
+
 			{
 				AllocateResourceParams arpFoo = new(cluster1, ComputeProtocol.Latest, new Requirements { Pool = poolFoo.ToString() });
 				Assert.AreEqual(agent1.Id, (await ComputeService.TryAllocateResourceAsync(arpFoo, CancellationToken.None))!.AgentId);
-				
+
 				AllocateResourceParams arpBar = new(cluster1, ComputeProtocol.Latest, new Requirements { Pool = poolBar.ToString() });
 				Assert.IsNull(await ComputeService.TryAllocateResourceAsync(arpBar, CancellationToken.None));
 			}
-			
+
 			{
 				AllocateResourceParams arpFoo = new(cluster2, ComputeProtocol.Latest, new Requirements { Pool = poolFoo.ToString() });
 				Assert.IsNull(await ComputeService.TryAllocateResourceAsync(arpFoo, CancellationToken.None));
-				
+
 				AllocateResourceParams arpBar = new(cluster2, ComputeProtocol.Latest, new Requirements { Pool = poolBar.ToString() });
 				Assert.AreEqual(agent2.Id, (await ComputeService.TryAllocateResourceAsync(arpBar, CancellationToken.None))!.AgentId);
 			}
 		}
-		
+
 		[TestMethod]
 		[DataRow("11.0.0.1", "11.0.0.1", null, false)]
 		[DataRow("11.0.0.1", "11.0.0.1", "200.20.20.20", false)]
@@ -293,24 +297,30 @@ namespace HordeServer.Tests.Compute
 		{
 			Assert.AreEqual(IPAddress.Parse(expectIp), ComputeService.ResolveRequesterIp(internalIp == null ? null : IPAddress.Parse(internalIp), usePublicIp, publicIp));
 		}
-		
+
 		[TestMethod]
 		public void Cluster_FindBest()
 		{
-			GlobalConfig globalConfig = new () { Networks = [
+			ComputeConfig computeConfig = new()
+			{
+				Networks = [
 				new NetworkConfig { CidrBlock = "11.0.0.0/16", Id = "network1", ComputeId = "compute1" },
 				new NetworkConfig { CidrBlock = "12.0.0.0/16", Id = "network2", ComputeId = "compute2" },
-			] };
-			globalConfig.PostLoad(new ServerSettings(), new List<ILoadedPlugin>());
-			Assert.AreEqual(new ClusterId("compute1"), ComputeService.FindBestComputeClusterId(globalConfig, IPAddress.Parse("11.0.0.1")));
-			Assert.AreEqual(new ClusterId("compute2"), ComputeService.FindBestComputeClusterId(globalConfig, IPAddress.Parse("12.0.1.1")));
-			Assert.ThrowsException<ComputeServiceException>(() => ComputeService.FindBestComputeClusterId(globalConfig, IPAddress.Parse("123.123.123.123")));
-			
-			GlobalConfig globalConfigCatchAll = new ()
+			]
+			};
+
+			PluginConfigOptions configOptions = new PluginConfigOptions(ConfigVersion.Latest, new HordeServer.Acls.AclConfig());
+			computeConfig.PostLoad(configOptions);
+
+			Assert.AreEqual(new ClusterId("compute1"), ComputeService.FindBestComputeClusterId(computeConfig, IPAddress.Parse("11.0.0.1")));
+			Assert.AreEqual(new ClusterId("compute2"), ComputeService.FindBestComputeClusterId(computeConfig, IPAddress.Parse("12.0.1.1")));
+			Assert.ThrowsException<ComputeServiceException>(() => ComputeService.FindBestComputeClusterId(computeConfig, IPAddress.Parse("123.123.123.123")));
+
+			ComputeConfig globalConfigCatchAll = new()
 			{
 				Networks = [new NetworkConfig { CidrBlock = "0.0.0.0/0", Id = "catchAll", ComputeId = "catchAll" }]
 			};
-			globalConfigCatchAll.PostLoad(new ServerSettings(), new List<ILoadedPlugin>());
+			globalConfigCatchAll.PostLoad(configOptions);
 			Assert.AreEqual(new ClusterId("catchAll"), ComputeService.FindBestComputeClusterId(globalConfigCatchAll, IPAddress.Parse("123.123.123.123")));
 		}
 
@@ -374,13 +384,14 @@ namespace HordeServer.Tests.Compute
 
 		private async Task<ComputeService> CreateComputeServiceAsync(string? tunnelAddress, string pool, ComputeClusterConfig ccc)
 		{
-			ServerSettings ss = new() { ComputeTunnelAddress = tunnelAddress };
+			IOptionsMonitor<ComputeConfig> computeConfig = ServiceProvider.GetRequiredService<IOptionsMonitor<ComputeConfig>>();
+			StaticComputeConfig ss = new() { ComputeTunnelAddress = tunnelAddress };
 			ComputeService cs = new(AgentCollection, LogCollection, AgentService, AgentRelayService, GetRedisServiceSingleton(),
-				new TestOptionsMonitor<ServerSettings>(ss), GlobalConfig, Clock, Tracer, Meter,
+				new TestOptionsMonitor<StaticComputeConfig>(ss), computeConfig, Clock, Tracer, Meter,
 				NullLogger<ComputeService>.Instance);
 			await CreateAgentAsync(new PoolId(pool), properties: ["ComputeIp=11.0.0.1", "ComputePort=5000"]);
-			GlobalConfig.CurrentValue.Compute = [ccc];
-			GlobalConfig.CurrentValue.PostLoad(ss, new List<ILoadedPlugin>());
+			computeConfig.CurrentValue.Clusters = [ccc];
+			computeConfig.CurrentValue.PostLoad(new PluginConfigOptions(ConfigVersion.Latest, new AclConfig()));
 			return cs;
 		}
 
@@ -392,8 +403,8 @@ namespace HordeServer.Tests.Compute
 			string? tunnelAddress = null,
 			string pool = "myPool")
 		{
-			await using ComputeService cs = await CreateComputeServiceAsync(tunnelAddress, pool, new ComputeClusterConfig { Id = _cluster1});
-			AllocateResourceParams arp = new(_cluster1, ComputeProtocol.Latest, new Requirements() { Pool = pool})
+			await using ComputeService cs = await CreateComputeServiceAsync(tunnelAddress, pool, new ComputeClusterConfig { Id = _cluster1 });
+			AllocateResourceParams arp = new(_cluster1, ComputeProtocol.Latest, new Requirements() { Pool = pool })
 			{
 				ConnectionMode = connectionMode,
 				Ports = ports ?? new Dictionary<string, int>(),
diff --git a/Engine/Source/Programs/Horde/HordeServer.Tests/Compute/TunnelServiceTests.cs b/Engine/Source/Programs/Horde/HordeServer.Tests/Compute/TunnelServiceTests.cs
--- a/Engine/Source/Programs/Horde/HordeServer.Tests/Compute/TunnelServiceTests.cs
+++ b/Engine/Source/Programs/Horde/HordeServer.Tests/Compute/TunnelServiceTests.cs
@@ -53,8 +53,8 @@ namespace HordeServer.Tests.Compute
 
 			_echoServerTask = StartEchoServerAsync(_echoPort, echoServerBufferSize, cts.Token);
 
-			ServerSettings serverSettings = new() { ComputeTunnelPort = _tunnelPort };
-			using TunnelService tunnelService = new(new OptionsWrapper<ServerSettings>(serverSettings), TestSetup.CreateConsoleLogger<TunnelService>());
+			StaticComputeConfig serverSettings = new() { ComputeTunnelPort = _tunnelPort };
+			using TunnelService tunnelService = new(new OptionsWrapper<StaticComputeConfig>(serverSettings), TestSetup.CreateConsoleLogger<TunnelService>());
 			tunnelService.Start(IPAddress.Any);
 			_tunnelServerTask = tunnelService.ServerTask!;
 
diff --git a/Engine/Source/Programs/Horde/HordeServer.Tests/Configuration/ConfigTests.cs b/Engine/Source/Programs/Horde/HordeServer.Tests/Configuration/ConfigTests.cs
--- a/Engine/Source/Programs/Horde/HordeServer.Tests/Configuration/ConfigTests.cs
+++ b/Engine/Source/Programs/Horde/HordeServer.Tests/Configuration/ConfigTests.cs
@@ -125,7 +125,7 @@ namespace HordeServer.Tests.Configuration
 			Assert.IsTrue(!ReferenceEquals(file1, file3));
 		}
 
-		private static NetworkConfig? GetNetworkConfig(GlobalConfig gc, string ip)
+		private static NetworkConfig? GetNetworkConfig(ComputeConfig gc, string ip)
 		{
 			bool result = gc.TryGetNetworkConfig(IPAddress.Parse(ip), out NetworkConfig? networkConfig);
 			return result ? networkConfig : null;
@@ -134,7 +134,7 @@ namespace HordeServer.Tests.Configuration
 		[TestMethod]
 		public void NetworkConfig()
 		{
-			GlobalConfig gc = new()
+			ComputeConfig gc = new()
 			{
 				Networks = new List<NetworkConfig>()
 				{
diff --git a/Engine/Source/Programs/Horde/HordeServer.Tests/Fleet/AwsAutoScalingLifecycleTest.cs b/Engine/Source/Programs/Horde/HordeServer.Tests/Fleet/AwsAutoScalingLifecycleTest.cs
--- a/Engine/Source/Programs/Horde/HordeServer.Tests/Fleet/AwsAutoScalingLifecycleTest.cs
+++ b/Engine/Source/Programs/Horde/HordeServer.Tests/Fleet/AwsAutoScalingLifecycleTest.cs
@@ -62,8 +62,8 @@ public class AwsAutoScalingLifecycleServiceTest : TestSetup
 			.Setup(x => x.CompleteLifecycleActionAsync(It.IsAny<CompleteLifecycleActionRequest>(), It.IsAny<CancellationToken>()))
 			.Returns(OnCompleteLifecycleActionAsync);
 
-		ServerSettings ss = new() { AwsAutoScalingQueueUrls = new[] { _queueUrl1, _queueUrl2 } };
-		_asgLifecycleService = new AwsAutoScalingLifecycleService(AgentService, GetRedisServiceSingleton(), AgentCollection, Clock, new TestOptionsMonitor<ServerSettings>(ss), ServiceProvider, Tracer, logger);
+		StaticComputeConfig ss = new() { AwsAutoScalingQueueUrls = new[] { _queueUrl1, _queueUrl2 } };
+		_asgLifecycleService = new AwsAutoScalingLifecycleService(AgentService, GetRedisServiceSingleton(), AgentCollection, Clock, new TestOptionsMonitor<StaticComputeConfig>(ss), ServiceProvider, Tracer, logger);
 		_asgLifecycleService.SetAmazonClientsTesting(_asgMock.Object, _fakeSqs);
 		await _asgLifecycleService.StartAsync(CancellationToken.None);
 	}
diff --git a/Engine/Source/Programs/Horde/HordeServer.Tests/Fleet/FleetServiceTest.cs b/Engine/Source/Programs/Horde/HordeServer.Tests/Fleet/FleetServiceTest.cs
--- a/Engine/Source/Programs/Horde/HordeServer.Tests/Fleet/FleetServiceTest.cs
+++ b/Engine/Source/Programs/Horde/HordeServer.Tests/Fleet/FleetServiceTest.cs
@@ -231,12 +231,13 @@ namespace HordeServer.Tests.Fleet
 		private FleetService GetFleetService(IFleetManager fleetManager, bool isDowntimeActive = false)
 		{
 			ILoggerFactory loggerFactory = ServiceProvider.GetRequiredService<ILoggerFactory>();
-			IOptions<ServerSettings> serverSettingsOpt = ServiceProvider.GetRequiredService<IOptions<ServerSettings>>();
+			IOptions<StaticComputeConfig> serverSettingsOpt = ServiceProvider.GetRequiredService<IOptions<StaticComputeConfig>>();
 			serverSettingsOpt.Value.FleetManagerV2 = FleetManagerType.AwsReuse;
 
+			IOptionsMonitor<ComputeConfig> computeConfig = ServiceProvider.GetRequiredService<IOptionsMonitor<ComputeConfig>>();
 			FleetService service = new(
 				AgentService, PoolCollection, new DowntimeServiceStub(isDowntimeActive), Meter,
-				new StubFleetManagerFactory(fleetManager), Clock, ServiceProvider.GetRequiredService<IEnumerable<IPoolSizeStrategyFactory>>(), serverSettingsOpt, GlobalConfig, Tracer, loggerFactory.CreateLogger<FleetService>());
+				new StubFleetManagerFactory(fleetManager), Clock, ServiceProvider.GetRequiredService<IEnumerable<IPoolSizeStrategyFactory>>(), serverSettingsOpt, computeConfig, Tracer, loggerFactory.CreateLogger<FleetService>());
 
 			return service;
 		}
diff --git a/Engine/Source/Programs/Horde/HordeServer.Tests/TestSetup.cs b/Engine/Source/Programs/Horde/HordeServer.Tests/TestSetup.cs
--- a/Engine/Source/Programs/Horde/HordeServer.Tests/TestSetup.cs
+++ b/Engine/Source/Programs/Horde/HordeServer.Tests/TestSetup.cs
@@ -115,6 +115,7 @@ namespace HordeServer.Tests
 		public TestSetup()
 		{
 			AddPlugin<AnalyticsPlugin>();
+			AddPlugin<ComputePlugin>();
 			AddPlugin<StoragePlugin>();
 		}
 
diff --git a/Engine/Source/Programs/Horde/HordeServer/Configuration/ConfigService.cs b/Engine/Source/Programs/Horde/HordeServer/Configuration/ConfigService.cs
--- a/Engine/Source/Programs/Horde/HordeServer/Configuration/ConfigService.cs
+++ b/Engine/Source/Programs/Horde/HordeServer/Configuration/ConfigService.cs
@@ -587,10 +587,7 @@ namespace HordeServer.Configuration
 				ObjectConfigNode configNode = new ObjectConfigNode(typeof(GlobalConfig));
 
 				JsonObject configObj = await context.PreprocessFileAsync(GetGlobalConfigUri(), configNode, cancellationToken);
-				CopyJsonNode(configObj, "TelemetryStores", "Plugins.Analytics.Stores");
-				CopyJsonNode(configObj, "Secrets", "Plugins.Secrets.Secrets");
-				CopyJsonNode(configObj, "Storage", "Plugins.Storage");
-				CopyJsonNode(configObj, "Tools", "Plugins.Tools.Tools");
+				UpgradeConfig(configObj);
 
 				GlobalConfig globalConfig = JsonSerializer.Deserialize<GlobalConfig>(configObj, _jsonOptions)!;
 				if (globalConfig.VersionEnum < ConfigVersion.Latest)
@@ -627,19 +624,29 @@ namespace HordeServer.Configuration
 				throw new ConfigException(context, ex.Message, ex);
 			}
 		}
-		/*
-		static bool TryGetChildNode(JsonNode node, string name, [NotNullWhen(true)] out JsonNode? child)
+
+		static void UpgradeConfig(JsonObject rootObj)
 		{
-			child = node[name];
-			if (child != null)
-			{
-				JsonObject? parentObj = node as JsonObject;
-				parentObj
-				parentObj.TryGetPropertyValue(
-			}
-			return child != null;
+			// Analytics
+			CopyJsonNode(rootObj, "TelemetryStores", "Plugins.Analytics.Stores");
+
+			// Compute
+			CopyJsonNode(rootObj, "Rates", "Plugins.Compute.Rates");
+			CopyJsonNode(rootObj, "Compute", "Plugins.Compute.Clusters");
+			CopyJsonNode(rootObj, "Pools", "Plugins.Compute.Pools");
+			CopyJsonNode(rootObj, "Software", "Plugins.Compute.Software");
+			CopyJsonNode(rootObj, "Networks", "Plugins.Compute.Network");
+
+			// Secrets
+			CopyJsonNode(rootObj, "Secrets", "Plugins.Secrets.Secrets");
+
+			// Storage
+			CopyJsonNode(rootObj, "Storage", "Plugins.Storage");
+
+			// Tools
+			CopyJsonNode(rootObj, "Tools", "Plugins.Tools.Tools");
 		}
-		*/
+
 		static void CopyJsonNode(JsonObject rootObj, string sourcePath, string targetPath)
 		{
 			JsonNode? sourceNode = rootObj;
diff --git a/Engine/Source/Programs/Horde/HordeServer/HordeServer.csproj b/Engine/Source/Programs/Horde/HordeServer/HordeServer.csproj
--- a/Engine/Source/Programs/Horde/HordeServer/HordeServer.csproj
+++ b/Engine/Source/Programs/Horde/HordeServer/HordeServer.csproj
@@ -95,29 +95,11 @@
   </ItemGroup>
 
   <ItemGroup>
-    <Protobuf Include="Protos\horde\job_rpc.proto" ProtoRoot="Protos" GrpcServices="Server" AdditionalImportDirs="..\..\Shared\EpicGames.Horde\Protos" />
-    <Protobuf Include="Protos\horde\job_rpc_messages.proto" ProtoRoot="Protos" />
-    <Protobuf Include="Protos\horde\enrollment_rpc.proto" ProtoRoot="Protos" GrpcServices="Server" />
-    <Protobuf Include="Protos\horde\enrollment_rpc_messages.proto" ProtoRoot="Protos" />
-    <Protobuf Include="Protos\horde\horde_rpc.proto" ProtoRoot="Protos" GrpcServices="Server" />
-    <Protobuf Include="Protos\horde\horde_rpc_messages.proto" ProtoRoot="Protos" />
-    <Protobuf Include="Protos\horde\relay_rpc.proto" ProtoRoot="Protos" />
-    <Protobuf Include="Protos\horde\common\agent_status.proto" ProtoRoot="Protos" />
+    <Protobuf Include="Protos\horde\job_rpc.proto" ProtoRoot="Protos" GrpcServices="Server" AdditionalImportDirs="..\..\Shared\EpicGames.Horde\Protos;..\Plugins\Compute\HordeServer.Compute\Protos" />
+    <Protobuf Include="Protos\horde\job_rpc_messages.proto" ProtoRoot="Protos" AdditionalImportDirs="..\Plugins\Compute\HordeServer.Compute\Protos" />
     <Protobuf Include="Protos\horde\common\io_hash.proto" ProtoRoot="Protos" />
     <Protobuf Include="Protos\horde\common\label_change.proto" ProtoRoot="Protos" />
-    <Protobuf Include="Protos\horde\common\lease_outcome.proto" ProtoRoot="Protos" />
-    <Protobuf Include="Protos\horde\common\lease_state.proto" ProtoRoot="Protos" />
     <Protobuf Include="Protos\horde\common\namespace_id.proto" ProtoRoot="Protos" />
-    <Protobuf Include="Protos\horde\messages\agent_capabilities.proto" ProtoRoot="Protos" />
-    <Protobuf Include="Protos\horde\messages\agent_requirements.proto" ProtoRoot="Protos" />
-    <Protobuf Include="Protos\horde\messages\agent_workspace.proto" ProtoRoot="Protos" />
-    <Protobuf Include="Protos\horde\messages\lease.proto" ProtoRoot="Protos" />
-    <Protobuf Include="Protos\horde\tasks\compute_task.proto" ProtoRoot="Protos" />
-    <Protobuf Include="Protos\horde\tasks\conform_task.proto" ProtoRoot="Protos" />
-    <Protobuf Include="Protos\horde\tasks\job_task.proto" ProtoRoot="Protos" />
-    <Protobuf Include="Protos\horde\tasks\restart_task.proto" ProtoRoot="Protos" />
-    <Protobuf Include="Protos\horde\tasks\shutdown_task.proto" ProtoRoot="Protos" />
-    <Protobuf Include="Protos\horde\tasks\upgrade_task.proto" ProtoRoot="Protos" />
     <Protobuf Include="Protos\grpc\health\v1\health.proto" ProtoRoot="Protos" GrpcServices="Server" />
   </ItemGroup>
 
@@ -135,6 +117,7 @@
     <ProjectReference Include="..\..\Shared\EpicGames.Slack\EpicGames.Slack.csproj" />
     <ProjectReference Include="..\HordeServer.Shared\HordeServer.Shared.csproj" />
     <ProjectReference Include="..\Plugins\Analytics\HordeServer.Analytics\HordeServer.Analytics.csproj" />
+    <ProjectReference Include="..\Plugins\Compute\HordeServer.Compute\HordeServer.Compute.csproj" />
     <ProjectReference Include="..\Plugins\Ddc\HordeServer.Ddc\HordeServer.Ddc.csproj" />
     <ProjectReference Include="..\Plugins\Secrets\HordeServer.Secrets\HordeServer.Secrets.csproj" />
     <ProjectReference Include="..\Plugins\Storage\HordeServer.Storage\HordeServer.Storage.csproj" />
diff --git a/Engine/Source/Programs/Horde/HordeServer/Plugins/IPlugin.cs b/Engine/Source/Programs/Horde/HordeServer/Plugins/IPlugin.cs
--- a/Engine/Source/Programs/Horde/HordeServer/Plugins/IPlugin.cs
+++ b/Engine/Source/Programs/Horde/HordeServer/Plugins/IPlugin.cs
@@ -52,8 +52,9 @@ namespace HordeServer.Plugins
 		/// Configure services for the plugin
 		/// </summary>
 		/// <param name="configuration">Static configuration data</param>
+		/// <param name="serverInfo">Information about the server</param>
 		/// <param name="serviceCollection">Service collection instance</param>
-		void ConfigureServices(IConfiguration configuration, IServiceCollection serviceCollection);
+		void ConfigureServices(IConfiguration configuration, IServerInfo serverInfo, IServiceCollection serviceCollection);
 	}
 
 	/// <summary>
@@ -64,7 +65,7 @@ namespace HordeServer.Plugins
 		/// <summary>
 		/// Adds a plugin's services to a service collection
 		/// </summary>
-		public static void AddPlugin(this IServiceCollection serviceCollection, ILoadedPlugin plugin, IConfiguration configuration)
-			=> plugin.ConfigureServices(configuration, serviceCollection);
+		public static void AddPlugin(this IServiceCollection serviceCollection, IServerInfo serverInfo, ILoadedPlugin plugin, IConfiguration configuration)
+			=> plugin.ConfigureServices(configuration, serverInfo, serviceCollection);
 	}
 }
diff --git a/Engine/Source/Programs/Horde/HordeServer/Plugins/PluginCollection.cs b/Engine/Source/Programs/Horde/HordeServer/Plugins/PluginCollection.cs
--- a/Engine/Source/Programs/Horde/HordeServer/Plugins/PluginCollection.cs
+++ b/Engine/Source/Programs/Horde/HordeServer/Plugins/PluginCollection.cs
@@ -48,15 +48,15 @@ namespace HordeServer.Plugins
 			public LoadedPlugin(IPluginMetadata metadata)
 				=> _metadata = metadata;
 
-			public void ConfigureServices(IConfiguration config, IServiceCollection serviceCollection)
+			public void ConfigureServices(IConfiguration config, IServerInfo serverInfo, IServiceCollection serviceCollection)
 			{
 				serviceCollection.AddPluginConfig<TGlobalConfig>(Name);
 
-				TStartup startup = CreateStartup(config);
+				TStartup startup = CreateStartup(config, serverInfo);
 				startup.ConfigureServices(serviceCollection);
 			}
 
-			static TStartup CreateStartup(IConfiguration configuration)
+			static TStartup CreateStartup(IConfiguration configuration, IServerInfo serverInfo)
 			{
 				ConstructorInfo? chosenConstructor = null;
 				ParameterInfo[]? chosenConstructorParams = null;
@@ -86,7 +86,15 @@ namespace HordeServer.Plugins
 					for (int idx = 0; idx < chosenConstructorParams.Length; idx++)
 					{
 						ParameterInfo parameter = chosenConstructorParams[idx];
-						if (parameter.ParameterType == typeof(TServerConfig))
+						if (parameter.ParameterType == typeof(IConfiguration))
+						{
+							arguments[idx] = configuration;
+						}
+						else if (parameter.ParameterType == typeof(IServerInfo))
+						{
+							arguments[idx] = serverInfo;
+						}
+						else if (parameter.ParameterType == typeof(TServerConfig))
 						{
 							arguments[idx] = new TServerConfig();
 							configuration.Bind(arguments[idx]);
@@ -107,7 +115,8 @@ namespace HordeServer.Plugins
 			{
 				foreach (ParameterInfo parameter in parameters)
 				{
-					if (parameter.ParameterType != typeof(TServerConfig))
+					if (parameter.ParameterType != typeof(IServerInfo)
+						&& parameter.ParameterType != typeof(TServerConfig))
 					{
 						return false;
 					}
diff --git a/Engine/Source/Programs/Horde/HordeServer/Server/DebugController.cs b/Engine/Source/Programs/Horde/HordeServer/Server/DebugController.cs
--- a/Engine/Source/Programs/Horde/HordeServer/Server/DebugController.cs
+++ b/Engine/Source/Programs/Horde/HordeServer/Server/DebugController.cs
@@ -199,27 +199,6 @@ namespace HordeServer.Server
 			return Empty;
 		}
 
-		/// <summary>
-		/// Get the network ID for a given IP address
-		/// </summary>
-		[HttpGet]
-		[Route("/api/v1/debug/network-id")]
-		public ActionResult<object> GetNetworkId([FromQuery] string? ipAddress = null)
-		{
-			if (!_globalConfig.Value.Authorize(ServerAclAction.Debug, User))
-			{
-				return Forbid(ServerAclAction.Debug);
-			}
-
-			if (ipAddress == null || !IPAddress.TryParse(ipAddress, out IPAddress? ip))
-			{
-				return BadRequest("Unable to read or convert query parameter 'ipAddress'");
-			}
-
-			_globalConfig.Value.TryGetNetworkConfig(ip, out NetworkConfig? networkConfig);
-			return networkConfig == null ? StatusCode(StatusCodes.Status500InternalServerError, "Unable to find a network config for the IP") : Ok(networkConfig);
-		}
-
 		/// <summary>
 		/// Add a port mapping for agent relay
 		/// </summary>
diff --git a/Engine/Source/Programs/Horde/HordeServer/Server/GlobalConfig.cs b/Engine/Source/Programs/Horde/HordeServer/Server/GlobalConfig.cs
--- a/Engine/Source/Programs/Horde/HordeServer/Server/GlobalConfig.cs
+++ b/Engine/Source/Programs/Horde/HordeServer/Server/GlobalConfig.cs
@@ -6,7 +6,6 @@ using System.ComponentModel.DataAnnotations;
 using System.Diagnostics;
 using System.Diagnostics.CodeAnalysis;
 using System.Linq;
-using System.Net;
 using System.Security.Claims;
 using System.Text.Json.Serialization;
 using EpicGames.Core;
@@ -18,10 +17,8 @@ using EpicGames.Horde.Compute;
 using EpicGames.Horde.Jobs.Templates;
 using EpicGames.Horde.Projects;
 using EpicGames.Horde.Streams;
-using EpicGames.Horde.Tools;
 using EpicGames.Horde.Users;
 using EpicGames.Perforce;
-using EpicGames.Serialization;
 using HordeServer.Acls;
 using HordeServer.Agents;
 using HordeServer.Agents.Pools;
@@ -132,7 +129,7 @@ namespace HordeServer.Server
 		/// <summary>
 		/// List of pools
 		/// </summary>
-		public List<PoolConfig> Pools { get; set; } = new List<PoolConfig>();
+		public List<PoolConfig> Pools => (Plugins.Values.OfType<ComputeConfig>().FirstOrDefault() ?? new ComputeConfig()).Pools;
 
 		/// <summary>
 		/// List of scheduled downtime
@@ -144,26 +141,6 @@ namespace HordeServer.Server
 		/// </summary>
 		public List<PerforceCluster> PerforceClusters { get; set; } = new List<PerforceCluster>();
 
-		/// <summary>
-		/// List of costs of a particular agent type
-		/// </summary>
-		public List<AgentSoftwareConfig> Software { get; set; } = new List<AgentSoftwareConfig>();
-
-		/// <summary>
-		/// List of costs of a particular agent type
-		/// </summary>
-		public List<AgentRateConfig> Rates { get; set; } = new List<AgentRateConfig>();
-
-		/// <summary>
-		/// List of networks
-		/// </summary>
-		public List<NetworkConfig> Networks { get; set; } = new List<NetworkConfig>();
-
-		/// <summary>
-		/// List of compute profiles
-		/// </summary>
-		public List<ComputeClusterConfig> Compute { get; set; } = new List<ComputeClusterConfig>();
-
 		/// <summary>
 		/// Device configuration
 		/// </summary>
@@ -225,10 +202,8 @@ namespace HordeServer.Server
 
 		private readonly Dictionary<ProjectId, ProjectConfig> _projectLookup = new Dictionary<ProjectId, ProjectConfig>();
 		private readonly Dictionary<StreamId, StreamConfig> _streamLookup = new Dictionary<StreamId, StreamConfig>();
-		private readonly Dictionary<ClusterId, ComputeClusterConfig> _computeClusterLookup = new Dictionary<ClusterId, ComputeClusterConfig>();
 		private readonly Dictionary<AclScopeName, AclConfig> _aclLookup = new Dictionary<AclScopeName, AclConfig>();
 		private readonly Dictionary<ArtifactType, ArtifactTypeConfig> _artifactTypeLookup = new Dictionary<ArtifactType, ArtifactTypeConfig>();
-		private readonly Dictionary<PoolId, PoolConfig> _poolLookup = new Dictionary<PoolId, PoolConfig>();
 
 		/// <inheritdoc cref="AclConfig.Authorize(AclAction, ClaimsPrincipal)"/>
 		public bool Authorize(AclAction action, ClaimsPrincipal user)
@@ -259,25 +234,12 @@ namespace HordeServer.Server
 				}
 			}
 
-			_computeClusterLookup.Clear();
-			foreach (ComputeClusterConfig computeCluster in Compute)
-			{
-				_computeClusterLookup.Add(computeCluster.Id, computeCluster);
-				computeCluster.PostLoad(this);
-			}
-
 			_artifactTypeLookup.Clear();
 			foreach (ArtifactTypeConfig artifactType in ArtifactTypes)
 			{
 				_artifactTypeLookup.Add(artifactType.Type, artifactType);
 			}
 
-			_poolLookup.Clear();
-			foreach (PoolConfig pool in Pools)
-			{
-				_poolLookup.Add(pool.Id, pool);
-			}
-			ConfigObject.MergeDefaults<string, PoolConfig>(Pools.Select(x => (x.Id.ToString(), x.Base?.ToString(), x)));
 			UpdateWorkspacesForPools();
 
 			// Ensure that all plugins have an entry in the global config so they can register their ACLs
@@ -495,77 +457,6 @@ namespace HordeServer.Server
 			return streamConfig.TryGetTemplate(templateId, out config);
 		}
 
-		/// <summary>
-		/// Attempts to get configuration for a pool from this object
-		/// </summary>
-		/// <param name="poolId">The pool identifier</param>
-		/// <param name="config">Configuration for the pool</param>
-		/// <returns>True if the pool configuration was found</returns>
-		public bool TryGetPool(PoolId poolId, [NotNullWhen(true)] out PoolConfig? config) => _poolLookup.TryGetValue(poolId, out config);
-
-		/// <summary>
-		/// Attempt to resolve an IP address to a network config
-		/// </summary>
-		/// <param name="ip">IP address to resolve</param>
-		/// <param name="networkConfig">Config for the network</param>
-		/// <returns>True if the IP address was resolved</returns>
-		public bool TryGetNetworkConfig(IPAddress ip, [NotNullWhen(true)] out NetworkConfig? networkConfig)
-		{
-			foreach (NetworkConfig nc in Networks)
-			{
-				if (nc.Id != null && IsIpInBlock(ip, nc.CidrBlock))
-				{
-					networkConfig = nc;
-					return true;
-				}
-			}
-
-			networkConfig = null;
-			return false;
-		}
-
-		private static bool IsIpInBlock(IPAddress ip, string? cidrBlock)
-		{
-			if (cidrBlock == null)
-			{
-				return false;
-			}
-
-			if (cidrBlock == "0.0.0.0/0")
-			{
-				return true;
-			}
-
-			string[] parts = cidrBlock.Split('/');
-			if (parts.Length != 2 || !IPAddress.TryParse(parts[0], out IPAddress? address) || !Int32.TryParse(parts[1], out int maskBits))
-			{
-				return false;
-			}
-
-			byte[] networkPrefixBytes = address.GetAddressBytes();
-			Array.Reverse(networkPrefixBytes);
-
-			uint networkPrefix = BitConverter.ToUInt32(networkPrefixBytes, 0);
-			uint subnetMask = 0xffffffff;
-			subnetMask <<= 32 - maskBits;
-			uint ipRangeStart = networkPrefix & subnetMask;
-			uint ipRangeEnd = networkPrefix | (subnetMask ^ 0xffffffff);
-
-			byte[] ipBytes = ip.GetAddressBytes();
-			Array.Reverse(ipBytes);
-			uint ipUint = BitConverter.ToUInt32(ipBytes, 0);
-			return ipUint > ipRangeStart && ipUint <= ipRangeEnd;
-		}
-
-		/// <summary>
-		/// Attempts to get compute cluster configuration from this object
-		/// </summary>
-		/// <param name="clusterId">Compute cluster id</param>
-		/// <param name="config">Receives the cluster configuration on success</param>
-		/// <returns>True on success</returns>
-		public bool TryGetComputeCluster(ClusterId clusterId, [NotNullWhen(true)] out ComputeClusterConfig? config)
-			=> _computeClusterLookup.TryGetValue(clusterId, out config);
-
 		/// <summary>
 		/// Authorizes a user to perform a given action
 		/// </summary>
@@ -781,72 +672,6 @@ namespace HordeServer.Server
 		}
 	}
 
-	/// <summary>
-	/// Selects different agent software versions by evaluating a condition
-	/// </summary>
-	[DebuggerDisplay("{ToolId}")]
-	public class AgentSoftwareConfig
-	{
-		/// <summary>
-		/// Tool identifier
-		/// </summary>
-		public ToolId ToolId { get; set; }
-
-		/// <summary>
-		/// Condition for using this channel
-		/// </summary>
-		public Condition? Condition { get; set; }
-	}
-
-	/// <summary>
-	/// Describes the monetary cost of agents matching a particular criteria
-	/// </summary>
-	public class AgentRateConfig
-	{
-		/// <summary>
-		/// Condition string
-		/// </summary>
-		[CbField("c")]
-		public Condition? Condition { get; set; }
-
-		/// <summary>
-		/// Rate for this agent
-		/// </summary>
-		[CbField("r")]
-		public double Rate { get; set; }
-	}
-
-	/// <summary>
-	/// Describes a network
-	/// The ID describes any logical grouping, such as region, availability zone, rack or office location. 
-	/// </summary>
-	public class NetworkConfig
-	{
-		/// <summary>
-		/// ID for this network
-		/// </summary>
-		[CbField("id")]
-		public string? Id { get; set; }
-
-		/// <summary>
-		/// CIDR block
-		/// </summary>
-		[CbField("cb")]
-		public string? CidrBlock { get; set; }
-
-		/// <summary>
-		/// Human-readable description
-		/// </summary>
-		[CbField("d")]
-		public string? Description { get; set; }
-
-		/// <summary>
-		/// Compute ID for this network (used when allocating compute resources)
-		/// </summary>
-		[CbField("cid")]
-		public string? ComputeId { get; set; }
-	}
-
 	/// <summary>
 	/// How frequently the maintence window repeats
 	/// </summary>
diff --git a/Engine/Source/Programs/Horde/HordeServer/ServerApp.cs b/Engine/Source/Programs/Horde/HordeServer/ServerApp.cs
--- a/Engine/Source/Programs/Horde/HordeServer/ServerApp.cs
+++ b/Engine/Source/Programs/Horde/HordeServer/ServerApp.cs
@@ -216,6 +216,19 @@ namespace HordeServer
 		/// </summary>
 		static readonly KeyValuePair<string, string>[] s_renamedConfigKeys = new[]
 		{
+			// Compute plugin
+			KeyValuePair.Create("Horde:EnableUpgradeTasks", "Horde:Plugins:Compute:EnableUpgradeTasks"),
+			KeyValuePair.Create("Horde:FleetManagerV2", "Horde:Plugins:Compute:FleetManagerV2"),
+			KeyValuePair.Create("Horde:FleetManagerV2Config", "Horde:Plugins:Compute:FleetManagerV2Config"),
+			KeyValuePair.Create("Horde:AwsAutoScalingQueueUrls", "Horde:Plugins:Compute:AwsAutoScalingQueueUrls"),
+			KeyValuePair.Create("Horde:AutoEnrollAgents", "Horde:Plugins:Compute:AutoEnrollAgents"),
+			KeyValuePair.Create("Horde:DefaultAgentPoolSizeStrategy", "Horde:Plugins:Compute:DefaultAgentPoolSizeStrategy"),
+			KeyValuePair.Create("Horde:AgentPoolScaleOutCooldownSeconds", "Horde:Plugins:Compute:AgentPoolScaleOutCooldownSeconds"),
+			KeyValuePair.Create("Horde:AgentPoolScaleInCooldownSeconds", "Horde:Plugins:Compute:AgentPoolScaleInCooldownSeconds"),
+			KeyValuePair.Create("Horde:ComputeTunnelPort", "Horde:Plugins:Compute:ComputeTunnelPort"),
+			KeyValuePair.Create("Horde:ComputeTunnelAddress", "Horde:Plugins:Compute:ComputeTunnelAddress"),
+			KeyValuePair.Create("Horde:WithAws", "Horde:Plugins:Compute:WithAws"),
+
 			// Storage plugin
 			KeyValuePair.Create("Horde:BundleCacheDir", "Horde:Plugins:Storage:BundleCacheDir"),
 			KeyValuePair.Create("Horde:BundleCacheSize", "Horde:Plugins:Storage:BundleCacheSize"),
diff --git a/Engine/Source/Programs/Horde/HordeServer/ServerInfo.cs b/Engine/Source/Programs/Horde/HordeServer/ServerInfo.cs
--- a/Engine/Source/Programs/Horde/HordeServer/ServerInfo.cs
+++ b/Engine/Source/Programs/Horde/HordeServer/ServerInfo.cs
@@ -1,6 +1,7 @@
 // Copyright Epic Games, Inc. All Rights Reserved.
 
 using EpicGames.Core;
+using Microsoft.Extensions.Configuration;
 using Microsoft.Extensions.Options;
 
 namespace HordeServer
@@ -25,13 +26,23 @@ namespace HordeServer
 		/// <inheritdoc/>
 		public DirectoryReference DataDir => ServerApp.DataDir;
 
+		/// <inheritdoc/>
+		public IConfiguration Configuration => _configuration;
+
+		/// <inheritdoc/>
+		public bool ReadOnlyMode => _serverSettings.Value.MongoReadOnlyMode;
+
+		readonly IConfiguration _configuration;
 		readonly IOptions<ServerSettings> _serverSettings;
 
 		/// <summary>
 		/// Constructor
 		/// </summary>
-		public ServerInfo(IOptions<ServerSettings> serverSettings)
-			=> _serverSettings = serverSettings;
+		public ServerInfo(IConfiguration configuration, IOptions<ServerSettings> serverSettings)
+		{
+			_configuration = configuration;
+			_serverSettings = serverSettings;
+		}
 
 		/// <inheritdoc/>
 		public bool IsRunModeActive(RunMode mode)
diff --git a/Engine/Source/Programs/Horde/HordeServer/ServerSettings.cs b/Engine/Source/Programs/Horde/HordeServer/ServerSettings.cs
--- a/Engine/Source/Programs/Horde/HordeServer/ServerSettings.cs
+++ b/Engine/Source/Programs/Horde/HordeServer/ServerSettings.cs
@@ -11,7 +11,6 @@ using EpicGames.Horde.Server;
 using EpicGames.Horde.Storage.Bundles;
 using EpicGames.Horde.Storage.Nodes;
 using EpicGames.Perforce;
-using HordeServer.Agents.Fleet;
 using HordeServer.Server;
 
 namespace HordeServer
@@ -157,17 +156,6 @@ namespace HordeServer
 		/// </summary>
 		public int Http2Port { get; set; } = 5002;
 
-		/// <summary>
-		/// Port to listen on for tunneling compute sockets to agents
-		/// </summary>
-		public int ComputeTunnelPort { get; set; }
-
-		/// <summary>
-		/// What address (host:port) clients should connect to for compute socket tunneling
-		/// Port may differ from <see cref="ComputeTunnelPort" /> if Horde server is behind a reverse proxy/firewall
-		/// </summary>
-		public string? ComputeTunnelAddress { get; set; }
-
 		/// <summary>
 		/// Connection string for the Mongo database
 		/// </summary>
@@ -246,11 +234,6 @@ namespace HordeServer
 		/// </summary>
 		public AuthMethod AuthMethod { get; set; } = AuthMethod.Anonymous;
 
-		/// <summary>
-		/// Whether to automatically enroll agents in the farm
-		/// </summary>
-		public bool AutoEnrollAgents { get; set; }
-
 		/// <summary>
 		/// Optional profile name to report through the /api/v1/server/auth endpoint. Allows sharing auth tokens between providers configured through
 		/// the same profile name in OidcToken.exe config files.
@@ -452,22 +435,6 @@ namespace HordeServer
 		/// </summary>
 		public bool LogSessionRequests { get; set; } = false;
 
-		/// <summary>
-		/// Default fleet manager to use (when not specified by pool)
-		/// </summary>
-		public FleetManagerType FleetManagerV2 { get; set; } = FleetManagerType.NoOp;
-
-		/// <summary>
-		/// Config for the fleet manager (serialized JSON)
-		/// </summary>
-		public string? FleetManagerV2Config { get; set; }
-
-		/// <summary>
-		/// AWS SQS queue URLs where lifecycle events from EC2 auto-scaling are received
-		/// <see cref="AwsAutoScalingLifecycleService" />
-		/// </summary>
-		public string[] AwsAutoScalingQueueUrls { get; set; } = Array.Empty<string>();
-
 		/// <summary>
 		/// Whether to run scheduled jobs.
 		/// </summary>
@@ -578,21 +545,6 @@ namespace HordeServer
 		/// </summary>
 		public string? DeviceReportChannel { get; set; }
 
-		/// <summary>
-		/// Default agent pool sizing strategy for pools that doesn't have one explicitly configured
-		/// </summary>
-		public PoolSizeStrategy DefaultAgentPoolSizeStrategy { get; set; } = PoolSizeStrategy.LeaseUtilization;
-
-		/// <summary>
-		/// Scale-out cooldown for auto-scaling agent pools (in seconds). Can be overridden by per-pool settings.
-		/// </summary>
-		public int AgentPoolScaleOutCooldownSeconds { get; set; } = 60; // 1 min
-
-		/// <summary>
-		/// Scale-in cooldown for auto-scaling agent pools (in seconds). Can be overridden by per-pool settings.
-		/// </summary>
-		public int AgentPoolScaleInCooldownSeconds { get; set; } = 1200; // 20 mins
-
 		/// <summary>
 		/// Set the minimum size of the global thread pool
 		/// This value has been found in need of tweaking to avoid timeouts with the Redis client during bursts
@@ -606,11 +558,6 @@ namespace HordeServer
 		/// </summary>
 		public bool WithDatadog { get; set; }
 
-		/// <summary>
-		/// Whether to enable Amazon Web Services (AWS) specific features
-		/// </summary>
-		public bool WithAws { get; set; } = false;
-
 		/// <summary>
 		/// Path to the root config file. Relative to the server.json file by default.
 		/// </summary>
@@ -631,11 +578,6 @@ namespace HordeServer
 		/// </summary>
 		public int PerforceConnectionPoolSize { get; set; } = 5;
 
-		/// <summary>
-		/// Whether to enable the upgrade task source.
-		/// </summary>
-		public bool EnableUpgradeTasks { get; set; } = true;
-
 		/// <summary>
 		/// Whether to enable the conform task source.
 		/// </summary>
diff --git a/Engine/Source/Programs/Horde/HordeServer/Startup.cs b/Engine/Source/Programs/Horde/HordeServer/Startup.cs
--- a/Engine/Source/Programs/Horde/HordeServer/Startup.cs
+++ b/Engine/Source/Programs/Horde/HordeServer/Startup.cs
@@ -10,12 +10,6 @@ using System.Reflection;
 using System.Runtime.InteropServices;
 using System.Threading;
 using System.Threading.Tasks;
-using Amazon;
-using Amazon.AutoScaling;
-using Amazon.CloudWatch;
-using Amazon.EC2;
-using Amazon.Extensions.NETCore.Setup;
-using Amazon.SQS;
 using EpicGames.AspNet;
 using EpicGames.Core;
 using EpicGames.Horde.Acls;
@@ -47,7 +41,6 @@ using HordeServer.Agents.Utilization;
 using HordeServer.Artifacts;
 using HordeServer.Auditing;
 using HordeServer.Authentication;
-using HordeServer.Aws;
 using HordeServer.Commits;
 using HordeServer.Compute;
 using HordeServer.Configuration;
@@ -292,7 +285,6 @@ namespace HordeServer
 		// This method gets called *multiple times* by the runtime. Use this method to add services to the container.
 		public void ConfigureServices(IServiceCollection services)
 		{
-			services.AddSingleton<IServerInfo, ServerInfo>();
 			services.AddSingleton<JsonSchemaCache>();
 
 			// IOptionsMonitor pattern for live updating of configuration settings
@@ -308,9 +300,13 @@ namespace HordeServer
 			ServerSettings settings = new();
 			BindServerSettings(Configuration, settings);
 
+			ServerInfo serverInfo = new ServerInfo(Configuration, Options.Create(settings));
+			services.AddSingleton<IServerInfo>(serverInfo);
+
 			// Register the plugin collection
 			PluginCollection pluginCollection = new PluginCollection();
 			pluginCollection.Add<AnalyticsPlugin>();
+			pluginCollection.Add<ComputePlugin>();
 			pluginCollection.Add<StoragePlugin>();
 			pluginCollection.Add<DdcPlugin>();
 			pluginCollection.Add<SecretsPlugin>();
@@ -322,7 +318,7 @@ namespace HordeServer
 			foreach (ILoadedPlugin plugin in pluginCollection.LoadedPlugins)
 			{
 				IConfigurationSection pluginConfig = pluginsConfig.GetSection(plugin.Name.ToString());
-				plugin.ConfigureServices(pluginConfig, services);
+				plugin.ConfigureServices(pluginConfig, serverInfo, services);
 			}
 
 			OpenTelemetryHelper.Configure(services, settings.OpenTelemetry);
@@ -424,24 +420,7 @@ namespace HordeServer
 			services.AddSingleton(typeof(IAuditLogFactory<>), typeof(AuditLogFactory<>));
 			services.AddSingleton(typeof(ISingletonDocument<>), typeof(SingletonDocument<>));
 
-			services.AddSingleton<AwsAutoScalingLifecycleService>();
-			services.AddSingleton<FleetService>();
-			services.AddSingleton<IFleetManagerFactory, FleetManagerFactory>();
-			services.AddSingleton<IPoolSizeStrategyFactory, NoOpPoolSizeStrategyFactory>();
 			services.AddSingleton<IPoolSizeStrategyFactory, JobQueueStrategyFactory>();
-			services.AddSingleton<IPoolSizeStrategyFactory, LeaseUtilizationStrategyFactory>();
-
-			// Associate IFleetManager interface with the default implementation from config for convenience
-			// Though most fleet managers are created on a per-pool basis
-			services.AddSingleton<IFleetManager>(ctx => ctx.GetRequiredService<IFleetManagerFactory>().CreateFleetManager(FleetManagerType.Default));
-
-			// Run the tunnel service for all run modes
-			services.AddSingleton<TunnelService>();
-			services.AddHostedService<TunnelService>(sp => sp.GetRequiredService<TunnelService>());
-
-			// Runs the agent relay service for all run modes to notify long-polling requests
-			services.AddSingleton<AgentRelayService>();
-			services.AddHostedService(provider => provider.GetRequiredService<AgentRelayService>());
 
 			services.AddSingleton<AclService>();
 			services.AddSingleton<IAclService>(sp => sp.GetRequiredService<AclService>());
@@ -516,26 +495,6 @@ namespace HordeServer
 				services.AddSingleton<IExternalIssueService, DefaultExternalIssueService>();
 			}
 
-			if (settings.WithAws)
-			{
-				AWSOptions awsOptions = Configuration.GetAWSOptions();
-				services.AddDefaultAWSOptions(awsOptions);
-				if (awsOptions.Region == null && Environment.GetEnvironmentVariable("AWS_REGION") == null)
-				{
-					awsOptions.Region = RegionEndpoint.USEast1;
-				}
-
-				services.AddAWSService<IAmazonCloudWatch>();
-				services.AddAWSService<IAmazonAutoScaling>();
-				services.AddAWSService<IAmazonSQS>();
-				services.AddAWSService<IAmazonEC2>();
-
-				services.AddSingleton<AwsCloudWatchMetricExporter>();
-
-				services.AddSingleton<IPoolSizeStrategyFactory, LeaseUtilizationAwsMetricStrategyFactory>();
-				services.AddSingleton<IPoolSizeStrategyFactory, ComputeQueueAwsMetricStrategyFactory>();
-			}
-
 			AuthenticationBuilder authBuilder = services.AddAuthentication(options =>
 				{
 					switch (settings.AuthMethod)
@@ -757,12 +716,6 @@ namespace HordeServer
 				{
 					services.AddHostedService(provider => provider.GetRequiredService<SlackNotificationSink>());
 				}
-
-				if (settings.WithAws)
-				{
-					services.AddHostedService(provider => provider.GetRequiredService<AwsAutoScalingLifecycleService>());
-					services.AddHostedService(provider => provider.GetRequiredService<AwsCloudWatchMetricExporter>());
-				}
 			}
 
 			services.AddHostedService(provider => provider.GetRequiredService<IExternalIssueService>());
@@ -777,12 +730,8 @@ namespace HordeServer
 				services.AddHostedService<ConformTaskSource>(provider => provider.GetRequiredService<ConformTaskSource>());
 				services.AddSingleton<ComputeTaskSource>();
 
-				services.AddSingleton<ITaskSource, UpgradeTaskSource>();
-				services.AddSingleton<ITaskSource, ShutdownTaskSource>();
-				services.AddSingleton<ITaskSource, RestartTaskSource>();
 				services.AddSingleton<ITaskSource, ConformTaskSource>(provider => provider.GetRequiredService<ConformTaskSource>());
 				services.AddSingleton<ITaskSource, JobTaskSource>(provider => provider.GetRequiredService<JobTaskSource>());
-				services.AddSingleton<ITaskSource, ComputeTaskSource>(provider => provider.GetRequiredService<ComputeTaskSource>());
 			}
 
 			// Allow longer to shutdown so we can debug missing cancellation tokens
@@ -980,18 +929,18 @@ namespace HordeServer
 		public sealed class AclActionBsonSerializer : SerializerBase<AclAction>
 		{
 			/// <inheritdoc/>
-			public override AclAction Deserialize(BsonDeserializationContext context, BsonDeserializationArgs args) 
+			public override AclAction Deserialize(BsonDeserializationContext context, BsonDeserializationArgs args)
 				=> new AclAction(context.Reader.ReadString());
 
 			/// <inheritdoc/>
-			public override void Serialize(BsonSerializationContext context, BsonSerializationArgs args, AclAction value) 
+			public override void Serialize(BsonSerializationContext context, BsonSerializationArgs args, AclAction value)
 				=> context.Writer.WriteString(value.Name);
 		}
 
 		public sealed class AclScopeNameBsonSerializer : SerializerBase<AclScopeName>
 		{
 			/// <inheritdoc/>
-			public override AclScopeName Deserialize(BsonDeserializationContext context, BsonDeserializationArgs args) 
+			public override AclScopeName Deserialize(BsonDeserializationContext context, BsonDeserializationArgs args)
 				=> new AclScopeName(context.Reader.ReadString());
 
 			/// <inheritdoc/>
@@ -1004,11 +953,11 @@ namespace HordeServer
 			readonly TConverter _converter = new TConverter();
 
 			/// <inheritdoc/>
-			public override TValue Deserialize(BsonDeserializationContext context, BsonDeserializationArgs args) 
+			public override TValue Deserialize(BsonDeserializationContext context, BsonDeserializationArgs args)
 				=> _converter.FromSubResourceId(new SubResourceId((ushort)context.Reader.ReadInt32()));
 
 			/// <inheritdoc/>
-			public override void Serialize(BsonSerializationContext context, BsonSerializationArgs args, TValue value) 
+			public override void Serialize(BsonSerializationContext context, BsonSerializationArgs args, TValue value)
 				=> context.Writer.WriteInt32(_converter.ToSubResourceId(value).Value);
 		}
 
diff --git a/Engine/Source/Programs/Horde/Plugins/Compute/HordeServer.Compute.Tests/HordeServer.Compute.Tests.csproj b/Engine/Source/Programs/Horde/Plugins/Compute/HordeServer.Compute.Tests/HordeServer.Compute.Tests.csproj
new file mode 100644
--- /dev/null
+++ b/Engine/Source/Programs/Horde/Plugins/Compute/HordeServer.Compute.Tests/HordeServer.Compute.Tests.csproj
@@ -0,0 +1,30 @@
+<Project Sdk="Microsoft.NET.Sdk">
+
+  <PropertyGroup>
+    <TargetFramework>net8.0</TargetFramework>
+    <ImplicitUsings>enable</ImplicitUsings>
+    <Nullable>enable</Nullable>
+
+    <IsPackable>false</IsPackable>
+    <IsTestProject>true</IsTestProject>
+    <RootNamespace>HordeServer.Tests</RootNamespace>
+  </PropertyGroup>
+
+  <ItemGroup>
+    <PackageReference Include="coverlet.collector" Version="6.0.0" />
+    <PackageReference Include="Microsoft.NET.Test.Sdk" Version="17.8.0" />
+    <PackageReference Include="MSTest.TestAdapter" Version="3.1.1" />
+    <PackageReference Include="MSTest.TestFramework" Version="3.1.1" />
+  </ItemGroup>
+
+  <ItemGroup>
+    <ProjectReference Include="..\HordeServer.Compute\HordeServer.Compute.csproj" />
+    <ProjectReference Include="..\..\..\HordeAgent.Tests\HordeAgent.Tests.csproj" />
+    <ProjectReference Include="..\..\..\HordeServer.Tests.Shared\HordeServer.Tests.Shared.csproj" />
+  </ItemGroup>
+
+  <ItemGroup>
+    <Using Include="Microsoft.VisualStudio.TestTools.UnitTesting" />
+  </ItemGroup>
+
+</Project>
