Horde: Move compute functionality into a plugin.

diff --git a/Engine/Source/Programs/Horde/Plugins/Compute/HordeServer.Compute/Agents/AgentService.cs b/Engine/Source/Programs/Horde/Plugins/Compute/HordeServer.Compute/Agents/AgentService.cs
new file mode 100644
--- /dev/null
+++ b/Engine/Source/Programs/Horde/Plugins/Compute/HordeServer.Compute/Agents/AgentService.cs
@@ -0,0 +1,1246 @@
+// Copyright Epic Games, Inc. All Rights Reserved.
+
+using System;
+using System.Collections.Generic;
+using System.Diagnostics;
+using System.Diagnostics.Metrics;
+using System.Linq;
+using System.Security.Claims;
+using System.Threading;
+using System.Threading.Tasks;
+using EpicGames.Core;
+using EpicGames.Horde.Agents;
+using EpicGames.Horde.Agents.Leases;
+using EpicGames.Horde.Agents.Pools;
+using EpicGames.Horde.Agents.Sessions;
+using EpicGames.Redis;
+using EpicGames.Serialization;
+using Google.Protobuf.WellKnownTypes;
+using HordeServer.Acls;
+using HordeServer.Agents.Leases;
+using HordeServer.Agents.Pools;
+using HordeServer.Agents.Sessions;
+using HordeServer.Auditing;
+using HordeServer.Server;
+using HordeServer.Tasks;
+using HordeServer.Utilities;
+using HordeCommon.Rpc;
+using HordeCommon.Rpc.Messages;
+using Microsoft.Extensions.Hosting;
+using Microsoft.Extensions.Logging;
+using OpenTelemetry.Trace;
+using StackExchange.Redis;
+
+#pragma warning disable CA2227 // Change x to be read-only by removing the property setter
+
+namespace HordeServer.Agents
+{
+	/// <summary>
+	/// Singleton used to store agent costs
+	/// </summary>
+	[RedisConverter(typeof(RedisCbConverter<>))]
+	public class AgentRateTable
+	{
+		/// <summary>
+		/// List of costs for different agent types
+		/// </summary>
+		[CbField]
+		public List<AgentRateConfig> Entries { get; set; } = new List<AgentRateConfig>();
+	}
+
+	/// <summary>
+	/// Wraps functionality for manipulating agents
+	/// </summary>
+	public sealed class AgentService : IHostedService, IAsyncDisposable
+	{
+		/// <summary>
+		/// Maximum time between updates for an agent to be considered online
+		/// </summary>
+		public static readonly TimeSpan SessionExpiryTime = TimeSpan.FromMinutes(5);
+
+		/// <summary>
+		/// Time before a session expires that we will poll until
+		/// </summary>
+		public static readonly TimeSpan SessionLongPollTime = TimeSpan.FromSeconds(55);
+
+		/// <summary>
+		/// Time after which a session will be renewed
+		/// </summary>
+		public static readonly TimeSpan SessionRenewTime = TimeSpan.FromSeconds(50);
+
+		/// <summary>
+		/// Collection of agent documents
+		/// </summary>
+		public IAgentCollection Agents { get; }
+
+		readonly ILeaseCollection _leases;
+		readonly ISessionCollection _sessions;
+		readonly IAclService _aclService;
+		readonly IDowntimeService _downtimeService;
+		readonly ITaskSource[] _taskSources;
+		readonly IRedisService _redisService;
+		readonly IHostApplicationLifetime _applicationLifetime;
+		readonly IClock _clock;
+		readonly Tracer _tracer;
+		readonly ILogger _logger;
+		readonly ITicker _ticker;
+		readonly ITicker _sharedTicker;
+
+		readonly RedisStringKey<AgentRateTable> _agentRateTableData = new ("agent-rates");
+
+		/// <summary>Lazily updated costs for different agent types</summary>
+		readonly AsyncCachedValue<AgentRateTable?> _cachedRates;
+
+		/// <summary>Lazily updated list of current pools</summary>
+		readonly AsyncCachedValue<IReadOnlyList<IPoolConfig>> _cachedPools;
+		
+		/// <summary>Manually updated cached list of current agents</summary>
+		IReadOnlyDictionary<AgentId, IAgent>? _cachedAgents;
+
+		/// <summary>All the agents currently performing a long poll for work on this server</summary>
+		readonly Dictionary<AgentId, CancellationTokenSource> _waitingAgents = new ();
+
+		/// <summary>OpenTelemetry measurements (gauges)</summary>
+		IEnumerable<Measurement<int>> _measurements = new List<Measurement<int>>();
+
+		/// <summary>Subscription for update events</summary>
+		IAsyncDisposable? _subscription;
+
+		/// <summary>
+		/// Constructor
+		/// </summary>
+		public AgentService(
+			IAgentCollection agents,
+			ILeaseCollection leases,
+			ISessionCollection sessions,
+			IAclService aclService,
+			IDowntimeService downtimeService,
+			IPoolCollection poolCollection,
+			IEnumerable<ITaskSource> taskSources,
+			IRedisService redisService,
+			IHostApplicationLifetime applicationLifetime,
+			IClock clock,
+			Tracer tracer,
+			Meter meter,
+			ILogger<AgentService> logger)
+		{
+			Agents = agents;
+			_leases = leases;
+			_sessions = sessions;
+			_aclService = aclService;
+			_downtimeService = downtimeService;
+			_cachedRates = new AsyncCachedValue<AgentRateTable?>(_ => redisService.GetDatabase().StringGetAsync(_agentRateTableData), TimeSpan.FromMinutes(2));
+			_cachedPools = new AsyncCachedValue<IReadOnlyList<IPoolConfig>>(poolCollection.GetConfigsAsync, TimeSpan.FromSeconds(30.0));
+			_taskSources = taskSources.ToArray();
+			_applicationLifetime = applicationLifetime;
+			_redisService = redisService;
+			_clock = clock;
+			_ticker = clock.AddTicker($"{nameof(AgentService)}.{nameof(TickAsync)}", TimeSpan.FromSeconds(30.0), TickAsync, logger);
+			_sharedTicker = clock.AddSharedTicker($"{nameof(AgentService)}.{nameof(TickSharedAsync)}", TimeSpan.FromSeconds(30.0), TickSharedAsync, logger);
+			_tracer = tracer;
+			_logger = logger;
+
+			meter.CreateObservableGauge("horde.agent.count", () => _measurements);
+		}
+
+		/// <inheritdoc/>
+		public async Task StartAsync(CancellationToken cancellationToken)
+		{
+			await _ticker.StartAsync();
+			await _sharedTicker.StartAsync();
+			_subscription = await Agents.SubscribeToUpdateEventsAsync(OnAgentUpdate);
+		}
+
+		/// <inheritdoc/>
+		public async Task StopAsync(CancellationToken cancellationToken)
+		{
+			if (_subscription != null)
+			{
+				await _subscription.DisposeAsync();
+				_subscription = null;
+			}
+			await _ticker.StopAsync();
+			await _sharedTicker.StopAsync();
+		}
+
+		/// <inheritdoc/>
+		public async ValueTask DisposeAsync()
+		{
+			await _cachedRates.DisposeAsync();
+			await _cachedPools.DisposeAsync();
+			await _ticker.DisposeAsync();
+			await _sharedTicker.DisposeAsync();
+		}
+
+		/// <summary>
+		/// Gets user-readable payload information
+		/// </summary>
+		/// <param name="payload">The payload data</param>
+		/// <param name="cancellationToken">Cancellation token for the operation</param>
+		/// <returns>Dictionary of key/value pairs for the payload</returns>
+		public async ValueTask<Dictionary<string, string>?> GetPayloadDetailsAsync(ReadOnlyMemory<byte>? payload, CancellationToken cancellationToken = default)
+		{
+			Dictionary<string, string>? details = null;
+			if (payload != null)
+			{
+				Any basePayload = Any.Parser.ParseFrom(payload.Value.ToArray());
+				foreach (ITaskSource taskSource in _taskSources)
+				{
+					if (basePayload.Is(taskSource.Descriptor))
+					{
+						details = new Dictionary<string, string>();
+						await taskSource.GetLeaseDetailsAsync(basePayload, details, cancellationToken);
+						break;
+					}
+				}
+			}
+			return details;
+		}
+
+		/// <summary>
+		/// Issues a bearer token for the given session id
+		/// </summary>
+		/// <param name="agentId">The agent id</param>
+		/// <param name="sessionId">The session id</param>
+		/// <param name="cancellationToken">Cancellation token for the operation</param>
+		/// <returns>Bearer token for the agent</returns>
+		public async ValueTask<string> IssueSessionTokenAsync(AgentId agentId, SessionId sessionId, CancellationToken cancellationToken = default)
+		{
+			List<AclClaimConfig> claims = new List<AclClaimConfig>();
+			claims.Add(HordeClaims.AgentRoleClaim);
+			claims.Add(HordeClaims.GetAgentClaim(agentId));
+			claims.Add(HordeClaims.GetSessionClaim(sessionId));
+			return await _aclService.IssueBearerTokenAsync(claims, null, cancellationToken);
+		}
+
+		/// <summary>
+		/// Register a new agent
+		/// </summary>
+		/// <param name="name">Name of the agent</param>
+		/// <param name="ephemeral">Whether the agent is ephemeral or not</param>
+		/// <param name="enrollmentKey">Key for enrolling the agent</param>
+		/// <param name="cancellationToken">Cancellation token for the operation</param>
+		/// <returns>Unique id for the agent</returns>
+		public Task<IAgent> CreateAgentAsync(string name, bool ephemeral, string enrollmentKey, CancellationToken cancellationToken = default)
+		{
+			return CreateAgentAsync(new AgentId(name), ephemeral, enrollmentKey, cancellationToken);
+		}
+
+		/// <summary>
+		/// Register a new agent
+		/// </summary>
+		/// <param name="agentId">Agent id</param>
+		/// <param name="ephemeral">Whether the agent is ephemeral or not</param>
+		/// <param name="enrollmentKey">Key for enrolling the agent</param>
+		/// <param name="cancellationToken">Cancellation token for the operation</param>
+		/// <returns>Unique id for the agent</returns>
+		public async Task<IAgent> CreateAgentAsync(AgentId agentId, bool ephemeral, string enrollmentKey, CancellationToken cancellationToken = default)
+		{
+			for (; ; )
+			{
+				IAgent? agent = await Agents.GetAsync(agentId, cancellationToken);
+				if (agent == null)
+				{
+					return await Agents.AddAsync(agentId, ephemeral, enrollmentKey, cancellationToken);
+				}
+
+				agent = await agent.TryResetAsync(ephemeral, enrollmentKey, cancellationToken);
+				if (agent != null)
+				{
+					return agent;
+				}
+			}
+		}
+
+		/// <summary>
+		/// Gets an agent by ID
+		/// </summary>
+		/// <param name="agentId">Unique id of the agent</param>
+		/// <param name="cancellationToken">Cancellation token for the operation</param>
+		/// <returns>The agent document</returns>
+		public Task<IAgent?> GetAgentAsync(AgentId agentId, CancellationToken cancellationToken = default)
+		{
+			return Agents.GetAsync(agentId, cancellationToken);
+		}
+
+		/// <summary>
+		/// Finds all agents matching certain criteria
+		/// </summary>
+		/// <param name="poolId">The pool containing the agent</param>
+		/// <param name="modifiedAfter">If set, only returns agents modified after this time</param>
+		/// <param name="property">If set, only return agents matching this property</param>
+		/// <param name="includeDeleted">If set, include agents marked as deleted</param>
+		/// <param name="index">Index within the list of results</param>
+		/// <param name="count">Number of results to return</param>
+		/// <param name="cancellationToken">Cancellation token for the operation</param>
+		/// <returns>List of agents matching the given criteria</returns>
+		public Task<IReadOnlyList<IAgent>> FindAgentsAsync(PoolId? poolId, DateTime? modifiedAfter, string? property, bool includeDeleted, int? index, int? count, CancellationToken cancellationToken)
+		{
+			return Agents.FindAsync(poolId, modifiedAfter, property, null, null, includeDeleted, index, count, true, cancellationToken);
+		}
+		
+		/// <summary>
+		/// Get all agents from local in-memory cache
+		/// </summary>
+		/// <param name="cancellationToken">Cancellation token for the operation</param>
+		/// <returns>List of all agents</returns>
+		public async Task<IReadOnlyList<IAgent>> GetCachedAgentsAsync(CancellationToken cancellationToken)
+		{
+			if (_cachedAgents == null)
+			{
+				// First time this method is called, cached agents may not be initialized
+				// This could lead to multiple concurrent agent refreshes (dog-piling), but that's fine as it should only happen once
+				await RefreshCachedAgentsAsync(cancellationToken);
+			}
+
+			if (_cachedAgents == null)
+			{
+				throw new Exception("Cached agents not initialized");
+			}
+
+			return _cachedAgents.Values.ToList();
+		}
+
+		/// <summary>
+		/// Update the current workspaces for an agent.
+		/// </summary>
+		/// <param name="agent">The agent to update</param>
+		/// <param name="workspaces">Current list of workspaces</param>
+		/// <param name="pendingConform">Whether the agent still needs to run another conform</param>
+		/// <param name="cancellationToken">Cancellation token for the operation</param>
+		/// <returns>New agent state</returns>
+		public async Task<bool> TryUpdateWorkspacesAsync(IAgent agent, List<AgentWorkspaceInfo> workspaces, bool pendingConform, CancellationToken cancellationToken)
+		{
+			_ = this;
+
+			IAgent? newAgent = await agent.TryUpdateWorkspacesAsync(workspaces, pendingConform, cancellationToken);
+			return newAgent != null;
+		}
+
+		/// <summary>
+		/// Marks the agent as deleted
+		/// </summary>
+		/// <param name="agent">The agent to delete</param>
+		/// <param name="forceDelete">Whether to fully delete the agent as opposed to just marking it as deleted</param>
+		/// <param name="cancellationToken">Cancellation token for the operation</param>
+		/// <returns>Async task</returns>
+		public async Task DeleteAgentAsync(IAgent? agent, bool forceDelete = false, CancellationToken cancellationToken = default)
+		{
+			if (agent == null)
+			{
+				return;
+			}
+
+			if (forceDelete)
+			{
+				await Agents.ForceDeleteAsync(agent.Id, cancellationToken);
+				return;
+			}
+
+			while (agent is { Deleted: false })
+			{
+				IAgent? newAgent = await agent.TryDeleteAsync(cancellationToken);
+				if (newAgent != null)
+				{
+					break;
+				}
+				agent = await GetAgentAsync(agent.Id, cancellationToken);
+			}
+		}
+
+		async ValueTask<List<PoolId>> GetDynamicPoolsAsync(IAgent agent, CancellationToken cancellationToken)
+		{
+			List<PoolId> newDynamicPools = new List<PoolId>();
+
+			IReadOnlyList<IPoolConfig> pools = await _cachedPools.GetAsync(cancellationToken);
+			foreach (IPoolConfig pool in pools)
+			{
+				if (pool.Condition != null && agent.SatisfiesCondition(pool.Condition))
+				{
+					newDynamicPools.Add(pool.Id);
+				}
+			}
+
+			return newDynamicPools;
+		}
+
+		/// <summary>
+		/// Callback for an agents 
+		/// </summary>
+		/// <param name="agentId"></param>
+		void OnAgentUpdate(AgentId agentId)
+		{
+			lock (_waitingAgents)
+			{
+				if (_waitingAgents.TryGetValue(agentId, out CancellationTokenSource? cancellationSource))
+				{
+					cancellationSource.Cancel();
+				}
+			}
+		}
+
+		/// <summary>
+		/// Creates a new agent session
+		/// </summary>
+		/// <param name="agent">The agent to create a session for</param>
+		/// <param name="status">Current status of the agent</param>
+		/// <param name="properties">Properties for the agent</param>
+		/// <param name="resources">Resources which the agent has</param>
+		/// <param name="version">Version of the software that's running</param>
+		/// <param name="cancellationToken">Cancellation token for the operation</param>
+		/// <returns>New agent state</returns>
+		public async Task<IAgent> CreateSessionAsync(IAgent agent, AgentStatus status, IReadOnlyList<string> properties, IReadOnlyDictionary<string, int> resources, string? version, CancellationToken cancellationToken = default)
+		{
+			DateTime? lastStatusChange = null;
+			for (; ; )
+			{
+				IAuditLogChannel<AgentId> agentLogger = Agents.GetLogger(agent.Id);
+
+				// Check if there's already a session running for this agent.
+				IAgent? newAgent;
+				if (agent.SessionId != null)
+				{
+					// Save last status change timestamp to avoid registering the change to "stopped" when it's re-created immediately after to "ok".
+					lastStatusChange = agent.LastStatusChange;
+
+					// Try to terminate the current session
+					await TryTerminateSessionAsync(agent, cancellationToken);
+				}
+				else
+				{
+					DateTime utcNow = _clock.UtcNow;
+
+					// Remove any outstanding leases
+					foreach (AgentLease lease in agent.Leases)
+					{
+						agentLogger.LogInformation("Removing outstanding lease {LeaseId}", lease.Id);
+						await RemoveLeaseAsync(agent, lease, utcNow, LeaseOutcome.Failed, null, cancellationToken);
+					}
+
+					// Create a new session document
+					ISession newSession = await _sessions.AddAsync(SessionIdUtils.GenerateNewId(), agent.Id, _clock.UtcNow, properties, resources, version, cancellationToken);
+					DateTime sessionExpiresAt = utcNow + SessionExpiryTime;
+
+					// Get the new pools for the agent
+					List<PoolId> dynamicPools = await GetDynamicPoolsAsync(agent, cancellationToken);
+
+					// Reset the agent to use the new session
+					newAgent = await agent.TryCreateSessionAsync(new CreateSessionOptions(newSession.Id, sessionExpiresAt, status, properties, resources, dynamicPools, lastStatusChange ?? utcNow, version), cancellationToken);
+					if (newAgent != null)
+					{
+						LogPropertyChanges(agentLogger, agent.Properties, newAgent.Properties);
+						agent = newAgent;
+						agentLogger.LogInformation("Session {SessionId} started", newSession.Id);
+						break;
+					}
+
+					// Remove the session we didn't use
+					await _sessions.DeleteAsync(newSession.Id, cancellationToken);
+				}
+
+				// Get the current agent state
+				newAgent = await GetAgentAsync(agent.Id, cancellationToken);
+				if (newAgent == null)
+				{
+					throw new InvalidOperationException($"Invalid agent id '{agent.Id}'");
+				}
+				agent = newAgent;
+			}
+			return agent;
+		}
+
+		/// <summary>
+		/// Determines whether a task source can currently issue tasks
+		/// </summary>
+		bool CanUseTaskSource(IAgent agent, ITaskSource taskSource)
+		{
+			TaskSourceFlags flags = taskSource.Flags;
+			if ((flags & TaskSourceFlags.AllowWhenBusy) == 0 && agent.Status == AgentStatus.Busy)
+			{
+				return false;
+			}
+			if ((flags & TaskSourceFlags.AllowWhenDisabled) == 0 && !agent.Enabled)
+			{
+				return false;
+			}
+			if ((flags & TaskSourceFlags.AllowDuringDowntime) == 0 && _downtimeService.IsDowntimeActive)
+			{
+				return false;
+			}
+			return true;
+		}
+
+		/// <summary>
+		/// Compare properties and write changes to audit log
+		/// </summary>
+		private static void LogPropertyChanges(IAuditLogChannel<AgentId> agentLogger, IReadOnlyList<string> before, IReadOnlyList<string> after)
+		{
+			const string AwsInstanceTypeKey = KnownPropertyNames.AwsInstanceType + "=";
+			string beforeProp = before.FirstOrDefault(x => x.StartsWith(AwsInstanceTypeKey, StringComparison.Ordinal), String.Empty);
+			string afterProp = after.FirstOrDefault(x => x.StartsWith(AwsInstanceTypeKey, StringComparison.Ordinal), String.Empty);
+
+			if (!String.IsNullOrEmpty(beforeProp) && !String.IsNullOrEmpty(afterProp) && beforeProp != afterProp)
+			{
+				string oldInstanceType = beforeProp.Replace(AwsInstanceTypeKey, "", StringComparison.Ordinal);
+				string newInstanceType = afterProp.Replace(AwsInstanceTypeKey, "", StringComparison.Ordinal);
+				agentLogger.LogInformation("AWS EC2 instance type changed from {OldInstanceType} to {NewInstanceType}", oldInstanceType, newInstanceType);
+			}
+		}
+
+		/// <summary>
+		/// Waits for a lease to be assigned to an agent
+		/// </summary>
+		/// <param name="agent">The agent to assign a lease to</param>
+		/// <param name="newLeases">Leases that the agent knows about</param>
+		/// <param name="cancellationToken"></param>
+		/// <returns>True if a lease was assigned, false otherwise</returns>
+		public async Task<IAgent?> WaitForLeaseAsync(IAgent? agent, IList<HordeCommon.Rpc.Messages.RpcLease> newLeases, CancellationToken cancellationToken = default)
+		{
+			HashSet<string> knownLeases = new HashSet<string>(newLeases.Select(x => x.Id), StringComparer.OrdinalIgnoreCase);
+			while (agent != null && agent.Leases.All(x => knownLeases.Contains(x.Id.ToString())))
+			{
+				if (!agent.SessionExpiresAt.HasValue)
+				{
+					break;
+				}
+
+				// Check we have some time to wait
+				DateTime utcNow = _clock.UtcNow;
+				TimeSpan maxWaitTime = (agent.SessionExpiresAt.Value - SessionExpiryTime + SessionLongPollTime) - utcNow;
+				if (maxWaitTime <= TimeSpan.Zero)
+				{
+					break;
+				}
+
+				// Create a cancellation token that will expire with the session
+				using CancellationTokenSource cancellationSource = CancellationTokenSource.CreateLinkedTokenSource(_applicationLifetime.ApplicationStopping, cancellationToken);
+				cancellationSource.CancelAfter(maxWaitTime);
+
+				// Assign a new lease
+				(ITaskSource, AgentLease)? result = null;
+				try
+				{
+					// Add the cancellation source to the set of waiting agents
+					lock (_waitingAgents)
+					{
+						_waitingAgents[agent.Id] = cancellationSource;
+					}
+
+					// Create all the tasks to wait for
+					List<Task<(ITaskSource, AgentLease)?>> tasks = new List<Task<(ITaskSource, AgentLease)?>>();
+					foreach (ITaskSource taskSource in _taskSources)
+					{
+						if (CanUseTaskSource(agent, taskSource) && !cancellationSource.IsCancellationRequested)
+						{
+							Task<(ITaskSource, AgentLease)?> task = await GuardedAssignLeaseAsync(taskSource, agent, cancellationSource);
+							tasks.Add(task);
+						}
+					}
+
+					// If no task source is valid, just add a delay
+					if (tasks.Count == 0)
+					{
+						_logger.LogInformation("No task source valid for agent {AgentId}. Waiting {WaitTimeMs} ms", agent.Id, maxWaitTime.TotalMilliseconds);
+						await AsyncUtils.DelayNoThrow(maxWaitTime, cancellationToken);
+						break;
+					}
+
+					// Wait for all the tasks to complete. Once the first task completes it will set the cancellation source, triggering the 
+					// others to terminate.
+					await Task.WhenAll(tasks);
+
+					// Find the first result
+					foreach (Task<(ITaskSource, AgentLease)?> task in tasks)
+					{
+						(ITaskSource, AgentLease)? taskResult;
+						if (task.TryGetResult(out taskResult) && taskResult != null)
+						{
+							(ITaskSource taskSource, AgentLease taskLease) = taskResult.Value;
+							if (result == null)
+							{
+								result = (taskSource, taskLease);
+							}
+							else
+							{
+								await taskSource.CancelLeaseAsync(agent, taskLease.Id, Any.Parser.ParseFrom(taskLease.Payload), CancellationToken.None);
+							}
+						}
+					}
+				}
+				finally
+				{
+					lock (_waitingAgents)
+					{
+						_waitingAgents.Remove(agent.Id);
+					}
+				}
+
+				// Exit if we didn't find any work to do. It may be that all the task sources returned null, in which case wait for the time period to expire.
+				if (result == null)
+				{
+					if (!cancellationSource.IsCancellationRequested)
+					{
+						await cancellationSource.Token.AsTask();
+					}
+					break;
+				}
+
+				// Get the resulting lease
+				(ITaskSource source, AgentLease lease) = result.Value;
+
+				// Add the new lease to the agent
+				IAgent? newAgent = await agent.TryAddLeaseAsync(lease, cancellationToken);
+				if (newAgent != null)
+				{
+					await source.OnLeaseStartedAsync(newAgent, lease.Id, Any.Parser.ParseFrom(lease.Payload), Agents.GetLogger(agent.Id), CancellationToken.None);
+					await CreateLeaseAsync(agent, lease, CancellationToken.None);
+					return newAgent;
+				}
+				else
+				{
+					_logger.LogInformation("Failed adding lease {LeaseId} for agent {AgentId}", lease.Id.ToString(), agent.Id.ToString());
+					await source.CancelLeaseAsync(agent, lease.Id, Any.Parser.ParseFrom(lease.Payload), CancellationToken.None);
+				}
+
+				// Update the agent
+				agent = await GetAgentAsync(agent.Id, cancellationToken);
+			}
+			return agent;
+		}
+
+		async Task<Task<(ITaskSource, AgentLease)?>> GuardedAssignLeaseAsync(ITaskSource source, IAgent agent, CancellationTokenSource cancellationSource)
+		{
+			CancellationToken cancellationToken = cancellationSource.Token;
+			try
+			{
+				Task<AgentLease?> task = await source.AssignLeaseAsync(agent, cancellationToken);
+				return task.ContinueWith(x => WrapAssignedLease(source, x, cancellationSource), TaskScheduler.Default);
+			}
+			catch (OperationCanceledException) when (cancellationToken.IsCancellationRequested)
+			{
+				return Task.FromResult<(ITaskSource, AgentLease)?>(null);
+			}
+			catch (Exception ex)
+			{
+				_logger.LogError(ex, "Exception while trying to assign lease");
+				return Task.FromResult<(ITaskSource, AgentLease)?>(null);
+			}
+		}
+
+		(ITaskSource, AgentLease)? WrapAssignedLease(ITaskSource source, Task<AgentLease?> task, CancellationTokenSource cancellationSource)
+		{
+			if (task.IsCanceled)
+			{
+				return null;
+			}
+			else if (task.TryGetResult(out AgentLease? lease))
+			{
+				if (lease == null)
+				{
+					return null;
+				}
+				else
+				{
+					cancellationSource.Cancel();
+					return (source, lease);
+				}
+			}
+			else if (task.IsFaulted)
+			{
+				_logger.LogError(task.Exception, "Exception while trying to assign lease");
+				return null;
+			}
+			else
+			{
+				_logger.LogWarning("Unhandled task state: {Status}", task.Status);
+				return null;
+			}
+		}
+
+		/// <summary>
+		/// Cancels the specified agent lease
+		/// </summary>
+		/// <param name="agent">Agent to cancel the lease on</param>
+		/// <param name="leaseId">The lease id to cancel</param>
+		/// <param name="cancellationToken">Cancellation token for the operation</param>
+		/// <returns></returns>
+		public async Task<bool> CancelLeaseAsync(IAgent agent, LeaseId leaseId, CancellationToken cancellationToken)
+		{
+			_ = this;
+
+			int index = 0;
+			while (index < agent.Leases.Count && agent.Leases[index].Id != leaseId)
+			{
+				index++;
+			}
+
+			if (index == agent.Leases.Count)
+			{
+				return false;
+			}
+
+			if (agent.Leases[index].State == LeaseState.Cancelled)
+			{
+				return false;
+			}
+
+			await agent.TryCancelLeaseAsync(index, cancellationToken);
+			return true;
+		}
+
+		/// <summary>
+		/// 
+		/// </summary>
+		public async Task<IAgent?> UpdateSessionWithWaitAsync(IAgent inAgent, SessionId sessionId, AgentStatus status, IReadOnlyList<string>? properties, IReadOnlyDictionary<string, int>? resources, IList<HordeCommon.Rpc.Messages.RpcLease> newLeases, CancellationToken cancellationToken = default)
+		{
+			IAgent? agent = inAgent;
+
+			// Capture the current agent update index. This allows us to detect if anything has changed.
+			uint updateIndex = agent.UpdateIndex;
+
+			// Update the agent session and return to the caller if anything changes
+			agent = await UpdateSessionAsync(agent, sessionId, status, properties, resources, newLeases, cancellationToken);
+			if (agent != null && agent.UpdateIndex == updateIndex && (agent.Leases.Count > 0 || agent.Status != AgentStatus.Stopping))
+			{
+				agent = await WaitForLeaseAsync(agent, newLeases, cancellationToken);
+			}
+			return agent;
+		}
+
+		/// <summary>
+		/// Updates the state of the current agent session
+		/// </summary>
+		/// <param name="inAgent">The current agent state</param>
+		/// <param name="sessionId">Id of the session</param>
+		/// <param name="status">New status for the agent</param>
+		/// <param name="properties">New agent properties</param>
+		/// <param name="resources">New agent resources</param>
+		/// <param name="newLeases">New list of leases for this session</param>
+		/// <param name="cancellationToken">Cancellation token for the operation</param>
+		/// <returns>Updated agent state</returns>
+		public async Task<IAgent?> UpdateSessionAsync(IAgent inAgent, SessionId sessionId, AgentStatus status, IReadOnlyList<string>? properties, IReadOnlyDictionary<string, int>? resources, IList<HordeCommon.Rpc.Messages.RpcLease> newLeases, CancellationToken cancellationToken = default)
+		{
+			DateTime utcNow = _clock.UtcNow;
+
+			Stopwatch timer = Stopwatch.StartNew();
+
+			IAgent? agent = inAgent;
+			while (agent != null)
+			{
+				// Check the session id is correct.
+				if (agent.SessionId != sessionId)
+				{
+					if (status == AgentStatus.Stopping)
+					{
+						break; // Harmless; agent is not doing any work.
+					}
+					else
+					{
+						throw new InvalidOperationException($"Invalid agent session {sessionId}");
+					}
+				}
+
+				// Check the session hasn't expired
+				if (!agent.IsSessionValid(utcNow))
+				{
+					throw new InvalidOperationException("Session has already expired");
+				}
+
+				// Extend the current session time if we're within a time period of the current time expiring. This reduces
+				// unnecessary DB writes a little, but it also allows us to skip the update and jump into a long poll state
+				// if there's still time on the session left.
+				DateTime? sessionExpiresAt = null;
+				if (!agent.SessionExpiresAt.HasValue || utcNow > (agent.SessionExpiresAt - SessionExpiryTime) + SessionRenewTime)
+				{
+					sessionExpiresAt = utcNow + SessionExpiryTime;
+				}
+
+				// Flag for whether the leases array should be updated
+				bool updateLeases = false;
+				List<AgentLease> leases = new List<AgentLease>(agent.Leases);
+
+				// Remove any completed leases from the agent
+				Dictionary<LeaseId, RpcLease> leaseIdToNewState = newLeases.ToDictionary(x => LeaseId.Parse(x.Id), x => x);
+				for (int idx = 0; idx < leases.Count; idx++)
+				{
+					AgentLease lease = leases[idx];
+					if (lease.State == LeaseState.Cancelled)
+					{
+						RpcLease? newLease;
+						if (!leaseIdToNewState.TryGetValue(lease.Id, out newLease) || newLease.State == RpcLeaseState.Cancelled || newLease.State == RpcLeaseState.Completed)
+						{
+							await RemoveLeaseAsync(agent, lease, utcNow, LeaseOutcome.Cancelled, null, cancellationToken);
+							leases.RemoveAt(idx--);
+							updateLeases = true;
+						}
+					}
+					else
+					{
+						RpcLease? newLease;
+						if (leaseIdToNewState.TryGetValue(lease.Id, out newLease) && (LeaseState)newLease.State != lease.State)
+						{
+							if (newLease.State == RpcLeaseState.Cancelled || newLease.State == RpcLeaseState.Completed)
+							{
+								await RemoveLeaseAsync(agent, lease, utcNow, (LeaseOutcome)newLease.Outcome, newLease.Output.ToByteArray(), cancellationToken);
+								leases.RemoveAt(idx--);
+							}
+							else if (newLease.State == RpcLeaseState.Active && lease.State == LeaseState.Pending)
+							{
+								lease.State = LeaseState.Active;
+							}
+							updateLeases = true;
+						}
+					}
+				}
+
+				// If the agent is stopping, cancel all the leases. Clear out the current session once it's complete.
+				if (status == AgentStatus.Stopping || status == AgentStatus.Busy)
+				{
+					foreach (AgentLease lease in leases)
+					{
+						if (lease.State != LeaseState.Cancelled)
+						{
+							lease.State = LeaseState.Cancelled;
+							updateLeases = true;
+						}
+					}
+				}
+
+				// Get the new dynamic pools for the agent
+				List<PoolId> dynamicPools = await GetDynamicPoolsAsync(agent, cancellationToken);
+
+				// Update the agent, and try to create new lease documents if we succeed
+				IAgent? newAgent = await agent.TryUpdateSessionAsync(new UpdateSessionOptions(status, sessionExpiresAt, properties, resources, dynamicPools, updateLeases ? leases : null), cancellationToken);
+				if (newAgent != null)
+				{
+					agent = newAgent;
+					break;
+				}
+
+				// Fetch the agent again
+				agent = await GetAgentAsync(agent.Id, cancellationToken);
+			}
+
+			// If the agent is stopping, terminate the session
+			while (agent != null && agent.Status == AgentStatus.Stopping && agent.Leases.Count == 0)
+			{
+				IAgent? terminatedAgent = await TryTerminateSessionAsync(agent, cancellationToken);
+				if (terminatedAgent != null)
+				{
+					_logger.LogInformation("Terminated session {SessionId} for {AgentId}; agent is stopping", agent.SessionId, agent.Id);
+					agent = terminatedAgent;
+					break;
+				}
+				agent = await GetAgentAsync(agent.Id, cancellationToken);
+			}
+
+			return agent;
+		}
+
+		/// <summary>
+		/// Terminates an existing session. Does not update the agent itself, if it's currently 
+		/// </summary>
+		/// <param name="agent">The agent whose current session should be terminated</param>
+		/// <param name="cancellationToken">Cancellation token for the operation</param>
+		/// <returns>An up-to-date IAgent if the session was terminated</returns>
+		private async Task<IAgent?> TryTerminateSessionAsync(IAgent agent, CancellationToken cancellationToken = default)
+		{
+			// Make sure the agent has a valid session id
+			if (agent.SessionId == null)
+			{
+				return agent;
+			}
+
+			// Get the time that the session finishes at
+			DateTime finishTime = _clock.UtcNow;
+			if (agent.SessionExpiresAt.HasValue && agent.SessionExpiresAt.Value < finishTime)
+			{
+				finishTime = agent.SessionExpiresAt.Value;
+			}
+
+			// Save off the session id and current leases
+			SessionId sessionId = agent.SessionId.Value;
+			List<AgentLease> leases = new List<AgentLease>(agent.Leases);
+
+			// Clear the current session
+			IAgent? newAgent = await agent.TryTerminateSessionAsync(cancellationToken);
+			if (newAgent != null)
+			{
+				agent = newAgent;
+
+				// Remove any outstanding leases
+				foreach (AgentLease lease in leases)
+				{
+					Agents.GetLogger(agent.Id).LogInformation("Removing lease {LeaseId} during session terminate...", lease.Id);
+					await RemoveLeaseAsync(agent, lease, finishTime, LeaseOutcome.Failed, null, cancellationToken);
+				}
+
+				// Update the session document
+				Agents.GetLogger(agent.Id).LogInformation("Terminated session {SessionId}", sessionId);
+				await _sessions.UpdateAsync(sessionId, finishTime, agent.Properties, agent.Resources, cancellationToken);
+				return agent;
+			}
+			return null;
+		}
+
+		/// <summary>
+		/// Creates a new lease document
+		/// </summary>
+		/// <param name="agent">Agent that will be executing the lease</param>
+		/// <param name="agentLease">The new agent lease</param>
+		/// <param name="cancellationToken">Cancellation token for the operation</param>
+		/// <returns>New lease document</returns>
+		internal Task<ILease> CreateLeaseAsync(IAgent agent, AgentLease agentLease, CancellationToken cancellationToken = default)
+		{
+			try
+			{
+				return _leases.AddAsync(agentLease.Id, agentLease.ParentId, agentLease.Name, agent.Id, agent.SessionId!.Value, agentLease.StreamId, agentLease.PoolId, agentLease.LogId, agentLease.StartTime, agentLease.Payload!, cancellationToken);
+			}
+			catch (Exception ex)
+			{
+				_logger.LogError(ex, "Unable to create lease {LeaseId} for agent {AgentId}; lease already exists?", agentLease.Id, agent.Id);
+				throw;
+			}
+		}
+
+		/// <summary>
+		/// Finds all leases matching a set of criteria
+		/// </summary>
+		/// <param name="agentId">Unqiue id of the agent executing this lease</param>
+		/// <param name="sessionId">Unique id of the agent session</param>
+		/// <param name="startTime">Start of the search window to return results for</param>
+		/// <param name="finishTime">End of the search window to return results for</param>
+		/// <param name="index">Index of the first result to return</param>
+		/// <param name="count">Number of results to return</param>
+		/// <param name="cancellationToken">Cancellation token for the operation</param>
+		/// <returns>List of leases matching the given criteria</returns>
+		public Task<IReadOnlyList<ILease>> FindLeasesAsync(AgentId? agentId, SessionId? sessionId, DateTime? startTime, DateTime? finishTime, int index, int count, CancellationToken cancellationToken = default)
+		{
+			return _leases.FindLeasesAsync(null, agentId, sessionId, startTime, finishTime, index, count, cancellationToken: cancellationToken);
+		}
+
+		/// <summary>
+		/// Finds all leases by finish time
+		/// </summary>
+		/// <param name="minFinishTime">Start of the search window to return results for</param>
+		/// <param name="maxFinishTime">End of the search window to return results for</param>
+		/// <param name="index">Index of the first result to return</param>
+		/// <param name="count">Number of results to return</param>
+		/// <param name="cancellationToken">Cancellation token for the operation</param>
+		/// <returns>List of leases matching the given criteria</returns>
+		public Task<IReadOnlyList<ILease>> FindLeasesByFinishTimeAsync(DateTime? minFinishTime, DateTime? maxFinishTime, int? index, int? count, CancellationToken cancellationToken = default)
+		{
+			return _leases.FindLeasesByFinishTimeAsync(minFinishTime, maxFinishTime, index, count, null, false, cancellationToken);
+		}
+
+		/// <summary>
+		/// Gets a specific lease
+		/// </summary>
+		/// <param name="leaseId">Unique id of the lease</param>
+		/// <param name="cancellationToken"></param>
+		/// <returns>The lease that was found, or null if it does not exist</returns>
+		public Task<ILease?> GetLeaseAsync(LeaseId leaseId, CancellationToken cancellationToken = default)
+		{
+			return _leases.GetAsync(leaseId, cancellationToken);
+		}
+
+		/// <summary>
+		/// Removes a lease with the given id. Updates the lease state in the database, and removes the item from the agent's leases array.
+		/// </summary>
+		/// <param name="agent">The agent to remove a lease from</param>
+		/// <param name="lease">The lease to cancel</param>
+		/// <param name="utcNow">The current time</param>
+		/// <param name="outcome">Final status of the lease</param>
+		/// <param name="output">Output from executing the task</param>
+		/// <param name="cancellationToken">Cancellation token for the operation</param>
+		/// <returns>Async task</returns>
+		private async Task RemoveLeaseAsync(IAgent agent, AgentLease lease, DateTime utcNow, LeaseOutcome outcome, byte[]? output, CancellationToken cancellationToken = default)
+		{
+			// Make sure the lease is terminated correctly
+			if (lease.Payload == null)
+			{
+				_logger.LogWarning("Removing lease {LeaseId} (no payload)", lease.Id);
+			}
+			else
+			{
+				Any any = Any.Parser.ParseFrom(lease.Payload);
+				_logger.LogInformation("Removing lease {LeaseId} ({LeaseType})", lease.Id, any.TypeUrl);
+
+				foreach (ITaskSource taskSource in _taskSources)
+				{
+					if (any.Is(taskSource.Descriptor))
+					{
+						await taskSource.OnLeaseFinishedAsync(agent, lease.Id, any, outcome, output, Agents.GetLogger(agent.Id), cancellationToken);
+						break;
+					}
+				}
+			}
+
+			// Figure out what time the lease finished
+			DateTime finishTime = utcNow;
+			if (agent.SessionExpiresAt.HasValue && agent.SessionExpiresAt.Value < finishTime)
+			{
+				finishTime = agent.SessionExpiresAt.Value;
+			}
+			if (lease.ExpiryTime.HasValue && lease.ExpiryTime.Value < finishTime)
+			{
+				finishTime = lease.ExpiryTime.Value;
+			}
+
+			// Update the lease
+			await _leases.TrySetOutcomeAsync(lease.Id, finishTime, outcome, output, cancellationToken);
+
+			// Temporarily disabling due to gRPC timeouts.
+#if false
+			// Terminate any child leases
+			List<ILease> childLeases = await _leases.FindLeasesAsync(parentId: lease.Id);
+			foreach (ILease childLease in childLeases)
+			{
+				if (childLease.Outcome == LeaseOutcome.Unspecified)
+				{
+					IAgent? otherAgent = await GetAgentAsync(childLease.AgentId);
+					if (otherAgent != null)
+					{
+						AgentLease? otherLease = otherAgent.Leases.FirstOrDefault(x => x.Id == childLease.Id);
+						if (otherLease != null)
+						{
+							_logger.LogInformation("Terminating child lease {LeaseId} with parent {ParentLeaseId}", otherLease.Id, lease.Id);
+							await RemoveLeaseAsync(otherAgent, otherLease, utcNow, LeaseOutcome.Failed, null);
+						}
+					}
+				}
+			}
+#endif
+		}
+
+		/// <summary>
+		/// Gets information about a particular session
+		/// </summary>
+		/// <param name="sessionId">The unique session id</param>
+		/// <returns>The session information</returns>
+		public Task<ISession?> GetSessionAsync(SessionId sessionId)
+		{
+			return _sessions.GetAsync(sessionId);
+		}
+
+		/// <summary>
+		/// Find sessions for the given agent
+		/// </summary>
+		/// <param name="agentId">The unique agent id</param>
+		/// <param name="startTime">Start time to include in the search</param>
+		/// <param name="finishTime">Finish time to include in the search</param>
+		/// <param name="index">Index of the first result to return</param>
+		/// <param name="count">Number of results to return</param>
+		/// <returns>List of sessions matching the given criteria</returns>
+		public Task<List<ISession>> FindSessionsAsync(AgentId agentId, DateTime? startTime, DateTime? finishTime, int index, int count)
+		{
+			return _sessions.FindAsync(agentId, startTime, finishTime, index, count);
+		}
+
+		/// <summary>
+		/// Update the rate table for different agent types
+		/// </summary>
+		/// <param name="entries">New entries for the rate table</param>
+		/// <returns></returns>
+		public async Task UpdateRateTableAsync(List<AgentRateConfig> entries)
+		{
+			await _redisService.GetDatabase().StringSetAsync(_agentRateTableData, new AgentRateTable { Entries = entries });
+		}
+
+		/// <summary>
+		/// Gets the rate for the given agent
+		/// </summary>
+		/// <param name="agentId">Agent id to query</param>
+		/// <param name="cancellationToken">Cancellation token for the operation</param>
+		/// <returns>Hourly rate of running the given agent</returns>
+		public async ValueTask<double?> GetRateAsync(AgentId agentId, CancellationToken cancellationToken = default)
+		{
+			RedisKey key = $"agent-rate/{agentId}";
+
+			// Try to get the current value
+			RedisValue value = await _redisService.GetDatabase().StringGetAsync(key).WaitAsync(cancellationToken);
+			if (!value.IsNull)
+			{
+				double rate = (double)value;
+				if (rate == 0.0)
+				{
+					return null;
+				}
+				else
+				{
+					return rate;
+				}
+			}
+			else
+			{
+				double rate = 0.0;
+
+				// Get the rate table
+				AgentRateTable? rateTable = await _cachedRates.GetAsync(cancellationToken);
+				if (rateTable != null && rateTable.Entries.Count > 0)
+				{
+					IAgent? agent = await GetAgentAsync(agentId, cancellationToken);
+					if (agent != null)
+					{
+						foreach (AgentRateConfig config in rateTable.Entries)
+						{
+							if (config.Condition != null && config.Condition.Evaluate(x => agent.GetPropertyValues(x)))
+							{
+								rate = config.Rate;
+								break;
+							}
+						}
+					}
+				}
+
+				// Cache it for future reference
+				await _redisService.GetDatabase().StringSetAsync(key, rate, TimeSpan.FromMinutes(5.0), flags: CommandFlags.FireAndForget).WaitAsync(cancellationToken);
+				return rate;
+			}
+		}
+
+		internal async ValueTask TickAsync(CancellationToken stoppingToken)
+		{
+			using TelemetrySpan span = _tracer.StartActiveSpan($"{nameof(AgentService)}.{nameof(TickAsync)}");
+			await RefreshCachedAgentsAsync(stoppingToken);
+			await CollectMetricsAsync(stoppingToken);
+		}
+
+		internal async ValueTask TickSharedAsync(CancellationToken stoppingToken)
+		{
+			using TelemetrySpan span = _tracer.StartActiveSpan($"{nameof(AgentService)}.{nameof(TickSharedAsync)}");
+
+			await TerminateExpiredSessionsAsync(stoppingToken);
+			await DeleteExpiredEphemeralAgentsAsync(stoppingToken);
+		}
+
+		private async Task TerminateExpiredSessionsAsync(CancellationToken cancellationToken)
+		{
+			using TelemetrySpan span = _tracer.StartActiveSpan($"{nameof(AgentService)}.{nameof(TerminateExpiredSessionsAsync)}");
+
+			int c = 0;
+			while (!cancellationToken.IsCancellationRequested)
+			{
+				// Find all the agents which are ready to be expired
+				const int MaxAgents = 100;
+				DateTime utcNow = _clock.UtcNow;
+				IReadOnlyList<IAgent> expiredAgents = await Agents.FindExpiredAsync(utcNow, MaxAgents, cancellationToken);
+
+				// Transition each agent to being offline
+				foreach (IAgent expiredAgent in expiredAgents)
+				{
+					cancellationToken.ThrowIfCancellationRequested();
+					_logger.LogDebug("Terminating session {SessionId} for agent {Agent}", expiredAgent.SessionId, expiredAgent.Id);
+					await TryTerminateSessionAsync(expiredAgent, cancellationToken);
+				}
+				c += expiredAgents.Count;
+
+				// Try again if we didn't fetch everything
+				if (expiredAgents.Count < MaxAgents)
+				{
+					break;
+				}
+			}
+			span.SetAttribute("NumAgentsTerminated", c);
+		}
+
+		private async Task DeleteExpiredEphemeralAgentsAsync(CancellationToken cancellationToken)
+		{
+			using TelemetrySpan span = _tracer.StartActiveSpan($"{nameof(AgentService)}.{nameof(DeleteExpiredEphemeralAgentsAsync)}");
+			int c = 0;
+			foreach (IAgent agent in await Agents.FindDeletedAsync(cancellationToken))
+			{
+				cancellationToken.ThrowIfCancellationRequested();
+				bool noStatusChangeDuringPeriod = _clock.UtcNow > agent.LastStatusChange + TimeSpan.FromHours(1);
+				if (agent is { Status: AgentStatus.Stopped, Ephemeral: true } && noStatusChangeDuringPeriod)
+				{
+					_logger.LogDebug("Deleting ephemeral agent {Agent}", agent.Id);
+					await DeleteAgentAsync(agent, true, cancellationToken);
+					c++;
+				}
+			}
+
+			span.SetAttribute("NumAgentsDeleted", c);
+		}
+
+		/// <summary>
+		/// Updates the in-memory cache of current agents
+		/// Called by ticker to avoid blocking reads. Updates in the background, in favor of slightly more stale agents
+		/// </summary>
+		/// <param name="cancellationToken">Cancellation token</param>
+		private async Task RefreshCachedAgentsAsync(CancellationToken cancellationToken = default)
+		{
+			using TelemetrySpan span = _tracer.StartActiveSpan($"{nameof(AgentService)}.{nameof(RefreshCachedAgentsAsync)}");
+			Dictionary<AgentId, IAgent> agents = new();
+			IReadOnlyList<IAgent> agentList = await Agents.FindAsync(consistentRead: false, cancellationToken: cancellationToken);
+			foreach (IAgent agent in agentList)
+			{
+				agents[agent.Id] = agent;
+			}
+			
+			Interlocked.Exchange(ref _cachedAgents, agents);
+		}
+
+		private async Task CollectMetricsAsync(CancellationToken cancellationToken = default)
+		{
+			using TelemetrySpan span = _tracer.StartActiveSpan($"{nameof(AgentService)}.{nameof(CollectMetricsAsync)}");
+
+			IReadOnlyList<IAgent> agentList = await GetCachedAgentsAsync(cancellationToken);
+			int numAgentsTotal = agentList.Count;
+			int numAgentsTotalDeleted = agentList.Count(a => a.Deleted);
+			int numAgentsTotalEnabled = agentList.Count(a => a.Enabled);
+			int numAgentsTotalDisabled = agentList.Count(a => !a.Enabled);
+			int numAgentsTotalOk = agentList.Count(a => a.Enabled && a.Status == AgentStatus.Ok);
+			int numAgentsTotalStopping = agentList.Count(a => a.Enabled && a.Status == AgentStatus.Stopping);
+			int numAgentsTotalUnhealthy = agentList.Count(a => a.Enabled && a.Status == AgentStatus.Unhealthy);
+			int numAgentsTotalBusy = agentList.Count(a => a.Enabled && a.Status == AgentStatus.Busy);
+			int numAgentsTotalUnspecified = agentList.Count(a => a.Enabled && a.Status == AgentStatus.Unspecified);
+
+			List<Measurement<int>> newMeasurements =
+			[
+				new Measurement<int>(numAgentsTotal),
+				new Measurement<int>(numAgentsTotalDeleted, new KeyValuePair<string, object?>("status", "deleted")),
+				new Measurement<int>(numAgentsTotalEnabled, new KeyValuePair<string, object?>("status", "enabled")),
+				new Measurement<int>(numAgentsTotalDisabled, new KeyValuePair<string, object?>("status", "disabled")),
+				new Measurement<int>(numAgentsTotalOk, new KeyValuePair<string, object?>("status", "ok")),
+				new Measurement<int>(numAgentsTotalStopping, new KeyValuePair<string, object?>("status", "stopping")),
+				new Measurement<int>(numAgentsTotalUnhealthy, new KeyValuePair<string, object?>("status", "unhealthy")),
+				new Measurement<int>(numAgentsTotalBusy, new KeyValuePair<string, object?>("status", "paused")),
+				new Measurement<int>(numAgentsTotalUnspecified, new KeyValuePair<string, object?>("status", "unspecified"))
+			];
+
+			_measurements = newMeasurements;
+		}
+
+		/// <summary>
+		/// Determines if the user is authorized to perform an action on a particular agent
+		/// </summary>
+		/// <param name="agent">The agent to check</param>
+		/// <param name="user">The principal to authorize</param>
+		/// <param name="reason">Reason for being authorized or not</param>
+		/// <returns>True if the action is authorized</returns>
+		public bool AuthorizeSession(IAgent agent, ClaimsPrincipal user, out string reason)
+		{
+			if (agent.SessionId == null)
+			{
+				reason = $"{nameof(agent.SessionId)} is null";
+				return false;
+			}
+
+			if (!user.HasSessionClaim(agent.SessionId.Value))
+			{
+				reason = $"Missing session claim for {agent.SessionId.Value}";
+				return false;
+			}
+
+			if (!agent.IsSessionValid(_clock.UtcNow))
+			{
+				reason = $"Session has expired";
+				return false;
+			}
+
+			reason = "Session is valid";
+			return true;
+		}
+	}
+}
diff --git a/Engine/Source/Programs/Horde/Plugins/Compute/HordeServer.Compute/Agents/IAgentCollection.cs b/Engine/Source/Programs/Horde/Plugins/Compute/HordeServer.Compute/Agents/IAgentCollection.cs
new file mode 100644
--- /dev/null
+++ b/Engine/Source/Programs/Horde/Plugins/Compute/HordeServer.Compute/Agents/IAgentCollection.cs
@@ -0,0 +1,121 @@
+// Copyright Epic Games, Inc. All Rights Reserved.
+
+using System;
+using System.Collections.Generic;
+using System.Threading;
+using System.Threading.Tasks;
+using EpicGames.Horde.Agents;
+using EpicGames.Horde.Agents.Leases;
+using EpicGames.Horde.Agents.Pools;
+using HordeServer.Auditing;
+
+#pragma warning disable CA1716 // rename parameter property so that it no longer conflicts with the reserved language keyword 'Property'
+
+namespace HordeServer.Agents
+{
+	/// <summary>
+	/// Interface for a collection of agent documents
+	/// </summary>
+	public interface IAgentCollection
+	{
+		/// <summary>
+		/// Adds a new agent with the given properties
+		/// </summary>
+		/// <param name="id">Id for the new agent</param>
+		/// <param name="ephemeral">Whether the agent is ephemeral or not</param>
+		/// <param name="enrollmentKey">Key used to identify a unique enrollment for the agent with this id</param>
+		/// <param name="cancellationToken">Cancellation token for the operation</param>
+		Task<IAgent> AddAsync(AgentId id, bool ephemeral, string enrollmentKey, CancellationToken cancellationToken = default);
+
+		/// <summary>
+		/// Deletes an agent
+		/// </summary>
+		/// <param name="agentId">Unique id of the agent</param>
+		/// <param name="cancellationToken">Cancellation token for the operation</param>
+		/// <returns>Async task</returns>
+		Task ForceDeleteAsync(AgentId agentId, CancellationToken cancellationToken = default);
+
+		/// <summary>
+		/// Gets an agent by ID
+		/// </summary>
+		/// <param name="agentId">Unique id of the agent</param>
+		/// <param name="cancellationToken">Cancellation token for the operation</param>
+		/// <returns>The agent document</returns>
+		Task<IAgent?> GetAsync(AgentId agentId, CancellationToken cancellationToken = default);
+
+		/// <summary>
+		/// Gets multiple agents by ID
+		/// </summary>
+		/// <param name="agentIds">List of unique IDs of the agents</param>
+		/// <param name="cancellationToken">Cancellation token for the operation</param>
+		/// <returns>The agent documents</returns>
+		Task<IReadOnlyList<IAgent>> GetManyAsync(List<AgentId> agentIds, CancellationToken cancellationToken = default);
+
+		/// <summary>
+		/// Finds all agents matching certain criteria
+		/// </summary>
+		/// <param name="poolId">The pool ID in string form containing the agent</param>
+		/// <param name="modifiedAfter">If set, only returns agents modified after this time</param>
+		/// <param name="property">Property to look for</param>
+		/// <param name="status">Status to look for</param>
+		/// <param name="enabled">Enabled/disabled status to look for</param>
+		/// <param name="includeDeleted">Whether agents marked as deleted should be included</param>
+		/// <param name="index">Index of the first result</param>
+		/// <param name="count">Number of results to return</param>
+		/// <param name="consistentRead">If the database read should be made to the replica server</param>
+		/// <param name="cancellationToken">Cancellation token for the operation</param>
+		/// <returns>List of agents matching the given criteria</returns>
+		Task<IReadOnlyList<IAgent>> FindAsync(PoolId? poolId = null, DateTime? modifiedAfter = null, string? property = null, AgentStatus? status = null, bool? enabled = null, bool includeDeleted = false, int? index = null, int? count = null, bool consistentRead = true, CancellationToken cancellationToken = default);
+
+		/// <summary>
+		/// Finds all agents with an expired session
+		/// </summary>
+		/// <param name="utcNow">The current time</param>
+		/// <param name="maxAgents">Maximum number of agents to return</param>
+		/// <param name="cancellationToken">Cancellation token for the operation</param>
+		/// <returns>List of agents</returns>
+		Task<IReadOnlyList<IAgent>> FindExpiredAsync(DateTime utcNow, int maxAgents, CancellationToken cancellationToken = default);
+
+		/// <summary>
+		/// Finds all agents marked as deleted
+		/// </summary>
+		/// <param name="cancellationToken">Cancellation token for the operation</param>
+		/// <returns>List of agents</returns>
+		Task<IReadOnlyList<IAgent>> FindDeletedAsync(CancellationToken cancellationToken = default);
+
+		/// <summary>
+		/// Finds all active agent lease IDs
+		/// </summary>
+		/// <param name="cancellationToken">Cancellation token for the operation</param>
+		/// <returns>List of agent lease IDs</returns>
+		Task<List<LeaseId>> FindActiveLeaseIdsAsync(CancellationToken cancellationToken = default);
+
+		/// <summary>
+		/// Get all child lease IDs
+		/// </summary>
+		/// <param name="id">Lease ID</param>
+		/// <param name="cancellationToken">Cancellation token for the operation</param>
+		/// <returns>List of agent lease IDs</returns>
+		Task<List<LeaseId>> GetChildLeaseIdsAsync(LeaseId id, CancellationToken cancellationToken = default);
+
+		/// <summary>
+		/// Gets the log channel for an agent
+		/// </summary>
+		/// <param name="agentId"></param>
+		/// <returns></returns>
+		IAuditLogChannel<AgentId> GetLogger(AgentId agentId);
+
+		/// <summary>
+		/// Sends a notification that an update event has ocurred
+		/// </summary>
+		/// <param name="agentId">Agent that has been updated</param>
+		Task PublishUpdateEventAsync(AgentId agentId);
+
+		/// <summary>
+		/// Subscribe to notifications on agent states being updated
+		/// </summary>
+		/// <param name="onUpdate">Callback for updates</param>
+		/// <returns>Disposable subscription object</returns>
+		Task<IAsyncDisposable> SubscribeToUpdateEventsAsync(Action<AgentId> onUpdate);
+	}
+}
diff --git a/Engine/Source/Programs/Horde/Plugins/Compute/HordeServer.Compute/Agents/Sessions/ISessionCollection.cs b/Engine/Source/Programs/Horde/Plugins/Compute/HordeServer.Compute/Agents/Sessions/ISessionCollection.cs
new file mode 100644
--- /dev/null
+++ b/Engine/Source/Programs/Horde/Plugins/Compute/HordeServer.Compute/Agents/Sessions/ISessionCollection.cs
@@ -0,0 +1,77 @@
+// Copyright Epic Games, Inc. All Rights Reserved.
+
+using System;
+using System.Collections.Generic;
+using System.Threading;
+using System.Threading.Tasks;
+using EpicGames.Horde.Agents;
+using EpicGames.Horde.Agents.Sessions;
+
+namespace HordeServer.Agents.Sessions
+{
+	/// <summary>
+	/// Interface for a collection of session documents
+	/// </summary>
+	public interface ISessionCollection
+	{
+		/// <summary>
+		/// Adds a new session to the collection
+		/// </summary>
+		/// <param name="id">The session id</param>
+		/// <param name="agentId">The agent this session is for</param>
+		/// <param name="startTime">Start time of this session</param>
+		/// <param name="properties">Properties of this agent at the time the session started</param>
+		/// <param name="resources">Resources which the agent has</param>
+		/// <param name="version">Version of the agent software</param>
+		/// <param name="cancellationToken">Cancellation token for the operation</param>
+		Task<ISession> AddAsync(SessionId id, AgentId agentId, DateTime startTime, IReadOnlyList<string>? properties, IReadOnlyDictionary<string, int>? resources, string? version, CancellationToken cancellationToken = default);
+
+		/// <summary>
+		/// Gets information about a particular session
+		/// </summary>
+		/// <param name="sessionId">The unique session id</param>
+		/// <param name="cancellationToken">Cancellation token for the operation</param>
+		/// <returns>The session information</returns>
+		Task<ISession?> GetAsync(SessionId sessionId, CancellationToken cancellationToken = default);
+
+		/// <summary>
+		/// Find sessions for the given agent
+		/// </summary>
+		/// <param name="agentId">The unique agent id</param>
+		/// <param name="startTime">Start time to include in the search</param>
+		/// <param name="finishTime">Finish time to include in the search</param>
+		/// <param name="index">Index of the first result to return</param>
+		/// <param name="count">Number of results to return</param>
+		/// <param name="cancellationToken">Cancellation token for the operation</param>
+		/// <returns>List of sessions matching the given criteria</returns>
+		Task<List<ISession>> FindAsync(AgentId agentId, DateTime? startTime, DateTime? finishTime, int index, int count, CancellationToken cancellationToken = default);
+
+		/// <summary>
+		/// Finds any active sessions
+		/// </summary>
+		/// <param name="index">Index of the first result to return</param>
+		/// <param name="count">Number of results to return</param>
+		/// <param name="cancellationToken">Cancellation token for the operation</param>
+		/// <returns>List of sessions</returns>
+		Task<List<ISession>> FindActiveSessionsAsync(int? index = null, int? count = null, CancellationToken cancellationToken = default);
+
+		/// <summary>
+		/// Update a session from the collection
+		/// </summary>
+		/// <param name="sessionId">The session to update</param>
+		/// <param name="finishTime">Time at which the session finished</param>
+		/// <param name="properties">The agent properties</param>
+		/// <param name="resources">Resources which the agent has</param>
+		/// <param name="cancellationToken">Cancellation token for the operation</param>
+		/// <returns>Async task</returns>
+		Task UpdateAsync(SessionId sessionId, DateTime? finishTime, IReadOnlyList<string>? properties, IReadOnlyDictionary<string, int>? resources, CancellationToken cancellationToken = default);
+
+		/// <summary>
+		/// Delete a session from the collection
+		/// </summary>
+		/// <param name="sessionId">The session id</param>
+		/// <param name="cancellationToken">Cancellation token for the operation</param>
+		/// <returns>Async task</returns>
+		Task DeleteAsync(SessionId sessionId, CancellationToken cancellationToken = default);
+	}
+}
diff --git a/Engine/Source/Programs/Horde/Plugins/Compute/HordeServer.Compute/Agents/Sessions/SessionCollection.cs b/Engine/Source/Programs/Horde/Plugins/Compute/HordeServer.Compute/Agents/Sessions/SessionCollection.cs
new file mode 100644
--- /dev/null
+++ b/Engine/Source/Programs/Horde/Plugins/Compute/HordeServer.Compute/Agents/Sessions/SessionCollection.cs
@@ -0,0 +1,150 @@
+// Copyright Epic Games, Inc. All Rights Reserved.
+
+using System;
+using System.Collections.Generic;
+using System.Linq;
+using System.Threading;
+using System.Threading.Tasks;
+using EpicGames.Horde.Agents;
+using EpicGames.Horde.Agents.Sessions;
+using HordeServer.Server;
+using HordeServer.Utilities;
+using MongoDB.Bson.Serialization.Attributes;
+using MongoDB.Driver;
+
+namespace HordeServer.Agents.Sessions
+{
+	/// <summary>
+	/// Collection of session documents
+	/// </summary>
+	public class SessionCollection : ISessionCollection
+	{
+		/// <summary>
+		/// Concrete implementation of ISession
+		/// </summary>
+		class SessionDocument : ISession
+		{
+			[BsonRequired, BsonId]
+			public SessionId Id { get; set; }
+
+			[BsonRequired]
+			public AgentId AgentId { get; set; }
+
+			public DateTime StartTime { get; set; }
+			public DateTime? FinishTime { get; set; }
+			public List<string> Properties { get; set; } = new List<string>();
+			public Dictionary<string, int> Resources { get; set; } = new Dictionary<string, int>();
+			public string Version { get; set; } = String.Empty;
+
+			IReadOnlyList<string> ISession.Properties => Properties;
+			IReadOnlyDictionary<string, int> ISession.Resources => Resources;
+
+			[BsonConstructor]
+			private SessionDocument()
+			{
+			}
+
+			public SessionDocument(SessionId id, AgentId agentId, DateTime startTime, IReadOnlyList<string>? properties, IReadOnlyDictionary<string, int>? resources, string? version)
+			{
+				Id = id;
+				AgentId = agentId;
+				StartTime = startTime;
+				if (properties != null)
+				{
+					Properties = new List<string>(properties);
+				}
+				if (resources != null)
+				{
+					Resources = new Dictionary<string, int>(resources);
+				}
+				if (version != null)
+				{
+					Version = version;
+				}
+			}
+		}
+
+		/// <summary>
+		/// Collection of session documents
+		/// </summary>
+		readonly IMongoCollection<SessionDocument> _sessions;
+
+		/// <summary>
+		/// Constructor
+		/// </summary>
+		/// <param name="mongoService">The database service</param>
+		public SessionCollection(IMongoService mongoService)
+		{
+			List<MongoIndex<SessionDocument>> indexes = new List<MongoIndex<SessionDocument>>();
+			indexes.Add(keys => keys.Ascending(x => x.AgentId).Ascending(x => x.StartTime).Ascending(x => x.FinishTime));
+			indexes.Add(keys => keys.Ascending(x => x.FinishTime));
+
+			_sessions = mongoService.GetCollection<SessionDocument>("Sessions", indexes);
+		}
+
+		/// <inheritdoc/>
+		public async Task<ISession> AddAsync(SessionId id, AgentId agentId, DateTime startTime, IReadOnlyList<string>? properties, IReadOnlyDictionary<string, int>? resources, string? version, CancellationToken cancellationToken = default)
+		{
+			SessionDocument newSession = new SessionDocument(id, agentId, startTime, properties, resources, version);
+			await _sessions.InsertOneAsync(newSession, null, cancellationToken);
+			return newSession;
+		}
+
+		/// <inheritdoc/>
+		public async Task<ISession?> GetAsync(SessionId sessionId, CancellationToken cancellationToken = default)
+		{
+			return await _sessions.Find(x => x.Id == sessionId).FirstOrDefaultAsync(cancellationToken);
+		}
+
+		/// <inheritdoc/>
+		public async Task<List<ISession>> FindAsync(AgentId agentId, DateTime? startTime, DateTime? finishTime, int index, int count, CancellationToken cancellationToken = default)
+		{
+			FilterDefinitionBuilder<SessionDocument> filterBuilder = Builders<SessionDocument>.Filter;
+
+			FilterDefinition<SessionDocument> filter = filterBuilder.Eq(x => x.AgentId, agentId);
+			if (startTime != null)
+			{
+				filter &= filterBuilder.Gte(x => x.StartTime, startTime.Value);
+			}
+			if (finishTime != null)
+			{
+				filter &= filterBuilder.Or(filterBuilder.Eq(x => x.FinishTime, null), filterBuilder.Lte(x => x.FinishTime, finishTime.Value));
+			}
+
+			List<SessionDocument> results = await _sessions.Find(filter).SortByDescending(x => x.StartTime).Skip(index).Limit(count).ToListAsync(cancellationToken);
+			return results.ConvertAll<ISession>(x => x);
+		}
+
+		/// <inheritdoc/>
+		public async Task<List<ISession>> FindActiveSessionsAsync(int? index, int? count, CancellationToken cancellationToken = default)
+		{
+			List<SessionDocument> results = await _sessions.Find(x => x.FinishTime == null).Range(index, count).ToListAsync(cancellationToken);
+			return results.ConvertAll<ISession>(x => x);
+		}
+
+		/// <inheritdoc/>
+		public Task UpdateAsync(SessionId sessionId, DateTime? finishTime, IReadOnlyList<string>? properties, IReadOnlyDictionary<string, int>? resources, CancellationToken cancellationToken = default)
+		{
+			List<UpdateDefinition<SessionDocument>> updates = new List<UpdateDefinition<SessionDocument>>();
+			if (finishTime != null)
+			{
+				updates.Add(Builders<SessionDocument>.Update.Set(x => x.FinishTime, finishTime));
+			}
+			if (properties != null)
+			{
+				updates.Add(Builders<SessionDocument>.Update.Set(x => x.Properties, new List<string>(properties)));
+			}
+			if (resources != null)
+			{
+				updates.Add(Builders<SessionDocument>.Update.Set(x => x.Resources, new Dictionary<string, int>(resources)));
+			}
+			return _sessions.FindOneAndUpdateAsync(x => x.Id == sessionId, Builders<SessionDocument>.Update.Combine(updates), cancellationToken: cancellationToken);
+		}
+
+		/// <inheritdoc/>
+		public Task DeleteAsync(SessionId sessionId, CancellationToken cancellationToken = default)
+		{
+			return _sessions.DeleteOneAsync(x => x.Id == sessionId, null, cancellationToken);
+		}
+	}
+}
diff --git a/Engine/Source/Programs/Horde/Plugins/Compute/HordeServer.Compute/StaticComputeConfig.cs b/Engine/Source/Programs/Horde/Plugins/Compute/HordeServer.Compute/StaticComputeConfig.cs
new file mode 100644
--- /dev/null
+++ b/Engine/Source/Programs/Horde/Plugins/Compute/HordeServer.Compute/StaticComputeConfig.cs
@@ -0,0 +1,69 @@
+﻿// Copyright Epic Games, Inc. All Rights Reserved.
+
+using HordeServer.Agents.Fleet;
+
+namespace HordeServer
+{
+	/// <summary>
+	/// Static configuration for the compute plugin
+	/// </summary>
+	public class StaticComputeConfig
+	{
+		/// <summary>
+		/// Whether to enable the upgrade task source.
+		/// </summary>
+		public bool EnableUpgradeTasks { get; set; } = true;
+
+		/// <summary>
+		/// Whether to enable Amazon Web Services (AWS) specific features
+		/// </summary>
+		public bool WithAws { get; set; } = false;
+
+		/// <summary>
+		/// Default fleet manager to use (when not specified by pool)
+		/// </summary>
+		public FleetManagerType FleetManagerV2 { get; set; } = FleetManagerType.NoOp;
+
+		/// <summary>
+		/// Config for the fleet manager (serialized JSON)
+		/// </summary>
+		public string? FleetManagerV2Config { get; set; }
+
+		/// <summary>
+		/// AWS SQS queue URLs where lifecycle events from EC2 auto-scaling are received
+		/// <see cref="AwsAutoScalingLifecycleService" />
+		/// </summary>
+		public string[] AwsAutoScalingQueueUrls { get; set; } = Array.Empty<string>();
+
+		/// <summary>
+		/// Whether to automatically enroll agents in the farm
+		/// </summary>
+		public bool AutoEnrollAgents { get; set; }
+
+		/// <summary>
+		/// Default agent pool sizing strategy for pools that doesn't have one explicitly configured
+		/// </summary>
+		public PoolSizeStrategy DefaultAgentPoolSizeStrategy { get; set; } = PoolSizeStrategy.LeaseUtilization;
+
+		/// <summary>
+		/// Scale-out cooldown for auto-scaling agent pools (in seconds). Can be overridden by per-pool settings.
+		/// </summary>
+		public int AgentPoolScaleOutCooldownSeconds { get; set; } = 60; // 1 min
+
+		/// <summary>
+		/// Scale-in cooldown for auto-scaling agent pools (in seconds). Can be overridden by per-pool settings.
+		/// </summary>
+		public int AgentPoolScaleInCooldownSeconds { get; set; } = 1200; // 20 mins
+
+		/// <summary>
+		/// Port to listen on for tunneling compute sockets to agents
+		/// </summary>
+		public int ComputeTunnelPort { get; set; }
+
+		/// <summary>
+		/// What address (host:port) clients should connect to for compute socket tunneling
+		/// Port may differ from <see cref="ComputeTunnelPort" /> if Horde server is behind a reverse proxy/firewall
+		/// </summary>
+		public string? ComputeTunnelAddress { get; set; }
+	}
+}
diff --git a/Engine/Source/Programs/Horde/Plugins/Tools/HordeServer.Tools.Tests/Tools/ToolsControllerTests.cs b/Engine/Source/Programs/Horde/Plugins/Tools/HordeServer.Tools.Tests/Tools/ToolsControllerTests.cs
--- a/Engine/Source/Programs/Horde/Plugins/Tools/HordeServer.Tools.Tests/Tools/ToolsControllerTests.cs
+++ b/Engine/Source/Programs/Horde/Plugins/Tools/HordeServer.Tools.Tests/Tools/ToolsControllerTests.cs
@@ -33,6 +33,7 @@ public class ToolsControllerTests
 
 		PluginCollection pluginCollection = new PluginCollection();
 		pluginCollection.Add<AnalyticsPlugin>();
+		pluginCollection.Add<ComputePlugin>();
 		pluginCollection.Add<StoragePlugin>();
 
 		StorageConfig storageConfig = new StorageConfig();
