Horde: Change references from AclService to IAclService.

#rnx

diff --git a/Engine/Source/Programs/Horde/HordeServer.Shared/Acls/IAclService.cs b/Engine/Source/Programs/Horde/HordeServer.Shared/Acls/IAclService.cs
--- a/Engine/Source/Programs/Horde/HordeServer.Shared/Acls/IAclService.cs
+++ b/Engine/Source/Programs/Horde/HordeServer.Shared/Acls/IAclService.cs
@@ -23,7 +23,7 @@ namespace HordeServer.Acls
 		/// <param name="expiry">Time that the token expires</param>
 		/// <param name="cancellationToken">Cancellation token for the operation</param>
 		/// <returns>JWT security token with a claim for creating new agents</returns>
-		ValueTask<string> IssueBearerTokenAsync(IEnumerable<AclClaimConfig> claims, TimeSpan? expiry, CancellationToken cancellationToken = default);
+		ValueTask<string> IssueBearerTokenAsync(IEnumerable<Claim> claims, TimeSpan? expiry, CancellationToken cancellationToken = default);
 
 		/// <summary>
 		/// Finds an ACL scope by name
@@ -46,5 +46,15 @@ namespace HordeServer.Acls
 		/// <param name="user">The principal to validate</param>
 		public static bool Authorize(this IAclService aclService, AclScopeName scopeName, AclAction action, ClaimsPrincipal user)
 			=> aclService.TryGetAclScope(scopeName, out AclConfig? scopeConfig) && scopeConfig.Authorize(action, user);
+
+		/// <summary>
+		/// Issues a bearer token with the given roles
+		/// </summary>
+		/// <param name="claims">List of claims to include</param>
+		/// <param name="expiry">Time that the token expires</param>
+		/// <param name="cancellationToken">Cancellation token for the operation</param>
+		/// <returns>JWT security token with a claim for creating new agents</returns>
+		public static async ValueTask<string> IssueBearerTokenAsync(this IAclService aclService, IEnumerable<AclClaimConfig> claims, TimeSpan? expiry, CancellationToken cancellationToken = default)
+			=> await aclService.IssueBearerTokenAsync(claims.Select(x => new Claim(x.Type, x.Value)), expiry, cancellationToken);
 	}
 }
diff --git a/Engine/Source/Programs/Horde/HordeServer.Tests.Shared/ServerServiceTest.cs b/Engine/Source/Programs/Horde/HordeServer.Tests.Shared/ServerServiceTest.cs
--- a/Engine/Source/Programs/Horde/HordeServer.Tests.Shared/ServerServiceTest.cs
+++ b/Engine/Source/Programs/Horde/HordeServer.Tests.Shared/ServerServiceTest.cs
@@ -29,7 +29,7 @@ namespace HordeServer.Tests
 		public FakeClock Clock => ServiceProvider.GetRequiredService<FakeClock>();
 		public IMemoryCache Cache => ServiceProvider.GetRequiredService<IMemoryCache>();
 
-		public AclService AclService => ServiceProvider.GetRequiredService<AclService>();
+		public IAclService AclService => ServiceProvider.GetRequiredService<IAclService>();
 		public IMongoService MongoService => ServiceProvider.GetRequiredService<IMongoService>();
 		public IDowntimeService DowntimeService => ServiceProvider.GetRequiredService<IDowntimeService>();
 		public LifetimeService LifetimeService => ServiceProvider.GetRequiredService<LifetimeService>();
@@ -119,8 +119,7 @@ namespace HordeServer.Tests
 			services.AddSingleton<IHostApplicationLifetime, AppLifetimeStub>();
 			services.AddSingleton<IHostEnvironment, WebHostEnvironmentStub>();
 
-			services.AddSingleton<AclService>();
-			services.AddSingleton<IAclService>(sp => sp.GetRequiredService<AclService>());
+			services.AddSingleton<IAclService, AclService>();
 			services.AddSingleton<IDowntimeService, DowntimeServiceStub>();
 			services.AddSingleton<LifetimeService>();
 
diff --git a/Engine/Source/Programs/Horde/HordeServer/Acls/AclService.cs b/Engine/Source/Programs/Horde/HordeServer/Acls/AclService.cs
--- a/Engine/Source/Programs/Horde/HordeServer/Acls/AclService.cs
+++ b/Engine/Source/Programs/Horde/HordeServer/Acls/AclService.cs
@@ -35,18 +35,6 @@ namespace HordeServer.Acls
 			_globalConfig = globalConfig;
 		}
 
-		/// <summary>
-		/// Issues a bearer token with the given roles
-		/// </summary>
-		/// <param name="claims">List of claims to include</param>
-		/// <param name="expiry">Time that the token expires</param>
-		/// <param name="cancellationToken">Cancellation token for the operation</param>
-		/// <returns>JWT security token with a claim for creating new agents</returns>
-		public async ValueTask<string> IssueBearerTokenAsync(IEnumerable<AclClaimConfig> claims, TimeSpan? expiry, CancellationToken cancellationToken = default)
-		{
-			return await IssueBearerTokenAsync(claims.Select(x => new Claim(x.Type, x.Value)), expiry, cancellationToken);
-		}
-
 		/// <summary>
 		/// Issues a bearer token with the given claims
 		/// </summary>
diff --git a/Engine/Source/Programs/Horde/HordeServer/Artifacts/ArtifactsController.cs b/Engine/Source/Programs/Horde/HordeServer/Artifacts/ArtifactsController.cs
--- a/Engine/Source/Programs/Horde/HordeServer/Artifacts/ArtifactsController.cs
+++ b/Engine/Source/Programs/Horde/HordeServer/Artifacts/ArtifactsController.cs
@@ -52,7 +52,7 @@ namespace HordeServer.Artifacts
 		readonly StorageService _storageService;
 		readonly ILeaseCollection _leaseCollection;
 		readonly IJobCollection _jobCollection;
-		readonly AclService _aclService;
+		readonly IAclService _aclService;
 		readonly UnsyncCache _unsyncCache;
 		readonly IBlockCache _blockCache;
 		readonly GlobalConfig _globalConfig;
@@ -61,7 +61,7 @@ namespace HordeServer.Artifacts
 		/// <summary>
 		/// Constructor
 		/// </summary>
-		public ArtifactsController(IArtifactCollection artifactCollection, StorageService storageService, ILeaseCollection leaseCollection, IJobCollection jobCollection, AclService aclService, UnsyncCache unsyncCache, IBlockCache blockCache, IOptionsSnapshot<GlobalConfig> globalConfig, ILogger<ArtifactsController> logger)
+		public ArtifactsController(IArtifactCollection artifactCollection, StorageService storageService, ILeaseCollection leaseCollection, IJobCollection jobCollection, IAclService aclService, UnsyncCache unsyncCache, IBlockCache blockCache, IOptionsSnapshot<GlobalConfig> globalConfig, ILogger<ArtifactsController> logger)
 		{
 			_artifactCollection = artifactCollection;
 			_storageService = storageService;
diff --git a/Engine/Source/Programs/Horde/HordeServer/Jobs/JobRpcService.cs b/Engine/Source/Programs/Horde/HordeServer/Jobs/JobRpcService.cs
--- a/Engine/Source/Programs/Horde/HordeServer/Jobs/JobRpcService.cs
+++ b/Engine/Source/Programs/Horde/HordeServer/Jobs/JobRpcService.cs
@@ -48,7 +48,7 @@ namespace HordeServer.Jobs
 	[Authorize]
 	public class JobRpcService : JobRpc.JobRpcBase
 	{
-		readonly AclService _aclService;
+		readonly IAclService _aclService;
 		readonly JobService _jobService;
 		readonly AgentService _agentService;
 		readonly PoolService _poolService;
@@ -68,7 +68,7 @@ namespace HordeServer.Jobs
 		/// <summary>
 		/// Constructor
 		/// </summary>
-		public JobRpcService(AclService aclService, JobService jobService, AgentService agentService, PoolService poolService, ConformTaskSource conformTaskSource, IArtifactCollection artifactCollection, IJobCollection jobCollection, ILogCollection logCollection, IGraphCollection graphs, ITestDataCollection testData, IJobStepRefCollection jobStepRefCollection, ITemplateCollection templateCollection, HttpClient httpClient, IClock clock, IOptionsSnapshot<GlobalConfig> globalConfig, ILogger<JobRpcService> logger)
+		public JobRpcService(IAclService aclService, JobService jobService, AgentService agentService, PoolService poolService, ConformTaskSource conformTaskSource, IArtifactCollection artifactCollection, IJobCollection jobCollection, ILogCollection logCollection, IGraphCollection graphs, ITestDataCollection testData, IJobStepRefCollection jobStepRefCollection, ITemplateCollection templateCollection, HttpClient httpClient, IClock clock, IOptionsSnapshot<GlobalConfig> globalConfig, ILogger<JobRpcService> logger)
 		{
 			_aclService = aclService;
 			_jobService = jobService;
diff --git a/Engine/Source/Programs/Horde/HordeServer/Jobs/JobTaskSource.cs b/Engine/Source/Programs/Horde/HordeServer/Jobs/JobTaskSource.cs
--- a/Engine/Source/Programs/Horde/HordeServer/Jobs/JobTaskSource.cs
+++ b/Engine/Source/Programs/Horde/HordeServer/Jobs/JobTaskSource.cs
@@ -193,7 +193,7 @@ namespace HordeServer.Jobs
 		/// </summary>
 		readonly Lazy<AgentService> _agentService;
 		
-		readonly AclService _aclService;
+		readonly IAclService _aclService;
 		readonly IStreamCollection _streamCollection;
 		readonly ILogCollection _logCollection;
 		readonly IJobCollection _jobs;
@@ -242,7 +242,7 @@ namespace HordeServer.Jobs
 		/// <summary>
 		/// Constructor
 		/// </summary>
-		public JobTaskSource(AclService aclService, Lazy<AgentService> agentService, IJobCollection jobs, IJobStepRefCollection jobStepRefs, IBisectTaskCollection bisectTasks, IGraphCollection graphs, IPoolCollection pools, PoolService poolService, IUgsMetadataCollection ugsMetadataCollection, IStreamCollection streamCollection, ILogCollection logCollection, PerforceLoadBalancer perforceLoadBalancer, IClock clock, IOptionsMonitor<ServerSettings> settings, IOptionsMonitor<GlobalConfig> globalConfig, Tracer tracer, ILogger<JobTaskSource> logger)
+		public JobTaskSource(IAclService aclService, Lazy<AgentService> agentService, IJobCollection jobs, IJobStepRefCollection jobStepRefs, IBisectTaskCollection bisectTasks, IGraphCollection graphs, IPoolCollection pools, PoolService poolService, IUgsMetadataCollection ugsMetadataCollection, IStreamCollection streamCollection, ILogCollection logCollection, PerforceLoadBalancer perforceLoadBalancer, IClock clock, IOptionsMonitor<ServerSettings> settings, IOptionsMonitor<GlobalConfig> globalConfig, Tracer tracer, ILogger<JobTaskSource> logger)
 		{
 			_aclService = aclService;
 			_agentService = agentService;
diff --git a/Engine/Source/Programs/Horde/HordeServer/Server/AdminController.cs b/Engine/Source/Programs/Horde/HordeServer/Server/AdminController.cs
--- a/Engine/Source/Programs/Horde/HordeServer/Server/AdminController.cs
+++ b/Engine/Source/Programs/Horde/HordeServer/Server/AdminController.cs
@@ -54,14 +54,14 @@ namespace HordeServer.Server
 	[Route("[controller]")]
 	public class AdminController : HordeControllerBase
 	{
-		readonly AclService _aclService;
+		readonly IAclService _aclService;
 		readonly IOptionsSnapshot<GlobalConfig> _globalConfig;
 		readonly IOptions<ServerSettings> _serverSettings;
 
 		/// <summary>
 		/// Constructor
 		/// </summary>
-		public AdminController(AclService aclService, IOptionsSnapshot<GlobalConfig> globalConfig, IOptions<ServerSettings> serverSettings)
+		public AdminController(IAclService aclService, IOptionsSnapshot<GlobalConfig> globalConfig, IOptions<ServerSettings> serverSettings)
 		{
 			_aclService = aclService;
 			_globalConfig = globalConfig;
diff --git a/Engine/Source/Programs/Horde/HordeServer/Server/RpcService.cs b/Engine/Source/Programs/Horde/HordeServer/Server/RpcService.cs
--- a/Engine/Source/Programs/Horde/HordeServer/Server/RpcService.cs
+++ b/Engine/Source/Programs/Horde/HordeServer/Server/RpcService.cs
@@ -44,14 +44,14 @@ namespace HordeServer.Server
 		readonly LifetimeService _lifetimeService;
 		readonly IToolCollection _toolCollection;
 		readonly IAgentTelemetryCollection _agentTelemetryCollection;
-		readonly AclService _aclService;
+		readonly IAclService _aclService;
 		readonly IOptionsSnapshot<GlobalConfig> _globalConfig;
 		readonly ILogger _logger;
 
 		/// <summary>
 		/// Constructor
 		/// </summary>
-		public RpcService(AgentService agentService, LifetimeService lifetimeService, IToolCollection toolCollection, IAgentTelemetryCollection agentTelemetryCollection, AclService aclService, IOptionsSnapshot<GlobalConfig> globalConfig, ILogger<RpcService> logger)
+		public RpcService(AgentService agentService, LifetimeService lifetimeService, IToolCollection toolCollection, IAgentTelemetryCollection agentTelemetryCollection, IAclService aclService, IOptionsSnapshot<GlobalConfig> globalConfig, ILogger<RpcService> logger)
 		{
 			_agentService = agentService;
 			_lifetimeService = lifetimeService;
diff --git a/Engine/Source/Programs/Horde/HordeServer/Startup.cs b/Engine/Source/Programs/Horde/HordeServer/Startup.cs
--- a/Engine/Source/Programs/Horde/HordeServer/Startup.cs
+++ b/Engine/Source/Programs/Horde/HordeServer/Startup.cs
@@ -422,8 +422,7 @@ namespace HordeServer
 
 			services.AddSingleton<IPoolSizeStrategyFactory, JobQueueStrategyFactory>();
 
-			services.AddSingleton<AclService>();
-			services.AddSingleton<IAclService>(sp => sp.GetRequiredService<AclService>());
+			services.AddSingleton<IAclService, AclService>();
 			services.AddSingleton<AgentService>();
 			services.AddSingleton(provider => new Lazy<AgentService>(provider.GetRequiredService<AgentService>));
 			services.AddHostedService<ArtifactExpirationService>();
