Horde: Move more classes into HordeServer.Shared.

#rnx

diff --git a/Engine/Source/Programs/Horde/HordeServer.Shared/Utilities/CaseInsensitiveDictionary.cs b/Engine/Source/Programs/Horde/HordeServer.Shared/Utilities/CaseInsensitiveDictionary.cs
new file mode 100644
--- /dev/null
+++ b/Engine/Source/Programs/Horde/HordeServer.Shared/Utilities/CaseInsensitiveDictionary.cs
@@ -0,0 +1,37 @@
+// Copyright Epic Games, Inc. All Rights Reserved.
+
+using System;
+using System.Collections.Generic;
+
+namespace HordeServer.Utilities
+{
+	/// <summary>
+	/// Case insensitive set of strings
+	/// </summary>
+	public class CaseInsensitiveDictionary<TValue> : Dictionary<string, TValue>
+	{
+		/// <summary>
+		/// Constructor
+		/// </summary>
+		public CaseInsensitiveDictionary()
+			: base(StringComparer.OrdinalIgnoreCase)
+		{
+		}
+
+		/// <summary>
+		/// Constructor
+		/// </summary>
+		public CaseInsensitiveDictionary(IDictionary<string, TValue> items)
+			: base(items, StringComparer.OrdinalIgnoreCase)
+		{
+		}
+
+		/// <summary>
+		/// Constructor
+		/// </summary>
+		public CaseInsensitiveDictionary(IEnumerable<KeyValuePair<string, TValue>> items)
+			: base(items, StringComparer.OrdinalIgnoreCase)
+		{
+		}
+	}
+}
diff --git a/Engine/Source/Programs/Horde/HordeServer.Shared/Utilities/CaseInsensitiveStringSet.cs b/Engine/Source/Programs/Horde/HordeServer.Shared/Utilities/CaseInsensitiveStringSet.cs
new file mode 100644
--- /dev/null
+++ b/Engine/Source/Programs/Horde/HordeServer.Shared/Utilities/CaseInsensitiveStringSet.cs
@@ -0,0 +1,29 @@
+// Copyright Epic Games, Inc. All Rights Reserved.
+
+using System;
+using System.Collections.Generic;
+
+namespace HordeServer.Utilities
+{
+	/// <summary>
+	/// Case insensitive set of strings
+	/// </summary>
+	public class CaseInsensitiveStringSet : HashSet<string>
+	{
+		/// <summary>
+		/// Constructor
+		/// </summary>
+		public CaseInsensitiveStringSet()
+			: base(StringComparer.OrdinalIgnoreCase)
+		{
+		}
+
+		/// <summary>
+		/// Constructor
+		/// </summary>
+		public CaseInsensitiveStringSet(IEnumerable<string> items)
+			: base(items, StringComparer.OrdinalIgnoreCase)
+		{
+		}
+	}
+}
diff --git a/Engine/Source/Programs/Horde/HordeServer.Shared/Utilities/InlineFileStreamResult.cs b/Engine/Source/Programs/Horde/HordeServer.Shared/Utilities/InlineFileStreamResult.cs
new file mode 100644
--- /dev/null
+++ b/Engine/Source/Programs/Horde/HordeServer.Shared/Utilities/InlineFileStreamResult.cs
@@ -0,0 +1,39 @@
+// Copyright Epic Games, Inc. All Rights Reserved.
+
+using System.Net.Mime;
+using System.Threading.Tasks;
+using Microsoft.AspNetCore.Mvc;
+
+namespace HordeServer.Utilities
+{
+	/// <summary>
+	/// Class to return a file stream without the "content-disposition: attachment" header
+	/// </summary>
+	public class InlineFileStreamResult : FileStreamResult
+	{
+		/// <summary>
+		/// The suggested download filename
+		/// </summary>
+		readonly string _fileName;
+
+		/// <summary>
+		/// Constructor
+		/// </summary>
+		public InlineFileStreamResult(System.IO.Stream stream, string mimeType, string fileName)
+			: base(stream, mimeType)
+		{
+			_fileName = fileName;
+		}
+
+		/// <inheritdoc/>
+		public override Task ExecuteResultAsync(ActionContext context)
+		{
+			ContentDisposition contentDisposition = new ContentDisposition();
+			contentDisposition.Inline = true;
+			contentDisposition.FileName = _fileName;
+			context.HttpContext.Response.Headers["Content-Disposition"] = contentDisposition.ToString();
+
+			return base.ExecuteResultAsync(context);
+		}
+	}
+}
diff --git a/Engine/Source/Programs/Horde/HordeServer/Utilities/CaseInsensitiveDictionary.cs b/Engine/Source/Programs/Horde/HordeServer/Utilities/CaseInsensitiveDictionary.cs
deleted file mode 100644
diff --git a/Engine/Source/Programs/Horde/HordeServer/Utilities/CaseInsensitiveStringSet.cs b/Engine/Source/Programs/Horde/HordeServer/Utilities/CaseInsensitiveStringSet.cs
deleted file mode 100644
diff --git a/Engine/Source/Programs/Horde/HordeServer/Utilities/InlineFileStreamResult.cs b/Engine/Source/Programs/Horde/HordeServer/Utilities/InlineFileStreamResult.cs
deleted file mode 100644
