Horde: Remove property redirect for pools to compute plugin.

#rnx

diff --git a/Engine/Source/Programs/Horde/HordeServer.Tests/Agents/Pools/PoolUpdateServiceTest.cs b/Engine/Source/Programs/Horde/HordeServer.Tests/Agents/Pools/PoolUpdateServiceTest.cs
--- a/Engine/Source/Programs/Horde/HordeServer.Tests/Agents/Pools/PoolUpdateServiceTest.cs
+++ b/Engine/Source/Programs/Horde/HordeServer.Tests/Agents/Pools/PoolUpdateServiceTest.cs
@@ -24,7 +24,7 @@ public class PoolUpdateServiceTest : TestSetup
 
 	public PoolUpdateServiceTest()
 	{
-		UpdateConfig(x => x.Pools.Clear());
+		UpdateConfig(x => x.Plugins.GetComputeConfig().Pools.Clear());
 		_pus = new PoolUpdateService(AgentService, PoolCollection, Clock, GlobalConfig, Tracer, new NullLogger<PoolUpdateService>());
 	}
 
@@ -86,7 +86,7 @@ public class PoolUpdateServiceTest : TestSetup
 	{
 		// Arrange
 		// Explicitly set the grace period for the pool to be longer than the default of 8 hours
-		UpdateConfig(config => config.Pools[0].ShutdownIfDisabledGracePeriod = TimeSpan.FromHours(24.0));
+		UpdateConfig(config => config.Plugins.GetComputeConfig().Pools[0].ShutdownIfDisabledGracePeriod = TimeSpan.FromHours(24.0));
 
 		// Act
 		await _pus.ShutdownDisabledAgentsAsync(CancellationToken.None);
@@ -102,7 +102,7 @@ public class PoolUpdateServiceTest : TestSetup
 	public async Task ShutdownDisabledAgents_WithAutoScalingOff_DoesNotRequestShutdownAsync()
 	{
 		// Arrange
-		UpdateConfig(config => config.Pools[0].EnableAutoscaling = false);
+		UpdateConfig(config => config.Plugins.GetComputeConfig().Pools[0].EnableAutoscaling = false);
 
 		// Act
 		await _pus.ShutdownDisabledAgentsAsync(CancellationToken.None);
diff --git a/Engine/Source/Programs/Horde/HordeServer.Tests/Agents/Pools/PoolsControllerTest.cs b/Engine/Source/Programs/Horde/HordeServer.Tests/Agents/Pools/PoolsControllerTest.cs
--- a/Engine/Source/Programs/Horde/HordeServer.Tests/Agents/Pools/PoolsControllerTest.cs
+++ b/Engine/Source/Programs/Horde/HordeServer.Tests/Agents/Pools/PoolsControllerTest.cs
@@ -15,7 +15,7 @@ namespace HordeServer.Tests.Agents.Pools
 	{
 		public PoolsControllerTest()
 		{
-			UpdateConfig(x => x.Pools.Clear());
+			UpdateConfig(x => x.Plugins.GetComputeConfig().Pools.Clear());
 		}
 
 		[TestMethod]
@@ -30,8 +30,9 @@ namespace HordeServer.Tests.Agents.Pools
 
 			UpdateConfig(config =>
 			{
-				config.Pools.Clear();
-				config.Pools.Add(poolConfig);
+				ComputeConfig computeConfig = config.Plugins.GetComputeConfig();
+				computeConfig.Pools.Clear();
+				computeConfig.Pools.Add(poolConfig);
 			});
 
 			ActionResult<List<object>> rawResult = await PoolsController.GetPoolsAsync();
@@ -59,8 +60,9 @@ namespace HordeServer.Tests.Agents.Pools
 
 			UpdateConfig(config =>
 			{
-				config.Pools.Clear();
-				config.Pools.Add(poolConfig);
+				ComputeConfig computeConfig = config.Plugins.GetComputeConfig();
+				computeConfig.Pools.Clear();
+				computeConfig.Pools.Add(poolConfig);
 			});
 
 			IPool? pool = await PoolService.GetPoolAsync(poolConfig.Id);
diff --git a/Engine/Source/Programs/Horde/HordeServer.Tests/Fleet/FleetServiceTest.cs b/Engine/Source/Programs/Horde/HordeServer.Tests/Fleet/FleetServiceTest.cs
--- a/Engine/Source/Programs/Horde/HordeServer.Tests/Fleet/FleetServiceTest.cs
+++ b/Engine/Source/Programs/Horde/HordeServer.Tests/Fleet/FleetServiceTest.cs
@@ -102,7 +102,7 @@ namespace HordeServer.Tests.Fleet
 
 		public FleetServiceTest()
 		{
-			UpdateConfig(x => x.Pools.Clear());
+			UpdateConfig(x => x.Plugins.GetComputeConfig().Pools.Clear());
 		}
 
 		[TestMethod]
diff --git a/Engine/Source/Programs/Horde/HordeServer.Tests/Server/RpcServiceTest.cs b/Engine/Source/Programs/Horde/HordeServer.Tests/Server/RpcServiceTest.cs
--- a/Engine/Source/Programs/Horde/HordeServer.Tests/Server/RpcServiceTest.cs
+++ b/Engine/Source/Programs/Horde/HordeServer.Tests/Server/RpcServiceTest.cs
@@ -273,7 +273,7 @@ namespace HordeServer.Tests.Server
 
 		public RpcServiceTest()
 		{
-			UpdateConfig(x => x.Pools.Clear());
+			UpdateConfig(x => x.Plugins.GetComputeConfig().Pools.Clear());
 		}
 
 		[TestMethod]
diff --git a/Engine/Source/Programs/Horde/HordeServer.Tests/TestSetup.cs b/Engine/Source/Programs/Horde/HordeServer.Tests/TestSetup.cs
--- a/Engine/Source/Programs/Horde/HordeServer.Tests/TestSetup.cs
+++ b/Engine/Source/Programs/Horde/HordeServer.Tests/TestSetup.cs
@@ -372,7 +372,7 @@ namespace HordeServer.Tests
 
 		protected async Task<IPool> CreatePoolAsync(PoolConfig poolConfig)
 		{
-			UpdateConfig(config => config.Pools.Add(poolConfig));
+			UpdateConfig(config => config.Plugins.GetComputeConfig().Pools.Add(poolConfig));
 			return await PoolCollection.GetAsync(poolConfig.Id) ?? throw new NotImplementedException();
 		}
 	}
diff --git a/Engine/Source/Programs/Horde/HordeServer/Server/GlobalConfig.cs b/Engine/Source/Programs/Horde/HordeServer/Server/GlobalConfig.cs
--- a/Engine/Source/Programs/Horde/HordeServer/Server/GlobalConfig.cs
+++ b/Engine/Source/Programs/Horde/HordeServer/Server/GlobalConfig.cs
@@ -126,11 +126,6 @@ namespace HordeServer.Server
 		/// </summary>
 		public List<ProjectConfig> Projects { get; set; } = new List<ProjectConfig>();
 
-		/// <summary>
-		/// List of pools
-		/// </summary>
-		public List<PoolConfig> Pools => (Plugins.Values.OfType<ComputeConfig>().FirstOrDefault() ?? new ComputeConfig()).Pools;
-
 		/// <summary>
 		/// List of scheduled downtime
 		/// </summary>
@@ -240,8 +235,6 @@ namespace HordeServer.Server
 				_artifactTypeLookup.Add(artifactType.Type, artifactType);
 			}
 
-			UpdateWorkspacesForPools();
-
 			// Ensure that all plugins have an entry in the global config so they can register their ACLs
 			foreach (ILoadedPlugin loadedPlugin in loadedPlugins)
 			{
@@ -258,6 +251,8 @@ namespace HordeServer.Server
 				pluginConfig.PostLoad(pluginConfigOptions);
 			}
 
+			UpdateWorkspacesForPools();
+
 			_aclLookup.Clear();
 			BuildAclScopeLookup(Acl, _aclLookup);
 
@@ -395,7 +390,8 @@ namespace HordeServer.Server
 			}
 
 			// Update the list of workspaces for each pool
-			foreach (PoolConfig pool in Pools)
+			ComputeConfig computeConfig = Plugins.GetComputeConfig();
+			foreach (PoolConfig pool in computeConfig.Pools)
 			{
 				// Get the new list of workspaces for this pool
 				List<AgentWorkspaceInfo>? newWorkspaces;
diff --git a/Engine/Source/Programs/Horde/HordeServer/Server/ServerController.cs b/Engine/Source/Programs/Horde/HordeServer/Server/ServerController.cs
--- a/Engine/Source/Programs/Horde/HordeServer/Server/ServerController.cs
+++ b/Engine/Source/Programs/Horde/HordeServer/Server/ServerController.cs
@@ -158,7 +158,7 @@ namespace HordeServer.Server
 
 			IPoolCollection poolCollection = _serviceProvider.GetRequiredService<IPoolCollection>();
 			List<IPoolConfig> poolConfigs = (await poolCollection.GetConfigsAsync(cancellationToken)).ToList();
-			HashSet<PoolId> removePoolIds = _globalConfig.Value.Pools.Select(x => x.Id).ToHashSet();
+			HashSet<PoolId> removePoolIds = _globalConfig.Value.Plugins.GetComputeConfig().Pools.Select(x => x.Id).ToHashSet();
 			poolConfigs.RemoveAll(x => removePoolIds.Contains(x.Id));
 
 			if (minAgents != null || maxAgents != null)
diff --git a/Engine/Source/Programs/Horde/Plugins/Compute/HordeServer.Compute/ComputePlugin.cs b/Engine/Source/Programs/Horde/Plugins/Compute/HordeServer.Compute/ComputePlugin.cs
--- a/Engine/Source/Programs/Horde/Plugins/Compute/HordeServer.Compute/ComputePlugin.cs
+++ b/Engine/Source/Programs/Horde/Plugins/Compute/HordeServer.Compute/ComputePlugin.cs
@@ -93,4 +93,19 @@ namespace HordeServer
 			}
 		}
 	}
+
+	/// <summary>
+	/// Helper methods for compute config
+	/// </summary>
+	public static class ComputePluginExtensions
+	{
+		/// <summary>
+		/// Get the compute plugin config
+		/// </summary>
+		public static ComputeConfig GetComputeConfig(this IDictionary<PluginName, IPluginConfig> plugins)
+		{
+			IPluginConfig pluginConfig = plugins[new PluginName("compute")];
+			return (ComputeConfig)pluginConfig;
+		}
+	}
 }
diff --git a/Engine/Source/Programs/Horde/Plugins/Compute/HordeServer.Compute/Server/PoolsServerController.cs b/Engine/Source/Programs/Horde/Plugins/Compute/HordeServer.Compute/Server/PoolsServerController.cs
new file mode 100644
--- /dev/null
+++ b/Engine/Source/Programs/Horde/Plugins/Compute/HordeServer.Compute/Server/PoolsServerController.cs
@@ -0,0 +1,119 @@
+// Copyright Epic Games, Inc. All Rights Reserved.
+
+using EpicGames.Horde.Agents.Pools;
+using HordeServer.Agents;
+using HordeServer.Agents.Pools;
+using HordeServer.Utilities;
+using Microsoft.AspNetCore.Authorization;
+using Microsoft.AspNetCore.Mvc;
+using Microsoft.Extensions.DependencyInjection;
+using Microsoft.Extensions.Options;
+
+namespace HordeServer.Server
+{
+	/// <summary>
+	/// Controller managing pools-related server commands
+	/// </summary>
+	[ApiController]
+	[Authorize]
+	[Route("[controller]")]
+	public class PoolsServerController : HordeControllerBase
+	{
+		readonly IServiceProvider _serviceProvider;
+		readonly IOptionsSnapshot<ComputeConfig> _computeConfig;
+
+		/// <summary>
+		/// Constructor
+		/// </summary>
+		public PoolsServerController(IServiceProvider serviceProvider, IOptionsSnapshot<ComputeConfig> computeConfig)
+		{
+			_serviceProvider = serviceProvider;
+			_computeConfig = computeConfig;
+		}
+
+		/// <summary>
+		/// Converts all legacy pools into config entries
+		/// </summary>
+		[HttpGet]
+		[Route("/api/v1/server/migrate/pool-config")]
+		public async Task<ActionResult<object>> MigratePoolsAsync([FromQuery] int? minAgents = null, [FromQuery] int? maxAgents = null, CancellationToken cancellationToken = default)
+		{
+			if (!_computeConfig.Value.Authorize(PoolAclAction.ListPools, User))
+			{
+				return Forbid(PoolAclAction.ListPools);
+			}
+
+			IPoolCollection poolCollection = _serviceProvider.GetRequiredService<IPoolCollection>();
+			List<IPoolConfig> poolConfigs = (await poolCollection.GetConfigsAsync(cancellationToken)).ToList();
+			HashSet<PoolId> removePoolIds = _computeConfig.Value.Pools.Select(x => x.Id).ToHashSet();
+			poolConfigs.RemoveAll(x => removePoolIds.Contains(x.Id));
+
+			if (minAgents != null || maxAgents != null)
+			{
+				IAgentCollection agentCollection = _serviceProvider.GetRequiredService<IAgentCollection>();
+
+				Dictionary<PoolId, int> poolIdToCount = new Dictionary<PoolId, int>();
+
+				IReadOnlyList<IAgent> agents = await agentCollection.FindAsync(cancellationToken: cancellationToken);
+				foreach (IAgent agent in agents)
+				{
+					foreach (PoolId poolId in agent.Pools)
+					{
+						int count;
+						if (!poolIdToCount.TryGetValue(poolId, out count))
+						{
+							count = 0;
+						}
+						poolIdToCount[poolId] = count + 1;
+					}
+				}
+
+				if (minAgents != null && minAgents.Value > 0)
+				{
+					poolConfigs.RemoveAll(x => !poolIdToCount.TryGetValue(x.Id, out int count) || count < minAgents.Value);
+				}
+				if (maxAgents != null)
+				{
+					poolConfigs.RemoveAll(x => poolIdToCount.TryGetValue(x.Id, out int count) && count > maxAgents.Value);
+				}
+			}
+
+			List<PoolConfig> configs = new List<PoolConfig>();
+			foreach (IPoolConfig currentConfig in poolConfigs.OrderBy(x => x.Id.Id.Text))
+			{
+				PoolConfig config = new PoolConfig();
+				config.Id = currentConfig.Id;
+				config.Name = currentConfig.Name;
+				config.Condition = currentConfig.Condition;
+				if (currentConfig.Properties != null && currentConfig.Properties.Count > 0 && (currentConfig.Properties.Count != 0 && (currentConfig.Properties.First().Key != "color" && currentConfig.Properties.First().Value != "0")))
+				{
+					config.Properties = new Dictionary<string, string>(currentConfig.Properties);
+				}
+				config.EnableAutoscaling = currentConfig.EnableAutoscaling;
+				config.MinAgents = currentConfig.MinAgents;
+				config.NumReserveAgents = currentConfig.NumReserveAgents;
+				config.ConformInterval = currentConfig.ConformInterval;
+				config.ScaleInCooldown = currentConfig.ScaleInCooldown;
+				config.ScaleOutCooldown = currentConfig.ScaleOutCooldown;
+				config.ShutdownIfDisabledGracePeriod = currentConfig.ShutdownIfDisabledGracePeriod;
+#pragma warning disable CS0618 // Type or member is obsolete
+				config.SizeStrategy = currentConfig.SizeStrategy;
+#pragma warning restore CS0618 // Type or member is obsolete
+				if (currentConfig.SizeStrategies != null && currentConfig.SizeStrategies.Count > 0)
+				{
+					config.SizeStrategies = currentConfig.SizeStrategies.ToList();
+				}
+				if (currentConfig.FleetManagers != null && currentConfig.FleetManagers.Count > 0)
+				{
+					config.FleetManagers = currentConfig.FleetManagers.ToList();
+				}
+				config.LeaseUtilizationSettings = currentConfig.LeaseUtilizationSettings;
+				config.JobQueueSettings = currentConfig.JobQueueSettings;
+				config.ComputeQueueAwsMetricSettings = currentConfig.ComputeQueueAwsMetricSettings;
+				configs.Add(config);
+			}
+
+			return new { Pools = configs };
+		}
+	}
+}
