Horde: Move debug endpoint attribute into shared library.

#rnx

diff --git a/Engine/Source/Programs/Horde/HordeServer.Shared/IServerInfo.cs b/Engine/Source/Programs/Horde/HordeServer.Shared/IServerInfo.cs
--- a/Engine/Source/Programs/Horde/HordeServer.Shared/IServerInfo.cs
+++ b/Engine/Source/Programs/Horde/HordeServer.Shared/IServerInfo.cs
@@ -73,6 +73,11 @@ namespace HordeServer
 		/// </summary>
 		bool ReadOnlyMode { get; }
 
+		/// <summary>
+		/// Whether to enable endpoints which are for debugging purposes
+		/// </summary>
+		bool EnableDebugEndpoints { get; }
+
 		/// <summary>
 		/// Helper method to check if this process has activated the given mode
 		/// </summary>
diff --git a/Engine/Source/Programs/Horde/HordeServer.Shared/Utilities/DebugEndpointAttribute.cs b/Engine/Source/Programs/Horde/HordeServer.Shared/Utilities/DebugEndpointAttribute.cs
new file mode 100644
--- /dev/null
+++ b/Engine/Source/Programs/Horde/HordeServer.Shared/Utilities/DebugEndpointAttribute.cs
@@ -0,0 +1,33 @@
+// Copyright Epic Games, Inc. All Rights Reserved.
+
+using System;
+using Microsoft.AspNetCore.Mvc;
+using Microsoft.AspNetCore.Mvc.Filters;
+using Microsoft.Extensions.DependencyInjection;
+using Microsoft.Extensions.Options;
+
+namespace HordeServer.Utilities
+{
+	/// <summary>
+	/// Only requests to attached controller to pass if debug endpoint is enabled in settings
+	/// Adds extra security for not enabling these admin endpoints by accident.
+	/// </summary>
+	[AttributeUsage(AttributeTargets.Class)]
+	public sealed class DebugEndpointAttribute : Attribute, IActionFilter
+	{
+		/// <inheritdoc />
+		public void OnActionExecuting(ActionExecutingContext context)
+		{
+			IServerInfo serverInfo = context.HttpContext.RequestServices.GetRequiredService<IServerInfo>();
+			if (!serverInfo.EnableDebugEndpoints)
+			{
+				context.Result = new ForbidResult();
+			}
+		}
+
+		/// <inheritdoc />
+		public void OnActionExecuted(ActionExecutedContext context)
+		{
+		}
+	}
+}
diff --git a/Engine/Source/Programs/Horde/HordeServer/Server/ServerController.cs b/Engine/Source/Programs/Horde/HordeServer/Server/ServerController.cs
--- a/Engine/Source/Programs/Horde/HordeServer/Server/ServerController.cs
+++ b/Engine/Source/Programs/Horde/HordeServer/Server/ServerController.cs
@@ -143,90 +143,5 @@ namespace HordeServer.Server
 			response.LocalRedirectUrls = settings.OidcLocalRedirectUrls;
 			return response;
 		}
-
-		/// <summary>
-		/// Converts all legacy pools into config entries
-		/// </summary>
-		[HttpGet]
-		[Route("/api/v1/server/migrate/pool-config")]
-		public async Task<ActionResult<object>> MigratePoolsAsync([FromQuery] int? minAgents = null, [FromQuery] int? maxAgents = null, CancellationToken cancellationToken = default)
-		{
-			if (!_globalConfig.Value.Authorize(PoolAclAction.ListPools, User))
-			{
-				return Forbid(PoolAclAction.ListPools);
-			}
-
-			IPoolCollection poolCollection = _serviceProvider.GetRequiredService<IPoolCollection>();
-			List<IPoolConfig> poolConfigs = (await poolCollection.GetConfigsAsync(cancellationToken)).ToList();
-			HashSet<PoolId> removePoolIds = _globalConfig.Value.Plugins.GetComputeConfig().Pools.Select(x => x.Id).ToHashSet();
-			poolConfigs.RemoveAll(x => removePoolIds.Contains(x.Id));
-
-			if (minAgents != null || maxAgents != null)
-			{
-				IAgentCollection agentCollection = _serviceProvider.GetRequiredService<IAgentCollection>();
-
-				Dictionary<PoolId, int> poolIdToCount = new Dictionary<PoolId, int>();
-
-				IReadOnlyList<IAgent> agents = await agentCollection.FindAsync(cancellationToken: cancellationToken);
-				foreach (IAgent agent in agents)
-				{
-					foreach (PoolId poolId in agent.Pools)
-					{
-						int count;
-						if (!poolIdToCount.TryGetValue(poolId, out count))
-						{
-							count = 0;
-						}
-						poolIdToCount[poolId] = count + 1;
-					}
-				}
-
-				if (minAgents != null && minAgents.Value > 0)
-				{
-					poolConfigs.RemoveAll(x => !poolIdToCount.TryGetValue(x.Id, out int count) || count < minAgents.Value);
-				}
-				if (maxAgents != null)
-				{
-					poolConfigs.RemoveAll(x => poolIdToCount.TryGetValue(x.Id, out int count) && count > maxAgents.Value);
-				}
-			}
-
-			List<PoolConfig> configs = new List<PoolConfig>();
-			foreach (IPoolConfig currentConfig in poolConfigs.OrderBy(x => x.Id.Id.Text))
-			{
-				PoolConfig config = new PoolConfig();
-				config.Id = currentConfig.Id;
-				config.Name = currentConfig.Name;
-				config.Condition = currentConfig.Condition;
-				if (currentConfig.Properties != null && currentConfig.Properties.Count > 0 && (currentConfig.Properties.Count != 0 && (currentConfig.Properties.First().Key != "color" && currentConfig.Properties.First().Value != "0")))
-				{
-					config.Properties = new Dictionary<string, string>(currentConfig.Properties);
-				}
-				config.EnableAutoscaling = currentConfig.EnableAutoscaling;
-				config.MinAgents = currentConfig.MinAgents;
-				config.NumReserveAgents = currentConfig.NumReserveAgents;
-				config.ConformInterval = currentConfig.ConformInterval;
-				config.ScaleInCooldown = currentConfig.ScaleInCooldown;
-				config.ScaleOutCooldown = currentConfig.ScaleOutCooldown;
-				config.ShutdownIfDisabledGracePeriod = currentConfig.ShutdownIfDisabledGracePeriod;
-#pragma warning disable CS0618 // Type or member is obsolete
-				config.SizeStrategy = currentConfig.SizeStrategy;
-#pragma warning restore CS0618 // Type or member is obsolete
-				if (currentConfig.SizeStrategies != null && currentConfig.SizeStrategies.Count > 0)
-				{
-					config.SizeStrategies = currentConfig.SizeStrategies.ToList();
-				}
-				if (currentConfig.FleetManagers != null && currentConfig.FleetManagers.Count > 0)
-				{
-					config.FleetManagers = currentConfig.FleetManagers.ToList();
-				}
-				config.LeaseUtilizationSettings = currentConfig.LeaseUtilizationSettings;
-				config.JobQueueSettings = currentConfig.JobQueueSettings;
-				config.ComputeQueueAwsMetricSettings = currentConfig.ComputeQueueAwsMetricSettings;
-				configs.Add(config);
-			}
-
-			return new { Pools = configs };
-		}
 	}
 }
diff --git a/Engine/Source/Programs/Horde/HordeServer/ServerInfo.cs b/Engine/Source/Programs/Horde/HordeServer/ServerInfo.cs
--- a/Engine/Source/Programs/Horde/HordeServer/ServerInfo.cs
+++ b/Engine/Source/Programs/Horde/HordeServer/ServerInfo.cs
@@ -29,6 +29,9 @@ namespace HordeServer
 		/// <inheritdoc/>
 		public IConfiguration Configuration => _configuration;
 
+		/// <inheritdoc/>
+		public bool EnableDebugEndpoints => _serverSettings.Value.EnableDebugEndpoints;
+
 		/// <inheritdoc/>
 		public bool ReadOnlyMode => _serverSettings.Value.MongoReadOnlyMode;
 
diff --git a/Engine/Source/Programs/Horde/HordeServer/Utilities/DebugEndpointAttribute.cs b/Engine/Source/Programs/Horde/HordeServer/Utilities/DebugEndpointAttribute.cs
deleted file mode 100644
