Horde: Remove GlobalConfig reference from ProjectConfig, allowing it to be more easily replaced with a build-specific context instance.

#rnx

diff --git a/Engine/Source/Programs/Horde/HordeServer/Commands/Test/ReplicationCommand.cs b/Engine/Source/Programs/Horde/HordeServer/Commands/Test/ReplicationCommand.cs
--- a/Engine/Source/Programs/Horde/HordeServer/Commands/Test/ReplicationCommand.cs
+++ b/Engine/Source/Programs/Horde/HordeServer/Commands/Test/ReplicationCommand.cs
@@ -79,7 +79,7 @@ namespace HordeServer.Commands.Test
 			}
 
 			PerforceReplicationOptions options = new PerforceReplicationOptions();
-			await perforceReplicator.RunOnceAsync(replicator, streamConfig, options, default);
+			await perforceReplicator.RunOnceAsync(replicator, globalConfig, streamConfig, options, default);
 
 			return 0;
 		}
diff --git a/Engine/Source/Programs/Horde/HordeServer/Projects/ProjectConfig.cs b/Engine/Source/Programs/Horde/HordeServer/Projects/ProjectConfig.cs
--- a/Engine/Source/Programs/Horde/HordeServer/Projects/ProjectConfig.cs
+++ b/Engine/Source/Programs/Horde/HordeServer/Projects/ProjectConfig.cs
@@ -29,12 +29,6 @@ namespace HordeServer.Projects
 	[DebuggerDisplay("{Id}")]
 	public class ProjectConfig
 	{
-		/// <summary>
-		/// Accessor for the global config owning this project
-		/// </summary>
-		[JsonIgnore]
-		public GlobalConfig GlobalConfig { get; private set; } = null!;
-
 		/// <summary>
 		/// The project id
 		/// </summary>
@@ -120,13 +114,11 @@ namespace HordeServer.Projects
 		/// Callback after this configuration has been read
 		/// </summary>
 		/// <param name="id">Id of this project</param>
-		/// <param name="globalConfig">The owning global config object</param>
-		public void PostLoad(ProjectId id, GlobalConfig globalConfig)
+		/// <param name="parentAcl">The owning global config object</param>
+		public void PostLoad(ProjectId id, AclConfig parentAcl)
 		{
 			Id = id;
-			GlobalConfig = globalConfig;
-
-			Acl.PostLoad(globalConfig.Acl, $"project:{Id}");
+			Acl.PostLoad(parentAcl, $"project:{Id}");
 
 			foreach (StreamConfig stream in Streams)
 			{
diff --git a/Engine/Source/Programs/Horde/HordeServer/Replicators/PerforceReplicator.cs b/Engine/Source/Programs/Horde/HordeServer/Replicators/PerforceReplicator.cs
--- a/Engine/Source/Programs/Horde/HordeServer/Replicators/PerforceReplicator.cs
+++ b/Engine/Source/Programs/Horde/HordeServer/Replicators/PerforceReplicator.cs
@@ -308,7 +308,7 @@ namespace HordeServer.Replicators
 		/// <summary>
 		/// Runs a replication loop for a stream
 		/// </summary>
-		public async Task RunAsync(ReplicatorId replicatorId, StreamConfig streamConfig, PerforceReplicationOptions options, CancellationToken cancellationToken = default)
+		public async Task RunAsync(ReplicatorId replicatorId, GlobalConfig globalConfig, StreamConfig streamConfig, PerforceReplicationOptions options, CancellationToken cancellationToken = default)
 		{
 			_logger.LogInformation("Starting replication background task for {ReplicatorId}", replicatorId);
 
@@ -321,14 +321,14 @@ namespace HordeServer.Replicators
 
 			while (!cancellationToken.IsCancellationRequested)
 			{
-				replicator = await RunOnceAsync(replicator, streamConfig, options, cancellationToken);
+				replicator = await RunOnceAsync(replicator, globalConfig, streamConfig, options, cancellationToken);
 			}
 		}
 
 		/// <summary>
 		/// Runs the replicator for a single change
 		/// </summary>
-		public async Task<IReplicator> RunOnceAsync(IReplicator replicator, StreamConfig streamConfig, PerforceReplicationOptions replicatorOptions, CancellationToken cancellationToken)
+		public async Task<IReplicator> RunOnceAsync(IReplicator replicator, GlobalConfig globalConfig, StreamConfig streamConfig, PerforceReplicationOptions replicatorOptions, CancellationToken cancellationToken)
 		{
 			RefName refName = GetRefName(replicator.Id);
 			RefName incRefName = GetIncrementalRefName(replicator.Id);
@@ -384,7 +384,7 @@ namespace HordeServer.Replicators
 			BlobSerializerOptions blobOptions = new BlobSerializerOptions();
 			try
 			{
-				replicator = await WriteInternalAsync(replicator, change, streamConfig, replicatorOptions, blobOptions, cancellationToken);
+				replicator = await WriteInternalAsync(replicator, change, globalConfig, streamConfig, replicatorOptions, blobOptions, cancellationToken);
 				return replicator;
 			}
 			catch (OperationCanceledException ex)
@@ -410,7 +410,7 @@ namespace HordeServer.Replicators
 			}
 		}
 
-		async Task<IReplicator> WriteInternalAsync(IReplicator replicator, int change, StreamConfig streamConfig, PerforceReplicationOptions options, BlobSerializerOptions blobOptions, CancellationToken cancellationToken = default)
+		async Task<IReplicator> WriteInternalAsync(IReplicator replicator, int change, GlobalConfig globalConfig, StreamConfig streamConfig, PerforceReplicationOptions options, BlobSerializerOptions blobOptions, CancellationToken cancellationToken = default)
 		{
 			using IStorageClient store = _storageService.CreateClient(Namespace.Perforce);
 
@@ -478,7 +478,7 @@ namespace HordeServer.Replicators
 			}
 
 			// Create a client to replicate from this stream
-			ReplicationClient clientInfo = await FindOrAddReplicationClientAsync(streamConfig, cancellationToken);
+			ReplicationClient clientInfo = await FindOrAddReplicationClientAsync(globalConfig, streamConfig, cancellationToken);
 
 			// Connect to the server and flush the workspace
 			using IPerforceConnection perforce = await PerforceConnection.CreateAsync(clientInfo.Settings, _logger);
@@ -893,7 +893,7 @@ namespace HordeServer.Replicators
 			return clientInfo;
 		}
 
-		async Task<ReplicationClient> FindOrAddReplicationClientAsync(StreamConfig streamConfig, CancellationToken cancellationToken = default)
+		async Task<ReplicationClient> FindOrAddReplicationClientAsync(GlobalConfig globalConfig, StreamConfig streamConfig, CancellationToken cancellationToken = default)
 		{
 			ReplicationClient? clientInfo = await FindReplicationClientAsync(streamConfig, cancellationToken);
 			if (clientInfo == null)
@@ -906,7 +906,6 @@ namespace HordeServer.Replicators
 
 				InfoRecord serverInfo = await perforce.GetInfoAsync(InfoOptions.ShortOutput, cancellationToken);
 
-				GlobalConfig globalConfig = streamConfig.ProjectConfig.GlobalConfig;
 				PerforceCluster cluster = globalConfig.GetPerforceCluster(streamConfig.ClusterName);
 				bool partitioned = cluster.SupportsPartitionedWorkspaces;
 
diff --git a/Engine/Source/Programs/Horde/HordeServer/Replicators/ReplicationService.cs b/Engine/Source/Programs/Horde/HordeServer/Replicators/ReplicationService.cs
--- a/Engine/Source/Programs/Horde/HordeServer/Replicators/ReplicationService.cs
+++ b/Engine/Source/Programs/Horde/HordeServer/Replicators/ReplicationService.cs
@@ -90,7 +90,7 @@ namespace HordeServer.Replicators
 								if (replicatorConfig.Enabled)
 								{
 									ReplicatorId replicatorId = new ReplicatorId(streamConfig.Id, replicatorConfig.Id);
-									replicatorIdToTask.Add(replicatorId, BackgroundTask.StartNew(ctx => RunReplicationGuardedAsync(replicatorId, streamConfig, ctx)));
+									replicatorIdToTask.Add(replicatorId, BackgroundTask.StartNew(ctx => RunReplicationGuardedAsync(replicatorId, globalConfig, streamConfig, ctx)));
 									removeReplicators.Remove(replicatorId);
 								}
 							}
@@ -113,7 +113,7 @@ namespace HordeServer.Replicators
 			}
 		}
 
-		async Task RunReplicationGuardedAsync(ReplicatorId replicatorId, StreamConfig streamConfig, CancellationToken cancellationToken)
+		async Task RunReplicationGuardedAsync(ReplicatorId replicatorId, GlobalConfig globalConfig, StreamConfig streamConfig, CancellationToken cancellationToken)
 		{
 			_logger.LogInformation("Started background task for replication of {ReplicatorId}", replicatorId);
 
@@ -128,7 +128,7 @@ namespace HordeServer.Replicators
 				.Build();
 
 			PerforceReplicationOptions replicationOptions = new PerforceReplicationOptions();
-			await pipeline.ExecuteAsync(async ctx => await _replicator.RunAsync(replicatorId, streamConfig, replicationOptions, ctx), cancellationToken);
+			await pipeline.ExecuteAsync(async ctx => await _replicator.RunAsync(replicatorId, globalConfig, streamConfig, replicationOptions, ctx), cancellationToken);
 		}
 	}
 }
diff --git a/Engine/Source/Programs/Horde/HordeServer/Server/GlobalConfig.cs b/Engine/Source/Programs/Horde/HordeServer/Server/GlobalConfig.cs
--- a/Engine/Source/Programs/Horde/HordeServer/Server/GlobalConfig.cs
+++ b/Engine/Source/Programs/Horde/HordeServer/Server/GlobalConfig.cs
@@ -221,7 +221,7 @@ namespace HordeServer.Server
 			foreach (ProjectConfig project in Projects)
 			{
 				_projectLookup.Add(project.Id, project);
-				project.PostLoad(project.Id, this);
+				project.PostLoad(project.Id, Acl);
 
 				foreach (StreamConfig stream in project.Streams)
 				{
