Horde: Move IHealthMonitor into HordeServer.Shared.

#rnx

diff --git a/Engine/Source/Programs/Horde/HordeServer.Shared/Server/IHealthMonitor.cs b/Engine/Source/Programs/Horde/HordeServer.Shared/Server/IHealthMonitor.cs
new file mode 100644
--- /dev/null
+++ b/Engine/Source/Programs/Horde/HordeServer.Shared/Server/IHealthMonitor.cs
@@ -0,0 +1,30 @@
+// Copyright Epic Games, Inc. All Rights Reserved.
+
+using Microsoft.Extensions.Diagnostics.HealthChecks;
+
+namespace HordeServer.Server
+{
+	/// <summary>
+	/// Allows reporting the health of a subsystem
+	/// </summary>
+	public interface IHealthMonitor
+	{
+		/// <summary>
+		/// Sets the name used for reporting status of this service
+		/// </summary>
+		void SetName(string name);
+
+		/// <summary>
+		/// Updates the current health of a system
+		/// </summary>
+		Task UpdateAsync(HealthStatus result, string? message = null, DateTimeOffset? timestamp = null);
+	}
+
+	/// <summary>
+	/// Typed implementation of <see cref="IHealthMonitor"/>
+	/// </summary>
+	/// <typeparam name="T">Type of the subsystem</typeparam>
+	public interface IHealthMonitor<T> : IHealthMonitor
+	{
+	}
+}
diff --git a/Engine/Source/Programs/Horde/HordeServer/Server/ServerStatusService.cs b/Engine/Source/Programs/Horde/HordeServer/Server/ServerStatusService.cs
--- a/Engine/Source/Programs/Horde/HordeServer/Server/ServerStatusService.cs
+++ b/Engine/Source/Programs/Horde/HordeServer/Server/ServerStatusService.cs
@@ -14,30 +14,6 @@ using StackExchange.Redis;
 
 namespace HordeServer.Server;
 
-/// <summary>
-/// Allows reporting the health of a subsystem
-/// </summary>
-public interface IHealthMonitor
-{
-	/// <summary>
-	/// Sets the name used for reporting status of this service
-	/// </summary>
-	void SetName(string name);
-
-	/// <summary>
-	/// Updates the current health of a system
-	/// </summary>
-	Task UpdateAsync(HealthStatus result, string? message = null, DateTimeOffset? timestamp = null);
-}
-
-/// <summary>
-/// Typed implementation of <see cref="IHealthMonitor"/>
-/// </summary>
-/// <typeparam name="T">Type of the subsystem</typeparam>
-public interface IHealthMonitor<T> : IHealthMonitor
-{
-}
-
 internal class HealthMonitor<T> : IHealthMonitor<T>
 {
 	private readonly ServerStatusService _statusService;
