Horde: Test fix.

#rnx

diff --git a/Engine/Source/Programs/Horde/HordeServer/Server/GlobalConfig.cs b/Engine/Source/Programs/Horde/HordeServer/Server/GlobalConfig.cs
--- a/Engine/Source/Programs/Horde/HordeServer/Server/GlobalConfig.cs
+++ b/Engine/Source/Programs/Horde/HordeServer/Server/GlobalConfig.cs
@@ -352,6 +352,13 @@ namespace HordeServer.Server
 
 		void UpdateWorkspacesForPools()
 		{
+			// Try to get the compute config, and skip if it isn't configured
+			ComputeConfig? computeConfig;
+			if (!Plugins.TryGetValue(new PluginName("compute"), out computeConfig))
+			{
+				return;
+			}
+
 			// Lookup table of pool id to workspaces
 			Dictionary<PoolId, AutoSdkConfig> poolToAutoSdkView = new Dictionary<PoolId, AutoSdkConfig>();
 			Dictionary<PoolId, List<AgentWorkspaceInfo>> poolToAgentWorkspaces = new Dictionary<PoolId, List<AgentWorkspaceInfo>>();
@@ -390,7 +397,6 @@ namespace HordeServer.Server
 			}
 
 			// Update the list of workspaces for each pool
-			ComputeConfig computeConfig = Plugins.GetComputeConfig();
 			foreach (PoolConfig pool in computeConfig.Pools)
 			{
 				// Get the new list of workspaces for this pool
