Handle property split increment name in Gauntlet Issue Handler

#rnx

diff --git a/Engine/Source/Programs/Horde/HordeServer.Tests/Issues/IssueServiceTests.cs b/Engine/Source/Programs/Horde/HordeServer.Tests/Issues/IssueServiceTests.cs
--- a/Engine/Source/Programs/Horde/HordeServer.Tests/Issues/IssueServiceTests.cs
+++ b/Engine/Source/Programs/Horde/HordeServer.Tests/Issues/IssueServiceTests.cs
@@ -6,6 +6,7 @@ using System.Buffers;
 using System.Collections.Generic;
 using System.Linq;
 using System.Runtime.InteropServices;
+using System.Text;
 using System.Text.Json;
 using System.Threading;
 using System.Threading.Tasks;
@@ -1953,6 +1954,25 @@ namespace HordeServer.Tests.Issues
 			Assert.AreEqual("Errors in Update Version Files", issue.Summary);
 		}
 
+		static private IEnumerable<JsonLogEvent> MultilineLogEvent(LogLevel level, EventId eventId, string format, Dictionary<string, object> properties)
+		{
+			DateTime time = new DateTime(2024, 1, 1, 0, 0, 0, 0, DateTimeKind.Utc);
+			string message = MessageTemplate.Render(format, properties);
+			LogEvent baseEvent = new LogEvent(time, level, eventId, message, format, properties, null);
+			JsonLogEvent jsonLogEvent = new JsonLogEvent(baseEvent);
+			ServerLogPacketBuilder writer = new ServerLogPacketBuilder();
+			writer.SanitizeAndWriteEvent(jsonLogEvent);
+			string[] output = Encoding.UTF8.GetString(writer.CreatePacket().Item1.Span).Split('\n');
+			foreach (string outData in output)
+			{
+				JsonLogEvent outJsonLogEvent;
+				if (JsonLogEvent.TryParse(Encoding.UTF8.GetBytes(outData), out outJsonLogEvent))
+				{
+					yield return outJsonLogEvent;
+				}
+			}
+		}
+
 		[TestMethod]
 		public async Task GauntletTestAsync()
 		{
@@ -2013,17 +2033,31 @@ namespace HordeServer.Tests.Issues
 				await UpdateCompleteStepAsync(job, 0, 0, JobStepOutcome.Failure);
 				await using (TestJsonLogger logger = await CreateLoggerAsync(job, 0, 1))
 				{
-					logger.LogError(KnownLogEvents.Gauntlet_FatalEvent, "Engine encounter a critical failure\n{Summary}\n{Callstack}", logMessage, logAlternateCallstack);
+					Dictionary<string, object> properties = new Dictionary<string, object>
+					{
+						["Summary"] = logMessage,
+						["Callstack"] = logAlternateCallstack
+					};
+					foreach (JsonLogEvent outJsonLogEvent in MultilineLogEvent(LogLevel.Error, KnownLogEvents.Gauntlet_FatalEvent, "Engine encounter a critical failure\n{Summary}\n{Callstack}", properties))
+					{
+						logger.LogJsonLogEvent(outJsonLogEvent);
+					}
 				}
 				await UpdateCompleteStepAsync(job, 0, 1, JobStepOutcome.Failure);
+				await using (TestJsonLogger logger = await CreateLoggerAsync(job, 0, 2))
+				{
+					LogValue logSummary = new LogValue(new Utf8String("Summary"), logMessage);
+					logger.LogError(KnownLogEvents.Gauntlet_FatalEvent, "{Summary}\n{Callstack}", logSummary, logAlternateCallstack);
+				}
+				await UpdateCompleteStepAsync(job, 0, 2, JobStepOutcome.Failure);
 
 				IReadOnlyList<IIssue> issues = await IssueCollection.FindIssuesAsync();
 				Assert.AreEqual(1, issues.Count);
 				Assert.AreEqual(1, issues[0].Fingerprints.Count);
-				Assert.AreEqual("Gauntlet:fatal:with-callstack", issues[0].Fingerprints[0].Type);
+				Assert.AreEqual("Gauntlet:fatal:with-callstack:", issues[0].Fingerprints[0].Type.Substring(0, 30));
 				Assert.AreEqual("hash:", issues[0].Fingerprints[0].Keys.First().Name.Substring(0, 5));
 				Assert.AreEqual(37, issues[0].Fingerprints[0].Keys.First().Name.Length);
-				Assert.AreEqual("Automation fatal errors in Update Version Files and Compile UnrealHeaderTool Win64", issues[0].Summary);
+				Assert.AreEqual("Automation fatal errors in Update Version Files, Compile UnrealHeaderTool Win64 and Compile ShooterGameEditor Win64", issues[0].Summary);
 			}
 			// #3
 			// Scenario: Gauntlet Test + Fatal event
@@ -2054,6 +2088,7 @@ namespace HordeServer.Tests.Issues
 				}
 				await UpdateCompleteStepAsync(job, 0, 0, JobStepOutcome.Failure);
 				await UpdateCompleteStepAsync(job, 0, 1, JobStepOutcome.Success);
+				await UpdateCompleteStepAsync(job, 0, 2, JobStepOutcome.Success);
 
 				IReadOnlyList<IIssue> issues = await IssueCollection.FindIssuesAsync();
 				Assert.AreEqual(2, issues.Count);
diff --git a/Engine/Source/Programs/Shared/EpicGames.Horde/Issues/Handlers/GauntletIssueHandler.cs b/Engine/Source/Programs/Shared/EpicGames.Horde/Issues/Handlers/GauntletIssueHandler.cs
--- a/Engine/Source/Programs/Shared/EpicGames.Horde/Issues/Handlers/GauntletIssueHandler.cs
+++ b/Engine/Source/Programs/Shared/EpicGames.Horde/Issues/Handlers/GauntletIssueHandler.cs
@@ -44,12 +44,12 @@ namespace EpicGames.Horde.Issues.Handlers
 		/// <summary>
 		/// Callstack log type property
 		/// </summary>
-		static readonly Utf8String s_callstackLogType = new Utf8String("Callstack");
+		const string CallstackLogType = "Callstack";
 
 		/// <summary>
 		/// Summary log type property
 		/// </summary>
-		static readonly Utf8String s_summaryLogType = new Utf8String("Summary");
+		const string SummaryLogType = "Summary";
 
 		/// <summary>
 		/// Max Message Length to hash
@@ -117,6 +117,7 @@ namespace EpicGames.Horde.Issues.Handlers
 					key += $":{_context.StreamId}:{_context.NodeName}";
 				}
 				keys.Add(key, IssueKeyType.None);
+				metadata.Add("Hash", hash.ToString());
 			}
 			else
 			{
@@ -150,23 +151,41 @@ namespace EpicGames.Horde.Issues.Handlers
 
 		private static bool EventHasCallstackProperty(IssueEvent issueEvent)
 		{
-			return issueEvent.Lines.Any(x => FindNestedPropertyOfType(x, s_callstackLogType) != null);
+			return issueEvent.Lines.Any(x => FindNestedPropertyOfType(x, CallstackLogType) != null);
 		}
 
 		private static string? GetSummaryProperty(IssueEvent issueEvent)
 		{
+			StringBuilder? summary = null;
 			foreach (JsonLogEvent logEvent in issueEvent.Lines)
 			{
-				JsonProperty? property = FindNestedPropertyOfType(logEvent, s_summaryLogType);
+				JsonProperty? property = FindNestedPropertyOfType(logEvent, SummaryLogType);
 				if (property != null)
 				{
-					return property.Value.Value.GetString();
+					if (summary == null)
+					{
+						summary = new StringBuilder();
+					}
+					JsonElement value = property.Value.Value;
+					if (value.ValueKind == JsonValueKind.String
+						// handle LogValue type
+						|| (value.TryGetProperty(LogEventPropertyName.Text.Span, out value) && value.ValueKind == JsonValueKind.String))
+					{
+						summary.Append(value.GetString() + '\n');
+						continue;
+					}
+					summary.Append(value.ToString() + '\n');
+				}
+				else if (summary != null)
+				{
+					// when property is null but not summary, we early exit since we expect property split to be contiguous
+					return summary.ToString();
 				}
 			}
-			return null;
+			return summary?.ToString();
 		}
 
-		private static JsonProperty? FindNestedPropertyOfType(JsonLogEvent logEvent, Utf8String type)
+		private static JsonProperty? FindNestedPropertyOfType(JsonLogEvent logEvent, string searchType)
 		{
 			JsonElement line = JsonDocument.Parse(logEvent.Data).RootElement;
 			JsonElement properties;
@@ -174,8 +193,13 @@ namespace EpicGames.Horde.Issues.Handlers
 			{
 				foreach (JsonProperty property in properties.EnumerateObject())
 				{
-					if (property.NameEquals(type.Span))
+					if (property.Name.StartsWith(searchType, System.StringComparison.OrdinalIgnoreCase))
 					{
+						// if name is longer, check if it is a split property pattern: {name}${index}
+						if (property.Name.Length > searchType.Length && property.Name.Substring(searchType.Length, 1) != "$")
+						{
+							continue;
+						}
 						return property;
 					}
 				}
@@ -197,7 +221,8 @@ namespace EpicGames.Horde.Issues.Handlers
 				issue.Metadata.Add("GauntletType", gauntletType);
 				if (hasCallstack)
 				{
-					issue.Type = $"{issue.Type}:with-callstack";
+					string? hash = issue.Metadata.FindValues("Hash").FirstOrDefault();
+					issue.Type = $"{issue.Type}:with-callstack:{hash}";
 				}
 				_issues.Add(issue);
 
