Horde: Respect Forwarded-For header for client's IP during compute cluster ID resolving

Also return a better error message when no network range or cluster is found.

diff --git a/Engine/Source/Programs/Horde/Plugins/Compute/HordeServer.Compute/Compute/ComputeController.cs b/Engine/Source/Programs/Horde/Plugins/Compute/HordeServer.Compute/Compute/ComputeController.cs
--- a/Engine/Source/Programs/Horde/Plugins/Compute/HordeServer.Compute/Compute/ComputeController.cs
+++ b/Engine/Source/Programs/Horde/Plugins/Compute/HordeServer.Compute/Compute/ComputeController.cs
@@ -53,10 +53,27 @@ namespace HordeServer.Compute
 		public ActionResult<GetClusterResponse> GetCluster([FromBody] AssignComputeRequest request)
 		{
 			using TelemetrySpan span = _tracer.StartActiveSpan($"{nameof(ComputeController)}.{nameof(GetCluster)}");
-			IPAddress requesterIp = ComputeService.ResolveRequesterIp(HttpContext.Connection.RemoteIpAddress, request.Connection?.PreferPublicIp, request.Connection?.ClientPublicIp);
-			ClusterId clusterId = ComputeService.FindBestComputeClusterId(_computeConfig.Value, requesterIp);
-			GetClusterResponse response = new () { ClusterId = clusterId };
-			return response;
+			IPAddress? requesterIp = null;
+			try
+			{
+				string? forwardedForHeader = HttpContext.Request.Headers["X-Forwarded-For"].FirstOrDefault();
+				IPAddress? clientIp = String.IsNullOrEmpty(forwardedForHeader)
+					? HttpContext.Connection.RemoteIpAddress
+					: IPAddress.Parse(forwardedForHeader);
+				
+				span.SetAttribute("clientIp", clientIp?.ToString());
+				requesterIp = ComputeService.ResolveRequesterIp(clientIp, request.Connection?.PreferPublicIp, request.Connection?.ClientPublicIp);
+				span.SetAttribute("requesterIp", requesterIp.ToString());
+				ClusterId clusterId = ComputeService.FindBestComputeClusterId(_computeConfig.Value, requesterIp);
+				span.SetAttribute("clusterId", clusterId.ToString());
+				
+				GetClusterResponse response = new() { ClusterId = clusterId };
+				return response;
+			}
+			catch (ComputeServiceException)
+			{
+				return StatusCode((int)HttpStatusCode.NotFound, $"Unable to resolve a compute cluster ID for IP {requesterIp?.ToString() ?? "null"}");
+			}
 		}
 
 		/// <summary>
