Horde: Move build functionality into a plugin.

diff --git a/Engine/Source/Programs/Horde/Drivers/JobDriver/JobDriver.csproj b/Engine/Source/Programs/Horde/Drivers/JobDriver/JobDriver.csproj
--- a/Engine/Source/Programs/Horde/Drivers/JobDriver/JobDriver.csproj
+++ b/Engine/Source/Programs/Horde/Drivers/JobDriver/JobDriver.csproj
@@ -91,9 +91,9 @@
   </ItemGroup>
   
   <ItemGroup>
-    <Protobuf Include="..\..\HordeServer\Protos\horde\job_rpc.proto" ProtoRoot="..\..\HordeServer\Protos" GrpcServices="Client" Access="Public" AdditionalImportDirs="..\..\..\Shared\EpicGames.Horde\Protos;..\..\Plugins\Compute\HordeServer.Compute\Protos"/>
-    <Protobuf Include="..\..\HordeServer\Protos\horde\job_rpc_messages.proto" ProtoRoot="..\..\HordeServer\Protos" Access="Public" AdditionalImportDirs="..\..\Plugins\Compute\HordeServer.Compute\Protos" />
-    <Protobuf Include="..\..\HordeServer\Protos\horde\common\label_change.proto" ProtoRoot="..\..\HordeServer\Protos" Access="Public" />
+    <Protobuf Include="..\..\Plugins\Build\HordeServer.Build\Protos\horde\job_rpc.proto" ProtoRoot="..\..\Plugins\Build\HordeServer.Build\Protos" GrpcServices="Client" Access="Public" AdditionalImportDirs="..\..\..\Shared\EpicGames.Horde\Protos;..\..\Plugins\Compute\HordeServer.Compute\Protos"/>
+    <Protobuf Include="..\..\Plugins\Build\HordeServer.Build\Protos\horde\job_rpc_messages.proto" ProtoRoot="..\..\Plugins\Build\HordeServer.Build\Protos" Access="Public" AdditionalImportDirs="..\..\Plugins\Compute\HordeServer.Compute\Protos" />
+    <Protobuf Include="..\..\Plugins\Build\HordeServer.Build\Protos\horde\common\label_change.proto" ProtoRoot="..\..\Plugins\Build\HordeServer.Build\Protos" Access="Public" />
     <Protobuf Include="..\..\Plugins\Compute\HordeServer.Compute\Protos\horde\messages\agent_workspace.proto" ProtoRoot="..\..\Plugins\Compute\HordeServer.Compute\Protos" Access="Public" />
     <Protobuf Include="..\..\Plugins\Compute\HordeServer.Compute\Protos\horde\tasks\conform_task.proto" ProtoRoot="..\..\Plugins\Compute\HordeServer.Compute\Protos" Access="Public" />
     <Protobuf Include="..\..\Plugins\Compute\HordeServer.Compute\Protos\horde\tasks\job_task.proto" ProtoRoot="..\..\Plugins\Compute\HordeServer.Compute\Protos" Access="Public" />
diff --git a/Engine/Source/Programs/Horde/HordeAgent/HordeAgent.csproj b/Engine/Source/Programs/Horde/HordeAgent/HordeAgent.csproj
--- a/Engine/Source/Programs/Horde/HordeAgent/HordeAgent.csproj
+++ b/Engine/Source/Programs/Horde/HordeAgent/HordeAgent.csproj
@@ -142,10 +142,8 @@
     <Protobuf Include="..\Plugins\Compute\HordeServer.Compute\Protos\horde\enrollment_rpc.proto" ProtoRoot="..\Plugins\Compute\HordeServer.Compute\Protos" GrpcServices="Client" Access="Public" AdditionalImportDirs="..\Plugins\Compute\HordeServer.Compute\Protos"/>
     <Protobuf Include="..\Plugins\Compute\HordeServer.Compute\Protos\horde\enrollment_rpc_messages.proto" ProtoRoot="..\Plugins\Compute\HordeServer.Compute\Protos" Access="Public" AdditionalImportDirs="..\Plugins\Compute\HordeServer.Compute\Protos"/>
     <Protobuf Include="..\Plugins\Compute\HordeServer.Compute\Protos\horde\common\agent_status.proto" ProtoRoot="..\Plugins\Compute\HordeServer.Compute\Protos" Access="Public" AdditionalImportDirs="..\Plugins\Compute\HordeServer.Compute\Protos"/>
-    <Protobuf Include="..\HordeServer\Protos\horde\common\io_hash.proto" ProtoRoot="..\HordeServer\Protos" Access="Public" />
     <Protobuf Include="..\Plugins\Compute\HordeServer.Compute\Protos\horde\common\lease_outcome.proto" ProtoRoot="..\Plugins\Compute\HordeServer.Compute\Protos" Access="Public" AdditionalImportDirs="..\Plugins\Compute\HordeServer.Compute\Protos"/>
     <Protobuf Include="..\Plugins\Compute\HordeServer.Compute\Protos\horde\common\lease_state.proto" ProtoRoot="..\Plugins\Compute\HordeServer.Compute\Protos" Access="Public" AdditionalImportDirs="..\Plugins\Compute\HordeServer.Compute\Protos"/>
-    <Protobuf Include="..\HordeServer\Protos\horde\common\namespace_id.proto" ProtoRoot="..\HordeServer\Protos" Access="Public" />
     <Protobuf Include="..\Plugins\Compute\HordeServer.Compute\Protos\horde\messages\agent_capabilities.proto" ProtoRoot="..\Plugins\Compute\HordeServer.Compute\Protos" Access="Public" AdditionalImportDirs="..\Plugins\Compute\HordeServer.Compute\Protos"/>
     <Protobuf Include="..\Plugins\Compute\HordeServer.Compute\Protos\horde\messages\agent_requirements.proto" ProtoRoot="..\Plugins\Compute\HordeServer.Compute\Protos" Access="Public" AdditionalImportDirs="..\Plugins\Compute\HordeServer.Compute\Protos"/>
     <Protobuf Include="..\Plugins\Compute\HordeServer.Compute\Protos\horde\messages\lease.proto" ProtoRoot="..\Plugins\Compute\HordeServer.Compute\Protos" Access="Public" AdditionalImportDirs="..\Plugins\Compute\HordeServer.Compute\Protos"/>
@@ -156,8 +154,8 @@
     <Protobuf Include="..\Plugins\Compute\HordeServer.Compute\Protos\horde\tasks\test_task.proto" ProtoRoot="..\Plugins\Compute\HordeServer.Compute\Protos" Access="Public" AdditionalImportDirs="..\Plugins\Compute\HordeServer.Compute\Protos"/>
     <Protobuf Include="..\Plugins\Compute\HordeServer.Compute\Protos\horde\tasks\upgrade_task.proto" ProtoRoot="..\Plugins\Compute\HordeServer.Compute\Protos" Access="Public" AdditionalImportDirs="..\Plugins\Compute\HordeServer.Compute\Protos"/>
 
-    <Protobuf Include="..\HordeServer\Protos\horde\job_rpc_messages.proto" ProtoRoot="..\HordeServer\Protos" Access="Public" AdditionalImportDirs="..\Plugins\Compute\HordeServer.Compute\Protos" />
-    <Protobuf Include="..\HordeServer\Protos\horde\common\label_change.proto" ProtoRoot="..\HordeServer\Protos" Access="Public" AdditionalImportDirs="..\Plugins\Compute\HordeServer.Compute\Protos" />
+    <Protobuf Include="..\Plugins\Build\HordeServer.Build\Protos\horde\job_rpc_messages.proto" ProtoRoot="..\Plugins\Build\HordeServer.Build\Protos" Access="Public" AdditionalImportDirs="..\Plugins\Compute\HordeServer.Compute\Protos" />
+    <Protobuf Include="..\Plugins\Build\HordeServer.Build\Protos\horde\common\label_change.proto" ProtoRoot="..\Plugins\Build\HordeServer.Build\Protos" Access="Public" AdditionalImportDirs="..\Plugins\Compute\HordeServer.Compute\Protos" />
     <Protobuf Include="..\Plugins\Compute\HordeServer.Compute\Protos\horde\messages\agent_workspace.proto" ProtoRoot="..\Plugins\Compute\HordeServer.Compute\Protos" Access="Public" />
     <Protobuf Include="..\Plugins\Compute\HordeServer.Compute\Protos\horde\tasks\conform_task.proto" ProtoRoot="..\Plugins\Compute\HordeServer.Compute\Protos" Access="Public" />
     <Protobuf Include="..\Plugins\Compute\HordeServer.Compute\Protos\horde\tasks\job_task.proto" ProtoRoot="..\Plugins\Compute\HordeServer.Compute\Protos" Access="Public" />
diff --git a/Engine/Source/Programs/Horde/HordeServer.Shared/Configuration/IConfigService.cs b/Engine/Source/Programs/Horde/HordeServer.Shared/Configuration/IConfigService.cs
new file mode 100644
--- /dev/null
+++ b/Engine/Source/Programs/Horde/HordeServer.Shared/Configuration/IConfigService.cs
@@ -0,0 +1,26 @@
+﻿// Copyright Epic Games, Inc. All Rights Reserved.
+
+using System;
+using System.Collections.Generic;
+using System.Linq;
+using System.Text;
+using System.Threading.Tasks;
+
+namespace HordeServer.Configuration
+{
+	/// <summary>
+	/// Interface for the config service
+	/// </summary>
+	public interface IConfigService
+	{
+		/// <summary>
+		/// Event for notifications that the config has been updated
+		/// </summary>
+		event Action<ConfigUpdateInfo>? OnConfigUpdate;
+
+		/// <summary>
+		/// Validate a new set of config files. Parses and runs PostLoad methods on them.
+		/// </summary>
+		Task<string?> ValidateAsync(Dictionary<Uri, byte[]> files, CancellationToken cancellationToken);
+	}
+}
diff --git a/Engine/Source/Programs/Horde/HordeServer.Shared/IServerInfo.cs b/Engine/Source/Programs/Horde/HordeServer.Shared/IServerInfo.cs
--- a/Engine/Source/Programs/Horde/HordeServer.Shared/IServerInfo.cs
+++ b/Engine/Source/Programs/Horde/HordeServer.Shared/IServerInfo.cs
@@ -78,6 +78,16 @@ namespace HordeServer
 		/// </summary>
 		bool EnableDebugEndpoints { get; }
 
+		/// <summary>
+		/// Url to use for generating links back to the server
+		/// </summary>
+		Uri ServerUrl { get; }
+
+		/// <summary>
+		/// Url to use for generating links back to the dashboard.
+		/// </summary>
+		Uri DashboardUrl { get; }
+
 		/// <summary>
 		/// Helper method to check if this process has activated the given mode
 		/// </summary>
diff --git a/Engine/Source/Programs/Horde/HordeServer.Shared/Utilities/HordeClaimTypes.cs b/Engine/Source/Programs/Horde/HordeServer.Shared/Utilities/HordeClaimTypes.cs
--- a/Engine/Source/Programs/Horde/HordeServer.Shared/Utilities/HordeClaimTypes.cs
+++ b/Engine/Source/Programs/Horde/HordeServer.Shared/Utilities/HordeClaimTypes.cs
@@ -3,6 +3,8 @@
 using System.Security.Claims;
 using EpicGames.Horde.Accounts;
 using EpicGames.Horde.Users;
+using HordeServer.Acls;
+using HordeServer.Server;
 using HordeServer.Users;
 
 namespace HordeServer.Utilities
@@ -124,7 +126,26 @@ namespace HordeServer.Utilities
 			string? idValue = principal.FindFirstValue(HordeClaimTypes.UserId);
 			return idValue == null ? null : UserId.Parse(idValue);
 		}
-		
+
+		/// <summary>
+		/// Determines whether the given user can masquerade as a given user
+		/// </summary>
+		/// <param name="user"></param>
+		/// <param name="userId"></param>
+		/// <returns></returns>
+		public static bool AuthorizeAsUser(this AclConfig aclConfig, ClaimsPrincipal user, UserId userId)
+		{
+			UserId? currentUserId = user.GetUserId();
+			if (currentUserId != null && currentUserId.Value == userId)
+			{
+				return true;
+			}
+			else
+			{
+				return aclConfig.Authorize(ServerAclAction.Impersonate, user);
+			}
+		}
+
 		/// <summary>
 		/// Gets the Horde username from a principal
 		/// </summary>
diff --git a/Engine/Source/Programs/Horde/HordeServer.Tests.Shared/ServerServiceTest.cs b/Engine/Source/Programs/Horde/HordeServer.Tests.Shared/ServerServiceTest.cs
--- a/Engine/Source/Programs/Horde/HordeServer.Tests.Shared/ServerServiceTest.cs
+++ b/Engine/Source/Programs/Horde/HordeServer.Tests.Shared/ServerServiceTest.cs
@@ -109,6 +109,7 @@ namespace HordeServer.Tests
 			services.AddSingleton(typeof(IAuditLogFactory<>), typeof(AuditLogFactory<>));
 
 			services.AddSingleton<ConfigService>();
+			services.AddSingleton<IConfigService>(sp => sp.GetRequiredService<ConfigService>());
 			services.AddSingleton<IOptionsFactory<GlobalConfig>>(sp => sp.GetRequiredService<ConfigService>());
 			services.AddSingleton<IOptionsChangeTokenSource<GlobalConfig>>(sp => sp.GetRequiredService<ConfigService>());
 
diff --git a/Engine/Source/Programs/Horde/HordeServer.Tests/Acls/AclTests.cs b/Engine/Source/Programs/Horde/HordeServer.Tests/Acls/AclTests.cs
--- a/Engine/Source/Programs/Horde/HordeServer.Tests/Acls/AclTests.cs
+++ b/Engine/Source/Programs/Horde/HordeServer.Tests/Acls/AclTests.cs
@@ -44,7 +44,7 @@ namespace HordeServer.Tests.Acls
 						ExcludeActions = new List<AclAction>{ JobAclAction.RetryJobStep }
 					},
 				};
-				globalConfig.Projects.Add(projectConfig);
+				globalConfig.Plugins.GetBuildConfig().Projects.Add(projectConfig);
 
 				StreamConfig streamConfig = new StreamConfig();
 				streamConfig.Id = new StreamId("ue5-main");
@@ -68,7 +68,7 @@ namespace HordeServer.Tests.Acls
 			});
 
 			GlobalConfig globalConfig = GlobalConfig.CurrentValue;
-			globalConfig.TryGetStream(new StreamId("ue5-main"), out StreamConfig? streamConfig);
+			globalConfig.Plugins.GetBuildConfig().TryGetStream(new StreamId("ue5-main"), out StreamConfig? streamConfig);
 
 			ClaimsIdentity identity = new ClaimsIdentity(null, new[] { new Claim("foo", "bar") });
 			ClaimsPrincipal principal = new ClaimsPrincipal(identity);
@@ -86,7 +86,7 @@ namespace HordeServer.Tests.Acls
 			{
 				ProjectConfig projectConfig = new ProjectConfig();
 				projectConfig.Id = new ProjectId("ue5");
-				globalConfig.Projects.Add(projectConfig);
+				globalConfig.Plugins.GetBuildConfig().Projects.Add(projectConfig);
 
 				StreamConfig streamConfig = new StreamConfig();
 				streamConfig.Id = new StreamId("ue5-main");
@@ -102,7 +102,7 @@ namespace HordeServer.Tests.Acls
 			});
 
 			GlobalConfig globalConfig = GlobalConfig.CurrentValue;
-			globalConfig.TryGetStream(new StreamId("ue5-main"), out StreamConfig? streamConfig);
+			globalConfig.Plugins.GetBuildConfig().TryGetStream(new StreamId("ue5-main"), out StreamConfig? streamConfig);
 
 			{
 				ClaimsIdentity identity = new ClaimsIdentity(null, new[] { new Claim("foo", "bar") });
diff --git a/Engine/Source/Programs/Horde/HordeServer.Tests/Agents/Pools/PoolUpdateServiceTest.cs b/Engine/Source/Programs/Horde/HordeServer.Tests/Agents/Pools/PoolUpdateServiceTest.cs
--- a/Engine/Source/Programs/Horde/HordeServer.Tests/Agents/Pools/PoolUpdateServiceTest.cs
+++ b/Engine/Source/Programs/Horde/HordeServer.Tests/Agents/Pools/PoolUpdateServiceTest.cs
@@ -8,7 +8,9 @@ using System.Threading.Tasks;
 using EpicGames.Horde.Agents;
 using HordeServer.Agents;
 using HordeServer.Agents.Pools;
+using Microsoft.Extensions.DependencyInjection;
 using Microsoft.Extensions.Logging.Abstractions;
+using Microsoft.Extensions.Options;
 using Microsoft.VisualStudio.TestTools.UnitTesting;
 
 namespace HordeServer.Tests.Agents.Pools;
@@ -25,7 +27,7 @@ public class PoolUpdateServiceTest : TestSetup
 	public PoolUpdateServiceTest()
 	{
 		UpdateConfig(x => x.Plugins.GetComputeConfig().Pools.Clear());
-		_pus = new PoolUpdateService(AgentService, PoolCollection, Clock, GlobalConfig, Tracer, new NullLogger<PoolUpdateService>());
+		_pus = new PoolUpdateService(AgentService, PoolCollection, Clock, ServiceProvider.GetRequiredService<IOptionsMonitor<BuildConfig>>(), Tracer, new NullLogger<PoolUpdateService>());
 	}
 
 	[TestInitialize]
@@ -69,7 +71,7 @@ public class PoolUpdateServiceTest : TestSetup
 	public async Task ShutdownDisabledAgents_WithGlobalGracePeriod_RequestsShutdownAsync()
 	{
 		// Arrange
-		UpdateConfig(config => config.AgentShutdownIfDisabledGracePeriod = TimeSpan.FromHours(3.0));
+		UpdateConfig(config => config.Plugins.GetBuildConfig().AgentShutdownIfDisabledGracePeriod = TimeSpan.FromHours(3.0));
 		
 		// Act
 		await _pus.ShutdownDisabledAgentsAsync(CancellationToken.None);
diff --git a/Engine/Source/Programs/Horde/HordeServer.Tests/Artifacts/ArtifactTests.cs b/Engine/Source/Programs/Horde/HordeServer.Tests/Artifacts/ArtifactTests.cs
--- a/Engine/Source/Programs/Horde/HordeServer.Tests/Artifacts/ArtifactTests.cs
+++ b/Engine/Source/Programs/Horde/HordeServer.Tests/Artifacts/ArtifactTests.cs
@@ -10,7 +10,6 @@ using EpicGames.Horde.Artifacts;
 using EpicGames.Horde.Streams;
 using HordeCommon;
 using HordeServer.Artifacts;
-using HordeServer.Server;
 using Microsoft.Extensions.DependencyInjection;
 using Microsoft.VisualStudio.TestTools.UnitTesting;
 
@@ -52,7 +51,7 @@ namespace HordeServer.Tests.Artifacts
 			ArtifactExpirationService expirationService = ServiceProvider.GetRequiredService<ArtifactExpirationService>();
 
 			ArtifactType type = new ArtifactType("my-artifact");
-			UpdateConfig(config => config.ArtifactTypes.Add(new ArtifactTypeConfig { Type = type, KeepDays = 1 }));
+			UpdateConfig(config => config.Plugins.GetBuildConfig().ArtifactTypes.Add(new ArtifactTypeConfig { Type = type, KeepDays = 1 }));
 
 			await expirationService.StartAsync(CancellationToken.None);
 
@@ -81,7 +80,7 @@ namespace HordeServer.Tests.Artifacts
 			DateTime startTime = Clock.UtcNow;
 
 			ArtifactType type = new ArtifactType("my-artifact");
-			UpdateConfig(config => config.ArtifactTypes.Add(new ArtifactTypeConfig { Type = type, KeepDays = 1 }));
+			UpdateConfig(config => config.Plugins.GetBuildConfig().ArtifactTypes.Add(new ArtifactTypeConfig { Type = type, KeepDays = 1 }));
 
 			ArtifactExpirationService expirationService = ServiceProvider.GetRequiredService<ArtifactExpirationService>();
 
@@ -98,7 +97,7 @@ namespace HordeServer.Tests.Artifacts
 				Assert.AreEqual(artifact.Id, artifacts[0].Id);
 			}
 
-			UpdateConfig(config => config.ArtifactTypes = new List<ArtifactTypeConfig> { new ArtifactTypeConfig { Type = type, KeepDays = 4 } });
+			UpdateConfig(config => config.Plugins.GetBuildConfig().ArtifactTypes = new List<ArtifactTypeConfig> { new ArtifactTypeConfig { Type = type, KeepDays = 4 } });
 			await Clock.AdvanceAsync(TimeSpan.FromDays(2.0));
 
 			{
@@ -106,7 +105,7 @@ namespace HordeServer.Tests.Artifacts
 				Assert.AreEqual(1, artifacts.Count);
 			}
 
-			UpdateConfig(config => config.ArtifactTypes = new List<ArtifactTypeConfig> { new ArtifactTypeConfig { Type = type, KeepDays = 1 } });
+			UpdateConfig(config => config.Plugins.GetBuildConfig().ArtifactTypes = new List<ArtifactTypeConfig> { new ArtifactTypeConfig { Type = type, KeepDays = 1 } });
 			await Clock.AdvanceAsync(TimeSpan.FromDays(1.0));
 
 			{
@@ -121,7 +120,7 @@ namespace HordeServer.Tests.Artifacts
 			DateTime startTime = Clock.UtcNow;
 
 			ArtifactType type = new ArtifactType("my-artifact");
-			UpdateConfig(config => config.ArtifactTypes.Add(new ArtifactTypeConfig { Type = type, KeepCount = 4 }));
+			UpdateConfig(config => config.Plugins.GetBuildConfig().ArtifactTypes.Add(new ArtifactTypeConfig { Type = type, KeepCount = 4 }));
 
 			ArtifactExpirationService expirationService = ServiceProvider.GetRequiredService<ArtifactExpirationService>();
 
@@ -161,7 +160,7 @@ namespace HordeServer.Tests.Artifacts
 				}
 			}
 
-			UpdateConfig(config => config.ArtifactTypes = new List<ArtifactTypeConfig> { new ArtifactTypeConfig { Type = type, KeepDays = 1 } });
+			UpdateConfig(config => config.Plugins.GetBuildConfig().ArtifactTypes = new List<ArtifactTypeConfig> { new ArtifactTypeConfig { Type = type, KeepDays = 1 } });
 			await Clock.AdvanceAsync(TimeSpan.FromDays(2.0));
 
 			{
diff --git a/Engine/Source/Programs/Horde/HordeServer.Tests/Configuration/ConfigTests.cs b/Engine/Source/Programs/Horde/HordeServer.Tests/Configuration/ConfigTests.cs
--- a/Engine/Source/Programs/Horde/HordeServer.Tests/Configuration/ConfigTests.cs
+++ b/Engine/Source/Programs/Horde/HordeServer.Tests/Configuration/ConfigTests.cs
@@ -172,16 +172,19 @@ namespace HordeServer.Tests.Configuration
 		[TestMethod]
 		public void WorkspaceConfig()
 		{
-			GlobalConfig gc = new();
 			Dictionary<string, WorkspaceConfig> inputWorkspaces = new();
 			inputWorkspaces["base"] = new WorkspaceConfig { Identifier = "base", Cluster = "myCluster", MinScratchSpace = 111 };
 			inputWorkspaces["subType"] = new WorkspaceConfig { Base = "base", Identifier = "subType", ConformDiskFreeSpace = 222 };
 			inputWorkspaces["subSubType"] = new WorkspaceConfig { Base = "subType", Identifier = "subSubType" };
-			
-			gc.Projects.Add(new ProjectConfig { Streams = [new StreamConfig { WorkspaceTypes = inputWorkspaces }]});
+
+			BuildConfig buildConfig = new BuildConfig();
+			buildConfig.Projects.Add(new ProjectConfig { Streams = [new StreamConfig { WorkspaceTypes = inputWorkspaces }]});
+
+			GlobalConfig gc = new();
+			gc.Plugins.AddBuildConfig(buildConfig);
 			gc.PostLoad(new ServerSettings(), new List<ILoadedPlugin>());
 
-			Dictionary<string,WorkspaceConfig> workspaces = gc.Projects[0].Streams[0].WorkspaceTypes;
+			Dictionary<string,WorkspaceConfig> workspaces = gc.Plugins.GetBuildConfig().Projects[0].Streams[0].WorkspaceTypes;
 			Assert.AreEqual(3, workspaces.Count);
 			Assert.AreEqual(111, workspaces["subType"].MinScratchSpace);
 			Assert.AreEqual(111, workspaces["subSubType"].MinScratchSpace);
@@ -191,10 +194,10 @@ namespace HordeServer.Tests.Configuration
 		[TestMethod]
 		public void WorkspaceInheritFromProjectConfig()
 		{
-			GlobalConfig gc = new();
-			
 			List<string> autoSdkViews = ["foo", "bar"];
-			gc.Projects.Add(new ProjectConfig
+
+			BuildConfig buildConfig = new BuildConfig();
+			buildConfig.Projects.Add(new ProjectConfig
 			{
 				WorkspaceTypes =
 				{
@@ -210,10 +213,12 @@ namespace HordeServer.Tests.Configuration
 					}}
 				]
 			});
-			
+
+			GlobalConfig gc = new();
+			gc.Plugins.AddBuildConfig(buildConfig);
 			gc.PostLoad(new ServerSettings(), new List<ILoadedPlugin>());
 
-			Dictionary<string, WorkspaceConfig> workspaces = gc.Projects[0].Streams[0].WorkspaceTypes;
+			Dictionary<string, WorkspaceConfig> workspaces = gc.Plugins.GetBuildConfig().Projects[0].Streams[0].WorkspaceTypes;
 			Assert.AreEqual(4, workspaces.Count);
 			
 			// Project-defined streams
diff --git a/Engine/Source/Programs/Horde/HordeServer.Tests/Devices/DeviceServiceTests.cs b/Engine/Source/Programs/Horde/HordeServer.Tests/Devices/DeviceServiceTests.cs
--- a/Engine/Source/Programs/Horde/HordeServer.Tests/Devices/DeviceServiceTests.cs
+++ b/Engine/Source/Programs/Horde/HordeServer.Tests/Devices/DeviceServiceTests.cs
@@ -158,7 +158,7 @@ namespace HordeServer.Tests.Devices
 				devices.Platforms.Add(platform);
 			}
 
-			UpdateConfig(x => x.Devices = devices);
+			UpdateConfig(x => x.Plugins.GetBuildConfig().Devices = devices);
 
 			for (int i = 1; i < 4; i++)
 			{
diff --git a/Engine/Source/Programs/Horde/HordeServer.Tests/Fixture.cs b/Engine/Source/Programs/Horde/HordeServer.Tests/Fixture.cs
--- a/Engine/Source/Programs/Horde/HordeServer.Tests/Fixture.cs
+++ b/Engine/Source/Programs/Horde/HordeServer.Tests/Fixture.cs
@@ -78,7 +78,11 @@ namespace HordeServer.Tests
 			ProjectId projectId = new ProjectId("ue5");
 			ProjectConfig projectConfig = new ProjectConfig { Id = projectId, Name = "UE5", Streams = new List<StreamConfig> { streamConfig } };
 
-			GlobalConfig globalConfig = new GlobalConfig { Projects = new List<ProjectConfig> { projectConfig } };
+			BuildConfig buildConfig = new BuildConfig();
+			buildConfig.Projects = new List<ProjectConfig> { projectConfig };
+
+			GlobalConfig globalConfig = new GlobalConfig();
+			globalConfig.Plugins.AddBuildConfig(buildConfig);
 			globalConfig.PostLoad(serverSettings, pluginCollection.LoadedPlugins);
 			configService.OverrideConfig(globalConfig);
 
diff --git a/Engine/Source/Programs/Horde/HordeServer.Tests/Fleet/JobQueueStrategyTest.cs b/Engine/Source/Programs/Horde/HordeServer.Tests/Fleet/JobQueueStrategyTest.cs
--- a/Engine/Source/Programs/Horde/HordeServer.Tests/Fleet/JobQueueStrategyTest.cs
+++ b/Engine/Source/Programs/Horde/HordeServer.Tests/Fleet/JobQueueStrategyTest.cs
@@ -17,6 +17,8 @@ using HordeServer.Jobs.Graphs;
 using HordeServer.Projects;
 using HordeServer.Server;
 using HordeServer.Streams;
+using Microsoft.Extensions.DependencyInjection;
+using Microsoft.Extensions.Options;
 using Microsoft.VisualStudio.TestTools.UnitTesting;
 
 namespace HordeServer.Tests.Fleet
@@ -138,7 +140,7 @@ namespace HordeServer.Tests.Fleet
 			projectConfig.Id = new ProjectId("ue5");
 			projectConfig.Streams.Add(streamConfig);
 
-			UpdateConfig(config => config.Projects = new List<ProjectConfig> { projectConfig });
+			UpdateConfig(config => config.Plugins.GetBuildConfig().Projects = new List<ProjectConfig> { projectConfig });
 
 			string nodeForAgentType1 = "bogusNodeOnAgentType" + ++s_uniqueId;
 			IGraph graph = await GraphCollection.AppendAsync(null, new()
@@ -161,7 +163,7 @@ namespace HordeServer.Tests.Fleet
 				await job.TryUpdateBatchAsync(job.Batches[0].Id, null, JobStepBatchState.Ready, null);
 			}
 
-			return (new(JobCollection, GraphCollection, StreamCollection, Clock, Cache, isDowntimeActive, GlobalConfig), poolSize, pool, agents);
+			return (new JobQueueStrategy(JobCollection, GraphCollection, StreamCollection, Clock, Cache, isDowntimeActive, ServiceProvider.GetRequiredService<IOptionsMonitor<BuildConfig>>()), poolSize, pool, agents);
 		}
 
 		private async Task<IJob> AddPlaceholderJobAsync(IGraph graph, StreamId streamId, string nodeNameToExecute)
diff --git a/Engine/Source/Programs/Horde/HordeServer.Tests/Issues/IssueServiceTests.cs b/Engine/Source/Programs/Horde/HordeServer.Tests/Issues/IssueServiceTests.cs
--- a/Engine/Source/Programs/Horde/HordeServer.Tests/Issues/IssueServiceTests.cs
+++ b/Engine/Source/Programs/Horde/HordeServer.Tests/Issues/IssueServiceTests.cs
@@ -24,7 +24,6 @@ using HordeServer.Issues;
 using HordeServer.Jobs;
 using HordeServer.Jobs.Graphs;
 using HordeServer.Logs;
-using HordeServer.Plugins;
 using HordeServer.Projects;
 using HordeServer.Server;
 using HordeServer.Storage;
@@ -155,9 +154,12 @@ namespace HordeServer.Tests.Issues
 			storageConfig.Backends.Add(new BackendConfig { Id = new BackendId("default-backend"), Type = StorageBackendType.Memory });
 			storageConfig.Namespaces.Add(new NamespaceConfig { Id = new NamespaceId("horde-logs"), Backend = new BackendId("default-backend"), GcDelayHrs = 0.0 });
 
+			BuildConfig buildConfig = new BuildConfig();
+			buildConfig.Projects.Add(projectConfig);
+
 			GlobalConfig globalConfig = new GlobalConfig();
-			globalConfig.Projects.Add(projectConfig);
-			globalConfig.Plugins.Add(new PluginName("storage"), storageConfig);
+			globalConfig.Plugins.AddBuildConfig(buildConfig);
+			globalConfig.Plugins.AddStorageConfig(storageConfig);
 			SetConfig(globalConfig);
 
 			static StreamConfig CreateStream(StreamId streamId, string streamName)
@@ -419,7 +421,7 @@ namespace HordeServer.Tests.Issues
 			// Scenario: Stream is deleted
 			// Expected: Issue is closed
 			{
-				UpdateConfig(x => x.Projects.Clear());
+				UpdateConfig(x => x.Plugins.GetBuildConfig().Projects.Clear());
 				await Clock.AdvanceAsync(TimeSpan.FromHours(1.0));
 
 				IReadOnlyList<IIssue> issues = await IssueCollection.FindIssuesAsync();
@@ -772,7 +774,7 @@ namespace HordeServer.Tests.Issues
 			Assert.AreEqual(issue2.Fingerprints.Count, 1);
 			Assert.IsTrue(issue2.Fingerprints[0].Type.Contains("ThreadSanitizer", StringComparison.OrdinalIgnoreCase));
 			Assert.IsTrue(issue2.Fingerprints[0].Type.Contains("data race", StringComparison.OrdinalIgnoreCase));
-			Assert.IsTrue(issue2.Fingerprints[0].Type.Contains("Paths.cpp", StringComparison.OrdinalIgnoreCase)); 
+			Assert.IsTrue(issue2.Fingerprints[0].Type.Contains("Paths.cpp", StringComparison.OrdinalIgnoreCase));
 		}
 
 		[TestMethod]
@@ -1957,7 +1959,7 @@ namespace HordeServer.Tests.Issues
 		static private IEnumerable<JsonLogEvent> MultilineLogEvent(LogLevel level, EventId eventId, string format, Dictionary<string, object> properties)
 		{
 			DateTime time = new DateTime(2024, 1, 1, 0, 0, 0, 0, DateTimeKind.Utc);
-			string message = MessageTemplate.Render(format, properties);
+			string message = MessageTemplate.Render(format, properties!);
 			LogEvent baseEvent = new LogEvent(time, level, eventId, message, format, properties, null);
 			JsonLogEvent jsonLogEvent = new JsonLogEvent(baseEvent);
 			ServerLogPacketBuilder writer = new ServerLogPacketBuilder();
diff --git a/Engine/Source/Programs/Horde/HordeServer.Tests/Jobs/Bisect/BisectTests.cs b/Engine/Source/Programs/Horde/HordeServer.Tests/Jobs/Bisect/BisectTests.cs
--- a/Engine/Source/Programs/Horde/HordeServer.Tests/Jobs/Bisect/BisectTests.cs
+++ b/Engine/Source/Programs/Horde/HordeServer.Tests/Jobs/Bisect/BisectTests.cs
@@ -142,7 +142,7 @@ namespace HordeServer.Tests.Jobs.Bisect
 			project.Id = new ProjectId("ue4");
 			project.Streams.Add(new StreamConfig { Id = StreamId, Templates = new List<TemplateRefConfig> { new TemplateRefConfig { Id = TemplateId } } });
 
-			UpdateConfig(x => x.Projects.Add(project));
+			UpdateConfig(x => x.Plugins.GetBuildConfig().Projects.Add(project));
 
 			// Create the graph
 			List<NewNode> nodes = new List<NewNode>();
diff --git a/Engine/Source/Programs/Horde/HordeServer.Tests/Jobs/JobExpirationTests.cs b/Engine/Source/Programs/Horde/HordeServer.Tests/Jobs/JobExpirationTests.cs
--- a/Engine/Source/Programs/Horde/HordeServer.Tests/Jobs/JobExpirationTests.cs
+++ b/Engine/Source/Programs/Horde/HordeServer.Tests/Jobs/JobExpirationTests.cs
@@ -32,8 +32,11 @@ namespace HordeServer.Tests.Jobs
 			ProjectConfig projectConfig = new ProjectConfig { Id = projectId };
 			projectConfig.Streams.Add(streamConfig);
 
+			BuildConfig buildConfig = new BuildConfig();
+			buildConfig.Projects.Add(projectConfig);
+
 			GlobalConfig globalConfig = new GlobalConfig();
-			globalConfig.Projects.Add(projectConfig);
+			globalConfig.Plugins.AddBuildConfig(buildConfig);
 
 			SetConfig(globalConfig);
 
diff --git a/Engine/Source/Programs/Horde/HordeServer.Tests/Jobs/JobServiceTests.cs b/Engine/Source/Programs/Horde/HordeServer.Tests/Jobs/JobServiceTests.cs
--- a/Engine/Source/Programs/Horde/HordeServer.Tests/Jobs/JobServiceTests.cs
+++ b/Engine/Source/Programs/Horde/HordeServer.Tests/Jobs/JobServiceTests.cs
@@ -48,8 +48,11 @@ namespace HordeServer.Tests.Jobs
 			ProjectConfig projectConfig = new ProjectConfig { Id = projectId };
 			projectConfig.Streams.Add(streamConfig);
 
+			BuildConfig buildConfig = new BuildConfig();
+			buildConfig.Projects.Add(projectConfig);
+
 			GlobalConfig globalConfig = new GlobalConfig();
-			globalConfig.Projects.Add(projectConfig);
+			globalConfig.Plugins.AddBuildConfig(buildConfig);
 
 			SetConfig(globalConfig);
 
@@ -89,7 +92,12 @@ namespace HordeServer.Tests.Jobs
 			ProjectId projectId = new ProjectId("ue5");
 			ProjectConfig projectConfig = new ProjectConfig { Id = projectId, Name = "UE4", Streams = new List<StreamConfig> { streamConfig } };
 
-			SetConfig(new GlobalConfig { Projects = new List<ProjectConfig> { projectConfig } });
+			BuildConfig buildConfig = new BuildConfig();
+			buildConfig.Projects = new List<ProjectConfig> { projectConfig };
+
+			GlobalConfig config = new GlobalConfig();
+			config.Plugins.AddBuildConfig(buildConfig);
+			SetConfig(config);
 
 			IUser user = await UserCollection.FindOrAddUserByLoginAsync("Bob");
 
@@ -234,7 +242,12 @@ namespace HordeServer.Tests.Jobs
 			ProjectId projectId = new ProjectId("ue5");
 			ProjectConfig projectConfig = new ProjectConfig { Id = projectId, Name = "UE5", Streams = new List<StreamConfig> { streamConfig } };
 
-			SetConfig(new GlobalConfig { Projects = new List<ProjectConfig> { projectConfig } });
+			BuildConfig buildConfig = new BuildConfig();
+			buildConfig.Projects = new List<ProjectConfig> { projectConfig };
+
+			GlobalConfig globalConfig = new();
+			globalConfig.Plugins.AddBuildConfig(buildConfig);
+			SetConfig(globalConfig);
 
 			// ----
 
diff --git a/Engine/Source/Programs/Horde/HordeServer.Tests/Jobs/Schedules/SchedulerTests.cs b/Engine/Source/Programs/Horde/HordeServer.Tests/Jobs/Schedules/SchedulerTests.cs
--- a/Engine/Source/Programs/Horde/HordeServer.Tests/Jobs/Schedules/SchedulerTests.cs
+++ b/Engine/Source/Programs/Horde/HordeServer.Tests/Jobs/Schedules/SchedulerTests.cs
@@ -37,8 +37,12 @@ namespace HordeServer.Tests.Jobs.Schedules
 		{
 			IUser bob = UserCollection.FindOrAddUserByLoginAsync("Bob").Result;
 
+			BuildConfig buildConfig = new BuildConfig();
+			buildConfig.Projects.Add(new ProjectConfig { Id = ProjectId, Name = "UE4" });
+
 			GlobalConfig globalConfig = new GlobalConfig();
-			globalConfig.Projects.Add(new ProjectConfig { Id = ProjectId, Name = "UE4" });
+			globalConfig.Plugins.AddBuildConfig(buildConfig);
+			
 			SetConfig(globalConfig);
 
 			_template = TemplateCollection.GetOrAddAsync(new TemplateConfig { Name = "Test template" }).Result;
@@ -61,7 +65,7 @@ namespace HordeServer.Tests.Jobs.Schedules
 			streamConfig.Tabs.Add(new TabConfig { Title = "foo", Templates = new List<TemplateId> { TemplateId } });
 			streamConfig.Templates.Add(new TemplateRefConfig { Id = TemplateId, Name = "Test", Schedule = schedule });
 
-			UpdateConfig(x => x.Projects[0].Streams = new List<StreamConfig> { streamConfig });
+			UpdateConfig(x => x.Plugins.GetBuildConfig().Projects[0].Streams = new List<StreamConfig> { streamConfig });
 
 			return await StreamCollection.GetAsync(streamConfig);
 		}
@@ -248,7 +252,7 @@ namespace HordeServer.Tests.Jobs.Schedules
 			Assert.AreEqual(100, jobs2[0].CodeChange);
 
 			StreamConfig? streamConfig;
-			GlobalConfig.CurrentValue.TryGetStream(StreamId, out streamConfig);
+			GlobalConfig.CurrentValue.Plugins.GetBuildConfig().TryGetStream(StreamId, out streamConfig);
 
 			IStream stream2 = (await StreamCollection.GetAsync(streamConfig!))!;
 			ITemplateSchedule schedule2 = stream2.Templates.First().Value.Schedule!;
@@ -435,7 +439,7 @@ namespace HordeServer.Tests.Jobs.Schedules
 			config.Name = "//UE5/Main";
 			config.Tabs.Add(new TabConfig { Title = "foo", Templates = new List<TemplateId> { newTemplateRefId1, newTemplateRefId2 } });
 			config.Templates = new() { new TemplateRefConfig { Id = newTemplateRefId1 }, new TemplateRefConfig { Id = newTemplateRefId2 } };
-			UpdateConfig(x => x.Projects[0].Streams = new List<StreamConfig> { config });
+			UpdateConfig(x => x.Plugins.GetBuildConfig().Projects[0].Streams = new List<StreamConfig> { config });
 
 			IStream stream = await StreamCollection.GetAsync(config);
 
@@ -527,7 +531,7 @@ namespace HordeServer.Tests.Jobs.Schedules
 			config.Tabs.Add(new TabConfig { Title = "foo", Templates = new List<TemplateId> { newTemplateRefId1, newTemplateRefId2 } });
 			config.Templates.Add(newTemplate1);
 			config.Templates.Add(newTemplate2);
-			UpdateConfig(x => x.Projects[0].Streams = new List<StreamConfig> { config });
+			UpdateConfig(x => x.Plugins.GetBuildConfig().Projects[0].Streams = new List<StreamConfig> { config });
 
 			IStream stream = await StreamCollection.GetAsync(config);
 
@@ -598,7 +602,7 @@ namespace HordeServer.Tests.Jobs.Schedules
 			Assert.AreEqual(100, jobs1[0].CodeChange);
 
 			// Make sure the job is registered
-			IStream? stream1 = await StreamCollection.GetAsync(GlobalConfig.CurrentValue.Streams[0]);
+			IStream? stream1 = await StreamCollection.GetAsync(GlobalConfig.CurrentValue.Plugins.GetBuildConfig().Streams[0]);
 			ITemplateRef templateRef1 = stream1!.Templates.First().Value;
 			Assert.AreEqual(1, templateRef1.Schedule!.ActiveJobs.Count);
 			Assert.AreEqual(jobs1[0].Id, templateRef1.Schedule!.ActiveJobs[0]);
@@ -607,7 +611,7 @@ namespace HordeServer.Tests.Jobs.Schedules
 			await SetScheduleAsync(schedule);
 
 			// Make sure the job is still registered
-			IStream? stream2 = await StreamCollection.GetAsync(GlobalConfig.CurrentValue.Streams[0]);
+			IStream? stream2 = await StreamCollection.GetAsync(GlobalConfig.CurrentValue.Plugins.GetBuildConfig().Streams[0]);
 			ITemplateRef templateRef2 = stream2!.Templates.First().Value;
 			Assert.AreEqual(1, templateRef2.Schedule!.ActiveJobs.Count);
 			Assert.AreEqual(jobs1[0].Id, templateRef2.Schedule!.ActiveJobs[0]);
@@ -691,7 +695,7 @@ namespace HordeServer.Tests.Jobs.Schedules
 			streamConfig.Name = "//UE5/Main";
 			streamConfig.Tabs.Add(new TabConfig { Title = "foo", Templates = new List<TemplateId> { TemplateId } });
 			streamConfig.Templates.Add(templateConfig);
-			UpdateConfig(x => x.Projects[0].Streams = new List<StreamConfig> { streamConfig });
+			UpdateConfig(x => x.Plugins.GetBuildConfig().Projects[0].Streams = new List<StreamConfig> { streamConfig });
 
 			// Make sure we don't have any jobs to start with
 			await ScheduleService.TickForTestingAsync();
diff --git a/Engine/Source/Programs/Horde/HordeServer.Tests/Jobs/TestData/TestDataTests.cs b/Engine/Source/Programs/Horde/HordeServer.Tests/Jobs/TestData/TestDataTests.cs
--- a/Engine/Source/Programs/Horde/HordeServer.Tests/Jobs/TestData/TestDataTests.cs
+++ b/Engine/Source/Programs/Horde/HordeServer.Tests/Jobs/TestData/TestDataTests.cs
@@ -91,8 +91,11 @@ namespace HordeServer.Tests.Jobs.TestData
 			projectConfig.Streams.Add(CreateStream(_mainStreamId, MainStreamName));
 			projectConfig.Streams.Add(CreateStream(_releaseStreamId, ReleaseStreamName));
 
+			BuildConfig buildConfig = new BuildConfig();
+			buildConfig.Projects.Add(projectConfig);
+
 			GlobalConfig globalConfig = new GlobalConfig();
-			globalConfig.Projects.Add(projectConfig);
+			globalConfig.Plugins.AddBuildConfig(buildConfig);
 
 			SetConfig(globalConfig);
 
diff --git a/Engine/Source/Programs/Horde/HordeServer.Tests/Perforce/PerforceLoadBalancerTests.cs b/Engine/Source/Programs/Horde/HordeServer.Tests/Perforce/PerforceLoadBalancerTests.cs
--- a/Engine/Source/Programs/Horde/HordeServer.Tests/Perforce/PerforceLoadBalancerTests.cs
+++ b/Engine/Source/Programs/Horde/HordeServer.Tests/Perforce/PerforceLoadBalancerTests.cs
@@ -81,7 +81,8 @@ public class PerforceLoadBalancerTests : TestSetup
 	private static GlobalConfig GetConfig(string perforceHost)
 	{
 		PerforceCluster perforceCluster = new() { Servers = [new PerforceServer { ServerAndPort = perforceHost }] };
-		GlobalConfig gc = new() { PerforceClusters = [perforceCluster] };
+		GlobalConfig gc = new();
+		gc.Plugins.GetBuildConfig().PerforceClusters = [perforceCluster];
 		return gc;
 	}
 
@@ -114,7 +115,7 @@ public class PerforceLoadBalancerTests : TestSetup
 	{
 #pragma warning disable CA2000 // Dispose objects before losing scope
 		HttpClient httpClient = new(new StubMessageHandler(HttpStatusCode.OK, httpCheckResponse));
-		return new(MongoService, GetRedisServiceSingleton(), LeaseCollection, Clock, httpClient, new TestOptionsMonitor<GlobalConfig>(gc), new FakeHealthMonitor<PerforceLoadBalancer>(), Tracer, _logger);
+		return new(MongoService, GetRedisServiceSingleton(), LeaseCollection, Clock, httpClient, new TestOptionsMonitor<BuildConfig>(gc.Plugins.GetBuildConfig()), new FakeHealthMonitor<PerforceLoadBalancer>(), Tracer, _logger);
 #pragma warning restore CA2000 // Dispose objects before losing scope		
 	}
 
diff --git a/Engine/Source/Programs/Horde/HordeServer.Tests/TestSetup.cs b/Engine/Source/Programs/Horde/HordeServer.Tests/TestSetup.cs
--- a/Engine/Source/Programs/Horde/HordeServer.Tests/TestSetup.cs
+++ b/Engine/Source/Programs/Horde/HordeServer.Tests/TestSetup.cs
@@ -115,6 +115,7 @@ namespace HordeServer.Tests
 		public TestSetup()
 		{
 			AddPlugin<AnalyticsPlugin>();
+			AddPlugin<BuildPlugin>();
 			AddPlugin<ComputePlugin>();
 			AddPlugin<StoragePlugin>();
 		}
diff --git a/Engine/Source/Programs/Horde/HordeServer/Commands/Test/ReplicationCommand.cs b/Engine/Source/Programs/Horde/HordeServer/Commands/Test/ReplicationCommand.cs
--- a/Engine/Source/Programs/Horde/HordeServer/Commands/Test/ReplicationCommand.cs
+++ b/Engine/Source/Programs/Horde/HordeServer/Commands/Test/ReplicationCommand.cs
@@ -54,7 +54,7 @@ namespace HordeServer.Commands.Test
 			IStreamCollection streamCollection = serviceProvider.GetRequiredService<IStreamCollection>();
 
 			StreamConfig? streamConfig;
-			if (!globalConfig.TryGetStream(new StreamId(StreamId), out streamConfig))
+			if (!globalConfig.Plugins.GetBuildConfig().TryGetStream(new StreamId(StreamId), out streamConfig))
 			{
 				throw new FatalErrorException($"Stream '{StreamId}' not found");
 			}
@@ -79,7 +79,7 @@ namespace HordeServer.Commands.Test
 			}
 
 			PerforceReplicationOptions options = new PerforceReplicationOptions();
-			await perforceReplicator.RunOnceAsync(replicator, globalConfig, streamConfig, options, default);
+			await perforceReplicator.RunOnceAsync(replicator, globalConfig.Plugins.GetBuildConfig(), streamConfig, options, default);
 
 			return 0;
 		}
diff --git a/Engine/Source/Programs/Horde/HordeServer/Configuration/ConfigService.cs b/Engine/Source/Programs/Horde/HordeServer/Configuration/ConfigService.cs
--- a/Engine/Source/Programs/Horde/HordeServer/Configuration/ConfigService.cs
+++ b/Engine/Source/Programs/Horde/HordeServer/Configuration/ConfigService.cs
@@ -33,7 +33,7 @@ namespace HordeServer.Configuration
 	/// <summary>
 	/// Service which processes runtime configuration data.
 	/// </summary>
-	public sealed class ConfigService : IOptionsFactory<GlobalConfig>, IOptionsChangeTokenSource<GlobalConfig>, IHostedService, IAsyncDisposable
+	public sealed class ConfigService : IConfigService, IOptionsFactory<GlobalConfig>, IOptionsChangeTokenSource<GlobalConfig>, IHostedService, IAsyncDisposable
 	{
 		// Index to all current config files 
 		[ProtoContract]
@@ -630,6 +630,15 @@ namespace HordeServer.Configuration
 			// Analytics
 			CopyJsonNode(rootObj, "TelemetryStores", "Plugins.Analytics.Stores");
 
+			// Build
+			CopyJsonNode(rootObj, "PerforceClusters", "Plugins.Build.PerforceClusters");
+			CopyJsonNode(rootObj, "Devices", "Plugins.Build.Devices");
+			CopyJsonNode(rootObj, "MaxConformCount", "Plugins.Build.MaxConformCount");
+			CopyJsonNode(rootObj, "AgentShutdownIfDisabledGracePeriod", "Plugins.Build.AgentShutdownIfDisabledGracePeriod");
+			CopyJsonNode(rootObj, "ArtifactTypes", "Plugins.Build.ArtifactTypes");
+			CopyJsonNode(rootObj, "Projects", "Plugins.Build.Projects");
+			CopyJsonNode(rootObj, "IssueFixedTag", "Plugins.Build.IssueFixedTag");
+
 			// Compute
 			CopyJsonNode(rootObj, "Rates", "Plugins.Compute.Rates");
 			CopyJsonNode(rootObj, "Compute", "Plugins.Compute.Clusters");
diff --git a/Engine/Source/Programs/Horde/HordeServer/Dashboard/DashboardController.cs b/Engine/Source/Programs/Horde/HordeServer/Dashboard/DashboardController.cs
--- a/Engine/Source/Programs/Horde/HordeServer/Dashboard/DashboardController.cs
+++ b/Engine/Source/Programs/Horde/HordeServer/Dashboard/DashboardController.cs
@@ -37,6 +37,7 @@ namespace HordeServer.Dashboard
 		/// Server settings
 		/// </summary>
 		readonly ServerSettings _settings;
+		readonly StaticBuildConfig _staticBuildConfig;
 
 		readonly IDashboardPreviewCollection _previewCollection;
 
@@ -47,16 +48,13 @@ namespace HordeServer.Dashboard
 		/// <summary>
 		/// Constructor
 		/// </summary>
-		/// <param name="previewCollection" />
-		/// <param name="hordeAccounts" />
-		/// <param name="serverSettings">Server settings</param>
-		/// <param name="globalConfig" />
-		public DashboardController(IDashboardPreviewCollection previewCollection, IAccountCollection hordeAccounts, IOptionsMonitor<ServerSettings> serverSettings, IOptionsSnapshot<GlobalConfig> globalConfig)
+		public DashboardController(IDashboardPreviewCollection previewCollection, IAccountCollection hordeAccounts, IOptionsMonitor<ServerSettings> serverSettings, IOptions<StaticBuildConfig> staticBuildConfig, IOptionsSnapshot<GlobalConfig> globalConfig)
 		{
 			_authenticationScheme = AccountController.GetAuthScheme(serverSettings.CurrentValue.AuthMethod);
 			_previewCollection = previewCollection;
 			_hordeAccounts = hordeAccounts;
 			_settings = serverSettings.CurrentValue;
+			_staticBuildConfig = staticBuildConfig.Value;
 			_globalConfig = globalConfig;
 		}
 
@@ -157,21 +155,21 @@ namespace HordeServer.Dashboard
 
 			dashboardConfigResponse.AuthMethod = _settings.AuthMethod;
 
-			if (_settings.JiraUrl != null)
+			if (_staticBuildConfig.JiraUrl != null)
 			{
 				dashboardConfigResponse.ExternalIssueServiceName = "Jira";
-				dashboardConfigResponse.ExternalIssueServiceUrl = _settings.JiraUrl.ToString().TrimEnd('/');
+				dashboardConfigResponse.ExternalIssueServiceUrl = _staticBuildConfig.JiraUrl.ToString().TrimEnd('/');
 			}
 
-			if (_settings.P4SwarmUrl != null)
+			if (_staticBuildConfig.P4SwarmUrl != null)
 			{
-				dashboardConfigResponse.PerforceSwarmUrl = _settings.P4SwarmUrl.ToString().TrimEnd('/');
+				dashboardConfigResponse.PerforceSwarmUrl = _staticBuildConfig.P4SwarmUrl.ToString().TrimEnd('/');
 			}
 
 			dashboardConfigResponse.HelpEmailAddress = _settings.HelpEmailAddress;
 			dashboardConfigResponse.HelpSlackChannel = _settings.HelpSlackChannel;
 
-			dashboardConfigResponse.DeviceProblemCooldownMinutes = _settings.DeviceProblemCooldownMinutes;
+			dashboardConfigResponse.DeviceProblemCooldownMinutes = _staticBuildConfig.DeviceProblemCooldownMinutes;
 
 			foreach (DashboardAgentCategoryConfig category in _globalConfig.Value.Dashboard.AgentCategories)
 			{
diff --git a/Engine/Source/Programs/Horde/HordeServer/HordeServer.csproj b/Engine/Source/Programs/Horde/HordeServer/HordeServer.csproj
--- a/Engine/Source/Programs/Horde/HordeServer/HordeServer.csproj
+++ b/Engine/Source/Programs/Horde/HordeServer/HordeServer.csproj
@@ -95,11 +95,6 @@
   </ItemGroup>
 
   <ItemGroup>
-    <Protobuf Include="Protos\horde\job_rpc.proto" ProtoRoot="Protos" GrpcServices="Server" AdditionalImportDirs="..\..\Shared\EpicGames.Horde\Protos;..\Plugins\Compute\HordeServer.Compute\Protos" />
-    <Protobuf Include="Protos\horde\job_rpc_messages.proto" ProtoRoot="Protos" AdditionalImportDirs="..\Plugins\Compute\HordeServer.Compute\Protos" />
-    <Protobuf Include="Protos\horde\common\io_hash.proto" ProtoRoot="Protos" />
-    <Protobuf Include="Protos\horde\common\label_change.proto" ProtoRoot="Protos" />
-    <Protobuf Include="Protos\horde\common\namespace_id.proto" ProtoRoot="Protos" />
     <Protobuf Include="Protos\grpc\health\v1\health.proto" ProtoRoot="Protos" GrpcServices="Server" />
   </ItemGroup>
 
@@ -110,13 +105,12 @@
     <ProjectReference Include="..\..\Shared\EpicGames.Horde\EpicGames.Horde.csproj" />
     <ProjectReference Include="..\..\Shared\EpicGames.IoHash\EpicGames.IoHash.csproj" />
     <ProjectReference Include="..\..\Shared\EpicGames.Oodle\EpicGames.Oodle.csproj" />
-    <ProjectReference Include="..\..\Shared\EpicGames.Perforce.Managed\EpicGames.Perforce.Managed.csproj" />
-    <ProjectReference Include="..\..\Shared\EpicGames.Perforce\EpicGames.Perforce.csproj" />
     <ProjectReference Include="..\..\Shared\EpicGames.Redis\EpicGames.Redis.csproj" />
     <ProjectReference Include="..\..\Shared\EpicGames.Serialization\EpicGames.Serialization.csproj" />
     <ProjectReference Include="..\..\Shared\EpicGames.Slack\EpicGames.Slack.csproj" />
     <ProjectReference Include="..\HordeServer.Shared\HordeServer.Shared.csproj" />
     <ProjectReference Include="..\Plugins\Analytics\HordeServer.Analytics\HordeServer.Analytics.csproj" />
+    <ProjectReference Include="..\Plugins\Build\HordeServer.Build\HordeServer.Build.csproj" />
     <ProjectReference Include="..\Plugins\Compute\HordeServer.Compute\HordeServer.Compute.csproj" />
     <ProjectReference Include="..\Plugins\Ddc\HordeServer.Ddc\HordeServer.Ddc.csproj" />
     <ProjectReference Include="..\Plugins\Secrets\HordeServer.Secrets\HordeServer.Secrets.csproj" />
diff --git a/Engine/Source/Programs/Horde/HordeServer/Server/DebugController.cs b/Engine/Source/Programs/Horde/HordeServer/Server/DebugController.cs
--- a/Engine/Source/Programs/Horde/HordeServer/Server/DebugController.cs
+++ b/Engine/Source/Programs/Horde/HordeServer/Server/DebugController.cs
@@ -142,7 +142,7 @@ namespace HordeServer.Server
 			byte[] data = _configService.Serialize(_globalConfig.Value);
 			GlobalConfig config = _configService.Deserialize(data)!;
 
-			foreach (ProjectConfig project in config.Projects)
+			foreach (ProjectConfig project in config.Plugins.GetBuildConfig().Projects)
 			{
 				project.Logo = null;
 			}
diff --git a/Engine/Source/Programs/Horde/HordeServer/Server/GlobalConfig.cs b/Engine/Source/Programs/Horde/HordeServer/Server/GlobalConfig.cs
--- a/Engine/Source/Programs/Horde/HordeServer/Server/GlobalConfig.cs
+++ b/Engine/Source/Programs/Horde/HordeServer/Server/GlobalConfig.cs
@@ -11,14 +11,8 @@ using System.Text.Json.Serialization;
 using EpicGames.Core;
 using EpicGames.Horde.Acls;
 using EpicGames.Horde.Agents.Pools;
-using EpicGames.Horde.Artifacts;
 using EpicGames.Horde.Common;
 using EpicGames.Horde.Compute;
-using EpicGames.Horde.Jobs.Templates;
-using EpicGames.Horde.Projects;
-using EpicGames.Horde.Streams;
-using EpicGames.Horde.Users;
-using EpicGames.Perforce;
 using HordeServer.Acls;
 using HordeServer.Agents;
 using HordeServer.Agents.Pools;
@@ -43,37 +37,6 @@ namespace HordeServer.Server
 {
 	using JsonObject = System.Text.Json.Nodes.JsonObject;
 
-	/// <summary>
-	/// Configuration for an artifact
-	/// </summary>
-	public class ArtifactTypeConfig
-	{
-		/// <summary>
-		/// Name of the artifact type
-		/// </summary>
-		public ArtifactType Type { get; set; }
-
-		/// <summary>
-		/// Legacy 'Name' property
-		/// </summary>
-		[Obsolete("Use Type instead")]
-		public ArtifactType Name
-		{
-			get => Type;
-			set => Type = value;
-		}
-
-		/// <summary>
-		/// Number of artifacts to retain
-		/// </summary>
-		public int? KeepCount { get; set; }
-
-		/// <summary>
-		/// Number of days to retain artifacts of this type
-		/// </summary>
-		public int? KeepDays { get; set; }
-	}
-
 	/// <summary>
 	/// Global configuration
 	/// </summary>
@@ -121,42 +84,11 @@ namespace HordeServer.Server
 		/// </summary>
 		public DashboardConfig Dashboard { get; set; } = new DashboardConfig();
 
-		/// <summary>
-		/// List of projects
-		/// </summary>
-		public List<ProjectConfig> Projects { get; set; } = new List<ProjectConfig>();
-
 		/// <summary>
 		/// List of scheduled downtime
 		/// </summary>
 		public List<ScheduledDowntime> Downtime { get; set; } = new List<ScheduledDowntime>();
 
-		/// <summary>
-		/// List of Perforce clusters
-		/// </summary>
-		public List<PerforceCluster> PerforceClusters { get; set; } = new List<PerforceCluster>();
-
-		/// <summary>
-		/// Device configuration
-		/// </summary>
-		public DeviceConfig? Devices { get; set; }
-
-		/// <summary>
-		/// Maximum number of conforms to run at once
-		/// </summary>
-		public int MaxConformCount { get; set; }
-
-		/// <summary>
-		/// Time to wait before shutting down an agent that has been disabled
-		/// Used if no value is set on the actual pool.
-		/// </summary>
-		public TimeSpan? AgentShutdownIfDisabledGracePeriod { get; set; } = null;
-
-		/// <summary>
-		/// Configuration for different artifact types
-		/// </summary>
-		public List<ArtifactTypeConfig> ArtifactTypes { get; set; } = new List<ArtifactTypeConfig>();
-
 		/// <summary>
 		/// Metrics to aggregate on the Horde server
 		/// </summary>
@@ -178,27 +110,13 @@ namespace HordeServer.Server
 		/// </summary>
 		public AclConfig Acl { get; set; } = new AclConfig();
 
-		/// <summary>
-		/// Commit tag to use for marking issues as fixed
-		/// </summary>
-		public string IssueFixedTag { get; set; } = "#horde";
-
 		/// <summary>
 		/// Accessor for the ACL scope lookup
 		/// </summary>
 		[JsonIgnore]
 		public IReadOnlyDictionary<AclScopeName, AclConfig> AclScopes => _aclLookup;
 
-		/// <summary>
-		/// Enumerates all the streams
-		/// </summary>
-		[JsonIgnore]
-		public IReadOnlyList<StreamConfig> Streams { get; private set; } = null!;
-
-		private readonly Dictionary<ProjectId, ProjectConfig> _projectLookup = new Dictionary<ProjectId, ProjectConfig>();
-		private readonly Dictionary<StreamId, StreamConfig> _streamLookup = new Dictionary<StreamId, StreamConfig>();
 		private readonly Dictionary<AclScopeName, AclConfig> _aclLookup = new Dictionary<AclScopeName, AclConfig>();
-		private readonly Dictionary<ArtifactType, ArtifactTypeConfig> _artifactTypeLookup = new Dictionary<ArtifactType, ArtifactTypeConfig>();
 
 		/// <inheritdoc cref="AclConfig.Authorize(AclAction, ClaimsPrincipal)"/>
 		public bool Authorize(AclAction action, ClaimsPrincipal user)
@@ -214,27 +132,6 @@ namespace HordeServer.Server
 			AclConfig defaultAcl = CreateRootAcl();
 			Acl.PostLoad(defaultAcl, defaultAcl.ScopeName);
 
-			Streams = Projects.SelectMany(x => x.Streams).ToList();
-
-			_projectLookup.Clear();
-			_streamLookup.Clear();
-			foreach (ProjectConfig project in Projects)
-			{
-				_projectLookup.Add(project.Id, project);
-				project.PostLoad(project.Id, Acl);
-
-				foreach (StreamConfig stream in project.Streams)
-				{
-					_streamLookup.Add(stream.Id, stream);
-				}
-			}
-
-			_artifactTypeLookup.Clear();
-			foreach (ArtifactTypeConfig artifactType in ArtifactTypes)
-			{
-				_artifactTypeLookup.Add(artifactType.Type, artifactType);
-			}
-
 			// Ensure that all plugins have an entry in the global config so they can register their ACLs
 			foreach (ILoadedPlugin loadedPlugin in loadedPlugins)
 			{
@@ -256,20 +153,24 @@ namespace HordeServer.Server
 			_aclLookup.Clear();
 			BuildAclScopeLookup(Acl, _aclLookup);
 
-			foreach (ProjectConfig project in Projects)
+			BuildConfig? buildConfig;
+			if (Plugins.TryGetBuildConfig(out buildConfig))
 			{
-				AclScopeName legacyProjectScopeName = Acl.ScopeName.Append($"p:{project.Id}");
-				_aclLookup.Add(legacyProjectScopeName, project.Acl);
-
-				foreach (StreamConfig stream in project.Streams)
+				foreach (ProjectConfig project in buildConfig.Projects)
 				{
-					AclScopeName legacyStreamScopeName = Acl.ScopeName.Append($"s:{stream.Id}");
-					_aclLookup.Add(legacyStreamScopeName, stream.Acl);
+					AclScopeName legacyProjectScopeName = Acl.ScopeName.Append($"p:{project.Id}");
+					_aclLookup.Add(legacyProjectScopeName, project.Acl);
 
-					foreach (TemplateRefConfig template in stream.Templates)
+					foreach (StreamConfig stream in project.Streams)
 					{
-						AclScopeName legacyTemplateScopeName = legacyStreamScopeName.Append($"t:{template.Id}");
-						_aclLookup.Add(legacyTemplateScopeName, template.Acl);
+						AclScopeName legacyStreamScopeName = Acl.ScopeName.Append($"s:{stream.Id}");
+						_aclLookup.Add(legacyStreamScopeName, stream.Acl);
+
+						foreach (TemplateRefConfig template in stream.Templates)
+						{
+							AclScopeName legacyTemplateScopeName = legacyStreamScopeName.Append($"t:{template.Id}");
+							_aclLookup.Add(legacyTemplateScopeName, template.Acl);
+						}
 					}
 				}
 			}
@@ -359,12 +260,18 @@ namespace HordeServer.Server
 				return;
 			}
 
+			BuildConfig? buildConfig;
+			if (!Plugins.TryGetValue(new PluginName("build"), out buildConfig))
+			{
+				return;
+			}
+
 			// Lookup table of pool id to workspaces
 			Dictionary<PoolId, AutoSdkConfig> poolToAutoSdkView = new Dictionary<PoolId, AutoSdkConfig>();
 			Dictionary<PoolId, List<AgentWorkspaceInfo>> poolToAgentWorkspaces = new Dictionary<PoolId, List<AgentWorkspaceInfo>>();
 
 			// Populate the workspace list from the current stream
-			foreach (StreamConfig streamConfig in Streams)
+			foreach (StreamConfig streamConfig in buildConfig.Streams)
 			{
 				foreach (KeyValuePair<string, AgentConfig> agentTypePair in streamConfig.AgentTypes)
 				{
@@ -418,47 +325,6 @@ namespace HordeServer.Server
 			}
 		}
 
-		/// <summary>
-		/// Attempts to get configuration for a specific artifact type
-		/// </summary>
-		/// <param name="type">The artifact type name</param>
-		/// <param name="config">Configuration for the stream</param>
-		/// <returns>True if the stream configuration was found</returns>
-		public bool TryGetArtifactType(ArtifactType type, [NotNullWhen(true)] out ArtifactTypeConfig? config) => _artifactTypeLookup.TryGetValue(type, out config);
-
-		/// <summary>
-		/// Attempts to get configuration for a project from this object
-		/// </summary>
-		/// <param name="projectId">The stream identifier</param>
-		/// <param name="config">Configuration for the stream</param>
-		/// <returns>True if the stream configuration was found</returns>
-		public bool TryGetProject(ProjectId projectId, [NotNullWhen(true)] out ProjectConfig? config) => _projectLookup.TryGetValue(projectId, out config);
-
-		/// <summary>
-		/// Attempts to get configuration for a stream from this object
-		/// </summary>
-		/// <param name="streamId">The stream identifier</param>
-		/// <param name="config">Configuration for the stream</param>
-		/// <returns>True if the stream configuration was found</returns>
-		public bool TryGetStream(StreamId streamId, [NotNullWhen(true)] out StreamConfig? config) => _streamLookup.TryGetValue(streamId, out config);
-
-		/// <summary>
-		/// Attempts to get configuration for a stream from this object
-		/// </summary>
-		/// <param name="streamId">The stream identifier</param>
-		/// <param name="templateId">Template identifier</param>
-		/// <param name="config">Configuration for the template</param>
-		/// <returns>True if the template configuration was found</returns>
-		public bool TryGetTemplate(StreamId streamId, TemplateId templateId, [NotNullWhen(true)] out TemplateRefConfig? config)
-		{
-			if (!_streamLookup.TryGetValue(streamId, out StreamConfig? streamConfig))
-			{
-				config = null;
-				return false;
-			}
-			return streamConfig.TryGetTemplate(templateId, out config);
-		}
-
 		/// <summary>
 		/// Authorizes a user to perform a given action
 		/// </summary>
@@ -476,86 +342,6 @@ namespace HordeServer.Server
 		public bool Authorize(AclScopeName scopeName, AclAction action, ClaimsPrincipal user)
 			=> _aclLookup.TryGetValue(scopeName, out AclConfig? scopeConfig) && scopeConfig.Authorize(action, user);
 
-		/// <summary>
-		/// Determines whether the given user can masquerade as a given user
-		/// </summary>
-		/// <param name="user"></param>
-		/// <param name="userId"></param>
-		/// <returns></returns>
-		public bool AuthorizeAsUser(ClaimsPrincipal user, UserId userId)
-		{
-			UserId? currentUserId = user.GetUserId();
-			if (currentUserId != null && currentUserId.Value == userId)
-			{
-				return true;
-			}
-			else
-			{
-				return this.Authorize(ServerAclAction.Impersonate, user);
-			}
-		}
-
-		/// <summary>
-		/// Finds a perforce cluster with the given name or that contains the provided server
-		/// </summary>
-		/// <param name="name">Name of the cluster</param>
-		/// <param name="serverAndPort">Find cluster which contains server</param>
-		/// <returns></returns>
-		public PerforceCluster GetPerforceCluster(string? name, string? serverAndPort = null)
-		{
-			PerforceCluster? cluster = FindPerforceCluster(name, serverAndPort);
-			if (cluster == null)
-			{
-				throw new Exception($"Unknown Perforce cluster '{name}'");
-			}
-			return cluster;
-		}
-
-		/// <summary>
-		/// Finds a perforce cluster with the given name or that contains the provided server
-		/// </summary>
-		/// <param name="name">Name of the cluster</param>
-		/// <param name="serverAndPort">Find cluster which contains server</param>
-		/// <returns></returns>
-		public PerforceCluster? FindPerforceCluster(string? name, string? serverAndPort = null)
-		{
-			List<PerforceCluster> clusters = PerforceClusters;
-
-			if (serverAndPort != null)
-			{
-				return clusters.FirstOrDefault(x => x.Servers.FirstOrDefault(server => String.Equals(server.ServerAndPort, serverAndPort, StringComparison.OrdinalIgnoreCase)) != null);
-			}
-
-			if (clusters.Count == 0)
-			{
-				clusters = DefaultClusters;
-			}
-
-			if (name == null)
-			{
-				return clusters.FirstOrDefault();
-			}
-			else
-			{
-				return clusters.FirstOrDefault(x => String.Equals(x.Name, name, StringComparison.OrdinalIgnoreCase));
-			}
-		}
-
-		static List<PerforceCluster> DefaultClusters { get; } = GetDefaultClusters();
-
-		static List<PerforceCluster> GetDefaultClusters()
-		{
-			PerforceServer server = new PerforceServer();
-			server.ServerAndPort = PerforceSettings.Default.ServerAndPort;
-
-			PerforceCluster cluster = new PerforceCluster();
-			cluster.Name = "Default";
-			cluster.CanImpersonate = false;
-			cluster.Servers.Add(server);
-
-			return new List<PerforceCluster> { cluster };
-		}
-
 		IReadOnlyList<string>? _cachedGroupClaims;
 
 		/// <summary>
diff --git a/Engine/Source/Programs/Horde/HordeServer/Server/ServerController.cs b/Engine/Source/Programs/Horde/HordeServer/Server/ServerController.cs
--- a/Engine/Source/Programs/Horde/HordeServer/Server/ServerController.cs
+++ b/Engine/Source/Programs/Horde/HordeServer/Server/ServerController.cs
@@ -1,23 +1,17 @@
 // Copyright Epic Games, Inc. All Rights Reserved.
 
 using System;
-using System.Collections.Generic;
 using System.Diagnostics;
-using System.Linq;
 using System.Reflection;
-using System.Threading;
 using System.Threading.Tasks;
 using EpicGames.Core;
 using EpicGames.Horde;
-using EpicGames.Horde.Agents.Pools;
 using EpicGames.Horde.Server;
 using HordeServer.Agents;
-using HordeServer.Agents.Pools;
 using HordeServer.Tools;
 using HordeServer.Utilities;
 using Microsoft.AspNetCore.Authorization;
 using Microsoft.AspNetCore.Mvc;
-using Microsoft.Extensions.DependencyInjection;
 using Microsoft.Extensions.Options;
 
 namespace HordeServer.Server
@@ -30,7 +24,6 @@ namespace HordeServer.Server
 	[Route("[controller]")]
 	public class ServerController : HordeControllerBase
 	{
-		readonly IServiceProvider _serviceProvider;
 		readonly IToolCollection _toolCollection;
 		readonly IClock _clock;
 		readonly IOptionsSnapshot<GlobalConfig> _globalConfig;
@@ -38,9 +31,8 @@ namespace HordeServer.Server
 		/// <summary>
 		/// Constructor
 		/// </summary>
-		public ServerController(IServiceProvider serviceProvider, IToolCollection toolCollection, IClock clock, IOptionsSnapshot<GlobalConfig> globalConfig)
+		public ServerController(IToolCollection toolCollection, IClock clock, IOptionsSnapshot<GlobalConfig> globalConfig)
 		{
-			_serviceProvider = serviceProvider;
 			_toolCollection = toolCollection;
 			_clock = clock;
 			_globalConfig = globalConfig;
diff --git a/Engine/Source/Programs/Horde/HordeServer/ServerApp.cs b/Engine/Source/Programs/Horde/HordeServer/ServerApp.cs
--- a/Engine/Source/Programs/Horde/HordeServer/ServerApp.cs
+++ b/Engine/Source/Programs/Horde/HordeServer/ServerApp.cs
@@ -4,6 +4,7 @@ using System;
 using System.Collections.Generic;
 using System.Diagnostics;
 using System.IO;
+using System.Linq;
 using System.Reflection;
 using System.Runtime.InteropServices;
 using System.Threading;
@@ -214,8 +215,36 @@ namespace HordeServer
 		/// <summary>
 		/// List of config settings to rename to new locations
 		/// </summary>
-		static readonly KeyValuePair<string, string>[] s_renamedConfigKeys = new[]
+		static readonly KeyValuePair<string, string>[] s_renamedConfigValues = new[]
 		{
+			// Build plugin
+			KeyValuePair.Create("Horde:UseLocalPerforceEnv", "Horde:Plugins:Build:UseLocalPerforceEnv"),
+			KeyValuePair.Create("Horde:PerforceConnectionPoolSize", "Horde:Plugins:Build:PerforceConnectionPoolSize"),
+			KeyValuePair.Create("Horde:EnableConformTasks", "Horde:Plugins:Build:EnableConformTasks"),
+			KeyValuePair.Create("Horde:P4SwarmUrl", "Horde:Plugins:Build:P4SwarmUrl"),
+			KeyValuePair.Create("Horde:JiraUsername", "Horde:Plugins:Build:JiraUsername"),
+			KeyValuePair.Create("Horde:JiraApiToken", "Horde:Plugins:Build:JiraApiToken"),
+			KeyValuePair.Create("Horde:JiraUrl", "Horde:Plugins:Build:JiraUrl"),
+			KeyValuePair.Create("Horde:SharedDeviceCheckoutDays", "Horde:Plugins:Build:SharedDeviceCheckoutDays"),
+			KeyValuePair.Create("Horde:DeviceProblemCooldownMinutes", "Horde:Plugins:Build:DeviceProblemCooldownMinutes"),
+			KeyValuePair.Create("Horde:DeviceReportChannel", "Horde:Plugins:Build:DeviceReportChannel"),
+			KeyValuePair.Create("Horde:DisableSchedules", "Horde:Plugins:Build:DisableSchedules"),
+			KeyValuePair.Create("Horde:SlackToken", "Horde:Plugins:Build:SlackToken"),
+			KeyValuePair.Create("Horde:SlackSocketToken", "Horde:Plugins:Build:SlackSocketToken"),
+			KeyValuePair.Create("Horde:SlackAdminToken", "Horde:Plugins:Build:SlackAdminToken"),
+			KeyValuePair.Create("Horde:SlackUsers", "Horde:Plugins:Build:SlackUsers"),
+			KeyValuePair.Create("Horde:SlackErrorPrefix", "Horde:Plugins:Build:SlackErrorPrefix"),
+			KeyValuePair.Create("Horde:SlackWarningPrefix", "Horde:Plugins:Build:SlackWarningPrefix"),
+			KeyValuePair.Create("Horde:ConfigNotificationChannel", "Horde:Plugins:Build:ConfigNotificationChannel"),
+			KeyValuePair.Create("Horde:UpdateStreamsNotificationChannel", "Horde:Plugins:Build:UpdateStreamsNotificationChannel"),
+			KeyValuePair.Create("Horde:JobNotificationChannel", "Horde:Plugins:Build:JobNotificationChannel"),
+			KeyValuePair.Create("Horde:AgentNotificationChannel", "Horde:Plugins:Build:AgentNotificationChannel"),
+			KeyValuePair.Create("Horde:TestDataRetainMonths", "Horde:Plugins:Build:TestDataRetainMonths"),
+			KeyValuePair.Create("Horde:BlockCacheDir", "Horde:Plugins:Build:BlockCacheDir"),
+			KeyValuePair.Create("Horde:BlockCacheSize", "Horde:Plugins:Build:BlockCacheSize"),
+			KeyValuePair.Create("Horde:Perforce:", "Horde:Plugins:Build:Perforce:"),
+			KeyValuePair.Create("Horde:Commits:", "Horde:Plugins:Build:Commits:"),
+
 			// Compute plugin
 			KeyValuePair.Create("Horde:EnableUpgradeTasks", "Horde:Plugins:Compute:EnableUpgradeTasks"),
 			KeyValuePair.Create("Horde:FleetManagerV2", "Horde:Plugins:Compute:FleetManagerV2"),
@@ -232,10 +261,10 @@ namespace HordeServer
 			// Storage plugin
 			KeyValuePair.Create("Horde:BundleCacheDir", "Horde:Plugins:Storage:BundleCacheDir"),
 			KeyValuePair.Create("Horde:BundleCacheSize", "Horde:Plugins:Storage:BundleCacheSize"),
-			KeyValuePair.Create("Horde:Backends", "Horde:Plugins:Storage:Backends"),
+			KeyValuePair.Create("Horde:Backends:", "Horde:Plugins:Storage:Backends:"),
 	
 			// Tools plugin
-			KeyValuePair.Create("Horde:BundledTools", "Horde:Plugins:Tools:BundledTools")
+			KeyValuePair.Create("Horde:BundledTools:", "Horde:Plugins:Tools:BundledTools:")
 		};
 
 		/// <summary>
@@ -267,12 +296,25 @@ namespace HordeServer
 			IConfiguration configuration = builder.Build();
 
 			List<KeyValuePair<string, string?>> remappedValues = new List<KeyValuePair<string, string?>>();
-			foreach ((string source, string target) in s_renamedConfigKeys)
+			foreach ((string source, string target) in s_renamedConfigValues)
 			{
-				string? value = configuration[source];
-				if (value != null && configuration[target] == null)
+				if (source.EndsWith(':'))
+				{
+					foreach (KeyValuePair<string, string?> pair in configuration.AsEnumerable().Where(x => x.Key.StartsWith(source, StringComparison.OrdinalIgnoreCase)))
+					{
+						if (pair.Value != null)
+						{
+							remappedValues.Add(new KeyValuePair<string, string?>(target + pair.Key.Substring(source.Length), pair.Value));
+						}
+					}
+				}
+				else
 				{
-					remappedValues.Add(new KeyValuePair<string, string?>(target, value));
+					string? value = configuration[source];
+					if (value != null && configuration[target] == null)
+					{
+						remappedValues.Add(new KeyValuePair<string, string?>(target, value));
+					}
 				}
 			}
 
diff --git a/Engine/Source/Programs/Horde/HordeServer/ServerInfo.cs b/Engine/Source/Programs/Horde/HordeServer/ServerInfo.cs
--- a/Engine/Source/Programs/Horde/HordeServer/ServerInfo.cs
+++ b/Engine/Source/Programs/Horde/HordeServer/ServerInfo.cs
@@ -1,5 +1,6 @@
 // Copyright Epic Games, Inc. All Rights Reserved.
 
+using System;
 using EpicGames.Core;
 using Microsoft.Extensions.Configuration;
 using Microsoft.Extensions.Options;
@@ -35,6 +36,12 @@ namespace HordeServer
 		/// <inheritdoc/>
 		public bool ReadOnlyMode => _serverSettings.Value.MongoReadOnlyMode;
 
+		/// <inheritdoc/>
+		public Uri ServerUrl => _serverSettings.Value.ServerUrl;
+
+		/// <inheritdoc/>
+		public Uri DashboardUrl => _serverSettings.Value.DashboardUrl;
+
 		readonly IConfiguration _configuration;
 		readonly IOptions<ServerSettings> _serverSettings;
 
diff --git a/Engine/Source/Programs/Horde/HordeServer/ServerSettings.cs b/Engine/Source/Programs/Horde/HordeServer/ServerSettings.cs
--- a/Engine/Source/Programs/Horde/HordeServer/ServerSettings.cs
+++ b/Engine/Source/Programs/Horde/HordeServer/ServerSettings.cs
@@ -2,16 +2,10 @@
 
 using System;
 using System.Collections.Generic;
-using System.ComponentModel;
 using System.Linq;
 using System.Net;
 using EpicGames.Core;
-using EpicGames.Horde;
 using EpicGames.Horde.Server;
-using EpicGames.Horde.Storage.Bundles;
-using EpicGames.Horde.Storage.Nodes;
-using EpicGames.Perforce;
-using HordeServer.Server;
 
 namespace HordeServer
 {
@@ -26,32 +20,6 @@ namespace HordeServer
 	{
 	}
 
-	/// <summary>
-	/// Options for the commit service
-	/// </summary>
-	public class CommitSettings
-	{
-		/// <summary>
-		/// Whether to mirror commit metadata to the database
-		/// </summary>
-		public bool ReplicateMetadata { get; set; } = true;
-
-		/// <summary>
-		/// Whether to mirror commit data to storage
-		/// </summary>
-		public bool ReplicateContent { get; set; } = true;
-
-		/// <summary>
-		/// Options for how objects are packed together
-		/// </summary>
-		public BundleOptions Bundle { get; set; } = new BundleOptions();
-
-		/// <summary>
-		/// Options for how objects are sliced
-		/// </summary>
-		public ChunkingOptions Chunking { get; set; } = new ChunkingOptions();
-	}
-
 	/// <summary>
 	/// OpenTelemetry configuration for collection and sending of traces and metrics.
 	/// </summary>
@@ -375,11 +343,6 @@ namespace HordeServer
 		/// </summary>
 		public bool EnableNewAgentsByDefault { get; set; } = false;
 
-		/// <summary>
-		/// The number of months to retain test data
-		/// </summary>
-		public int TestDataRetainMonths { get; set; } = 6;
-
 		/// <summary>
 		/// Interval between rebuilding the schedule queue with a DB query.
 		/// </summary>
@@ -435,66 +398,11 @@ namespace HordeServer
 		/// </summary>
 		public bool LogSessionRequests { get; set; } = false;
 
-		/// <summary>
-		/// Whether to run scheduled jobs.
-		/// </summary>
-		public bool DisableSchedules { get; set; }
-
 		/// <summary>
 		/// Timezone for evaluating schedules
 		/// </summary>
 		public string? ScheduleTimeZone { get; set; }
 
-		/// <summary>
-		/// Bot token for interacting with Slack (xoxb-*)
-		/// </summary>
-		public string? SlackToken { get; set; }
-
-		/// <summary>
-		/// Token for opening a socket to slack (xapp-*)
-		/// </summary>
-		public string? SlackSocketToken { get; set; }
-
-		/// <summary>
-		/// Admin user token for Slack (xoxp-*). This is only required when using the admin endpoints to invite users.
-		/// </summary>
-		public string? SlackAdminToken { get; set; }
-
-		/// <summary>
-		/// Filtered list of slack users to send notifications to. Should be Slack user ids, separated by commas.
-		/// </summary>
-		public string? SlackUsers { get; set; }
-
-		/// <summary>
-		/// Prefix to use when reporting errors
-		/// </summary>
-		public string SlackErrorPrefix { get; set; } = ":horde-error: ";
-
-		/// <summary>
-		/// Prefix to use when reporting warnings
-		/// </summary>
-		public string SlackWarningPrefix { get; set; } = ":horde-warning: ";
-
-		/// <summary>
-		/// Channel for sending messages related to config update failures
-		/// </summary>
-		public string? ConfigNotificationChannel { get; set; }
-
-		/// <summary>
-		/// Channel to send stream notification update failures to
-		/// </summary>
-		public string? UpdateStreamsNotificationChannel { get; set; }
-
-		/// <summary>
-		/// Slack channel to send job related notifications to. Multiple channels can be specified, separated by ;
-		/// </summary>
-		public string? JobNotificationChannel { get; set; }
-
-		/// <summary>
-		/// Slack channel to send agent related notifications to.
-		/// </summary>
-		public string? AgentNotificationChannel { get; set; }
-
 		/// <summary>
 		/// The URl to use for generating links back to the dashboard.
 		/// </summary>
@@ -510,41 +418,6 @@ namespace HordeServer
 		/// </summary>
 		public string? HelpSlackChannel { get; set; }
 
-		/// <summary>
-		/// Url of P4 Swarm installation
-		/// </summary>
-		public Uri? P4SwarmUrl { get; set; }
-
-		/// <summary>
-		/// The Jira service account user name
-		/// </summary>
-		public string? JiraUsername { get; set; }
-
-		/// <summary>
-		/// The Jira service account API token
-		/// </summary>
-		public string? JiraApiToken { get; set; }
-
-		/// <summary>
-		/// The Uri for the Jira installation
-		/// </summary>
-		public Uri? JiraUrl { get; set; }
-
-		/// <summary>
-		/// The number of days shared device checkouts are held
-		/// </summary>
-		public int SharedDeviceCheckoutDays { get; set; } = 3;
-
-		/// <summary>
-		/// The number of cooldown minutes for device problems
-		/// </summary>
-		public int DeviceProblemCooldownMinutes { get; set; } = 10;
-
-		/// <summary>
-		/// Channel to send device reports to
-		/// </summary>
-		public string? DeviceReportChannel { get; set; }
-
 		/// <summary>
 		/// Set the minimum size of the global thread pool
 		/// This value has been found in need of tweaking to avoid timeouts with the Redis client during bursts
@@ -563,26 +436,6 @@ namespace HordeServer
 		/// </summary>
 		public string ConfigPath { get; set; } = "globals.json";
 
-		/// <summary>
-		/// Perforce connections for use by the Horde server (not agents)
-		/// </summary>
-		public List<PerforceConnectionSettings> Perforce { get; set; } = new List<PerforceConnectionSettings>();
-
-		/// <summary>
-		/// Whether to use the local Perforce environment
-		/// </summary>
-		public bool UseLocalPerforceEnv { get; set; }
-
-		/// <summary>
-		/// Number of pooled perforce connections to keep
-		/// </summary>
-		public int PerforceConnectionPoolSize { get; set; } = 5;
-
-		/// <summary>
-		/// Whether to enable the conform task source.
-		/// </summary>
-		public bool EnableConformTasks { get; set; } = true;
-
 		/// <summary>
 		/// Forces configuration data to be read and updated as part of appplication startup, rather than on a schedule. Useful when running locally.
 		/// </summary>
@@ -593,31 +446,11 @@ namespace HordeServer
 		/// </summary>
 		public bool OpenBrowser { get; set; } = false;
 
-		/// <summary>
-		/// Directory to store the fine-grained block cache. This caches individual exports embedded in bundles.
-		/// </summary>
-		public string? BlockCacheDir { get; set; }
-
-		/// <summary>
-		/// Maximum size of the block cache. Accepts standard binary suffixes. Currently only allocates in multiples of 1024mb.
-		/// </summary>
-		public string BlockCacheSize { get; set; } = "4gb";
-
-		/// <summary>
-		/// Accessor for the block cache size in bytes
-		/// </summary>
-		public long BlockCacheSizeBytes => StringUtils.ParseBytesString(BlockCacheSize);
-
 		/// <summary>
 		/// Experimental features to enable on the server.
 		/// </summary>
 		public FeatureFlagSettings FeatureFlags { get; set; } = new();
 
-		/// <summary>
-		/// Options for the commit service
-		/// </summary>
-		public CommitSettings Commits { get; set; } = new CommitSettings();
-
 		/// <summary>
 		/// Options for OpenTelemetry
 		/// </summary>
@@ -649,94 +482,4 @@ namespace HordeServer
 			}
 		}
 	}
-
-	/// <summary>
-	/// Identifier for a pool
-	/// </summary>
-	/// <param name="Id">Id to construct from</param>
-	[JsonSchemaString]
-	[StringIdConverter(typeof(PerforceConnectionIdConverter))]
-	[TypeConverter(typeof(StringIdTypeConverter<PerforceConnectionId, PerforceConnectionIdConverter>))]
-	public record struct PerforceConnectionId(StringId Id)
-	{
-		/// <summary>
-		/// Constructor
-		/// </summary>
-		public PerforceConnectionId(string id) : this(new StringId(id))
-		{
-		}
-
-		/// <inheritdoc/>
-		public override string ToString() => Id.ToString();
-	}
-
-	/// <summary>
-	/// Converter to and from <see cref="StringId"/> instances.
-	/// </summary>
-	class PerforceConnectionIdConverter : StringIdConverter<PerforceConnectionId>
-	{
-		/// <inheritdoc/>
-		public override PerforceConnectionId FromStringId(StringId id) => new PerforceConnectionId(id);
-
-		/// <inheritdoc/>
-		public override StringId ToStringId(PerforceConnectionId value) => value.Id;
-	}
-
-	/// <summary>
-	/// Perforce connection information for use by the Horde server (for reading config files, etc...)
-	/// </summary>
-	public class PerforceConnectionSettings
-	{
-		/// <summary>
-		/// Identifier for the default perforce connection profile
-		/// </summary>
-		public static PerforceConnectionId Default { get; } = new PerforceConnectionId("default");
-
-		/// <summary>
-		/// Identifier for this server
-		/// </summary>
-		public PerforceConnectionId Id { get; set; } = Default;
-
-		/// <summary>
-		/// Server and port
-		/// </summary>
-		public string? ServerAndPort { get; set; }
-
-		/// <summary>
-		/// Credentials for the server
-		/// </summary>
-		public PerforceCredentials? Credentials { get; set; }
-
-		/// <summary>
-		/// Create a <see cref="PerforceSettings"/> object with these settings as overrides
-		/// </summary>
-		/// <returns>New perforce settings object</returns>
-		public PerforceSettings ToPerforceSettings()
-		{
-			PerforceSettings settings = new PerforceSettings(PerforceSettings.Default);
-			settings.PreferNativeClient = true;
-
-			if (!String.IsNullOrEmpty(ServerAndPort))
-			{
-				settings.ServerAndPort = ServerAndPort;
-			}
-			if (Credentials != null)
-			{
-				if (!String.IsNullOrEmpty(Credentials.UserName))
-				{
-					settings.UserName = Credentials.UserName;
-				}
-
-				if (!String.IsNullOrEmpty(Credentials.Ticket))
-				{
-					settings.Password = Credentials.Ticket;
-				}
-				else if (!String.IsNullOrEmpty(Credentials.Password))
-				{
-					settings.Password = Credentials.Password;
-				}
-			}
-			return settings;
-		}
-	}
 }
diff --git a/Engine/Source/Programs/Horde/HordeServer/Startup.cs b/Engine/Source/Programs/Horde/HordeServer/Startup.cs
--- a/Engine/Source/Programs/Horde/HordeServer/Startup.cs
+++ b/Engine/Source/Programs/Horde/HordeServer/Startup.cs
@@ -41,7 +41,6 @@ using HordeServer.Agents.Utilization;
 using HordeServer.Artifacts;
 using HordeServer.Auditing;
 using HordeServer.Authentication;
-using HordeServer.Commits;
 using HordeServer.Compute;
 using HordeServer.Configuration;
 using HordeServer.Dashboard;
@@ -49,25 +48,17 @@ using HordeServer.Devices;
 using HordeServer.Issues;
 using HordeServer.Issues.External;
 using HordeServer.Jobs;
-using HordeServer.Jobs.Bisect;
-using HordeServer.Jobs.Graphs;
 using HordeServer.Jobs.Schedules;
-using HordeServer.Jobs.Templates;
 using HordeServer.Jobs.TestData;
-using HordeServer.Jobs.Timing;
 using HordeServer.Logs;
 using HordeServer.Notifications;
-using HordeServer.Notifications.Sinks;
 using HordeServer.Perforce;
 using HordeServer.Plugins;
-using HordeServer.Replicators;
 using HordeServer.Server;
 using HordeServer.Server.Notices;
 using HordeServer.ServiceAccounts;
 using HordeServer.Storage;
-using HordeServer.Streams;
 using HordeServer.Tasks;
-using HordeServer.Ugs;
 using HordeServer.Users;
 using HordeServer.Utilities;
 using Microsoft.AspNetCore.Authentication;
@@ -306,10 +297,11 @@ namespace HordeServer
 			// Register the plugin collection
 			PluginCollection pluginCollection = new PluginCollection();
 			pluginCollection.Add<AnalyticsPlugin>();
+			pluginCollection.Add<BuildPlugin>();
 			pluginCollection.Add<ComputePlugin>();
-			pluginCollection.Add<StoragePlugin>();
 			pluginCollection.Add<DdcPlugin>();
 			pluginCollection.Add<SecretsPlugin>();
+			pluginCollection.Add<StoragePlugin>();
 			pluginCollection.Add<ToolsPlugin>();
 			services.AddSingleton<IPluginCollection>(pluginCollection);
 
@@ -375,29 +367,12 @@ namespace HordeServer
 
 			services.AddSingleton<IAccountCollection, AccountCollection>();
 			services.AddSingleton<IAgentCollection, AgentCollection>();
-			services.AddSingleton<IArtifactCollection, ArtifactCollection>();
-			services.AddSingleton<IGraphCollection, GraphCollection>();
-			services.AddSingleton<IssueCollection>();
-			services.AddSingleton<IIssueCollection>(sp => sp.GetRequiredService<IssueCollection>());
-			services.AddSingleton<ILogExtIssueProvider>(sp => sp.GetRequiredService<IssueCollection>());
-			services.AddSingleton<IJobCollection, JobCollection>();
-			services.AddSingleton<IJobStepRefCollection, JobStepRefCollection>();
-			services.AddSingleton<IJobTimingCollection, JobTimingCollection>();
 			services.AddSingleton<ILeaseCollection, LeaseCollection>();
 			services.AddSingleton<ILogCollection, LogCollection>();
 			services.AddSingleton<INotificationTriggerCollection, NotificationTriggerCollection>();
 			services.AddSingleton<IPoolCollection, PoolCollection>();
-			services.AddSingleton<IBisectTaskCollection, BisectTaskCollection>();
-			services.AddSingleton<IReplicatorCollection, ReplicatorCollection>();
 			services.AddSingleton<IServiceAccountCollection, ServiceAccountCollection>();
 			services.AddSingleton<ISessionCollection, SessionCollection>();
-			services.AddSingleton<ISubscriptionCollection, SubscriptionCollection>();
-			services.AddSingleton<IStreamCollection, StreamCollection>();
-			services.AddSingleton<ITemplateCollection, TemplateCollection>();
-			services.AddSingleton<ITestDataCollection, TestDataCollection>();
-			services.AddSingleton<IUtilizationDataCollection, UtilizationDataCollection>();
-			services.AddSingleton<ITemplateCollection, TemplateCollection>();
-			services.AddSingleton<IUgsMetadataCollection, UgsMetadataCollection>();
 			services.AddSingleton<IUserCollection, UserCollectionV2>();
 			services.AddSingleton<IDeviceCollection, DeviceCollection>();
 			services.AddSingleton<INoticeCollection, NoticeCollection>();
@@ -408,6 +383,7 @@ namespace HordeServer
 			services.AddSingleton<IConfigSource, PerforceConfigSource>();
 
 			services.AddSingleton<ConfigService>();
+			services.AddSingleton<IConfigService>(sp => sp.GetRequiredService<ConfigService>());
 			services.AddSingleton<IOptionsFactory<GlobalConfig>>(sp => sp.GetRequiredService<ConfigService>());
 			services.AddSingleton<IOptionsChangeTokenSource<GlobalConfig>>(sp => sp.GetRequiredService<ConfigService>());
 
@@ -420,12 +396,9 @@ namespace HordeServer
 			services.AddSingleton(typeof(IAuditLogFactory<>), typeof(AuditLogFactory<>));
 			services.AddSingleton(typeof(ISingletonDocument<>), typeof(SingletonDocument<>));
 
-			services.AddSingleton<IPoolSizeStrategyFactory, JobQueueStrategyFactory>();
-
 			services.AddSingleton<IAclService, AclService>();
 			services.AddSingleton<AgentService>();
 			services.AddSingleton(provider => new Lazy<AgentService>(provider.GetRequiredService<AgentService>));
-			services.AddHostedService<ArtifactExpirationService>();
 			services.AddSingleton<ConsistencyService>();
 			services.AddSingleton<RequestTrackerService>();
 			services.AddSingleton<ComputeService>();
@@ -433,7 +406,6 @@ namespace HordeServer
 			services.AddSingleton<MongoService>();
 			services.AddSingleton<IMongoService>(sp => sp.GetRequiredService<MongoService>());
 			services.AddSingleton<GlobalsService>();
-			services.AddSingleton<ICommitService, CommitService>();
 			services.AddSingleton<IClock, Clock>();
 			services.AddSingleton<IDowntimeService, DowntimeService>();
 			services.AddSingleton<IssueService>();
@@ -450,49 +422,12 @@ namespace HordeServer
 			services.AddSingleton<INotificationService, NotificationService>();
 			services.AddSingleton<UnsyncCache>();
 
-			if (settings.Commits.ReplicateMetadata)
-			{
-				services.AddSingleton<PerforceServiceCache>();
-				services.AddSingleton<IPerforceService>(sp => sp.GetRequiredService<PerforceServiceCache>());
-			}
-			else
-			{
-				services.AddSingleton<PerforceService>();
-				services.AddSingleton<IPerforceService>(sp => sp.GetRequiredService<PerforceService>());
-			}
-			services.AddSingleton<PerforceReplicator>();
-
-			services.AddSingleton<PerforceLoadBalancer>();
 			services.AddSingleton<PoolService>();
-			services.AddSingleton<ReplicationService>();
 			services.AddSingleton<ScheduleService>();
 
-			if (settings.SlackToken != null)
-			{
-				services.AddSingleton<SlackNotificationSink>();
-				services.AddSingleton<IAvatarService, SlackNotificationSink>(sp => sp.GetRequiredService<SlackNotificationSink>());
-				services.AddSingleton<INotificationSink, SlackNotificationSink>(sp => sp.GetRequiredService<SlackNotificationSink>());
-			}
-			else
-			{
-				services.AddSingleton<IAvatarService, NullAvatarService>();
-			}
-
 			services.AddScoped<OAuthControllerFilter>();
 
-			services.AddSingleton<DeviceService>();
 			services.AddSingleton<NoticeService>();
-			services.AddSingleton<IBlockCache>(sp => CreateBlockCache(sp));
-			services.AddSingleton<TestDataService>();
-
-			if (settings.JiraUrl != null)
-			{
-				services.AddSingleton<IExternalIssueService, JiraService>();
-			}
-			else
-			{
-				services.AddSingleton<IExternalIssueService, DefaultExternalIssueService>();
-			}
 
 			AuthenticationBuilder authBuilder = services.AddAuthentication(options =>
 				{
@@ -678,7 +613,6 @@ namespace HordeServer
 			if (settings.IsRunModeActive(RunMode.Worker) && !settings.MongoReadOnlyMode)
 			{
 				services.AddHostedService<AgentReportService>();
-				services.AddHostedService<BisectService>();
 				services.AddHostedService(provider => provider.GetRequiredService<FleetService>());
 				services.AddHostedService(provider => provider.GetRequiredService<ConsistencyService>());
 				services.AddHostedService(provider => provider.GetRequiredService<IssueService>());
@@ -695,26 +629,6 @@ namespace HordeServer
 				services.AddHostedService(provider => provider.GetRequiredService<ComputeService>());
 				services.AddHostedService(provider => provider.GetRequiredService<EnrollmentService>());
 				services.AddHostedService(provider => provider.GetRequiredService<StorageService>());
-
-				if (settings.Commits.ReplicateMetadata)
-				{
-					services.AddHostedService(provider => provider.GetRequiredService<PerforceServiceCache>());
-				}
-
-				if (settings.Commits.ReplicateContent)
-				{
-					services.AddHostedService(provider => provider.GetRequiredService<ReplicationService>());
-				}
-
-				if (!settings.DisableSchedules)
-				{
-					services.AddHostedService(provider => provider.GetRequiredService<ScheduleService>());
-				}
-
-				if (settings.SlackToken != null)
-				{
-					services.AddHostedService(provider => provider.GetRequiredService<SlackNotificationSink>());
-				}
 			}
 
 			services.AddHostedService(provider => provider.GetRequiredService<IExternalIssueService>());
@@ -825,13 +739,6 @@ namespace HordeServer
 			OnAddHealthChecks(services);
 		}
 
-		static BlockCache CreateBlockCache(IServiceProvider serviceProvider)
-		{
-			ServerSettings serverSettings = serviceProvider.GetRequiredService<IOptions<ServerSettings>>().Value;
-			DirectoryReference cacheDir = DirectoryReference.Combine(ServerApp.DataDir, String.IsNullOrEmpty(serverSettings.BlockCacheDir) ? "BlockCache" : serverSettings.BlockCacheDir);
-			return BlockCache.Create(cacheDir, (int)(serverSettings.BlockCacheSizeBytes / (1024 * 1024 * 1024)));
-		}
-
 		public static void ConfigureFormatters()
 		{
 			LogValueFormatter.RegisterTypeAnnotation<AgentId>("AgentId");
diff --git a/Engine/Source/Programs/Horde/HordeServer/Utilities/Clock.cs b/Engine/Source/Programs/Horde/HordeServer/Utilities/Clock.cs
--- a/Engine/Source/Programs/Horde/HordeServer/Utilities/Clock.cs
+++ b/Engine/Source/Programs/Horde/HordeServer/Utilities/Clock.cs
@@ -15,40 +15,6 @@ using TimeZoneConverter;
 
 namespace HordeCommon
 {
-	/// <summary>
-	/// Placeholder interface for ITicker
-	/// </summary>
-	public sealed class NullTicker : ITicker
-	{
-		/// <inheritdoc/>
-		public ValueTask DisposeAsync() => new ValueTask();
-
-		/// <inheritdoc/>
-		public Task StartAsync() => Task.CompletedTask;
-
-		/// <inheritdoc/>
-		public Task StopAsync() => Task.CompletedTask;
-	}
-
-	/// <summary>
-	/// Extension methods for <see cref="IClock"/>
-	/// </summary>
-	public static class ClockExtensions
-	{
-		/// <summary>
-		/// Gets the start of the day for the given datetime in UTC, respecting the configured timezone.
-		/// </summary>
-		/// <param name="timeZone">Time zone to adjust for</param>
-		/// <param name="time">Time to convert</param>
-		/// <returns>UTC datetime for the start of the day</returns>
-		public static DateTime GetStartOfDayUtc(this TimeZoneInfo timeZone, DateTime time)
-		{
-			DateTime currentTimeLocal = TimeZoneInfo.ConvertTime(time, timeZone);
-			DateTime startOfDayLocal = currentTimeLocal - currentTimeLocal.TimeOfDay;
-			return TimeZoneInfo.ConvertTime(startOfDayLocal, timeZone, TimeZoneInfo.Utc);
-		}
-	}
-
 	/// <summary>
 	/// Implementation of <see cref="IClock"/> which returns the current time
 	/// </summary>
diff --git a/Engine/Source/Programs/Horde/HordeServer/Utilities/OpenTelemetry.cs b/Engine/Source/Programs/Horde/HordeServer/Utilities/OpenTelemetry.cs
--- a/Engine/Source/Programs/Horde/HordeServer/Utilities/OpenTelemetry.cs
+++ b/Engine/Source/Programs/Horde/HordeServer/Utilities/OpenTelemetry.cs
@@ -7,10 +7,6 @@ using System.Diagnostics.Metrics;
 using System.Linq;
 using System.Net.Http;
 using System.Runtime.CompilerServices;
-using EpicGames.Core;
-using EpicGames.Horde.Jobs;
-using EpicGames.Horde.Jobs.Templates;
-using EpicGames.Horde.Streams;
 using Microsoft.AspNetCore.Http;
 using Microsoft.Extensions.DependencyInjection;
 using Microsoft.Extensions.Logging;
@@ -89,18 +85,18 @@ public static class OpenTelemetryHelper
 		void DatadogAspNetRequestEnricher(Activity activity, HttpRequest request)
 		{
 			activity.DisplayName = $"{request.Method} {request.Headers.Host}{request.Path}";
-			
+
 			// From https://docs.datadoghq.com/standard-attributes/
 			activity.SetTag("service.name", settings.ServiceName);
 			activity.SetTag("operation.name", "http.request");
-			
+
 			// Resolve client's IP via headers or actual TCP/IP remote IP
 			string? forwardedForHeader = request.HttpContext.Request.Headers["X-Forwarded-For"].FirstOrDefault();
 			string? clientIp = String.IsNullOrEmpty(forwardedForHeader)
 				? request.HttpContext.Connection.RemoteIpAddress?.ToString()
 				: forwardedForHeader;
 			activity.SetTag("http.client_ip", clientIp);
-			
+
 			// Header sent by the dashboard to indicate how long a user has been inactive for a particular browser page (in seconds)
 			if (request.Headers.TryGetValue("X-Horde-LastUserActivity", out StringValues values))
 			{
@@ -111,7 +107,7 @@ public static class OpenTelemetryHelper
 				}
 			}
 		}
-		
+
 		void DatadogAspNetResponseEnricher(Activity activity, HttpResponse response)
 		{
 			// The request hook above is executed too early in the middleware chain so user related information must be read here
@@ -244,31 +240,6 @@ public static class OpenTelemetryHelper
 /// </summary>
 public static class OpenTelemetrySpanExtensions
 {
-	/// <summary>Set a key:value tag on the span</summary>
-	/// <returns>This span instance, for chaining</returns>
-	public static TelemetrySpan SetAttribute(this TelemetrySpan span, string key, ContentHash value)
-	{
-		span.SetAttribute(key, value.ToString());
-		return span;
-	}
-
-	/// <summary>Set a key:value tag on the span</summary>
-	/// <returns>This span instance, for chaining</returns>
-	public static TelemetrySpan SetAttribute(this TelemetrySpan span, string key, SubResourceId value)
-	{
-		span.SetAttribute(key, value.ToString());
-		return span;
-	}
-
-	/// <inheritdoc cref="TelemetrySpan.SetAttribute(System.String, System.String)"/>
-	public static TelemetrySpan SetAttribute(this TelemetrySpan span, string key, StreamId? value) => span.SetAttribute(key, value?.ToString());
-
-	/// <inheritdoc cref="TelemetrySpan.SetAttribute(System.String, System.String)"/>
-	public static TelemetrySpan SetAttribute(this TelemetrySpan span, string key, TemplateId? value) => span.SetAttribute(key, value?.ToString());
-
-	/// <inheritdoc cref="TelemetrySpan.SetAttribute(System.String, System.String)"/>
-	public static TelemetrySpan SetAttribute(this TelemetrySpan span, string key, TemplateId[]? values) => span.SetAttribute(key, values != null ? String.Join(',', values.Select(x => x.Id.ToString())) : null);
-
 	/// <summary>
 	/// Start a MongoDB-based tracing span
 	/// </summary>
