Horde: Move build functionality into a plugin.

diff --git a/Engine/Source/Programs/Horde/Plugins/Build/HordeServer.Build/Jobs/Templates/ITemplateCollection.cs b/Engine/Source/Programs/Horde/Plugins/Build/HordeServer.Build/Jobs/Templates/ITemplateCollection.cs
new file mode 100644
--- /dev/null
+++ b/Engine/Source/Programs/Horde/Plugins/Build/HordeServer.Build/Jobs/Templates/ITemplateCollection.cs
@@ -0,0 +1,27 @@
+// Copyright Epic Games, Inc. All Rights Reserved.
+
+using System.Threading.Tasks;
+using EpicGames.Core;
+
+namespace HordeServer.Jobs.Templates
+{
+	/// <summary>
+	/// Interface for a collection of template documents
+	/// </summary>
+	public interface ITemplateCollection
+	{
+		/// <summary>
+		/// Gets a template by ID
+		/// </summary>
+		/// <param name="templateId">Unique id of the template</param>
+		/// <returns>The template document</returns>
+		Task<ITemplate?> GetAsync(ContentHash templateId);
+
+		/// <summary>
+		/// Gets a template instance from its configuration
+		/// </summary>
+		/// <param name="templateConfig">The template configuration</param>
+		/// <returns>Template instance</returns>
+		Task<ITemplate> GetOrAddAsync(TemplateConfig templateConfig);
+	}
+}
diff --git a/Engine/Source/Programs/Horde/Plugins/Build/HordeServer.Build/StaticBuildConfig.cs b/Engine/Source/Programs/Horde/Plugins/Build/HordeServer.Build/StaticBuildConfig.cs
new file mode 100644
--- /dev/null
+++ b/Engine/Source/Programs/Horde/Plugins/Build/HordeServer.Build/StaticBuildConfig.cs
@@ -0,0 +1,270 @@
+﻿// Copyright Epic Games, Inc. All Rights Reserved.
+
+using System;
+using System.Collections.Generic;
+using System.ComponentModel;
+using System.Linq;
+using System.Text;
+using System.Threading.Tasks;
+using EpicGames.Core;
+using EpicGames.Horde;
+using EpicGames.Horde.Storage.Bundles;
+using EpicGames.Horde.Storage.Nodes;
+using EpicGames.Perforce;
+
+namespace HordeServer
+{
+	public class StaticBuildConfig
+	{
+		/// <summary>
+		/// Perforce connections for use by the Horde server (not agents)
+		/// </summary>
+		public List<PerforceConnectionSettings> Perforce { get; set; } = new List<PerforceConnectionSettings>();
+
+		/// <summary>
+		/// Whether to use the local Perforce environment
+		/// </summary>
+		public bool UseLocalPerforceEnv { get; set; }
+
+		/// <summary>
+		/// Number of pooled perforce connections to keep
+		/// </summary>
+		public int PerforceConnectionPoolSize { get; set; } = 5;
+
+		/// <summary>
+		/// Whether to enable the conform task source.
+		/// </summary>
+		public bool EnableConformTasks { get; set; } = true;
+
+		/// <summary>
+		/// Url of P4 Swarm installation
+		/// </summary>
+		public Uri? P4SwarmUrl { get; set; }
+
+		/// <summary>
+		/// The Jira service account user name
+		/// </summary>
+		public string? JiraUsername { get; set; }
+
+		/// <summary>
+		/// The Jira service account API token
+		/// </summary>
+		public string? JiraApiToken { get; set; }
+
+		/// <summary>
+		/// The Uri for the Jira installation
+		/// </summary>
+		public Uri? JiraUrl { get; set; }
+
+		/// <summary>
+		/// The number of days shared device checkouts are held
+		/// </summary>
+		public int SharedDeviceCheckoutDays { get; set; } = 3;
+
+		/// <summary>
+		/// The number of cooldown minutes for device problems
+		/// </summary>
+		public int DeviceProblemCooldownMinutes { get; set; } = 10;
+
+		/// <summary>
+		/// Channel to send device reports to
+		/// </summary>
+		public string? DeviceReportChannel { get; set; }
+
+		/// <summary>
+		/// Whether to run scheduled jobs.
+		/// </summary>
+		public bool DisableSchedules { get; set; }
+
+		/// <summary>
+		/// Bot token for interacting with Slack (xoxb-*)
+		/// </summary>
+		public string? SlackToken { get; set; }
+
+		/// <summary>
+		/// Token for opening a socket to slack (xapp-*)
+		/// </summary>
+		public string? SlackSocketToken { get; set; }
+
+		/// <summary>
+		/// Admin user token for Slack (xoxp-*). This is only required when using the admin endpoints to invite users.
+		/// </summary>
+		public string? SlackAdminToken { get; set; }
+
+		/// <summary>
+		/// Filtered list of slack users to send notifications to. Should be Slack user ids, separated by commas.
+		/// </summary>
+		public string? SlackUsers { get; set; }
+
+		/// <summary>
+		/// Prefix to use when reporting errors
+		/// </summary>
+		public string SlackErrorPrefix { get; set; } = ":horde-error: ";
+
+		/// <summary>
+		/// Prefix to use when reporting warnings
+		/// </summary>
+		public string SlackWarningPrefix { get; set; } = ":horde-warning: ";
+
+		/// <summary>
+		/// Channel for sending messages related to config update failures
+		/// </summary>
+		public string? ConfigNotificationChannel { get; set; }
+
+		/// <summary>
+		/// Channel to send stream notification update failures to
+		/// </summary>
+		public string? UpdateStreamsNotificationChannel { get; set; }
+
+		/// <summary>
+		/// Slack channel to send job related notifications to. Multiple channels can be specified, separated by ;
+		/// </summary>
+		public string? JobNotificationChannel { get; set; }
+
+		/// <summary>
+		/// Slack channel to send agent related notifications to.
+		/// </summary>
+		public string? AgentNotificationChannel { get; set; }
+
+		/// <summary>
+		/// The number of months to retain test data
+		/// </summary>
+		public int TestDataRetainMonths { get; set; } = 6;
+
+		/// <summary>
+		/// Directory to store the fine-grained block cache. This caches individual exports embedded in bundles.
+		/// </summary>
+		public string? BlockCacheDir { get; set; }
+
+		/// <summary>
+		/// Maximum size of the block cache. Accepts standard binary suffixes. Currently only allocates in multiples of 1024mb.
+		/// </summary>
+		public string BlockCacheSize { get; set; } = "4gb";
+
+		/// <summary>
+		/// Accessor for the block cache size in bytes
+		/// </summary>
+		public long BlockCacheSizeBytes => StringUtils.ParseBytesString(BlockCacheSize);
+
+		/// <summary>
+		/// Options for the commit service
+		/// </summary>
+		public CommitSettings Commits { get; set; } = new CommitSettings();
+	}
+
+	/// <summary>
+	/// Identifier for a pool
+	/// </summary>
+	/// <param name="Id">Id to construct from</param>
+	[JsonSchemaString]
+	[StringIdConverter(typeof(PerforceConnectionIdConverter))]
+	[TypeConverter(typeof(StringIdTypeConverter<PerforceConnectionId, PerforceConnectionIdConverter>))]
+	public record struct PerforceConnectionId(StringId Id)
+	{
+		/// <summary>
+		/// Constructor
+		/// </summary>
+		public PerforceConnectionId(string id) : this(new StringId(id))
+		{
+		}
+
+		/// <inheritdoc/>
+		public override string ToString() => Id.ToString();
+	}
+
+	/// <summary>
+	/// Converter to and from <see cref="StringId"/> instances.
+	/// </summary>
+	class PerforceConnectionIdConverter : StringIdConverter<PerforceConnectionId>
+	{
+		/// <inheritdoc/>
+		public override PerforceConnectionId FromStringId(StringId id) => new PerforceConnectionId(id);
+
+		/// <inheritdoc/>
+		public override StringId ToStringId(PerforceConnectionId value) => value.Id;
+	}
+
+	/// <summary>
+	/// Perforce connection information for use by the Horde server (for reading config files, etc...)
+	/// </summary>
+	public class PerforceConnectionSettings
+	{
+		/// <summary>
+		/// Identifier for the default perforce connection profile
+		/// </summary>
+		public static PerforceConnectionId Default { get; } = new PerforceConnectionId("default");
+
+		/// <summary>
+		/// Identifier for this server
+		/// </summary>
+		public PerforceConnectionId Id { get; set; } = Default;
+
+		/// <summary>
+		/// Server and port
+		/// </summary>
+		public string? ServerAndPort { get; set; }
+
+		/// <summary>
+		/// Credentials for the server
+		/// </summary>
+		public PerforceCredentials? Credentials { get; set; }
+
+		/// <summary>
+		/// Create a <see cref="PerforceSettings"/> object with these settings as overrides
+		/// </summary>
+		/// <returns>New perforce settings object</returns>
+		public PerforceSettings ToPerforceSettings()
+		{
+			PerforceSettings settings = new PerforceSettings(PerforceSettings.Default);
+			settings.PreferNativeClient = true;
+
+			if (!String.IsNullOrEmpty(ServerAndPort))
+			{
+				settings.ServerAndPort = ServerAndPort;
+			}
+			if (Credentials != null)
+			{
+				if (!String.IsNullOrEmpty(Credentials.UserName))
+				{
+					settings.UserName = Credentials.UserName;
+				}
+
+				if (!String.IsNullOrEmpty(Credentials.Ticket))
+				{
+					settings.Password = Credentials.Ticket;
+				}
+				else if (!String.IsNullOrEmpty(Credentials.Password))
+				{
+					settings.Password = Credentials.Password;
+				}
+			}
+			return settings;
+		}
+	}
+
+	/// <summary>
+	/// Options for the commit service
+	/// </summary>
+	public class CommitSettings
+	{
+		/// <summary>
+		/// Whether to mirror commit metadata to the database
+		/// </summary>
+		public bool ReplicateMetadata { get; set; } = true;
+
+		/// <summary>
+		/// Whether to mirror commit data to storage
+		/// </summary>
+		public bool ReplicateContent { get; set; } = true;
+
+		/// <summary>
+		/// Options for how objects are packed together
+		/// </summary>
+		public BundleOptions Bundle { get; set; } = new BundleOptions();
+
+		/// <summary>
+		/// Options for how objects are sliced
+		/// </summary>
+		public ChunkingOptions Chunking { get; set; } = new ChunkingOptions();
+	}
+}
diff --git a/Engine/Source/Programs/Horde/Plugins/Build/HordeServer.Build/Streams/IStreamCollection.cs b/Engine/Source/Programs/Horde/Plugins/Build/HordeServer.Build/Streams/IStreamCollection.cs
new file mode 100644
--- /dev/null
+++ b/Engine/Source/Programs/Horde/Plugins/Build/HordeServer.Build/Streams/IStreamCollection.cs
@@ -0,0 +1,143 @@
+// Copyright Epic Games, Inc. All Rights Reserved.
+
+using System;
+using System.Collections.Generic;
+using System.Linq;
+using System.Threading;
+using System.Threading.Tasks;
+using EpicGames.Horde.Jobs;
+using EpicGames.Horde.Jobs.Templates;
+
+namespace HordeServer.Streams
+{
+	/// <summary>
+	/// Collection of stream documents
+	/// </summary>
+	public interface IStreamCollection
+	{
+		/// <summary>
+		/// Gets a stream by ID
+		/// </summary>
+		/// <param name="streamConfig">The stream config object</param>
+		/// <param name="cancellationToken">Cancellation token for the operation</param>
+		/// <returns>The stream document</returns>
+		Task<IStream> GetAsync(StreamConfig streamConfig, CancellationToken cancellationToken = default);
+
+		/// <summary>
+		/// Gets a stream by ID
+		/// </summary>
+		/// <param name="streamConfigs">The stream config object</param>
+		/// <param name="cancellationToken">Cancellation token for the operation</param>
+		/// <returns>The stream document</returns>
+		Task<IReadOnlyList<IStream>> GetAsync(IReadOnlyList<StreamConfig> streamConfigs, CancellationToken cancellationToken = default);
+
+		/// <summary>
+		/// Updates user-facing properties for an existing stream
+		/// </summary>
+		/// <param name="stream">The stream to update</param>
+		/// <param name="newPausedUntil">The new datetime for pausing builds</param>
+		/// <param name="newPauseComment">The reason for pausing</param>
+		/// <param name="cancellationToken">Cancellation token for the operation</param>
+		/// <returns>The updated stream if successful, null otherwise</returns>
+		Task<IStream?> TryUpdatePauseStateAsync(IStream stream, DateTime? newPausedUntil, string? newPauseComment, CancellationToken cancellationToken = default);
+
+		/// <summary>
+		/// Attempts to update the last trigger time for a schedule
+		/// </summary>
+		/// <param name="stream">The stream to update</param>
+		/// <param name="templateRefId">The template ref id</param>
+		/// <param name="lastTriggerTimeUtc">New last trigger time for the schedule</param>
+		/// <param name="lastTriggerChange">New last trigger changelist for the schedule</param>
+		/// <param name="newActiveJobs">New list of active jobs</param>
+		/// <param name="cancellationToken">Cancellation token for the operation</param>
+		/// <returns>The updated stream if successful, null otherwise</returns>
+		Task<IStream?> TryUpdateScheduleTriggerAsync(IStream stream, TemplateId templateRefId, DateTime? lastTriggerTimeUtc, int? lastTriggerChange, List<JobId> newActiveJobs, CancellationToken cancellationToken = default);
+
+		/// <summary>
+		/// Attempts to update a stream template ref
+		/// </summary>
+		/// <param name="streamInterface">The stream containing the template ref</param>
+		/// <param name="templateRefId">The template ref to update</param>
+		/// <param name="stepStates">The stream states to update, pass an empty list to clear all step states, otherwise will be a partial update based on included step updates</param>
+		/// <param name="cancellationToken">Cancellation token for the operation</param>
+		/// <returns></returns>
+		Task<IStream?> TryUpdateTemplateRefAsync(IStream streamInterface, TemplateId templateRefId, List<UpdateStepStateRequest>? stepStates = null, CancellationToken cancellationToken = default);
+	}
+
+	/// <summary>
+	/// Extension methods for <see cref="StreamCollection"/>
+	/// </summary>
+	public static class StreamCollectionExtensions
+	{
+		/// <summary>
+		/// Updates an existing stream
+		/// </summary>
+		/// <param name="streamCollection">The stream collection</param>
+		/// <param name="stream">The stream to update</param>
+		/// <param name="newPausedUntil">The new datetime for pausing builds</param>
+		/// <param name="newPauseComment">The reason for pausing</param>
+		/// <param name="cancellationToken">Cancellation token for the operation</param>
+		/// <returns>Async task object</returns>
+		public static async Task<IStream?> UpdatePauseStateAsync(this IStreamCollection streamCollection, IStream? stream, DateTime? newPausedUntil = null, string? newPauseComment = null, CancellationToken cancellationToken = default)
+		{
+			for (; stream != null; stream = await streamCollection.GetAsync(stream.Config, cancellationToken))
+			{
+				IStream? newStream = await streamCollection.TryUpdatePauseStateAsync(stream, newPausedUntil, newPauseComment, cancellationToken);
+				if (newStream != null)
+				{
+					return newStream;
+				}
+			}
+			return null;
+		}
+
+		/// <summary>
+		/// Attempts to update the last trigger time for a schedule
+		/// </summary>
+		/// <param name="streamCollection">The stream collection</param>
+		/// <param name="stream">The stream to update</param>
+		/// <param name="templateRefId">The template ref id</param>
+		/// <param name="lastTriggerTimeUtc"></param>
+		/// <param name="lastTriggerChange"></param>
+		/// <param name="addJobs">Jobs to add</param>
+		/// <param name="removeJobs">Jobs to remove</param>
+		/// <param name="cancellationToken">Cancellation token for the operation</param>
+		/// <returns>True if the stream was updated</returns>
+		public static async Task<IStream?> UpdateScheduleTriggerAsync(this IStreamCollection streamCollection, IStream stream, TemplateId templateRefId, DateTime? lastTriggerTimeUtc = null, int? lastTriggerChange = null, List<JobId>? addJobs = null, List<JobId>? removeJobs = null, CancellationToken cancellationToken = default)
+		{
+			IStream? newStream = stream;
+			while (newStream != null)
+			{
+				ITemplateRef? templateRef;
+				if (!newStream.Templates.TryGetValue(templateRefId, out templateRef))
+				{
+					break;
+				}
+				if (templateRef.Schedule == null)
+				{
+					break;
+				}
+
+				IEnumerable<JobId> newActiveJobs = templateRef.Schedule.ActiveJobs;
+				if (removeJobs != null)
+				{
+					newActiveJobs = newActiveJobs.Except(removeJobs);
+				}
+				if (addJobs != null)
+				{
+					newActiveJobs = newActiveJobs.Union(addJobs);
+				}
+
+				newStream = await streamCollection.TryUpdateScheduleTriggerAsync(newStream, templateRefId, lastTriggerTimeUtc, lastTriggerChange, newActiveJobs.ToList(), cancellationToken);
+
+				if (newStream != null)
+				{
+					return newStream;
+				}
+
+				newStream = await streamCollection.GetAsync(stream.Config, cancellationToken);
+			}
+			return null;
+		}
+	}
+}
diff --git a/Engine/Source/Programs/Horde/Plugins/Build/HordeServer.Build/Streams/StreamCollection.cs b/Engine/Source/Programs/Horde/Plugins/Build/HordeServer.Build/Streams/StreamCollection.cs
new file mode 100644
--- /dev/null
+++ b/Engine/Source/Programs/Horde/Plugins/Build/HordeServer.Build/Streams/StreamCollection.cs
@@ -0,0 +1,511 @@
+// Copyright Epic Games, Inc. All Rights Reserved.
+
+using System;
+using System.Collections.Generic;
+using System.IO;
+using System.Linq;
+using System.Threading;
+using System.Threading.Tasks;
+using EpicGames.Core;
+using EpicGames.Horde.Jobs;
+using EpicGames.Horde.Jobs.Templates;
+using EpicGames.Horde.Streams;
+using EpicGames.Horde.Users;
+using HordeServer.Server;
+using HordeServer.Utilities;
+using MongoDB.Bson.Serialization.Attributes;
+using MongoDB.Driver;
+
+namespace HordeServer.Streams
+{
+	/// <summary>
+	/// Collection of stream documents
+	/// </summary>
+	public class StreamCollection : IStreamCollection
+	{
+		/// <summary>
+		/// Information about a stream
+		/// </summary>
+		class StreamDoc : IStream
+		{
+			[BsonRequired, BsonId]
+			public StreamId Id { get; set; }
+
+			public Dictionary<TemplateId, TemplateRefDoc> Templates { get; set; } = new Dictionary<TemplateId, TemplateRefDoc>();
+			public DateTime? PausedUntil { get; set; }
+			public string? PauseComment { get; set; }
+
+			public int UpdateIndex { get; set; }
+
+			[BsonConstructor]
+			private StreamDoc()
+			{
+			}
+
+			public StreamDoc(StreamId id)
+			{
+				Id = id;
+			}
+
+			#region IStream Implementation
+
+			[BsonIgnore]
+			public StreamConfig Config { get; private set; } = null!;
+
+			[BsonIgnore]
+			IReadOnlyDictionary<TemplateId, ITemplateRef>? _cachedTemplates;
+
+			IReadOnlyDictionary<TemplateId, ITemplateRef> IStream.Templates
+			{
+				get
+				{
+					_cachedTemplates ??= Templates.ToDictionary(x => x.Key, x => (ITemplateRef)x.Value);
+					return _cachedTemplates;
+				}
+			}
+
+			public bool PostLoad(StreamConfig config, DateTime utcNow)
+			{
+				Config = config;
+
+				bool replaceDocument = false;
+
+				Dictionary<TemplateId, TemplateRefDoc> templateRefDocs = new Dictionary<TemplateId, TemplateRefDoc>(config.Templates.Count);
+				foreach (TemplateRefConfig templateRefConfig in config.Templates)
+				{
+					TemplateRefDoc? templateRefDoc;
+					if (!Templates.TryGetValue(templateRefConfig.Id, out templateRefDoc))
+					{
+						templateRefDoc = new TemplateRefDoc();
+						replaceDocument = true;
+					}
+					replaceDocument |= templateRefDoc.PostLoad(templateRefConfig, utcNow);
+
+					templateRefDocs.Add(templateRefConfig.Id, templateRefDoc);
+				}
+
+				replaceDocument |= Templates.Count != templateRefDocs.Count;
+				Templates = templateRefDocs;
+
+				return replaceDocument;
+			}
+
+			#endregion
+		}
+
+		class TemplateRefDoc : ITemplateRef
+		{
+			[BsonIgnoreIfNull]
+			public TemplateScheduleDoc? Schedule { get; set; }
+
+			[BsonIgnoreIfNull]
+			public List<TemplateStepDoc>? StepStates { get; set; }
+
+			#region ITemplateRef implementation
+
+			[BsonIgnore]
+			public TemplateId Id { get; private set; }
+
+			[BsonIgnore]
+			public TemplateRefConfig Config { get; private set; } = null!;
+
+			ITemplateSchedule? ITemplateRef.Schedule => Schedule;
+			IReadOnlyList<ITemplateStep> ITemplateRef.StepStates => (IReadOnlyList<ITemplateStep>?)StepStates ?? Array.Empty<ITemplateStep>();
+
+			public bool PostLoad(TemplateRefConfig config, DateTime utcNow)
+			{
+				Id = config.Id;
+				Config = config;
+
+				bool replaceDocument = false;
+
+				if (Config.Schedule == null)
+				{
+					if (Schedule != null)
+					{
+						Schedule = null;
+						replaceDocument = true;
+					}
+				}
+				else
+				{
+					if (Schedule == null)
+					{
+						Schedule = new TemplateScheduleDoc();
+						Schedule.LastTriggerTimeUtc = utcNow;
+						replaceDocument = true;
+					}
+				}
+
+				Schedule?.PostLoad(this);
+
+				if (StepStates != null)
+				{
+					int count = StepStates.RemoveAll(x => x.PausedByUserId == null);
+					replaceDocument |= count > 0;
+				}
+
+				return replaceDocument;
+			}
+
+			#endregion
+		}
+
+		class TemplateScheduleDoc : ITemplateSchedule
+		{
+			public int LastTriggerChange { get; set; }
+
+			[BsonIgnoreIfNull, Obsolete("Use LastTriggerTimeUtc instead")]
+			public DateTimeOffset? LastTriggerTime { get; set; }
+
+			public DateTime LastTriggerTimeUtc { get; set; }
+			public List<JobId> ActiveJobs { get; set; } = new List<JobId>();
+
+			#region ITemplateSchedule implementation
+
+			[BsonIgnore]
+			TemplateRefDoc? _owner;
+
+			[BsonIgnore]
+			public ScheduleConfig Config => _owner?.Config.Schedule!;
+
+			IReadOnlyList<JobId> ITemplateSchedule.ActiveJobs => ActiveJobs;
+
+			public void PostLoad(TemplateRefDoc owner)
+			{
+				_owner = owner;
+
+#pragma warning disable CS0618 // Type or member is obsolete
+				if (LastTriggerTime.HasValue)
+				{
+					LastTriggerTimeUtc = LastTriggerTime.Value.UtcDateTime;
+					LastTriggerTime = null;
+				}
+#pragma warning restore CS0618 // Type or member is obsolete
+			}
+
+			#endregion
+		}
+
+		class TemplateStepDoc : ITemplateStep
+		{
+			public string Name { get; set; } = String.Empty;
+			public UserId? PausedByUserId { get; set; }
+			public DateTime? PauseTimeUtc { get; set; }
+
+			UserId ITemplateStep.PausedByUserId => PausedByUserId ?? UserId.Empty;
+			DateTime ITemplateStep.PauseTimeUtc => PauseTimeUtc ?? DateTime.MinValue;
+
+			public TemplateStepDoc()
+			{
+			}
+
+			public TemplateStepDoc(string name, UserId pausedByUserId, DateTime pauseTimeUtc)
+			{
+				Name = name;
+				PausedByUserId = pausedByUserId;
+				PauseTimeUtc = pauseTimeUtc;
+			}
+		}
+
+		readonly IMongoCollection<StreamDoc> _streams;
+		readonly IClock _clock;
+
+		/// <summary>
+		/// Constructor
+		/// </summary>
+		/// <param name="mongoService">The database service instance</param>
+		/// <param name="clock"></param>
+		public StreamCollection(IMongoService mongoService, IClock clock)
+		{
+			_streams = mongoService.GetCollection<StreamDoc>("Streams");
+			_clock = clock;
+		}
+
+		/// <inheritdoc/>
+		public async Task<IStream> GetAsync(StreamConfig streamConfig, CancellationToken cancellationToken)
+			=> await GetInternalAsync(streamConfig, cancellationToken);
+
+		async Task<StreamDoc> GetInternalAsync(StreamConfig streamConfig, CancellationToken cancellationToken)
+		{
+			FindOneAndUpdateOptions<StreamDoc, StreamDoc> options = new FindOneAndUpdateOptions<StreamDoc, StreamDoc>();
+			options.IsUpsert = true;
+			options.ReturnDocument = ReturnDocument.After;
+
+			for (; ; )
+			{
+				StreamDoc stream = await _streams.FindOneAndUpdateAsync<StreamDoc>(x => x.Id == streamConfig.Id, Builders<StreamDoc>.Update.SetOnInsert(x => x.UpdateIndex, 0), options, cancellationToken);
+				if (await PostLoadAsync(stream, streamConfig, cancellationToken))
+				{
+					return stream;
+				}
+			}
+		}
+
+		/// <inheritdoc/>
+		public async Task<IReadOnlyList<IStream>> GetAsync(IReadOnlyList<StreamConfig> streamConfigs, CancellationToken cancellationToken)
+			=> await GetInternalAsync(streamConfigs, cancellationToken);
+
+		async Task<List<StreamDoc>> GetInternalAsync(IReadOnlyList<StreamConfig> streamConfigs, CancellationToken cancellationToken)
+		{
+			FilterDefinition<StreamDoc> filter = Builders<StreamDoc>.Filter.In(x => x.Id, streamConfigs.Select(x => x.Id));
+			List<StreamDoc> matches = await _streams.Find(filter).ToListAsync(cancellationToken);
+
+			List<StreamDoc> results = new List<StreamDoc>(streamConfigs.Count);
+			foreach (StreamConfig streamConfig in streamConfigs)
+			{
+				StreamDoc? stream = matches.FirstOrDefault(x => x.Id == streamConfig.Id);
+				if (stream == null)
+				{
+					stream = await GetInternalAsync(streamConfig, cancellationToken);
+				}
+				else
+				{
+					stream.PostLoad(streamConfig, _clock.UtcNow);
+				}
+				results.Add(stream);
+			}
+			return results;
+		}
+
+		async ValueTask<bool> PostLoadAsync(StreamDoc streamDoc, StreamConfig streamConfig, CancellationToken cancellationToken)
+		{
+			if (streamDoc.PostLoad(streamConfig, _clock.UtcNow))
+			{
+				int updateIndex = streamDoc.UpdateIndex++;
+
+				ReplaceOneResult result = await _streams.ReplaceOneAsync(x => x.Id == streamDoc.Id && x.UpdateIndex == updateIndex, streamDoc, cancellationToken: cancellationToken);
+				if (result.MatchedCount != 1)
+				{
+					return false;
+				}
+			}
+			return true;
+		}
+
+		/// <inheritdoc/>
+		public async Task<IStream?> TryUpdateTemplateRefAsync(IStream streamInterface, TemplateId templateId, List<UpdateStepStateRequest>? stepStates = null, CancellationToken cancellationToken = default)
+		{
+			StreamDoc stream = (StreamDoc)streamInterface;
+
+			TemplateRefDoc? templateRef;
+			if (!stream.Templates.TryGetValue(templateId, out templateRef))
+			{
+				return null;
+			}
+
+			UpdateDefinitionBuilder<StreamDoc> updateBuilder = Builders<StreamDoc>.Update;
+			List<UpdateDefinition<StreamDoc>> updates = new List<UpdateDefinition<StreamDoc>>();
+
+			Dictionary<TemplateId, TemplateRefDoc> newTemplates = new Dictionary<TemplateId, TemplateRefDoc>(stream.Templates);
+
+			// clear
+			if (stepStates != null && stepStates.Count == 0)
+			{
+				bool hasUpdates = false;
+				foreach (KeyValuePair<TemplateId, TemplateRefDoc> entry in newTemplates)
+				{
+					if (entry.Value.StepStates != null)
+					{
+						hasUpdates = true;
+						entry.Value.StepStates = null;
+					}
+				}
+
+				if (hasUpdates)
+				{
+					updates.Add(updateBuilder.Set(x => x.Templates, newTemplates));
+				}
+			}
+			else if (stepStates != null)
+			{
+				// get currently valid step states
+				List<TemplateStepDoc> newStepStates = templateRef.StepStates?.ToList() ?? new List<TemplateStepDoc>();
+
+				// generate update list
+				foreach (UpdateStepStateRequest updateState in stepStates)
+				{
+					int stateIndex = newStepStates.FindIndex(x => x.Name == updateState.Name);
+
+					UserId? pausedByUserId = updateState.PausedByUserId != null ? UserId.Parse(updateState.PausedByUserId) : null;
+
+					if (stateIndex == -1)
+					{
+						// if this is a new state without anything set, ignore it
+						if (pausedByUserId != null)
+						{
+							newStepStates.Add(new TemplateStepDoc(updateState.Name, pausedByUserId.Value, _clock.UtcNow));
+						}
+					}
+					else
+					{
+						if (pausedByUserId == null)
+						{
+							newStepStates.RemoveAt(stateIndex);
+						}
+						else
+						{
+							newStepStates[stateIndex].PausedByUserId = pausedByUserId.Value;
+						}
+					}
+				}
+
+				if (newStepStates.Count == 0)
+				{
+					templateRef.StepStates = null;
+				}
+				else
+				{
+					templateRef.StepStates = newStepStates;
+				}
+
+				updates.Add(updateBuilder.Set(x => x.Templates, newTemplates));
+			}
+
+			if (updates.Count == 0)
+			{
+				return streamInterface;
+			}
+
+			return await TryUpdateStreamAsync(stream, updateBuilder.Combine(updates), cancellationToken);
+
+		}
+
+		/// <inheritdoc/>
+		public async Task<IStream?> TryUpdatePauseStateAsync(IStream streamInterface, DateTime? newPausedUntil, string? newPauseComment, CancellationToken cancellationToken)
+		{
+			StreamDoc stream = (StreamDoc)streamInterface;
+
+			UpdateDefinitionBuilder<StreamDoc> updateBuilder = Builders<StreamDoc>.Update;
+
+			List<UpdateDefinition<StreamDoc>> updates = new List<UpdateDefinition<StreamDoc>>();
+			stream.PausedUntil = newPausedUntil;
+			stream.PauseComment = newPauseComment;
+			updates.Add(updateBuilder.Set(x => x.PausedUntil, newPausedUntil));
+			updates.Add(updateBuilder.Set(x => x.PauseComment, newPauseComment));
+
+			return await TryUpdateStreamAsync(stream, updateBuilder.Combine(updates), cancellationToken);
+		}
+
+		/// <inheritdoc/>
+		public async Task<IStream?> TryUpdateScheduleTriggerAsync(IStream streamInterface, TemplateId templateId, DateTime? lastTriggerTimeUtc, int? lastTriggerChange, List<JobId> newActiveJobs, CancellationToken cancellationToken)
+		{
+			StreamDoc stream = (StreamDoc)streamInterface;
+			TemplateRefDoc template = stream.Templates[templateId];
+			TemplateScheduleDoc schedule = template.Schedule!;
+
+			// Build the updates. MongoDB driver cannot parse TemplateRefId in expression tree; need to specify field name explicitly
+			List<UpdateDefinition<StreamDoc>> updates = new List<UpdateDefinition<StreamDoc>>();
+			if (lastTriggerTimeUtc.HasValue && lastTriggerTimeUtc.Value != schedule.LastTriggerTimeUtc)
+			{
+				FieldDefinition<StreamDoc, DateTime> lastTriggerTimeField = $"{nameof(stream.Templates)}.{templateId}.{nameof(template.Schedule)}.{nameof(schedule.LastTriggerTimeUtc)}";
+				updates.Add(Builders<StreamDoc>.Update.Set(lastTriggerTimeField, lastTriggerTimeUtc.Value));
+				schedule.LastTriggerTimeUtc = lastTriggerTimeUtc.Value;
+			}
+			if (lastTriggerChange.HasValue && lastTriggerChange.Value > schedule.LastTriggerChange)
+			{
+				FieldDefinition<StreamDoc, int> lastTriggerChangeField = $"{nameof(stream.Templates)}.{templateId}.{nameof(template.Schedule)}.{nameof(schedule.LastTriggerChange)}";
+				updates.Add(Builders<StreamDoc>.Update.Set(lastTriggerChangeField, lastTriggerChange.Value));
+				schedule.LastTriggerChange = lastTriggerChange.Value;
+			}
+			if (newActiveJobs != null)
+			{
+				FieldDefinition<StreamDoc, List<JobId>> field = $"{nameof(stream.Templates)}.{templateId}.{nameof(template.Schedule)}.{nameof(schedule.ActiveJobs)}";
+				updates.Add(Builders<StreamDoc>.Update.Set(field, newActiveJobs));
+				schedule.ActiveJobs = newActiveJobs;
+			}
+
+			return (updates.Count == 0) ? streamInterface : await TryUpdateStreamAsync(stream, Builders<StreamDoc>.Update.Combine(updates), cancellationToken);
+		}
+
+		/// <summary>
+		/// Update a stream
+		/// </summary>
+		/// <param name="stream">The stream to update</param>
+		/// <param name="update">The update definition</param>
+		/// <param name="cancellationToken">Cancellation token for the operation</param>
+		/// <returns>The updated document, or null the update failed</returns>
+		private async Task<StreamDoc?> TryUpdateStreamAsync(StreamDoc stream, UpdateDefinition<StreamDoc> update, CancellationToken cancellationToken)
+		{
+			FilterDefinition<StreamDoc> filter = Builders<StreamDoc>.Filter.Expr(x => x.Id == stream.Id && x.UpdateIndex == stream.UpdateIndex);
+			update = update.Set(x => x.UpdateIndex, stream.UpdateIndex + 1);
+
+			FindOneAndUpdateOptions<StreamDoc> options = new FindOneAndUpdateOptions<StreamDoc> { ReturnDocument = ReturnDocument.After };
+
+			StreamDoc? result = await _streams.FindOneAndUpdateAsync(filter, update, options, cancellationToken);
+			if (result != null)
+			{
+				result.PostLoad(stream.Config, _clock.UtcNow);
+			}
+			return result;
+		}
+
+		/// <summary>
+		/// Checks the stream definition for consistency
+		/// </summary>
+		public static void Validate(StreamId streamId, StreamConfig config)
+		{
+			HashSet<TemplateId> remainingTemplates = new HashSet<TemplateId>(config.Templates.Select(x => x.Id));
+
+			// Check the default preflight template is valid
+			if (config.DefaultPreflight != null)
+			{
+				if (config.DefaultPreflight.TemplateId != null && !remainingTemplates.Contains(config.DefaultPreflight.TemplateId.Value))
+				{
+					throw new InvalidStreamException($"Default preflight template was listed as '{config.DefaultPreflight.TemplateId.Value}', but no template was found by that name");
+				}
+			}
+
+			// Check the chained jobs are valid
+			foreach (TemplateRefConfig templateRef in config.Templates)
+			{
+				if (templateRef.ChainedJobs != null)
+				{
+					foreach (ChainedJobTemplateConfig chainedJob in templateRef.ChainedJobs)
+					{
+						if (!remainingTemplates.Contains(chainedJob.TemplateId))
+						{
+							throw new InvalidDataException($"Invalid template ref id '{chainedJob.TemplateId}");
+						}
+					}
+				}
+			}
+
+			// Check that all the templates are referenced by a tab
+			HashSet<TemplateId> undefinedTemplates = new();
+			foreach (TabConfig tab in config.Tabs)
+			{
+				if (tab.Templates != null)
+				{
+					remainingTemplates.ExceptWith(tab.Templates);
+					foreach (TemplateId templateId in tab.Templates)
+					{
+						if (config.Templates.Find(x => x.Id == templateId) == null)
+						{
+							undefinedTemplates.Add(templateId);
+						}
+					}
+				}
+			}
+			if (remainingTemplates.Count > 0)
+			{
+				throw new InvalidStreamException(String.Join("\n", remainingTemplates.Select(x => $"Template '{x}' is not listed on any tab for {streamId}")));
+			}
+
+			if (undefinedTemplates.Count > 0)
+			{
+				throw new InvalidStreamException(String.Join("\n", undefinedTemplates.Select(x => $"Template '{x}' is not defined for {streamId}")));
+			}
+
+			// Check that all the agent types reference valid workspace names
+			foreach (KeyValuePair<string, AgentConfig> pair in config.AgentTypes)
+			{
+				string? workspaceTypeName = pair.Value.Workspace;
+				if (workspaceTypeName != null && !config.WorkspaceTypes.ContainsKey(workspaceTypeName))
+				{
+					throw new InvalidStreamException($"Agent type '{pair.Key}' references undefined workspace type '{pair.Value.Workspace}' in {streamId}");
+				}
+			}
+		}
+	}
+}
diff --git a/Engine/Source/Programs/Horde/Plugins/Build/HordeServer.Build/Utilities/OpenTelemetryExtensions.cs b/Engine/Source/Programs/Horde/Plugins/Build/HordeServer.Build/Utilities/OpenTelemetryExtensions.cs
new file mode 100644
--- /dev/null
+++ b/Engine/Source/Programs/Horde/Plugins/Build/HordeServer.Build/Utilities/OpenTelemetryExtensions.cs
@@ -0,0 +1,66 @@
+// Copyright Epic Games, Inc. All Rights Reserved.
+
+using System;
+using System.Collections.Generic;
+using System.IO;
+using System.Linq;
+using System.Runtime.CompilerServices;
+using System.Text;
+using System.Threading;
+using System.Threading.Tasks;
+using EpicGames.Core;
+using EpicGames.Horde.Commits;
+using EpicGames.Horde.Jobs;
+using EpicGames.Horde.Jobs.Templates;
+using EpicGames.Horde.Streams;
+using EpicGames.Horde.Users;
+using EpicGames.Perforce;
+using EpicGames.Redis;
+using HordeServer.Commits;
+using HordeServer.Server;
+using HordeServer.Streams;
+using HordeServer.Users;
+using HordeServer.Utilities;
+using Microsoft.Extensions.Hosting;
+using Microsoft.Extensions.Logging;
+using Microsoft.Extensions.Options;
+using MongoDB.Bson;
+using MongoDB.Bson.Serialization.Attributes;
+using MongoDB.Bson.Serialization.Options;
+using MongoDB.Driver;
+using OpenTelemetry.Trace;
+using StackExchange.Redis;
+
+namespace HordeServer.Utilities
+{
+	/// <summary>
+	/// Extensions for OpenTelemetry spans
+	/// </summary>
+	public static class OpenTelemetryExtensions
+	{
+		/// <summary>Set a key:value tag on the span</summary>
+		/// <returns>This span instance, for chaining</returns>
+		public static TelemetrySpan SetAttribute(this TelemetrySpan span, string key, ContentHash value)
+		{
+			span.SetAttribute(key, value.ToString());
+			return span;
+		}
+
+		/// <summary>Set a key:value tag on the span</summary>
+		/// <returns>This span instance, for chaining</returns>
+		public static TelemetrySpan SetAttribute(this TelemetrySpan span, string key, SubResourceId value)
+		{
+			span.SetAttribute(key, value.ToString());
+			return span;
+		}
+
+		/// <inheritdoc cref="TelemetrySpan.SetAttribute(System.String, System.String)"/>
+		public static TelemetrySpan SetAttribute(this TelemetrySpan span, string key, StreamId? value) => span.SetAttribute(key, value?.ToString());
+
+		/// <inheritdoc cref="TelemetrySpan.SetAttribute(System.String, System.String)"/>
+		public static TelemetrySpan SetAttribute(this TelemetrySpan span, string key, TemplateId? value) => span.SetAttribute(key, value?.ToString());
+
+		/// <inheritdoc cref="TelemetrySpan.SetAttribute(System.String, System.String)"/>
+		public static TelemetrySpan SetAttribute(this TelemetrySpan span, string key, TemplateId[]? values) => span.SetAttribute(key, values != null ? String.Join(',', values.Select(x => x.Id.ToString())) : null);
+	}
+}
diff --git a/Engine/Source/Programs/Horde/Plugins/Tools/HordeServer.Tools.Tests/Tools/ToolsControllerTests.cs b/Engine/Source/Programs/Horde/Plugins/Tools/HordeServer.Tools.Tests/Tools/ToolsControllerTests.cs
--- a/Engine/Source/Programs/Horde/Plugins/Tools/HordeServer.Tools.Tests/Tools/ToolsControllerTests.cs
+++ b/Engine/Source/Programs/Horde/Plugins/Tools/HordeServer.Tools.Tests/Tools/ToolsControllerTests.cs
@@ -33,6 +33,7 @@ public class ToolsControllerTests
 
 		PluginCollection pluginCollection = new PluginCollection();
 		pluginCollection.Add<AnalyticsPlugin>();
+		pluginCollection.Add<BuildPlugin>();
 		pluginCollection.Add<ComputePlugin>();
 		pluginCollection.Add<StoragePlugin>();
 
