Horde: Fix binding of plugin server config instances.

diff --git a/Engine/Source/Programs/Horde/HordeServer/Plugins/PluginCollection.cs b/Engine/Source/Programs/Horde/HordeServer/Plugins/PluginCollection.cs
--- a/Engine/Source/Programs/Horde/HordeServer/Plugins/PluginCollection.cs
+++ b/Engine/Source/Programs/Horde/HordeServer/Plugins/PluginCollection.cs
@@ -50,6 +50,7 @@ namespace HordeServer.Plugins
 
 			public void ConfigureServices(IConfiguration config, IServerInfo serverInfo, IServiceCollection serviceCollection)
 			{
+				serviceCollection.Configure<TServerConfig>(config);
 				serviceCollection.AddPluginConfig<TGlobalConfig>(Name);
 
 				TStartup startup = CreateStartup(config, serverInfo);
