Horde: Add legacy shims for project config serialization.

#rnx

diff --git a/Engine/Source/Programs/Horde/HordeServer/Server/GlobalConfig.cs b/Engine/Source/Programs/Horde/HordeServer/Server/GlobalConfig.cs
--- a/Engine/Source/Programs/Horde/HordeServer/Server/GlobalConfig.cs
+++ b/Engine/Source/Programs/Horde/HordeServer/Server/GlobalConfig.cs
@@ -89,6 +89,12 @@ namespace HordeServer.Server
 		/// </summary>
 		public List<ScheduledDowntime> Downtime { get; set; } = new List<ScheduledDowntime>();
 
+		/// <summary>
+		/// Metrics to aggregate on the Horde server
+		/// </summary>
+		[Obsolete("Use Plugins.Build.Projects instead")]
+		public List<LegacyProjectConfig> Projects { get; set; } = new List<LegacyProjectConfig>();
+
 		/// <summary>
 		/// Metrics to aggregate on the Horde server
 		/// </summary>
@@ -381,21 +387,28 @@ namespace HordeServer.Server
 		}
 	}
 
-	/// <summary>
-	/// Shim for supporting includes and macros in telemetry settings before creating the analytics plugin
-	/// </summary>
-	public class LegacyTelemetryStoreConfig
+#pragma warning disable
+	public class LegacyProjectConfig
 	{
-		/// <summary>
-		/// Includes for other configuration files
-		/// </summary>
 		public List<ConfigInclude> Include { get; set; } = new List<ConfigInclude>();
+		public List<ConfigMacro> Macros { get; set; } = new List<ConfigMacro>();
+		public List<LegacyStreamConfig> Streams { get; set; } = new List<LegacyStreamConfig>();
+		public ConfigResource? Logo { get; set; }
+		public ConfigResource? LogoDarkTheme { get; set; }
+	}
 
-		/// <summary>
-		/// Macros within this configuration
-		/// </summary>
+	public class LegacyStreamConfig
+	{
+		public List<ConfigInclude> Include { get; set; } = new List<ConfigInclude>();
+		public List<ConfigMacro> Macros { get; set; } = new List<ConfigMacro>();
+	}
+
+	public class LegacyTelemetryStoreConfig
+	{
+		public List<ConfigInclude> Include { get; set; } = new List<ConfigInclude>();
 		public List<ConfigMacro> Macros { get; set; } = new List<ConfigMacro>();
 	}
+#pragma warning restore
 
 	/// <summary>
 	/// Profile for executing compute requests
