Horde: Fix preprocessing of remapped fields.

#rnx

diff --git a/Engine/Source/Programs/Horde/HordeServer.Shared/Configuration/ConfigNode.cs b/Engine/Source/Programs/Horde/HordeServer.Shared/Configuration/ConfigNode.cs
--- a/Engine/Source/Programs/Horde/HordeServer.Shared/Configuration/ConfigNode.cs
+++ b/Engine/Source/Programs/Horde/HordeServer.Shared/Configuration/ConfigNode.cs
@@ -2,6 +2,7 @@
 
 using System;
 using System.Collections.Generic;
+using System.Diagnostics.CodeAnalysis;
 using System.Reflection;
 using System.Text.Json;
 using System.Text.Json.Nodes;
@@ -396,6 +397,14 @@ namespace HordeServer.Configuration
 		/// </summary>
 		public Dictionary<string, ConfigNode> Properties { get; }
 
+		/// <summary>
+		/// Constructor
+		/// </summary>
+		public ObjectConfigNode(bool isIncludeRoot, bool isMacroScope)
+			: this(isIncludeRoot, isMacroScope, Array.Empty<KeyValuePair<string, ConfigNode>>())
+		{
+		}
+
 		/// <summary>
 		/// Constructor
 		/// </summary>
@@ -629,5 +638,55 @@ namespace HordeServer.Configuration
 				}
 			}
 		}
+
+		/// <summary>
+		/// Adds a new node at a given path
+		/// </summary>
+		/// <param name="path">Path to the node to add</param>
+		/// <param name="node">The node to add the requested path</param>
+		public void AddChildNode(string path, ConfigNode node)
+		{
+			int nextDotIdx = path.IndexOf('.', StringComparison.Ordinal);
+			if (nextDotIdx == -1)
+			{
+				Properties[path] = node;
+			}
+			else
+			{
+				string nextName = path.Substring(0, nextDotIdx);
+
+				ConfigNode? nextNode;
+				if (!Properties.TryGetValue(nextName, out nextNode))
+				{
+					nextNode = new ObjectConfigNode(false, false);
+					Properties.Add(nextName, nextNode);
+				}
+
+				ObjectConfigNode nextObj = (ObjectConfigNode)nextNode;
+				nextObj.AddChildNode(path.Substring(nextDotIdx + 1), node);
+			}
+		}
+
+		/// <summary>
+		/// Tries to get a node by path
+		/// </summary>
+		/// <param name="path">Path to the node to find</param>
+		/// <param name="node">The node at the requested path</param>
+		public bool TryGetChildNode(string path, [NotNullWhen(true)] out ConfigNode? node)
+		{
+			int nextDotIdx = path.IndexOf('.', StringComparison.Ordinal);
+			if (nextDotIdx == -1)
+			{
+				return Properties.TryGetValue(path, out node);
+			}
+
+			if (Properties.TryGetValue(path.Substring(0, nextDotIdx), out ConfigNode? nextNode) && nextNode is ObjectConfigNode nextObj)
+			{
+				return nextObj.TryGetChildNode(path.Substring(nextDotIdx + 1), out node);
+			}
+
+			node = null;
+			return false;
+		}
 	}
 }
diff --git a/Engine/Source/Programs/Horde/HordeServer/Configuration/ConfigService.cs b/Engine/Source/Programs/Horde/HordeServer/Configuration/ConfigService.cs
--- a/Engine/Source/Programs/Horde/HordeServer/Configuration/ConfigService.cs
+++ b/Engine/Source/Programs/Horde/HordeServer/Configuration/ConfigService.cs
@@ -2,6 +2,7 @@
 
 using System;
 using System.Collections.Generic;
+using System.Diagnostics.CodeAnalysis;
 using System.IO;
 using System.IO.Compression;
 using System.Linq;
@@ -10,6 +11,7 @@ using System.Text.Json;
 using System.Text.Json.Nodes;
 using System.Threading;
 using System.Threading.Tasks;
+using Amazon.EC2.Model;
 using EpicGames.Core;
 using EpicGames.Horde.Users;
 using EpicGames.Redis;
@@ -17,6 +19,7 @@ using HordeServer.Plugins;
 using HordeServer.Server;
 using HordeServer.Users;
 using HordeServer.Utilities;
+using Microsoft.Extensions.Azure;
 using Microsoft.Extensions.Diagnostics.HealthChecks;
 using Microsoft.Extensions.Hosting;
 using Microsoft.Extensions.Logging;
@@ -584,11 +587,39 @@ namespace HordeServer.Configuration
 				// Read the new config in
 				Uri globalConfigUri = GetGlobalConfigUri();
 
+				// Generate the schema for the preprocessor
 				ObjectConfigNode configNode = new ObjectConfigNode(typeof(GlobalConfig));
 
+				ObjectConfigNode pluginsObj = new ObjectConfigNode(false, false);
+				configNode.Properties["Plugins"] = pluginsObj;
+
+				foreach (ILoadedPlugin plugin in _pluginCollection.LoadedPlugins)
+				{
+					if (plugin.GlobalConfigType != null)
+					{
+						pluginsObj.Properties.Add(plugin.Name.ToString(), new ObjectConfigNode(plugin.GlobalConfigType));
+					}
+				}
+
+				// Copy the schema from any remapped nodes to the source node
+				foreach (KeyValuePair<string, string> remapConfigPair in s_remapConfigValues)
+				{
+					if (configNode.TryGetChildNode(remapConfigPair.Value, out ConfigNode? node))
+					{
+						configNode.AddChildNode(remapConfigPair.Key, node);
+					}
+				}
+
+				// Preprocess the file
 				JsonObject configObj = await context.PreprocessFileAsync(GetGlobalConfigUri(), configNode, cancellationToken);
-				UpgradeConfig(configObj);
 
+				// Copy the preprocessed data into the right location
+				foreach ((string source, string target) in s_remapConfigValues)
+				{
+					CopyJsonNode(configObj, source, target);
+				}
+
+				// Bind it to the target object
 				GlobalConfig globalConfig = JsonSerializer.Deserialize<GlobalConfig>(configObj, _jsonOptions)!;
 				if (globalConfig.VersionEnum < ConfigVersion.Latest)
 				{
@@ -625,36 +656,36 @@ namespace HordeServer.Configuration
 			}
 		}
 
-		static void UpgradeConfig(JsonObject rootObj)
+		static KeyValuePair<string, string>[] s_remapConfigValues = new[]
 		{
 			// Analytics
-			CopyJsonNode(rootObj, "TelemetryStores", "Plugins.Analytics.Stores");
+			KeyValuePair.Create("TelemetryStores", "Plugins.Analytics.Stores"),
 
 			// Build
-			CopyJsonNode(rootObj, "PerforceClusters", "Plugins.Build.PerforceClusters");
-			CopyJsonNode(rootObj, "Devices", "Plugins.Build.Devices");
-			CopyJsonNode(rootObj, "MaxConformCount", "Plugins.Build.MaxConformCount");
-			CopyJsonNode(rootObj, "AgentShutdownIfDisabledGracePeriod", "Plugins.Build.AgentShutdownIfDisabledGracePeriod");
-			CopyJsonNode(rootObj, "ArtifactTypes", "Plugins.Build.ArtifactTypes");
-			CopyJsonNode(rootObj, "Projects", "Plugins.Build.Projects");
-			CopyJsonNode(rootObj, "IssueFixedTag", "Plugins.Build.IssueFixedTag");
+			KeyValuePair.Create("PerforceClusters", "Plugins.Build.PerforceClusters"),
+			KeyValuePair.Create("Devices", "Plugins.Build.Devices"),
+			KeyValuePair.Create("MaxConformCount", "Plugins.Build.MaxConformCount"),
+			KeyValuePair.Create("AgentShutdownIfDisabledGracePeriod", "Plugins.Build.AgentShutdownIfDisabledGracePeriod"),
+			KeyValuePair.Create("ArtifactTypes", "Plugins.Build.ArtifactTypes"),
+			KeyValuePair.Create("Projects", "Plugins.Build.Projects"),
+			KeyValuePair.Create("IssueFixedTag", "Plugins.Build.IssueFixedTag"),
 
 			// Compute
-			CopyJsonNode(rootObj, "Rates", "Plugins.Compute.Rates");
-			CopyJsonNode(rootObj, "Compute", "Plugins.Compute.Clusters");
-			CopyJsonNode(rootObj, "Pools", "Plugins.Compute.Pools");
-			CopyJsonNode(rootObj, "Software", "Plugins.Compute.Software");
-			CopyJsonNode(rootObj, "Networks", "Plugins.Compute.Network");
+			KeyValuePair.Create("Rates", "Plugins.Compute.Rates"),
+			KeyValuePair.Create("Compute", "Plugins.Compute.Clusters"),
+			KeyValuePair.Create("Pools", "Plugins.Compute.Pools"),
+			KeyValuePair.Create("Software", "Plugins.Compute.Software"),
+			KeyValuePair.Create("Networks", "Plugins.Compute.Network"),
 
 			// Secrets
-			CopyJsonNode(rootObj, "Secrets", "Plugins.Secrets.Secrets");
+			KeyValuePair.Create("Secrets", "Plugins.Secrets.Secrets"),
 
 			// Storage
-			CopyJsonNode(rootObj, "Storage", "Plugins.Storage");
+			KeyValuePair.Create("Storage", "Plugins.Storage"),
 
 			// Tools
-			CopyJsonNode(rootObj, "Tools", "Plugins.Tools.Tools");
-		}
+			KeyValuePair.Create("Tools", "Plugins.Tools.Tools")
+		};
 
 		static void CopyJsonNode(JsonObject rootObj, string sourcePath, string targetPath)
 		{
diff --git a/Engine/Source/Programs/Horde/HordeServer/Server/GlobalConfig.cs b/Engine/Source/Programs/Horde/HordeServer/Server/GlobalConfig.cs
--- a/Engine/Source/Programs/Horde/HordeServer/Server/GlobalConfig.cs
+++ b/Engine/Source/Programs/Horde/HordeServer/Server/GlobalConfig.cs
@@ -89,18 +89,6 @@ namespace HordeServer.Server
 		/// </summary>
 		public List<ScheduledDowntime> Downtime { get; set; } = new List<ScheduledDowntime>();
 
-		/// <summary>
-		/// Metrics to aggregate on the Horde server
-		/// </summary>
-		[Obsolete("Use Plugins.Build.Projects instead")]
-		public List<LegacyProjectConfig> Projects { get; set; } = new List<LegacyProjectConfig>();
-
-		/// <summary>
-		/// Metrics to aggregate on the Horde server
-		/// </summary>
-		[Obsolete("Use Plugins.Analytics.Stores instead")]
-		public List<LegacyTelemetryStoreConfig> TelemetryStores { get; set; } = new List<LegacyTelemetryStoreConfig>();
-
 		/// <summary>
 		/// Plugin config objects
 		/// </summary>
@@ -387,29 +375,6 @@ namespace HordeServer.Server
 		}
 	}
 
-#pragma warning disable
-	public class LegacyProjectConfig
-	{
-		public List<ConfigInclude> Include { get; set; } = new List<ConfigInclude>();
-		public List<ConfigMacro> Macros { get; set; } = new List<ConfigMacro>();
-		public List<LegacyStreamConfig> Streams { get; set; } = new List<LegacyStreamConfig>();
-		public ConfigResource? Logo { get; set; }
-		public ConfigResource? LogoDarkTheme { get; set; }
-	}
-
-	public class LegacyStreamConfig
-	{
-		public List<ConfigInclude> Include { get; set; } = new List<ConfigInclude>();
-		public List<ConfigMacro> Macros { get; set; } = new List<ConfigMacro>();
-	}
-
-	public class LegacyTelemetryStoreConfig
-	{
-		public List<ConfigInclude> Include { get; set; } = new List<ConfigInclude>();
-		public List<ConfigMacro> Macros { get; set; } = new List<ConfigMacro>();
-	}
-#pragma warning restore
-
 	/// <summary>
 	/// Profile for executing compute requests
 	/// </summary>
