Horde: Fix static analysis warnings.

#rnx

diff --git a/Engine/Source/Programs/Horde/HordeServer/Configuration/ConfigService.cs b/Engine/Source/Programs/Horde/HordeServer/Configuration/ConfigService.cs
--- a/Engine/Source/Programs/Horde/HordeServer/Configuration/ConfigService.cs
+++ b/Engine/Source/Programs/Horde/HordeServer/Configuration/ConfigService.cs
@@ -2,7 +2,6 @@
 
 using System;
 using System.Collections.Generic;
-using System.Diagnostics.CodeAnalysis;
 using System.IO;
 using System.IO.Compression;
 using System.Linq;
@@ -11,7 +10,6 @@ using System.Text.Json;
 using System.Text.Json.Nodes;
 using System.Threading;
 using System.Threading.Tasks;
-using Amazon.EC2.Model;
 using EpicGames.Core;
 using EpicGames.Horde.Users;
 using EpicGames.Redis;
@@ -19,7 +17,6 @@ using HordeServer.Plugins;
 using HordeServer.Server;
 using HordeServer.Users;
 using HordeServer.Utilities;
-using Microsoft.Extensions.Azure;
 using Microsoft.Extensions.Diagnostics.HealthChecks;
 using Microsoft.Extensions.Hosting;
 using Microsoft.Extensions.Logging;
@@ -656,7 +653,7 @@ namespace HordeServer.Configuration
 			}
 		}
 
-		static KeyValuePair<string, string>[] s_remapConfigValues = new[]
+		static readonly KeyValuePair<string, string>[] s_remapConfigValues = new[]
 		{
 			// Analytics
 			KeyValuePair.Create("TelemetryStores", "Plugins.Analytics.Stores"),
