Horde: Move a couple more services into the compute plugin.

#rnx

diff --git a/Engine/Source/Programs/Horde/HordeServer.Shared/Server/ILifetimeService.cs b/Engine/Source/Programs/Horde/HordeServer.Shared/Server/ILifetimeService.cs
new file mode 100644
--- /dev/null
+++ b/Engine/Source/Programs/Horde/HordeServer.Shared/Server/ILifetimeService.cs
@@ -0,0 +1,26 @@
+// Copyright Epic Games, Inc. All Rights Reserved.
+
+using System;
+using System.Collections.Generic;
+using System.Linq;
+using System.Text;
+using System.Threading.Tasks;
+
+namespace HordeServer.Server
+{
+	/// <summary>
+	/// Interface for the server lifetime service
+	/// </summary>
+	public interface ILifetimeService
+	{
+		/// <summary>
+		/// Returns true if the server is stopping
+		/// </summary>
+		bool IsStopping { get; }
+
+		/// <summary>
+		/// Gets an awaitable task for the server stopping
+		/// </summary>
+		Task StoppingTask { get; }
+	}
+}
diff --git a/Engine/Source/Programs/Horde/HordeServer/Server/ConsistencyService.cs b/Engine/Source/Programs/Horde/HordeServer/Server/ConsistencyService.cs
deleted file mode 100644
diff --git a/Engine/Source/Programs/Horde/HordeServer/Server/LifetimeService.cs b/Engine/Source/Programs/Horde/HordeServer/Server/LifetimeService.cs
--- a/Engine/Source/Programs/Horde/HordeServer/Server/LifetimeService.cs
+++ b/Engine/Source/Programs/Horde/HordeServer/Server/LifetimeService.cs
@@ -14,7 +14,7 @@ namespace HordeServer.Server
 	/// <summary>
 	/// Service containing an async task that allows long polling operations to complete early if the server is shutting down
 	/// </summary>
-	public sealed class LifetimeService : IHostedService, IAsyncDisposable
+	public sealed class LifetimeService : ILifetimeService, IHostedService, IAsyncDisposable
 	{
 		/// <summary>
 		/// Writer for log output
diff --git a/Engine/Source/Programs/Horde/HordeServer/Server/RpcService.cs b/Engine/Source/Programs/Horde/HordeServer/Server/RpcService.cs
deleted file mode 100644
diff --git a/Engine/Source/Programs/Horde/HordeServer/Startup.cs b/Engine/Source/Programs/Horde/HordeServer/Startup.cs
--- a/Engine/Source/Programs/Horde/HordeServer/Startup.cs
+++ b/Engine/Source/Programs/Horde/HordeServer/Startup.cs
@@ -399,9 +399,7 @@ namespace HordeServer
 			services.AddSingleton<IAclService, AclService>();
 			services.AddSingleton<AgentService>();
 			services.AddSingleton(provider => new Lazy<AgentService>(provider.GetRequiredService<AgentService>));
-			services.AddSingleton<ConsistencyService>();
 			services.AddSingleton<RequestTrackerService>();
-			services.AddSingleton<ComputeService>();
 			services.AddSingleton<MongoCommandTracer>();
 			services.AddSingleton<MongoService>();
 			services.AddSingleton<IMongoService>(sp => sp.GetRequiredService<MongoService>());
@@ -414,15 +412,11 @@ namespace HordeServer
 			services.AddSingleton<LifetimeService>();
 			services.AddHostedService(provider => provider.GetRequiredService<LifetimeService>());
 			services.AddSingleton(typeof(IHealthMonitor<>), typeof(HealthMonitor<>));
-			services.AddSingleton<EnrollmentService>();
 			services.AddSingleton<ServerStatusService>();
 			services.AddHostedService(provider => provider.GetRequiredService<ServerStatusService>());
-			services.AddSingleton<LogTailService>();
-			services.AddHostedService(provider => provider.GetRequiredService<LogTailService>()); // NB: Runs even on workers, to receive tail notifications.
 			services.AddSingleton<INotificationService, NotificationService>();
 			services.AddSingleton<UnsyncCache>();
 
-			services.AddSingleton<PoolService>();
 			services.AddSingleton<ScheduleService>();
 
 			services.AddScoped<OAuthControllerFilter>();
@@ -604,17 +598,12 @@ namespace HordeServer
 			// Hosted service that needs to run no matter the run mode of the process (server vs worker)
 			services.AddHostedService(provider => (DowntimeService)provider.GetRequiredService<IDowntimeService>());
 
-			// Always run agent service too; need to be able to listen to Redis for events on any server.
-			services.AddHostedService(provider => provider.GetRequiredService<AgentService>());
-
 			// Notifications can be triggered from any instance, so always make sure we're ticking the background task.
 			services.AddHostedService(provider => (NotificationService)provider.GetRequiredService<INotificationService>());
 
 			if (settings.IsRunModeActive(RunMode.Worker) && !settings.MongoReadOnlyMode)
 			{
 				services.AddHostedService<AgentReportService>();
-				services.AddHostedService(provider => provider.GetRequiredService<FleetService>());
-				services.AddHostedService(provider => provider.GetRequiredService<ConsistencyService>());
 				services.AddHostedService(provider => provider.GetRequiredService<IssueService>());
 				services.AddHostedService<IssueReportService>();
 				services.AddHostedService<IssueTagService>();
diff --git a/Engine/Source/Programs/Horde/Plugins/Compute/HordeServer.Compute/ComputePlugin.cs b/Engine/Source/Programs/Horde/Plugins/Compute/HordeServer.Compute/ComputePlugin.cs
--- a/Engine/Source/Programs/Horde/Plugins/Compute/HordeServer.Compute/ComputePlugin.cs
+++ b/Engine/Source/Programs/Horde/Plugins/Compute/HordeServer.Compute/ComputePlugin.cs
@@ -6,11 +6,16 @@ using Amazon.CloudWatch;
 using Amazon.EC2;
 using Amazon.Extensions.NETCore.Setup;
 using Amazon.SQS;
+using HordeServer.Agents;
+using HordeServer.Agents.Enrollment;
 using HordeServer.Agents.Fleet;
+using HordeServer.Agents.Pools;
 using HordeServer.Agents.Relay;
 using HordeServer.Aws;
 using HordeServer.Compute;
+using HordeServer.Logs;
 using HordeServer.Plugins;
+using HordeServer.Server;
 using HordeServer.Tasks;
 using Microsoft.Extensions.Configuration;
 using Microsoft.Extensions.DependencyInjection;
@@ -56,6 +61,18 @@ namespace HordeServer
 			services.AddSingleton<AgentRelayService>();
 			services.AddHostedService(provider => provider.GetRequiredService<AgentRelayService>());
 
+			services.AddSingleton<ConsistencyService>();
+			services.AddSingleton<ComputeService>();
+			services.AddSingleton<EnrollmentService>();
+			services.AddSingleton<LogTailService>();
+			services.AddSingleton<PoolService>();
+
+			// Always run tail service on workers, to receive tail notifications.
+			services.AddHostedService(provider => provider.GetRequiredService<LogTailService>());
+
+			// Always run agent service too; need to be able to listen to Redis for events on any server.
+			services.AddHostedService(provider => provider.GetRequiredService<AgentService>());
+
 			if (!_serverInfo.ReadOnlyMode)
 			{
 				services.AddSingleton<ComputeTaskSource>();
@@ -64,6 +81,12 @@ namespace HordeServer
 				services.AddSingleton<ITaskSource, UpgradeTaskSource>();
 				services.AddSingleton<ITaskSource, ShutdownTaskSource>();
 				services.AddSingleton<ITaskSource, RestartTaskSource>();
+
+				if (_serverInfo.IsRunModeActive(RunMode.Worker))
+				{
+					services.AddHostedService(provider => provider.GetRequiredService<FleetService>());
+					services.AddHostedService(provider => provider.GetRequiredService<ConsistencyService>());
+				}
 			}
 
 			if (_staticComputeConfig.WithAws)
diff --git a/Engine/Source/Programs/Horde/Plugins/Compute/HordeServer.Compute/Server/ConsistencyService.cs b/Engine/Source/Programs/Horde/Plugins/Compute/HordeServer.Compute/Server/ConsistencyService.cs
new file mode 100644
--- /dev/null
+++ b/Engine/Source/Programs/Horde/Plugins/Compute/HordeServer.Compute/Server/ConsistencyService.cs
@@ -0,0 +1,101 @@
+// Copyright Epic Games, Inc. All Rights Reserved.
+
+using System;
+using System.Collections.Generic;
+using System.Linq;
+using System.Threading;
+using System.Threading.Tasks;
+using EpicGames.Core;
+using EpicGames.Horde.Agents;
+using EpicGames.Horde.Agents.Leases;
+using EpicGames.Horde.Agents.Sessions;
+using HordeServer.Agents;
+using HordeServer.Agents.Leases;
+using HordeServer.Agents.Sessions;
+using Microsoft.Extensions.Hosting;
+using Microsoft.Extensions.Logging;
+
+namespace HordeServer.Server
+{
+	/// <summary>
+	/// Service which checks the database for consistency and fixes up any errors
+	/// </summary>
+	class ConsistencyService : IHostedService, IAsyncDisposable
+	{
+		readonly IAgentCollection _agentCollection;
+		readonly ISessionCollection _sessionCollection;
+		readonly ILeaseCollection _leaseCollection;
+		readonly IClock _clock;
+		readonly ITicker _ticker;
+		readonly ILogger<ConsistencyService> _logger;
+
+		/// <summary>
+		/// Constructor
+		/// </summary>
+		public ConsistencyService(IAgentCollection agentCollection, ISessionCollection sessionCollection, ILeaseCollection leaseCollection, IClock clock, ILogger<ConsistencyService> logger)
+		{
+			_agentCollection = agentCollection;
+			_sessionCollection = sessionCollection;
+			_leaseCollection = leaseCollection;
+			_clock = clock;
+			_ticker = clock.AddSharedTicker<ConsistencyService>(TimeSpan.FromMinutes(20.0), TickLeaderAsync, logger);
+			_logger = logger;
+		}
+
+		/// <inheritdoc/>
+		public Task StartAsync(CancellationToken cancellationToken) => _ticker.StartAsync();
+
+		/// <inheritdoc/>
+		public Task StopAsync(CancellationToken cancellationToken) => _ticker.StopAsync();
+
+		/// <inheritdoc/>
+		public async ValueTask DisposeAsync() => await _ticker.DisposeAsync();
+
+		/// <summary>
+		/// Poll for inconsistencies in the database
+		/// </summary>
+		/// <param name="cancellationToken">Stopping token</param>
+		/// <returns>Async task</returns>
+		async ValueTask TickLeaderAsync(CancellationToken cancellationToken)
+		{
+			_logger.LogInformation("Ticking consistency service...");
+
+			// Find all the active sessions
+			List<ISession> sessions = await _sessionCollection.FindActiveSessionsAsync(cancellationToken: cancellationToken);
+			Dictionary<SessionId, ISession> sessionIdToInstance = sessions.ToDictionary(x => x.Id, x => x);
+
+			// Find all the active agents
+			IReadOnlyList<IAgent> agents = await _agentCollection.FindAsync(status: AgentStatus.Ok, cancellationToken: cancellationToken);
+			Dictionary<AgentId, IAgent> agentIdToInstance = agents.ToDictionary(x => x.Id, x => x);
+
+			// Find any sessions that do not have a finish time despite their agents running something else
+			DateTime utcNow = _clock.UtcNow;
+			foreach (ISession session in sessions)
+			{
+				if (!agentIdToInstance.TryGetValue(session.AgentId, out IAgent? agent) || agent.SessionId != session.Id)
+				{
+					agent = await _agentCollection.GetAsync(session.AgentId, cancellationToken);
+					if (agent == null || agent.SessionId != session.Id)
+					{
+						_logger.LogWarning("Forcing agent {AgentId} session {SessionId} to complete.", session.AgentId, session.Id);
+						await _sessionCollection.UpdateAsync(session.Id, utcNow, null, null, cancellationToken);
+						sessionIdToInstance.Remove(session.Id);
+					}
+				}
+			}
+
+			// Find any leases that are still running when their session has terminated
+			IReadOnlyList<ILease> leases = await _leaseCollection.FindActiveLeasesAsync(cancellationToken: cancellationToken);
+			foreach (ILease lease in leases)
+			{
+				if (!sessionIdToInstance.ContainsKey(lease.SessionId))
+				{
+					ISession? session = await _sessionCollection.GetAsync(lease.SessionId, cancellationToken);
+					DateTime finishTime = session?.FinishTime ?? DateTime.UtcNow;
+					_logger.LogWarning("Setting finish time for lease {LeaseId} to {FinishTime}", lease.Id, finishTime);
+					await _leaseCollection.TrySetOutcomeAsync(lease.Id, finishTime, LeaseOutcome.Cancelled, null, cancellationToken);
+				}
+			}
+		}
+	}
+}
diff --git a/Engine/Source/Programs/Horde/Plugins/Compute/HordeServer.Compute/Server/RpcService.cs b/Engine/Source/Programs/Horde/Plugins/Compute/HordeServer.Compute/Server/RpcService.cs
new file mode 100644
--- /dev/null
+++ b/Engine/Source/Programs/Horde/Plugins/Compute/HordeServer.Compute/Server/RpcService.cs
@@ -0,0 +1,454 @@
+// Copyright Epic Games, Inc. All Rights Reserved.
+
+using System;
+using System.Buffers;
+using System.Collections.Generic;
+using System.IO;
+using System.Linq;
+using System.Net;
+using System.Security.Claims;
+using System.Threading;
+using System.Threading.Tasks;
+using EpicGames.Horde.Agents;
+using EpicGames.Horde.Agents.Sessions;
+using EpicGames.Horde.Tools;
+using Google.Protobuf;
+using Google.Protobuf.WellKnownTypes;
+using Grpc.Core;
+using HordeCommon.Rpc;
+using HordeCommon.Rpc.Messages;
+using HordeServer.Acls;
+using HordeServer.Agents;
+using HordeServer.Agents.Sessions;
+using HordeServer.Agents.Telemetry;
+using HordeServer.Tools;
+using HordeServer.Utilities;
+using Microsoft.AspNetCore.Authorization;
+using Microsoft.Extensions.Logging;
+using Microsoft.Extensions.Options;
+
+namespace HordeServer.Server
+{
+	/// <summary>
+	/// Implements the Horde gRPC service for bots updating their status and dequeing work
+	/// </summary>
+	[Authorize]
+	public class RpcService : HordeRpc.HordeRpcBase
+	{
+		/// <summary>
+		/// Timeout before closing a long-polling request (client will retry again) 
+		/// </summary>
+		internal TimeSpan _longPollTimeout = TimeSpan.FromMinutes(9);
+
+		readonly AgentService _agentService;
+		readonly ILifetimeService _lifetimeService;
+		readonly IToolCollection _toolCollection;
+		readonly IAgentTelemetryCollection _agentTelemetryCollection;
+		readonly IAclService _aclService;
+		readonly IOptionsSnapshot<ComputeConfig> _computeConfig;
+		readonly ILogger _logger;
+
+		/// <summary>
+		/// Constructor
+		/// </summary>
+		public RpcService(AgentService agentService, ILifetimeService lifetimeService, IToolCollection toolCollection, IAgentTelemetryCollection agentTelemetryCollection, IAclService aclService, IOptionsSnapshot<ComputeConfig> computeConfig, ILogger<RpcService> logger)
+		{
+			_agentService = agentService;
+			_lifetimeService = lifetimeService;
+			_toolCollection = toolCollection;
+			_agentTelemetryCollection = agentTelemetryCollection;
+			_aclService = aclService;
+			_computeConfig = computeConfig;
+			_logger = logger;
+		}
+
+		/// <summary>
+		/// Waits until the server is terminating
+		/// </summary>
+		/// <param name="reader">Request reader</param>
+		/// <param name="writer">Response writer</param>
+		/// <param name="context">Context for the call</param>
+		/// <returns>Response object</returns>
+		[Obsolete("RPC connection management now done natively")]
+		public override async Task QueryServerState(IAsyncStreamReader<RpcQueryServerStateRequest> reader, IServerStreamWriter<RpcQueryServerStateResponse> writer, ServerCallContext context)
+		{
+			if (await reader.MoveNext())
+			{
+				RpcQueryServerStateRequest request = reader.Current;
+				_logger.LogInformation("Start server query for client {Name}", request.Name);
+
+				// Return the current response
+				RpcQueryServerStateResponse response = new RpcQueryServerStateResponse();
+				response.Name = Dns.GetHostName();
+				await writer.WriteAsync(response);
+
+				// Move to the next request from the client. This should always be the end of the stream, but will not occur until the client stops requesting responses.
+				Task<bool> moveNextTask = reader.MoveNext();
+
+				// Wait for the client to close the stream or a shutdown to start
+				Task longPollDelay = Task.Delay(_longPollTimeout);
+				Task waitTask = await Task.WhenAny(moveNextTask, _lifetimeService.StoppingTask, longPollDelay);
+
+				if (waitTask == moveNextTask)
+				{
+					throw new Exception("Unexpected request to QueryServerState posted from client.");
+				}
+				else if (waitTask == _lifetimeService.StoppingTask)
+				{
+					_logger.LogInformation("Notifying client {Name} of server shutdown", request.Name);
+					await writer.WriteAsync(response);
+				}
+				else if (waitTask == longPollDelay)
+				{
+					// Send same response as server shutdown. In the agent perspective, they will be identical.
+					await writer.WriteAsync(response);
+				}
+			}
+		}
+
+		/// <summary>
+		/// Waits until the server is terminating
+		/// </summary>
+		/// <param name="reader">Request reader</param>
+		/// <param name="writer">Response writer</param>
+		/// <param name="context">Context for the call</param>
+		/// <returns>Response object</returns>
+		[Obsolete("RPC connection management now done natively")]
+		public override async Task QueryServerStateV2(IAsyncStreamReader<RpcQueryServerStateRequest> reader, IServerStreamWriter<RpcQueryServerStateResponse> writer, ServerCallContext context)
+		{
+			if (await reader.MoveNext())
+			{
+				RpcQueryServerStateRequest request = reader.Current;
+				_logger.LogDebug("Start server query for client {Name}", request.Name);
+
+				try
+				{
+					// Return the current response
+					RpcQueryServerStateResponse response = new RpcQueryServerStateResponse();
+					response.Name = Dns.GetHostName();
+					response.Stopping = _lifetimeService.IsStopping;
+					await writer.WriteAsync(response);
+
+					// Move to the next request from the client. This should always be the end of the stream, but will not occur until the client stops requesting responses.
+					Task<bool> moveNextTask = reader.MoveNext();
+
+					// Wait for the client to close the stream or a shutdown to start
+					if (await Task.WhenAny(moveNextTask, _lifetimeService.StoppingTask) == _lifetimeService.StoppingTask)
+					{
+						response.Stopping = true;
+						await writer.WriteAsync(response);
+					}
+
+					// Wait until the client has finished sending
+					while (await moveNextTask)
+					{
+						moveNextTask = reader.MoveNext();
+					}
+				}
+				catch (Exception ex)
+				{
+					_logger.LogError(ex, "Exception in QueryServerState for {Name}", request.Name);
+					throw;
+				}
+			}
+		}
+
+		static void CopyPropertyToResource(string name, List<string> properties, Dictionary<string, int> resources)
+		{
+			foreach (string property in properties)
+			{
+				if (property.Length > name.Length && property.StartsWith(name, StringComparison.OrdinalIgnoreCase) && property[name.Length] == '=')
+				{
+					int value;
+					if (Int32.TryParse(property.AsSpan(name.Length + 1), out value))
+					{
+						resources[name] = value;
+					}
+				}
+			}
+		}
+
+		static void GetCapabilities(RpcAgentCapabilities? capabilities, out List<string> properties, out Dictionary<string, int> resources)
+		{
+			properties = new List<string>();
+			resources = new Dictionary<string, int>();
+
+			if (capabilities == null)
+			{
+				return;
+			}
+			properties.AddRange(capabilities.Properties);
+
+			if (capabilities.Devices.Count <= 0)
+			{
+				return;
+			}
+
+			RpcDeviceCapabilities device = capabilities.Devices[0];
+			if (device.Properties == null)
+			{
+				return;
+			}
+
+			properties.AddRange(device.Properties);
+			CopyPropertyToResource(KnownPropertyNames.LogicalCores, properties, resources);
+			CopyPropertyToResource(KnownPropertyNames.Ram, properties, resources);
+		}
+
+		/// <summary>
+		/// Creates a new agent
+		/// </summary>
+		/// <param name="request">Request to create a new agent</param>
+		/// <param name="context">Context for the RPC call</param>
+		/// <returns>Information about the new agent</returns>
+		public override async Task<RpcCreateAgentResponse> CreateAgent(RpcCreateAgentRequest request, ServerCallContext context)
+		{
+			using IDisposable? scope = _logger.BeginScope("CreateAgent({AgentId})", request.Name.ToString());
+
+			if (!_computeConfig.Value.Authorize(AgentAclAction.CreateAgent, context.GetHttpContext().User))
+			{
+				throw new StructuredRpcException(StatusCode.PermissionDenied, "User is not authenticated to create new agents");
+			}
+
+			IAgent agent = await _agentService.CreateAgentAsync(new AgentId(request.Name), request.Ephemeral, "", context.CancellationToken);
+
+			List<AclClaimConfig> claims = new List<AclClaimConfig>();
+			claims.Add(new AclClaimConfig(HordeClaimTypes.Agent, agent.Id.ToString()));
+
+			RpcCreateAgentResponse response = new RpcCreateAgentResponse();
+			response.Id = agent.Id.ToString();
+			response.Token = await _aclService.IssueBearerTokenAsync(claims, null, context.CancellationToken);
+
+			return response;
+		}
+
+		/// <summary>
+		/// Creates a new session
+		/// </summary>
+		/// <param name="request">Request to create a new agent</param>
+		/// <param name="context">Context for the RPC call</param>
+		/// <returns>Information about the new agent</returns>
+		public override async Task<RpcCreateSessionResponse> CreateSession(RpcCreateSessionRequest request, ServerCallContext context)
+		{
+			if (request.Capabilities == null)
+			{
+				throw new StructuredRpcException(StatusCode.InvalidArgument, "Capabilities may not be null");
+			}
+
+			AgentId agentId = new AgentId(request.Id);
+			using IDisposable? scope = _logger.BeginScope("CreateSession({AgentId})", agentId.ToString());
+
+			ComputeConfig computeConfig = _computeConfig.Value;
+
+			// Find the agent
+			IAgent? agent = await _agentService.GetAgentAsync(agentId);
+			if (agent == null)
+			{
+				if (!computeConfig.Authorize(AgentAclAction.CreateAgent, context.GetHttpContext().User))
+				{
+					throw new StructuredRpcException(StatusCode.PermissionDenied, "User is not authenticated to create new agents");
+				}
+
+				agent = await _agentService.CreateAgentAsync(agentId, false, "");
+			}
+
+			// Check the enrollment key in the user token matches
+			string enrollmentKey = context.GetHttpContext().User.FindFirstValue(HordeClaimTypes.AgentEnrollmentKey) ?? String.Empty;
+			if (!String.Equals(enrollmentKey, agent.EnrollmentKey, StringComparison.OrdinalIgnoreCase))
+			{
+				_logger.LogError("Enrollment key does not match for {AgentId} (was {OldKey}, now {NewKey})", agent.Id, agent.EnrollmentKey, enrollmentKey);
+				throw new StructuredRpcException(StatusCode.PermissionDenied, $"Enrollment key does not match for {agent.Id}");
+			}
+
+			// Make sure we're allowed to create sessions on this agent
+			ClaimsPrincipal user = context.GetHttpContext().User;
+			if (!computeConfig.Authorize(SessionAclAction.CreateSession, user) && !user.HasAgentClaim(agentId))
+			{
+				throw new StructuredRpcException(StatusCode.PermissionDenied, "User is not authenticated to create session for {AgentId}", agentId);
+			}
+
+			// Get the known properties for this agent
+			GetCapabilities(request.Capabilities, out List<string> properties, out Dictionary<string, int> resources);
+
+			// Create a new session
+			agent = await _agentService.CreateSessionAsync(agent, (AgentStatus)request.Status, properties, resources, request.Version, context.CancellationToken);
+			if (agent == null)
+			{
+				throw new StructuredRpcException(StatusCode.NotFound, "Agent {AgentId} not found", agentId);
+			}
+
+			// Create the response
+			RpcCreateSessionResponse response = new RpcCreateSessionResponse();
+			response.AgentId = agent.Id.ToString();
+			response.SessionId = agent.SessionId.ToString();
+			response.ExpiryTime = Timestamp.FromDateTime(agent.SessionExpiresAt!.Value);
+			response.Token = await _agentService.IssueSessionTokenAsync(agent.Id, agent.SessionId!.Value, context.CancellationToken);
+			return response;
+		}
+
+		/// <summary>
+		/// Updates an agent session
+		/// </summary>
+		/// <param name="reader">Request to create a new agent</param>
+		/// <param name="writer">Writer for response objects</param>
+		/// <param name="context">Context for the RPC call</param>
+		/// <returns>Information about the new agent</returns>
+		public override async Task UpdateSession(IAsyncStreamReader<RpcUpdateSessionRequest> reader, IServerStreamWriter<RpcUpdateSessionResponse> writer, ServerCallContext context)
+		{
+			// Read the request object
+			Task<bool> nextRequestTask = reader.MoveNext();
+			if (await nextRequestTask)
+			{
+				RpcUpdateSessionRequest request = reader.Current;
+				using IDisposable? scope = _logger.BeginScope("UpdateSession for agent {AgentId}, session {SessionId}", request.AgentId, request.SessionId);
+
+				_logger.LogDebug("Updating session for {AgentId}", request.AgentId);
+				foreach (RpcLease lease in request.Leases)
+				{
+					_logger.LogDebug("Session {SessionId}, Lease {LeaseId} - State: {LeaseState}, Outcome: {LeaseOutcome}", request.SessionId, lease.Id, lease.State, lease.Outcome);
+				}
+
+				// Get a task for moving to the next item. This will only complete once the call has closed.
+				using CancellationTokenSource cancellationSource = CancellationTokenSource.CreateLinkedTokenSource(context.CancellationToken);
+				nextRequestTask = reader.MoveNext();
+				nextRequestTask = nextRequestTask.ContinueWith(task =>
+				{
+					cancellationSource.Cancel();
+					return task.IsCanceled ? false : task.Result;
+				}, TaskScheduler.Current);
+
+				// Get the current agent state
+				IAgent? agent = await _agentService.GetAgentAsync(new AgentId(request.AgentId));
+				if (agent != null)
+				{
+					SessionId sessionId = SessionId.Parse(request.SessionId);
+
+					// Check we're authorized to update it
+					if (agent.SessionId != sessionId)
+					{
+						throw new StructuredRpcException(StatusCode.PermissionDenied, "Agent {AgentId} has completed session {SessionId}; now executing session {NewSessionId}. Cannot update state.", request.AgentId, sessionId, agent.SessionId?.ToString() ?? "(None)");
+					}
+					if (!_agentService.AuthorizeSession(agent, context.GetHttpContext().User, out string authReason))
+					{
+						throw new StructuredRpcException(StatusCode.PermissionDenied, "Not authenticated for {AgentId}. Reason {Reason}", request.AgentId, authReason);
+					}
+
+					// Get the new capabilities of this agent
+					List<string>? properties = null;
+					Dictionary<string, int>? resources = null;
+					if (request.Capabilities != null)
+					{
+						GetCapabilities(request.Capabilities, out properties, out resources);
+					}
+
+					// Update the session
+					try
+					{
+						agent = await _agentService.UpdateSessionWithWaitAsync(agent, sessionId, (AgentStatus)request.Status, properties, resources, request.Leases, cancellationSource.Token);
+					}
+					catch (OperationCanceledException)
+					{
+						// Ignore cancellation due to a message having been received
+					}
+					catch (Exception ex)
+					{
+						_logger.LogError(ex, "Swallowed exception while updating session for {AgentId}.", request.AgentId);
+						throw new StructuredRpcException(StatusCode.Internal, "Failed updating session. Reason: {Reason}", ex.Message);
+					}
+				}
+
+				// Handle the invalid agent case
+				if (agent == null)
+				{
+					throw new StructuredRpcException(StatusCode.NotFound, "Invalid agent name '{AgentId}'", request.AgentId);
+				}
+
+				// Create the new session info
+				if (!context.CancellationToken.IsCancellationRequested)
+				{
+					RpcUpdateSessionResponse response = new RpcUpdateSessionResponse();
+					response.Leases.Add(agent.Leases.Select(x => x.ToRpcMessage()));
+					response.ExpiryTime = (agent.SessionExpiresAt == null) ? new Timestamp() : Timestamp.FromDateTime(agent.SessionExpiresAt.Value);
+					response.Status = (RpcAgentStatus)agent.Status;
+					await writer.WriteAsync(response);
+				}
+
+				// Wait for the client to close the stream
+				try
+				{
+					while (await nextRequestTask)
+					{
+						nextRequestTask = reader.MoveNext();
+					}
+				}
+				catch (Exception ex)
+				{
+					_logger.LogDebug(ex, "Ignoring exception while finishing UpdateSession request.");
+				}
+			}
+		}
+
+		/// <summary>
+		/// Downloads a new agent archive
+		/// </summary>
+		/// <param name="request">Request arguments</param>
+		/// <param name="responseStream">Writer for the output data</param>
+		/// <param name="context">Context for the RPC call</param>
+		/// <returns>Information about the new agent</returns>
+		public override async Task DownloadSoftware(RpcDownloadSoftwareRequest request, IServerStreamWriter<RpcDownloadSoftwareResponse> responseStream, ServerCallContext context)
+		{
+			int colonIdx = request.Version.IndexOf(':', StringComparison.Ordinal);
+			ToolId toolId = new ToolId(request.Version.Substring(0, colonIdx));
+			string version = request.Version.Substring(colonIdx + 1);
+
+			ITool? tool = await _toolCollection.GetAsync(toolId, context.CancellationToken);
+			if (tool == null)
+			{
+				throw new StructuredRpcException(StatusCode.NotFound, $"Missing tool {toolId}");
+			}
+
+			if (!tool.Public && !tool.Authorize(ToolAclAction.DownloadTool, context.GetHttpContext().User))
+			{
+				throw new StructuredRpcException(StatusCode.PermissionDenied, $"User does not have DownloadTool entitlement for {toolId}");
+			}
+
+			IToolDeployment? deployment = tool.Deployments.LastOrDefault(x => x.Version.Equals(version, StringComparison.Ordinal));
+			if (deployment == null)
+			{
+				throw new StructuredRpcException(StatusCode.NotFound, $"Missing tool version {version}");
+			}
+
+			await using Stream stream = await deployment.OpenZipStreamAsync(context.CancellationToken);
+			using (IMemoryOwner<byte> buffer = MemoryPool<byte>.Shared.Rent(128 * 1024))
+			{
+				long totalWritten = 0;
+				for (; ; )
+				{
+					int read = await stream.ReadAsync(buffer.Memory, context.CancellationToken);
+					if (read == 0)
+					{
+						break;
+					}
+
+					RpcDownloadSoftwareResponse response = new RpcDownloadSoftwareResponse();
+					response.Data = UnsafeByteOperations.UnsafeWrap(buffer.Memory.Slice(0, read));
+					await responseStream.WriteAsync(response);
+
+					totalWritten += response.Data.Length;
+				}
+				_logger.LogInformation("Agent software zip is {Size:n0} bytes", totalWritten);
+			}
+		}
+
+		/// <inheritdoc/>
+		public override Task<Empty> UploadTelemetry(RpcUploadTelemetryRequest request, ServerCallContext context)
+		{
+			_logger.LogDebug("Posting telemetry data for {AgentId}", new AgentId(request.AgentId));
+
+			NewAgentTelemetry telemetry = new NewAgentTelemetry(request.UserCpu, request.IdleCpu, request.SystemCpu, (int)request.FreeRam, (int)request.UsedRam, (int)request.TotalRam, (long)request.FreeDisk, (long)request.TotalDisk);
+			_agentTelemetryCollection.Add(new AgentId(request.AgentId), telemetry);
+
+			return Task.FromResult(new Empty());
+		}
+	}
+}
