Horde: Allow ACLs to expose alternative legacy names to add to the lookup dictionary, removing dependency on projects/streams in GlobalConfig.

#rnx

diff --git a/Engine/Source/Programs/Horde/HordeServer.Shared/Acls/AclConfig.cs b/Engine/Source/Programs/Horde/HordeServer.Shared/Acls/AclConfig.cs
--- a/Engine/Source/Programs/Horde/HordeServer.Shared/Acls/AclConfig.cs
+++ b/Engine/Source/Programs/Horde/HordeServer.Shared/Acls/AclConfig.cs
@@ -31,6 +31,12 @@ namespace HordeServer.Acls
 		[JsonIgnore]
 		public AclScopeName ScopeName { get; set; }
 
+		/// <summary>
+		/// Legacy aliases for this scope
+		/// </summary>
+		[JsonIgnore]
+		public AclScopeName[]? LegacyScopeNames { get; set; }
+
 		/// <summary>
 		/// ACLs which are parented to this
 		/// </summary>
diff --git a/Engine/Source/Programs/Horde/HordeServer/Server/GlobalConfig.cs b/Engine/Source/Programs/Horde/HordeServer/Server/GlobalConfig.cs
--- a/Engine/Source/Programs/Horde/HordeServer/Server/GlobalConfig.cs
+++ b/Engine/Source/Programs/Horde/HordeServer/Server/GlobalConfig.cs
@@ -2,8 +2,6 @@
 
 using System;
 using System.Collections.Generic;
-using System.ComponentModel.DataAnnotations;
-using System.Diagnostics;
 using System.Diagnostics.CodeAnalysis;
 using System.Linq;
 using System.Security.Claims;
@@ -11,8 +9,6 @@ using System.Text.Json.Serialization;
 using EpicGames.Core;
 using EpicGames.Horde.Acls;
 using EpicGames.Horde.Agents.Pools;
-using EpicGames.Horde.Common;
-using EpicGames.Horde.Compute;
 using HordeServer.Acls;
 using HordeServer.Agents;
 using HordeServer.Agents.Pools;
@@ -146,28 +142,6 @@ namespace HordeServer.Server
 
 			_aclLookup.Clear();
 			BuildAclScopeLookup(Acl, _aclLookup);
-
-			BuildConfig? buildConfig;
-			if (Plugins.TryGetBuildConfig(out buildConfig))
-			{
-				foreach (ProjectConfig project in buildConfig.Projects)
-				{
-					AclScopeName legacyProjectScopeName = Acl.ScopeName.Append($"p:{project.Id}");
-					_aclLookup.Add(legacyProjectScopeName, project.Acl);
-
-					foreach (StreamConfig stream in project.Streams)
-					{
-						AclScopeName legacyStreamScopeName = Acl.ScopeName.Append($"s:{stream.Id}");
-						_aclLookup.Add(legacyStreamScopeName, stream.Acl);
-
-						foreach (TemplateRefConfig template in stream.Templates)
-						{
-							AclScopeName legacyTemplateScopeName = legacyStreamScopeName.Append($"t:{template.Id}");
-							_aclLookup.Add(legacyTemplateScopeName, template.Acl);
-						}
-					}
-				}
-			}
 		}
 
 		/// <summary>
@@ -236,6 +210,13 @@ namespace HordeServer.Server
 		static void BuildAclScopeLookup(AclConfig acl, Dictionary<AclScopeName, AclConfig> aclLookup)
 		{
 			aclLookup.Add(acl.ScopeName, acl);
+			if (acl.LegacyScopeNames != null)
+			{
+				foreach (AclScopeName legacyScopeName in acl.LegacyScopeNames)
+				{
+					aclLookup.Add(legacyScopeName, acl);
+				}
+			}
 			if (acl.Children != null)
 			{
 				foreach (AclConfig childAcl in acl.Children)
@@ -375,69 +356,6 @@ namespace HordeServer.Server
 		}
 	}
 
-	/// <summary>
-	/// Profile for executing compute requests
-	/// </summary>
-	[DebuggerDisplay("{Id}")]
-	public class ComputeClusterConfig
-	{
-		/// <summary>
-		/// The owning global config instance
-		/// </summary>
-		[JsonIgnore]
-		public GlobalConfig GlobalConfig { get; private set; } = null!;
-
-		/// <summary>
-		/// Name of the partition
-		/// </summary>
-		public ClusterId Id { get; set; } = new ClusterId("default");
-
-		/// <summary>
-		/// Name of the namespace to use
-		/// </summary>
-		public string NamespaceId { get; set; } = "horde.compute";
-
-		/// <summary>
-		/// Name of the input bucket
-		/// </summary>
-		public string RequestBucketId { get; set; } = "requests";
-
-		/// <summary>
-		/// Name of the output bucket
-		/// </summary>
-		public string ResponseBucketId { get; set; } = "responses";
-
-		/// <summary>
-		/// Filter for agents to include
-		/// </summary>
-		public Condition? Condition { get; set; }
-
-		/// <summary>
-		/// Access control list
-		/// </summary>
-		public AclConfig Acl { get; set; } = new AclConfig();
-
-		/// <summary>
-		/// Callback post loading this config file
-		/// </summary>
-		/// <param name="globalConfig">The global config instance</param>
-		public void PostLoad(GlobalConfig globalConfig)
-		{
-			GlobalConfig = globalConfig;
-			Acl.PostLoad(globalConfig.Acl, $"compute:{Id}");
-		}
-
-		/// <summary>
-		/// Authorizes a user to perform a given action
-		/// </summary>
-		/// <param name="action">The action being performed</param>
-		/// <param name="user">The principal to validate</param>
-		public bool Authorize(AclAction action, ClaimsPrincipal user)
-		{
-			return Acl?.Authorize(action, user) ?? GlobalConfig.Authorize(action, user);
-		}
-	}
-
 	/// <summary>
 	/// How frequently the maintence window repeats
 	/// </summary>
@@ -549,137 +467,5 @@ namespace HordeServer.Server
 			}
 		}
 	}
-
-	/// <summary>
-	/// Path to a platform and stream to use for syncing AutoSDK
-	/// </summary>
-	public class AutoSdkWorkspace
-	{
-		/// <summary>
-		/// Name of this workspace
-		/// </summary>
-		public string? Name { get; set; }
-
-		/// <summary>
-		/// The agent properties to check (eg. "OSFamily=Windows")
-		/// </summary>
-		public List<string> Properties { get; set; } = new List<string>();
-
-		/// <summary>
-		/// Username for logging in to the server
-		/// </summary>
-		public string? UserName { get; set; }
-
-		/// <summary>
-		/// Stream to use
-		/// </summary>
-		[Required]
-		public string? Stream { get; set; }
-	}
-
-	/// <summary>
-	/// Information about an individual Perforce server
-	/// </summary>
-	public class PerforceServer
-	{
-		/// <summary>
-		/// The server and port. The server may be a DNS entry with multiple records, in which case it will be actively load balanced.
-		/// </summary>
-		public string ServerAndPort { get; set; } = "perforce:1666";
-
-		/// <summary>
-		/// Whether to query the healthcheck address under each server
-		/// </summary>
-		public bool HealthCheck { get; set; }
-
-		/// <summary>
-		/// Whether to resolve the DNS entries and load balance between different hosts
-		/// </summary>
-		public bool ResolveDns { get; set; }
-
-		/// <summary>
-		/// Maximum number of simultaneous conforms on this server
-		/// </summary>
-		public int MaxConformCount { get; set; }
-
-		/// <summary>
-		/// Optional condition for a machine to be eligable to use this server
-		/// </summary>
-		public Condition? Condition { get; set; }
-
-		/// <summary>
-		/// List of properties for an agent to be eligable to use this server
-		/// </summary>
-		public List<string>? Properties { get; set; }
-	}
-
-	/// <summary>
-	/// Credentials for a Perforce user
-	/// </summary>
-	public class PerforceCredentials
-	{
-		/// <summary>
-		/// The username
-		/// </summary>
-		public string UserName { get; set; } = String.Empty;
-
-		/// <summary>
-		/// Password for the user
-		/// </summary>
-		public string? Password { get; set; } = String.Empty;
-
-		/// <summary>
-		/// Login ticket for the user (will be used instead of password if set)
-		/// </summary>
-		public string? Ticket { get; set; } = String.Empty;
-	}
-
-	/// <summary>
-	/// Information about a cluster of Perforce servers. 
-	/// </summary>
-	[DebuggerDisplay("{Name}")]
-	public class PerforceCluster
-	{
-		/// <summary>
-		/// The default cluster name
-		/// </summary>
-		public const string DefaultName = "Default";
-
-		/// <summary>
-		/// Name of the cluster
-		/// </summary>
-		[Required]
-		public string Name { get; set; } = null!;
-
-		/// <summary>
-		/// Username for Horde to log in to this server. Will use the first account specified below if not overridden.
-		/// </summary>
-		public string? ServiceAccount { get; set; }
-
-		/// <summary>
-		/// Whether the service account can impersonate other users
-		/// </summary>
-		public bool CanImpersonate { get; set; } = true;
-
-		/// <summary>
-		/// Whether to use partitioned workspaces on this server
-		/// </summary>
-		public bool SupportsPartitionedWorkspaces { get; set; } = false;
-
-		/// <summary>
-		/// List of servers
-		/// </summary>
-		public List<PerforceServer> Servers { get; set; } = new List<PerforceServer>();
-
-		/// <summary>
-		/// List of server credentials
-		/// </summary>
-		public List<PerforceCredentials> Credentials { get; set; } = new List<PerforceCredentials>();
-
-		/// <summary>
-		/// List of autosdk streams
-		/// </summary>
-		public List<AutoSdkWorkspace> AutoSdk { get; set; } = new List<AutoSdkWorkspace>();
-	}
 }
 
diff --git a/Engine/Source/Programs/Horde/HordeServer/Startup.cs b/Engine/Source/Programs/Horde/HordeServer/Startup.cs
--- a/Engine/Source/Programs/Horde/HordeServer/Startup.cs
+++ b/Engine/Source/Programs/Horde/HordeServer/Startup.cs
@@ -410,6 +410,7 @@ namespace HordeServer
 			services.AddSingleton<JobService>();
 			services.AddSingleton<ILogExtAuthProvider>(sp => sp.GetRequiredService<JobService>());
 			services.AddSingleton<LifetimeService>();
+			services.AddSingleton<ILifetimeService>(sp => sp.GetRequiredService<LifetimeService>());
 			services.AddHostedService(provider => provider.GetRequiredService<LifetimeService>());
 			services.AddSingleton(typeof(IHealthMonitor<>), typeof(HealthMonitor<>));
 			services.AddSingleton<ServerStatusService>();
diff --git a/Engine/Source/Programs/Horde/Plugins/Build/HordeServer.Build/Projects/ProjectConfig.cs b/Engine/Source/Programs/Horde/Plugins/Build/HordeServer.Build/Projects/ProjectConfig.cs
--- a/Engine/Source/Programs/Horde/Plugins/Build/HordeServer.Build/Projects/ProjectConfig.cs
+++ b/Engine/Source/Programs/Horde/Plugins/Build/HordeServer.Build/Projects/ProjectConfig.cs
@@ -121,6 +121,7 @@ namespace HordeServer.Projects
 		{
 			Id = id;
 			Acl.PostLoad(parentAcl, $"project:{Id}");
+			Acl.LegacyScopeNames = new AclScopeName[] { parentAcl.ScopeName.Append($"p:{Id}") };
 
 			foreach (StreamConfig stream in Streams)
 			{
diff --git a/Engine/Source/Programs/Horde/Plugins/Build/HordeServer.Build/Streams/StreamConfig.cs b/Engine/Source/Programs/Horde/Plugins/Build/HordeServer.Build/Streams/StreamConfig.cs
--- a/Engine/Source/Programs/Horde/Plugins/Build/HordeServer.Build/Streams/StreamConfig.cs
+++ b/Engine/Source/Programs/Horde/Plugins/Build/HordeServer.Build/Streams/StreamConfig.cs
@@ -259,6 +259,7 @@ namespace HordeServer.Streams
 			ProjectConfig = projectConfig;
 
 			Acl.PostLoad(projectConfig.Acl, $"stream:{Id}");
+			Acl.LegacyScopeNames = projectConfig.Acl.LegacyScopeNames?.Select(x => x.Append($"s:{Id}")).ToArray();
 
 			JobOptions.MergeDefaults(projectConfig.JobOptions);
 
@@ -872,6 +873,7 @@ namespace HordeServer.Streams
 			}
 
 			Acl.PostLoad(streamConfig.Acl, $"template:{Id}");
+			Acl.LegacyScopeNames = streamConfig.Acl.LegacyScopeNames?.Select(x => x.Append($"t:{Id}")).ToArray();
 		}
 	}
 
