Horde: Static analysis fix.

#rnx

diff --git a/Engine/Source/Programs/Horde/HordeServer/Startup.cs b/Engine/Source/Programs/Horde/HordeServer/Startup.cs
--- a/Engine/Source/Programs/Horde/HordeServer/Startup.cs
+++ b/Engine/Source/Programs/Horde/HordeServer/Startup.cs
@@ -31,7 +31,6 @@ using HordeServer.Accounts;
 using HordeServer.Acls;
 using HordeServer.Agents;
 using HordeServer.Agents.Enrollment;
-using HordeServer.Agents.Fleet;
 using HordeServer.Agents.Leases;
 using HordeServer.Agents.Pools;
 using HordeServer.Agents.Relay;
