Horde: Test fix.

#rnx

diff --git a/Engine/Source/Programs/Horde/HordeServer.Tests.Shared/ServerServiceTest.cs b/Engine/Source/Programs/Horde/HordeServer.Tests.Shared/ServerServiceTest.cs
--- a/Engine/Source/Programs/Horde/HordeServer.Tests.Shared/ServerServiceTest.cs
+++ b/Engine/Source/Programs/Horde/HordeServer.Tests.Shared/ServerServiceTest.cs
@@ -123,6 +123,7 @@ namespace HordeServer.Tests
 			services.AddSingleton<IAclService, AclService>();
 			services.AddSingleton<IDowntimeService, DowntimeServiceStub>();
 			services.AddSingleton<LifetimeService>();
+			services.AddSingleton<ILifetimeService>(sp => sp.GetRequiredService<LifetimeService>());
 
 			services.AddSingleton(typeof(IHealthMonitor<>), typeof(HealthMonitor<>));
 			services.AddSingleton<ServerStatusService>();
