Horde: Move service registration into the appropriate plugins.

#rnx

diff --git a/Engine/Source/Programs/Horde/HordeServer/Startup.cs b/Engine/Source/Programs/Horde/HordeServer/Startup.cs
--- a/Engine/Source/Programs/Horde/HordeServer/Startup.cs
+++ b/Engine/Source/Programs/Horde/HordeServer/Startup.cs
@@ -365,14 +365,9 @@ namespace HordeServer
 			services.AddGrpcReflection();
 
 			services.AddSingleton<IAccountCollection, AccountCollection>();
-			services.AddSingleton<IAgentCollection, AgentCollection>();
-			services.AddSingleton<ILeaseCollection, LeaseCollection>();
-			services.AddSingleton<ILogCollection, LogCollection>();
-			services.AddSingleton<INotificationTriggerCollection, NotificationTriggerCollection>();
-			services.AddSingleton<IPoolCollection, PoolCollection>();
 			services.AddSingleton<IServiceAccountCollection, ServiceAccountCollection>();
-			services.AddSingleton<ISessionCollection, SessionCollection>();
 			services.AddSingleton<IUserCollection, UserCollectionV2>();
+			services.AddSingleton<INotificationTriggerCollection, NotificationTriggerCollection>();
 			services.AddSingleton<IDeviceCollection, DeviceCollection>();
 			services.AddSingleton<INoticeCollection, NoticeCollection>();
 			services.AddSingleton<IDashboardPreviewCollection, DashboardPreviewCollection>();
@@ -396,8 +391,6 @@ namespace HordeServer
 			services.AddSingleton(typeof(ISingletonDocument<>), typeof(SingletonDocument<>));
 
 			services.AddSingleton<IAclService, AclService>();
-			services.AddSingleton<AgentService>();
-			services.AddSingleton(provider => new Lazy<AgentService>(provider.GetRequiredService<AgentService>));
 			services.AddSingleton<RequestTrackerService>();
 			services.AddSingleton<MongoCommandTracer>();
 			services.AddSingleton<MongoService>();
@@ -405,19 +398,12 @@ namespace HordeServer
 			services.AddSingleton<GlobalsService>();
 			services.AddSingleton<IClock, Clock>();
 			services.AddSingleton<IDowntimeService, DowntimeService>();
-			services.AddSingleton<IssueService>();
-			services.AddSingleton<JobService>();
-			services.AddSingleton<ILogExtAuthProvider>(sp => sp.GetRequiredService<JobService>());
 			services.AddSingleton<LifetimeService>();
 			services.AddSingleton<ILifetimeService>(sp => sp.GetRequiredService<LifetimeService>());
 			services.AddHostedService(provider => provider.GetRequiredService<LifetimeService>());
 			services.AddSingleton(typeof(IHealthMonitor<>), typeof(HealthMonitor<>));
 			services.AddSingleton<ServerStatusService>();
 			services.AddHostedService(provider => provider.GetRequiredService<ServerStatusService>());
-			services.AddSingleton<INotificationService, NotificationService>();
-			services.AddSingleton<UnsyncCache>();
-
-			services.AddSingleton<ScheduleService>();
 
 			services.AddScoped<OAuthControllerFilter>();
 
@@ -590,50 +576,13 @@ namespace HordeServer
 						.Build();
 				});
 
-			// Create the agent telemetry collection, and register the hosted service so we can flush from any server.
-			services.AddSingleton<AgentTelemetryCollection>();
-			services.AddSingleton<IAgentTelemetryCollection>(sp => sp.GetRequiredService<AgentTelemetryCollection>());
-			services.AddHostedService(provider => provider.GetRequiredService<AgentTelemetryCollection>());
-
 			// Hosted service that needs to run no matter the run mode of the process (server vs worker)
 			services.AddHostedService(provider => (DowntimeService)provider.GetRequiredService<IDowntimeService>());
 
-			// Notifications can be triggered from any instance, so always make sure we're ticking the background task.
-			services.AddHostedService(provider => (NotificationService)provider.GetRequiredService<INotificationService>());
 
 			if (settings.IsRunModeActive(RunMode.Worker) && !settings.MongoReadOnlyMode)
 			{
-				services.AddHostedService<AgentReportService>();
-				services.AddHostedService(provider => provider.GetRequiredService<IssueService>());
-				services.AddHostedService<IssueReportService>();
-				services.AddHostedService<IssueTagService>();
-				services.AddHostedService<JobExpirationService>();
 				services.AddHostedService<MetricService>();
-				services.AddHostedService(provider => provider.GetRequiredService<PerforceLoadBalancer>());
-				services.AddHostedService<PoolUpdateService>();
-				services.AddHostedService<UtilizationDataService>();
-				services.AddHostedService(provider => provider.GetRequiredService<DeviceService>());
-				services.AddHostedService<DeviceReportService>();
-				services.AddHostedService(provider => provider.GetRequiredService<TestDataService>());
-				services.AddHostedService(provider => provider.GetRequiredService<ComputeService>());
-				services.AddHostedService(provider => provider.GetRequiredService<EnrollmentService>());
-				services.AddHostedService(provider => provider.GetRequiredService<StorageService>());
-			}
-
-			services.AddHostedService(provider => provider.GetRequiredService<IExternalIssueService>());
-
-			// Task sources. Order of registration is important here; it dictates the priority in which sources are served.
-			services.AddSingleton<JobTaskSource>();
-
-			if (!settings.MongoReadOnlyMode)
-			{
-				services.AddHostedService<JobTaskSource>(provider => provider.GetRequiredService<JobTaskSource>());
-				services.AddSingleton<ConformTaskSource>();
-				services.AddHostedService<ConformTaskSource>(provider => provider.GetRequiredService<ConformTaskSource>());
-				services.AddSingleton<ComputeTaskSource>();
-
-				services.AddSingleton<ITaskSource, ConformTaskSource>(provider => provider.GetRequiredService<ConformTaskSource>());
-				services.AddSingleton<ITaskSource, JobTaskSource>(provider => provider.GetRequiredService<JobTaskSource>());
 			}
 
 			// Allow longer to shutdown so we can debug missing cancellation tokens
diff --git a/Engine/Source/Programs/Horde/Plugins/Build/HordeServer.Build/BuildPlugin.cs b/Engine/Source/Programs/Horde/Plugins/Build/HordeServer.Build/BuildPlugin.cs
--- a/Engine/Source/Programs/Horde/Plugins/Build/HordeServer.Build/BuildPlugin.cs
+++ b/Engine/Source/Programs/Horde/Plugins/Build/HordeServer.Build/BuildPlugin.cs
@@ -30,6 +30,7 @@ using HordeServer.Perforce;
 using HordeServer.Plugins;
 using HordeServer.Replicators;
 using HordeServer.Streams;
+using HordeServer.Tasks;
 using HordeServer.Ugs;
 using HordeServer.Users;
 using Microsoft.Extensions.DependencyInjection;
@@ -80,6 +81,47 @@ namespace HordeServer
 			services.AddSingleton<IBlockCache>(sp => CreateBlockCache(sp));
 			services.AddSingleton<TestDataService>();
 
+			services.AddSingleton<IssueService>();
+			services.AddSingleton<JobService>();
+			services.AddSingleton<ILogExtAuthProvider>(sp => sp.GetRequiredService<JobService>());
+			services.AddSingleton<INotificationService, NotificationService>();
+			services.AddSingleton<UnsyncCache>();
+
+			services.AddSingleton<ScheduleService>();
+
+			// Notifications can be triggered from any instance, so always make sure we're ticking the background task.
+			services.AddHostedService(provider => (NotificationService)provider.GetRequiredService<INotificationService>());
+
+			if (_serverInfo.IsRunModeActive(RunMode.Worker) && !_serverInfo.ReadOnlyMode)
+			{
+				services.AddHostedService<AgentReportService>();
+				services.AddHostedService(provider => provider.GetRequiredService<IssueService>());
+				services.AddHostedService<IssueReportService>();
+				services.AddHostedService<IssueTagService>();
+				services.AddHostedService<JobExpirationService>();
+				services.AddHostedService(provider => provider.GetRequiredService<PerforceLoadBalancer>());
+				services.AddHostedService<PoolUpdateService>();
+				services.AddHostedService<UtilizationDataService>();
+				services.AddHostedService(provider => provider.GetRequiredService<DeviceService>());
+				services.AddHostedService<DeviceReportService>();
+				services.AddHostedService(provider => provider.GetRequiredService<TestDataService>());
+			}
+
+			services.AddHostedService(provider => provider.GetRequiredService<IExternalIssueService>());
+
+			// Task sources. Order of registration is important here; it dictates the priority in which sources are served.
+			services.AddSingleton<JobTaskSource>();
+
+			if (!_serverInfo.ReadOnlyMode)
+			{
+				services.AddHostedService<JobTaskSource>(provider => provider.GetRequiredService<JobTaskSource>());
+				services.AddSingleton<ConformTaskSource>();
+				services.AddHostedService<ConformTaskSource>(provider => provider.GetRequiredService<ConformTaskSource>());
+
+				services.AddSingleton<ITaskSource, ConformTaskSource>(provider => provider.GetRequiredService<ConformTaskSource>());
+				services.AddSingleton<ITaskSource, JobTaskSource>(provider => provider.GetRequiredService<JobTaskSource>());
+			}
+
 			if (_staticConfig.Commits.ReplicateMetadata)
 			{
 				services.AddSingleton<PerforceServiceCache>();
diff --git a/Engine/Source/Programs/Horde/Plugins/Compute/HordeServer.Compute/ComputePlugin.cs b/Engine/Source/Programs/Horde/Plugins/Compute/HordeServer.Compute/ComputePlugin.cs
--- a/Engine/Source/Programs/Horde/Plugins/Compute/HordeServer.Compute/ComputePlugin.cs
+++ b/Engine/Source/Programs/Horde/Plugins/Compute/HordeServer.Compute/ComputePlugin.cs
@@ -9,8 +9,11 @@ using Amazon.SQS;
 using HordeServer.Agents;
 using HordeServer.Agents.Enrollment;
 using HordeServer.Agents.Fleet;
+using HordeServer.Agents.Leases;
 using HordeServer.Agents.Pools;
 using HordeServer.Agents.Relay;
+using HordeServer.Agents.Sessions;
+using HordeServer.Agents.Telemetry;
 using HordeServer.Aws;
 using HordeServer.Compute;
 using HordeServer.Logs;
@@ -43,6 +46,15 @@ namespace HordeServer
 		/// <inheritdoc/>
 		public void ConfigureServices(IServiceCollection services)
 		{
+			services.AddSingleton<IAgentCollection, AgentCollection>();
+			services.AddSingleton<ILeaseCollection, LeaseCollection>();
+			services.AddSingleton<ILogCollection, LogCollection>();
+			services.AddSingleton<IPoolCollection, PoolCollection>();
+			services.AddSingleton<ISessionCollection, SessionCollection>();
+
+			services.AddSingleton<AgentService>();
+			services.AddSingleton(provider => new Lazy<AgentService>(provider.GetRequiredService<AgentService>));
+
 			services.AddSingleton<AwsAutoScalingLifecycleService>();
 			services.AddSingleton<FleetService>();
 			services.AddSingleton<IFleetManagerFactory, FleetManagerFactory>();
@@ -73,6 +85,13 @@ namespace HordeServer
 			// Always run agent service too; need to be able to listen to Redis for events on any server.
 			services.AddHostedService(provider => provider.GetRequiredService<AgentService>());
 
+			// Create the agent telemetry collection, and register the hosted service so we can flush from any server.
+			services.AddSingleton<AgentTelemetryCollection>();
+			services.AddSingleton<IAgentTelemetryCollection>(sp => sp.GetRequiredService<AgentTelemetryCollection>());
+			services.AddHostedService(provider => provider.GetRequiredService<AgentTelemetryCollection>());
+
+
+
 			if (!_serverInfo.ReadOnlyMode)
 			{
 				services.AddSingleton<ComputeTaskSource>();
@@ -86,6 +105,8 @@ namespace HordeServer
 				{
 					services.AddHostedService(provider => provider.GetRequiredService<FleetService>());
 					services.AddHostedService(provider => provider.GetRequiredService<ConsistencyService>());
+					services.AddHostedService(provider => provider.GetRequiredService<ComputeService>());
+					services.AddHostedService(provider => provider.GetRequiredService<EnrollmentService>());
 				}
 			}
 
diff --git a/Engine/Source/Programs/Horde/Plugins/Storage/HordeServer.Storage/StoragePlugin.cs b/Engine/Source/Programs/Horde/Plugins/Storage/HordeServer.Storage/StoragePlugin.cs
--- a/Engine/Source/Programs/Horde/Plugins/Storage/HordeServer.Storage/StoragePlugin.cs
+++ b/Engine/Source/Programs/Horde/Plugins/Storage/HordeServer.Storage/StoragePlugin.cs
@@ -23,10 +23,12 @@ namespace HordeServer.Storage
 	[Plugin("Storage", GlobalConfigType = typeof(StorageConfig), ServerConfigType = typeof(StaticStorageConfig))]
 	public class StoragePlugin : IPluginStartup
 	{
+		readonly IServerInfo _serverInfo;
 		readonly StaticStorageConfig _staticConfig;
 
-		public StoragePlugin(StaticStorageConfig staticConfig)
+		public StoragePlugin(IServerInfo serverInfo, StaticStorageConfig staticConfig)
 		{
+			_serverInfo = serverInfo;
 			_staticConfig = staticConfig;
 		}
 
@@ -43,6 +45,11 @@ namespace HordeServer.Storage
 
 			services.AddSingleton<BundleCache>();
 			services.AddSingleton<StorageBackendCache>(CreateStorageBackendCache);
+
+			if (_serverInfo.IsRunModeActive(RunMode.Worker) && !_serverInfo.ReadOnlyMode)
+			{
+				services.AddHostedService(provider => provider.GetRequiredService<StorageService>());
+			}
 		}
 
 		StorageBackendCache CreateStorageBackendCache(IServiceProvider serviceProvider)
