Horde: Move some tests into HordeServer.Compute.Tests.

#rnx

diff --git a/Engine/Source/Programs/Horde/HordeServer.Tests/Compute/RedisTaskSchedulerTest.cs b/Engine/Source/Programs/Horde/HordeServer.Tests/Compute/RedisTaskSchedulerTest.cs
deleted file mode 100644
diff --git a/Engine/Source/Programs/Horde/HordeServer.Tests/Compute/TunnelServiceTests.cs b/Engine/Source/Programs/Horde/HordeServer.Tests/Compute/TunnelServiceTests.cs
deleted file mode 100644
diff --git a/Engine/Source/Programs/Horde/HordeServer.Tests/Logs/LogIndexingTests.cs b/Engine/Source/Programs/Horde/HordeServer.Tests/Logs/LogIndexingTests.cs
deleted file mode 100644
diff --git a/Engine/Source/Programs/Horde/HordeServer.Tests/Logs/LogServiceTests.cs b/Engine/Source/Programs/Horde/HordeServer.Tests/Logs/LogServiceTests.cs
deleted file mode 100644
diff --git a/Engine/Source/Programs/Horde/HordeServer.Tests/Logs/LogTailServiceTests.cs b/Engine/Source/Programs/Horde/HordeServer.Tests/Logs/LogTailServiceTests.cs
deleted file mode 100644
diff --git a/Engine/Source/Programs/Horde/Plugins/Compute/HordeServer.Compute.Tests/Compute/RedisTaskSchedulerTest.cs b/Engine/Source/Programs/Horde/Plugins/Compute/HordeServer.Compute.Tests/Compute/RedisTaskSchedulerTest.cs
new file mode 100644
--- /dev/null
+++ b/Engine/Source/Programs/Horde/Plugins/Compute/HordeServer.Compute.Tests/Compute/RedisTaskSchedulerTest.cs
@@ -0,0 +1,98 @@
+// Copyright Epic Games, Inc. All Rights Reserved.
+
+using System;
+using System.Threading;
+using System.Threading.Tasks;
+using HordeServer.Compute;
+using HordeServer.Server;
+using Microsoft.Extensions.Logging;
+using Microsoft.VisualStudio.TestTools.UnitTesting;
+using StackExchange.Redis;
+
+namespace HordeServer.Tests.Compute
+{
+	[TestClass]
+	public class RedisTaskSchedulerTest : ComputeTestSetup
+	{
+		private readonly RedisTaskScheduler<string, string> _scheduler;
+
+		public RedisTaskSchedulerTest()
+		{
+			using ILoggerFactory loggerFactory = LoggerFactory.Create(builder => { builder.AddConsole(); });
+			IRedisService redisService = GetRedisServiceSingleton();
+			_scheduler = new RedisTaskScheduler<string, string>(redisService.ConnectionPool, new RedisKey("myBaseKey"), loggerFactory.CreateLogger<RedisTaskScheduler<string, string>>());
+		}
+
+		public override async ValueTask DisposeAsync()
+		{
+			await _scheduler.DisposeAsync();
+			await base.DisposeAsync();
+			GC.SuppressFinalize(this);
+		}
+
+		[TestMethod]
+		public async Task EnqueueAndDequeueTaskAsync()
+		{
+			await _scheduler.EnqueueAsync("queue", "task1", true);
+			string? dequeuedTask = await _scheduler.DequeueAsync("queue");
+			Assert.AreEqual("task1", dequeuedTask);
+		}
+
+		[TestMethod]
+		[Ignore("Currently fails as nothing picks up tasks(?)")]
+		public async Task DequeueWithPredicateFromMiddleOfQueueAsync()
+		{
+			using CancellationTokenSource cts = new(2000);
+			await _scheduler.EnqueueAsync("queue", "task1", true);
+			await _scheduler.EnqueueAsync("queue", "task2", true);
+			await _scheduler.EnqueueAsync("queue", "task3", true);
+			Task<(string, string)?> dequeueTask = await _scheduler.DequeueAsync((x) => ValueTask.FromResult(x == "task2"), cts.Token);
+
+			(string queueId, string task) = (await dequeueTask.WaitAsync(cts.Token))!.Value;
+
+			Assert.AreEqual("queue", queueId);
+			Assert.AreEqual("task2", task);
+		}
+
+		[TestMethod]
+		public async Task EnqueueTaskAtFrontAsync()
+		{
+			await _scheduler.EnqueueAsync("queue", "task1", true);
+			await _scheduler.EnqueueAsync("queue", "task2", true);
+			Assert.AreEqual("task2", await _scheduler.DequeueAsync("queue"));
+		}
+
+		[TestMethod]
+		public async Task EnqueueTaskAtBackAsync()
+		{
+			await _scheduler.EnqueueAsync("queue", "task1", false);
+			await _scheduler.EnqueueAsync("queue", "task2", false);
+			Assert.AreEqual("task1", await _scheduler.DequeueAsync("queue"));
+		}
+
+		[TestMethod]
+		public async Task EnqueueTasksOnSeparateQueuesAsync()
+		{
+			await _scheduler.EnqueueAsync("queue1", "task1", true);
+			await _scheduler.EnqueueAsync("queue2", "task2", true);
+			Assert.AreEqual("task1", await _scheduler.DequeueAsync("queue1"));
+			Assert.AreEqual("task2", await _scheduler.DequeueAsync("queue2"));
+		}
+
+		[TestMethod]
+		public async Task DequeueEmptyQueueAsync()
+		{
+			Assert.IsNull(await _scheduler.DequeueAsync("queue"));
+		}
+
+		[TestMethod]
+		public async Task DequeuingMakesQueueActiveAsync()
+		{
+			await _scheduler.EnqueueAsync("queue", "task1", true);
+			Assert.AreEqual(1, (await _scheduler.GetInactiveQueuesAsync()).Count);
+
+			await _scheduler.DequeueAsync("queue");
+			Assert.AreEqual(0, (await _scheduler.GetInactiveQueuesAsync()).Count);
+		}
+	}
+}
\ No newline at end of file
diff --git a/Engine/Source/Programs/Horde/Plugins/Compute/HordeServer.Compute.Tests/Compute/TunnelServiceTests.cs b/Engine/Source/Programs/Horde/Plugins/Compute/HordeServer.Compute.Tests/Compute/TunnelServiceTests.cs
new file mode 100644
--- /dev/null
+++ b/Engine/Source/Programs/Horde/Plugins/Compute/HordeServer.Compute.Tests/Compute/TunnelServiceTests.cs
@@ -0,0 +1,176 @@
+// Copyright Epic Games, Inc. All Rights Reserved.
+
+using System;
+using System.Collections.Generic;
+using System.IO;
+using System.Net;
+using System.Net.Sockets;
+using System.Text;
+using System.Threading;
+using System.Threading.Tasks;
+using EpicGames.Core;
+using EpicGames.Horde.Compute.Clients;
+using HordeServer.Compute;
+using Microsoft.Extensions.Logging;
+using Microsoft.Extensions.Options;
+using Microsoft.VisualStudio.TestTools.UnitTesting;
+
+namespace HordeServer.Tests.Compute
+{
+	[TestClass]
+	public sealed class TunnelServiceTests
+	{
+		private static readonly Random s_random = new(112233);
+		private readonly string _hostname = "127.0.0.1";
+		private readonly int _echoPort = GetAvailablePort();
+		private readonly int _tunnelPort = GetAvailablePort();
+
+		private Task? _tunnelServerTask;
+		private Task? _echoServerTask;
+
+		/// <summary>
+		/// Find an available TCP/IP port
+		/// </summary>
+		/// <returns>Port number available</returns>
+		public static int GetAvailablePort()
+		{
+			using TcpListener listener = new(IPAddress.Loopback, 0);
+			listener.Start();
+			int port = ((IPEndPoint)listener.LocalEndpoint).Port;
+			listener.Stop();
+			return port;
+		}
+
+		[TestCleanup]
+		public async Task CleanupAsync()
+		{
+			if (_echoServerTask != null)
+			{
+				await _echoServerTask.IgnoreCanceledExceptionsAsync().ConfigureAwait(false);
+			}
+
+			if (_tunnelServerTask != null)
+			{
+				await _tunnelServerTask.IgnoreCanceledExceptionsAsync().ConfigureAwait(false);
+			}
+		}
+
+		[TestMethod]
+		[DataRow(20)]
+		[DataRow(1024)]
+		[DataRow(4096)]
+		[DataRow(4070)]
+		[DataRow(20000)]
+		public async Task SendRecvTestAsync(int echoServerBufferSize)
+		{
+			using CancellationTokenSource cts = new(10000);
+
+			_echoServerTask = StartEchoServerAsync(_echoPort, echoServerBufferSize, cts.Token);
+
+			StaticComputeConfig serverSettings = new() { ComputeTunnelPort = _tunnelPort };
+			using TunnelService tunnelService = new(new OptionsWrapper<StaticComputeConfig>(serverSettings), ComputeTestSetup.CreateConsoleLogger<TunnelService>());
+			tunnelService.Start(IPAddress.Any);
+			_tunnelServerTask = tunnelService.ServerTask!;
+
+			using TcpClient client = new();
+			await client.ConnectAsync(_hostname, _tunnelPort, cts.Token);
+			await using NetworkStream stream = client.GetStream();
+			using StreamReader streamReader = new(stream);
+			await using StreamWriter streamWriter = new(stream) { AutoFlush = true };
+
+			await streamWriter.WriteLineAsync(new TunnelHandshakeRequest(_hostname, _echoPort).Serialize());
+
+			string? handshakeResponseStr = await ReadLineWithTimeoutAsync(streamReader, TimeSpan.FromSeconds(3));
+			TunnelHandshakeResponse handshakeResponse = TunnelHandshakeResponse.Deserialize(handshakeResponseStr);
+			if (!handshakeResponse.IsSuccess)
+			{
+				throw new Exception("Handshake failed! " + handshakeResponse.Message);
+			}
+
+			try
+			{
+				await AssertSendReceiveAsync(streamReader, streamWriter, 1);
+				await AssertSendReceiveAsync(streamReader, streamWriter, echoServerBufferSize / 3);
+				await AssertSendReceiveAsync(streamReader, streamWriter, echoServerBufferSize - 1);
+				await AssertSendReceiveAsync(streamReader, streamWriter, echoServerBufferSize);
+				await AssertSendReceiveAsync(streamReader, streamWriter, echoServerBufferSize + 1);
+				await AssertSendReceiveAsync(streamReader, streamWriter, echoServerBufferSize * 10 + 3);
+				await AssertSendReceiveAsync(streamReader, streamWriter, echoServerBufferSize * 20);
+			}
+			finally
+			{
+				await cts.CancelAsync();
+				await tunnelService.StopAsync(CancellationToken.None);
+			}
+		}
+
+		private static async Task AssertSendReceiveAsync(StreamReader reader, StreamWriter writer, int length)
+		{
+			string randomString = GenerateRandomString(length);
+			await writer.WriteLineAsync(randomString);
+			string? response = await ReadLineWithTimeoutAsync(reader, TimeSpan.FromSeconds(3));
+			Assert.AreEqual(randomString, response);
+		}
+
+		private static string GenerateRandomString(int length)
+		{
+			const string Chars = "ABCDEFGHIJKLMNOPQRSTUVWXYZabcdefghijklmnopqrstuvwxyz0123456789";
+			StringBuilder result = new(length);
+
+			for (int i = 0; i < length; i++)
+			{
+				result.Append(Chars[s_random.Next(Chars.Length)]);
+			}
+
+			return result.ToString();
+		}
+
+		private static async Task StartEchoServerAsync(int port, int bufferSize, CancellationToken cancellationToken)
+		{
+			using TcpListener listener = new(IPAddress.Any, port);
+			listener.Start();
+
+			Console.WriteLine($"Echo server started on port {port}...");
+			List<Task> tasks = new();
+
+			while (!cancellationToken.IsCancellationRequested)
+			{
+				TcpClient client = await listener.AcceptTcpClientAsync(cancellationToken);
+				Console.WriteLine("Accepted connection from client...");
+				tasks.Add(HandleEchoClientAsync(client, bufferSize, cancellationToken));
+			}
+
+			await Task.WhenAll(tasks);
+		}
+
+		private static async Task HandleEchoClientAsync(TcpClient client, int bufferSize, CancellationToken cancellationToken)
+		{
+			await using (NetworkStream stream = client.GetStream())
+			{
+				Console.WriteLine("Echo client connected!");
+				byte[] buffer = new byte[bufferSize];
+				int bytesRead;
+
+				while ((bytesRead = await stream.ReadAsync(buffer, 0, buffer.Length, cancellationToken)) > 0)
+				{
+					await stream.WriteAsync(buffer, 0, bytesRead, cancellationToken);
+				}
+			}
+
+			client.Close();
+			Console.WriteLine("Client connection closed.");
+		}
+
+		private static async Task<string?> ReadLineWithTimeoutAsync(StreamReader reader, TimeSpan timeout)
+		{
+			Task<string?> readLineTask = reader.ReadLineAsync();
+			Task delayTask = Task.Delay(timeout);
+			Task completedTask = await Task.WhenAny(readLineTask, delayTask);
+			if (completedTask == readLineTask)
+			{
+				return await readLineTask;
+			}
+			throw new TimeoutException("Read operation timed out.");
+		}
+	}
+}
diff --git a/Engine/Source/Programs/Horde/Plugins/Compute/HordeServer.Compute.Tests/HordeServer.Compute.Tests.csproj b/Engine/Source/Programs/Horde/Plugins/Compute/HordeServer.Compute.Tests/HordeServer.Compute.Tests.csproj
--- a/Engine/Source/Programs/Horde/Plugins/Compute/HordeServer.Compute.Tests/HordeServer.Compute.Tests.csproj
+++ b/Engine/Source/Programs/Horde/Plugins/Compute/HordeServer.Compute.Tests/HordeServer.Compute.Tests.csproj
@@ -19,8 +19,12 @@
 
   <ItemGroup>
     <ProjectReference Include="..\HordeServer.Compute\HordeServer.Compute.csproj" />
+    <ProjectReference Include="..\..\Storage\HordeServer.Storage\HordeServer.Storage.csproj" />
     <ProjectReference Include="..\..\..\HordeAgent.Tests\HordeAgent.Tests.csproj" />
     <ProjectReference Include="..\..\..\HordeServer.Tests.Shared\HordeServer.Tests.Shared.csproj" />
+    <ProjectReference Include="..\..\..\HordeAgent\HordeAgent.csproj">
+      <Aliases>HordeAgent</Aliases>
+    </ProjectReference>
   </ItemGroup>
 
   <ItemGroup>
diff --git a/Engine/Source/Programs/Horde/Plugins/Compute/HordeServer.Compute.Tests/Logs/LogIndexingTests.cs b/Engine/Source/Programs/Horde/Plugins/Compute/HordeServer.Compute.Tests/Logs/LogIndexingTests.cs
new file mode 100644
--- /dev/null
+++ b/Engine/Source/Programs/Horde/Plugins/Compute/HordeServer.Compute.Tests/Logs/LogIndexingTests.cs
@@ -0,0 +1,247 @@
+// Copyright Epic Games, Inc. All Rights Reserved.
+
+using System;
+using System.Collections.Generic;
+using System.IO;
+using System.Linq;
+using System.Text;
+using System.Threading;
+using System.Threading.Tasks;
+using EpicGames.Horde.Jobs;
+using EpicGames.Horde.Logs;
+using EpicGames.Horde.Storage;
+using HordeServer.Jobs;
+using HordeServer.Logs;
+using HordeServer.Storage;
+using Microsoft.Extensions.Logging.Abstractions;
+using Microsoft.VisualStudio.TestTools.UnitTesting;
+
+namespace HordeServer.Tests.Logs
+{
+	class TestLogWriter : IAsyncDisposable
+	{
+		ILog _logFile;
+		readonly IStorageClient _storageClient;
+		readonly IBlobWriter _blobWriter;
+		readonly LogBuilder _builder;
+		int _lineCount;
+
+		public TestLogWriter(ILog logFile, StorageService storageService)
+		{
+			_logFile = logFile;
+			_storageClient = storageService.CreateClient(Namespace.Logs);
+			_blobWriter = _storageClient.CreateBlobWriter(logFile.RefName);
+			_builder = new LogBuilder((logFile.Type == LogType.Text) ? LogFormat.Text : LogFormat.Json, NullLogger.Instance);
+		}
+
+		public async ValueTask DisposeAsync()
+		{
+			await _blobWriter.DisposeAsync();
+			_storageClient.Dispose();
+		}
+
+		public async Task<ILog> WriteDataAsync(ReadOnlyMemory<byte> data)
+		{
+			_builder.WriteData(data);
+			_lineCount += data.Span.Count((byte)'\n');
+			_logFile = await _logFile.UpdateLineCountAsync(_lineCount, false, CancellationToken.None);
+			return await FlushAsync(false);
+		}
+
+		public async Task<ILog> FlushAsync(bool complete = true)
+		{
+			IBlobRef<LogNode> handle = await _builder.FlushAsync(_blobWriter, complete, CancellationToken.None);
+			await _storageClient.WriteRefAsync(_logFile.RefName, handle);
+			_logFile = await _logFile.UpdateLineCountAsync(_lineCount, complete, CancellationToken.None);
+			return _logFile;
+		}
+	}
+
+	[TestClass]
+	public class LogIndexingTests : ComputeTestSetup
+	{
+		private readonly byte[] _data = HordeServer.Compute.Tests.Properties.Resources.TextFile;
+
+		[TestMethod]
+		public async Task IndexTestsAsync()
+		{
+			JobId jobId = JobIdUtils.GenerateNewId();
+			ILog logFile = await LogCollection.AddAsync(jobId, null, null, LogType.Text);
+
+			// Write the test data to the log file in blocks
+			await using (TestLogWriter logWriter = new TestLogWriter(logFile, StorageService))
+			{
+				int offset = 0;
+				int lineIndex = 0;
+				while (offset < _data.Length)
+				{
+					int length = 0;
+					int lineCount = 0;
+
+					for (int idx = offset; idx < Math.Min(_data.Length, offset + 18830); idx++)
+					{
+						if (_data[idx] == '\n')
+						{
+							length = (idx + 1) - offset;
+							lineCount++;
+						}
+					}
+
+					await logWriter.WriteDataAsync(_data.AsMemory(offset, length));
+
+					offset += length;
+					lineIndex += lineCount;
+				}
+
+				await logWriter.FlushAsync();
+			}
+
+			// Read the data back out and check it's the same
+			byte[] readData = new byte[_data.Length];
+			using (Stream stream = await logFile.OpenRawStreamAsync())
+			{
+				int readSize = await stream.ReadAsync(readData, 0, readData.Length);
+				Assert.AreEqual(readData.Length, readSize);
+
+				int equalSize = 0;
+				while (equalSize < _data.Length && _data[equalSize] == readData[equalSize])
+				{
+					equalSize++;
+				}
+
+				Assert.AreEqual(equalSize, readSize);
+			}
+
+			// Test some searches
+			await SearchLogDataTestAsync(logFile.Id);
+		}
+
+		[TestMethod]
+		public void TrieTests()
+		{
+			ReadOnlyTrieBuilder builder = new ReadOnlyTrieBuilder();
+
+			ulong[] values = { 1, 2, 3, 5, 8, 13, 21, 34, 55, 89, 144, 233, 377, 610, 987, 1597, 2584, 4181, 6765, 10946, 0xfedcba9876543210UL };
+			foreach (ulong value in values)
+			{
+				builder.Add(value);
+			}
+
+			ReadOnlyTrie trie = builder.Build();
+			Assert.IsTrue(Enumerable.SequenceEqual(trie.EnumerateRange(0, UInt64.MaxValue), values));
+			Assert.IsTrue(Enumerable.SequenceEqual(trie.EnumerateRange(0, 90), values.Where(x => x <= 90)));
+			Assert.IsTrue(Enumerable.SequenceEqual(trie.EnumerateRange(2, 89), values.Where(x => x >= 2 && x <= 89)));
+		}
+
+		[TestMethod]
+		public async Task PartialTokenTestsAsync()
+		{
+			JobId jobId = JobIdUtils.GenerateNewId();
+			ILog logFile = await LogCollection.AddAsync(jobId, null, null, LogType.Text);
+
+			string[] lines =
+			{
+				"abcdefghi\n",
+				"jklmno123\n",
+				"pqrst99uv\n",
+				"wx\n"
+			};
+
+			int length = 0;
+			await using (TestLogWriter logWriter = new TestLogWriter(logFile, StorageService))
+			{
+				for (int lineIdx = 0; lineIdx < lines.Length; lineIdx++)
+				{
+					logFile = await logWriter.WriteDataAsync(Encoding.UTF8.GetBytes(lines[lineIdx]));
+					length += lines[lineIdx].Length;
+				}
+				logFile = await logWriter.FlushAsync(true);
+			}
+
+			for (int lineIdx = 0; lineIdx < lines.Length; lineIdx++)
+			{
+				for (int strLen = 1; strLen < 7; strLen++)
+				{
+					for (int strOfs = 0; strOfs + strLen < lines[lineIdx].Length - 1; strOfs++)
+					{
+						string str = lines[lineIdx].Substring(strOfs, strLen);
+
+						SearchStats stats = new SearchStats();
+						List<int> results = await logFile.SearchLogDataAsync(str, 0, 5, stats, CancellationToken.None);
+						Assert.AreEqual(1, results.Count);
+						Assert.AreEqual(lineIdx, results[0]);
+
+						Assert.AreEqual(1, stats.NumScannedBlocks);
+						Assert.AreEqual(3, stats.NumSkippedBlocks);
+						Assert.AreEqual(0, stats.NumFalsePositiveBlocks);
+					}
+				}
+			}
+		}
+
+		[TestMethod]
+		public async Task AppendIndexTestsAsync()
+		{
+			JobId jobId = JobIdUtils.GenerateNewId();
+			ILog logFile = await LogCollection.AddAsync(jobId, null, null, LogType.Text);
+
+			await using (TestLogWriter writer = new TestLogWriter(logFile, StorageService))
+			{
+				logFile = await writer.WriteDataAsync(Encoding.UTF8.GetBytes("abc\n"));
+				logFile = await writer.WriteDataAsync(Encoding.UTF8.GetBytes("def\n"));
+				logFile = await writer.WriteDataAsync(Encoding.UTF8.GetBytes("ghi\n"));
+				logFile = await writer.FlushAsync(true);
+			}
+
+			{
+				SearchStats stats = new SearchStats();
+				List<int> results = await logFile.SearchLogDataAsync("abc", 0, 5, stats, CancellationToken.None);
+				Assert.AreEqual(1, results.Count);
+				Assert.AreEqual(0, results[0]);
+
+				Assert.AreEqual(1, stats.NumScannedBlocks); // abc
+				Assert.AreEqual(2, stats.NumSkippedBlocks); // def + ghi
+				Assert.AreEqual(0, stats.NumFalsePositiveBlocks);
+			}
+			{
+				SearchStats stats = new SearchStats();
+				List<int> results = await logFile.SearchLogDataAsync("def", 0, 5, stats, CancellationToken.None);
+				Assert.AreEqual(1, results.Count);
+				Assert.AreEqual(1, results[0]);
+
+				Assert.AreEqual(1, stats.NumScannedBlocks); // def
+				Assert.AreEqual(2, stats.NumSkippedBlocks); // abc + ghi
+				Assert.AreEqual(0, stats.NumFalsePositiveBlocks);
+			}
+			{
+				SearchStats stats = new SearchStats();
+				List<int> results = await logFile.SearchLogDataAsync("ghi", 0, 5, stats, CancellationToken.None);
+				Assert.AreEqual(1, results.Count);
+				Assert.AreEqual(2, results[0]);
+
+				Assert.AreEqual(1, stats.NumScannedBlocks); // ghi
+				Assert.AreEqual(2, stats.NumSkippedBlocks); // abc + def
+				Assert.AreEqual(0, stats.NumFalsePositiveBlocks);
+			}
+		}
+
+		async Task SearchLogDataTestAsync(LogId logId)
+		{
+			ILog? logFile = await LogCollection.GetAsync(logId, CancellationToken.None);
+			Assert.IsNotNull(logFile);
+
+			await SearchLogDataTestAsync(logFile, "HISPANIOLA", 0, 4, new[] { 1503, 1520, 1525, 1595 });
+			await SearchLogDataTestAsync(logFile, "Hispaniola", 0, 4, new[] { 1503, 1520, 1525, 1595 });
+			await SearchLogDataTestAsync(logFile, "HizpaniolZ", 0, 4, Array.Empty<int>());
+			await SearchLogDataTestAsync(logFile, "Pieces of eight!", 0, 100, new[] { 2227, 2228, 5840, 5841, 7520 });
+			await SearchLogDataTestAsync(logFile, "NEWSLETTER", 0, 100, new[] { 7886 });
+		}
+
+		static async Task SearchLogDataTestAsync(ILog logFile, string text, int firstLine, int count, int[] expectedLines)
+		{
+			SearchStats stats = new SearchStats();
+			List<int> lines = await logFile.SearchLogDataAsync(text, firstLine, count, stats, CancellationToken.None);
+			Assert.IsTrue(lines.SequenceEqual(expectedLines));
+		}
+	}
+}
diff --git a/Engine/Source/Programs/Horde/Plugins/Compute/HordeServer.Compute.Tests/Logs/LogServiceTests.cs b/Engine/Source/Programs/Horde/Plugins/Compute/HordeServer.Compute.Tests/Logs/LogServiceTests.cs
new file mode 100644
--- /dev/null
+++ b/Engine/Source/Programs/Horde/Plugins/Compute/HordeServer.Compute.Tests/Logs/LogServiceTests.cs
@@ -0,0 +1,240 @@
+// Copyright Epic Games, Inc. All Rights Reserved.
+
+using System;
+using System.Collections.Generic;
+using System.IO;
+using System.Linq;
+using System.Security.Claims;
+using System.Text;
+using System.Threading;
+using System.Threading.Tasks;
+using EpicGames.Core;
+using EpicGames.Horde.Agents.Sessions;
+using EpicGames.Horde.Jobs;
+using EpicGames.Horde.Logs;
+using HordeServer.Agents.Sessions;
+using HordeServer.Jobs;
+using HordeServer.Logs;
+using HordeServer.Utilities;
+using Microsoft.VisualStudio.TestTools.UnitTesting;
+
+namespace HordeServer.Tests.Logs
+{
+	[TestClass]
+	public sealed class LogServiceTest : ComputeTestSetup
+	{
+		[TestMethod]
+		public void PlainTextDecoder()
+		{
+			string jsonText = "{\"time\":\"2022-01-26T14:18:29\",\"level\":\"Error\",\"message\":\"       \\tdepends on: /mnt/horde/U5\\u002BR5.0\\u002BInc\\u002BMin/Sync/Engine/Binaries/Linux/libUnrealEditor-Core.so\",\"id\":1,\"line\":5,\"lineCount\":1028}";
+			byte[] jsonData = Encoding.UTF8.GetBytes(jsonText);
+
+			byte[] outputData = new byte[jsonData.Length];
+			int outputLength = HordeServer.Logs.LogCollection.ConvertToPlainText(jsonData, outputData, 0);
+			string outputText = Encoding.UTF8.GetString(outputData, 0, outputLength);
+
+			string expectText = "       \tdepends on: /mnt/horde/U5+R5.0+Inc+Min/Sync/Engine/Binaries/Linux/libUnrealEditor-Core.so\n";
+			Assert.AreEqual(outputText, expectText);
+		}
+
+		[TestMethod]
+		public async Task WriteLogLifecycleOldTestAsync()
+		{
+			JobId jobId = JobIdUtils.GenerateNewId();
+			ILog logFile = await LogCollection.AddAsync(jobId, null, null, LogType.Text, logId: null, cancellationToken: CancellationToken.None);
+
+			await using (TestLogWriter writer = new TestLogWriter(logFile, StorageService))
+			{
+				await writer.WriteDataAsync(Encoding.ASCII.GetBytes("hello\n"));
+				await writer.FlushAsync(false);
+
+				await writer.WriteDataAsync(Encoding.ASCII.GetBytes("foo\nbar\n"));
+				await writer.FlushAsync(false);
+
+				await writer.WriteDataAsync(Encoding.ASCII.GetBytes("baz\n"));
+				await writer.FlushAsync(true);
+			}
+
+			logFile = (await LogCollection.GetAsync(logFile.Id, CancellationToken.None))!;
+
+			Assert.AreEqual("hello", await ReadLogAsync(logFile, 0, 5));
+			Assert.AreEqual("foo\nbar\nbaz\n", await ReadLogAsync(logFile, 6, 12));
+
+			List<Utf8String> lines = await logFile.ReadLinesAsync(0, 100, CancellationToken.None);
+			Assert.AreEqual("hello", lines[0].ToString());
+			Assert.AreEqual("foo", lines[1].ToString());
+			Assert.AreEqual("bar", lines[2].ToString());
+			Assert.AreEqual("baz", lines[3].ToString());
+		}
+
+		private static async Task AssertLineOffsetAsync(ILogCollection logCollection, LogId logId, int lineIndex, int clampedLineIndex, long offset)
+		{
+			ILog? logFile = await logCollection.GetAsync(logId, CancellationToken.None);
+			Assert.IsNotNull(logFile);
+
+			List<Utf8String> lines = await logFile.ReadLinesAsync(0, lineIndex + 1, CancellationToken.None);
+			Assert.AreEqual(clampedLineIndex, Math.Min(lineIndex, lines.Count));
+			Assert.AreEqual(offset, lines.Take(lineIndex).Sum(x => x.Length + 1));
+		}
+
+		static async Task<string> ReadLogAsync(ILog log, long offset, long length)
+		{
+			using Stream stream = await log.OpenRawStreamAsync();
+
+			byte[] prefix = new byte[(int)offset];
+			await stream.ReadGreedyAsync(prefix);
+
+			byte[] data = new byte[(int)length];
+			int actualLength = await stream.ReadGreedyAsync(data);
+
+			return Encoding.UTF8.GetString(data.AsSpan(0, actualLength));
+		}
+
+		[TestMethod]
+		public async Task WriteLogLifecycleAsync()
+		{
+			JobId jobId = JobIdUtils.GenerateNewId();
+			ILog log = await LogCollection.AddAsync(jobId, null, null, LogType.Text);
+
+			string str1 = "hello\n";
+			string str2 = "foo\nbar\n";
+			string str3 = "baz\nqux\nquux\n";
+			string str4 = "quuz\n";
+
+			int lineIndex = 0;
+			int offset = 0;
+
+			await using TestLogWriter logWriter = new TestLogWriter(log, StorageService);
+
+			// First write with flush. Will become chunk #1
+			log = await logWriter.WriteDataAsync(Encoding.ASCII.GetBytes(str1));
+			Assert.AreEqual(str1, await ReadLogAsync(log, 0, str1.Length));
+			Assert.AreEqual(str1, await ReadLogAsync(log, 0, str1.Length + 100)); // Reading too far is valid?
+			await AssertLineOffsetAsync(LogCollection, log.Id, 0, 0, 0);
+			await AssertLineOffsetAsync(LogCollection, log.Id, 1, 1, 6);
+			await AssertLineOffsetAsync(LogCollection, log.Id, 1, 1, 6);
+			await AssertLineOffsetAsync(LogCollection, log.Id, 1, 1, 6);
+
+			// Second write without flushing. Will become chunk #2
+			offset += str1.Length;
+			lineIndex += str1.Count(f => f == '\n');
+			log = await logWriter.WriteDataAsync(Encoding.ASCII.GetBytes(str2));
+
+			Assert.AreEqual(str1 + str2, await ReadLogAsync(log, 0, str1.Length + str2.Length));
+			await AssertLineOffsetAsync(LogCollection, log.Id, 0, 0, 0);
+			await AssertLineOffsetAsync(LogCollection, log.Id, 1, 1, 6);
+			await AssertLineOffsetAsync(LogCollection, log.Id, 2, 2, 10);
+			await AssertLineOffsetAsync(LogCollection, log.Id, 3, 3, 14);
+
+			// Third write without flushing. Will become chunk #2
+			offset += str2.Length;
+			lineIndex += str3.Count(f => f == '\n');
+			log = await logWriter.WriteDataAsync(Encoding.ASCII.GetBytes(str3));
+			Assert.AreEqual(str1 + str2 + str3, await ReadLogAsync(log, 0, str1.Length + str2.Length + str3.Length));
+			await AssertLineOffsetAsync(LogCollection, log.Id, 0, 0, 0);
+			await AssertLineOffsetAsync(LogCollection, log.Id, 1, 1, 6);
+			await AssertLineOffsetAsync(LogCollection, log.Id, 2, 2, 10);
+			await AssertLineOffsetAsync(LogCollection, log.Id, 3, 3, 14);
+			await AssertLineOffsetAsync(LogCollection, log.Id, 4, 4, 18);
+
+			// Fourth write with flush. Will become chunk #2
+			offset += str3.Length;
+			lineIndex += str4.Count(f => f == '\n');
+			log = await logWriter.WriteDataAsync(Encoding.ASCII.GetBytes(str4));
+			Assert.AreEqual(str1 + str2 + str3 + str4,
+			await ReadLogAsync(log, 0, str1.Length + str2.Length + str3.Length + str4.Length));
+			Assert.AreEqual(str3 + str4,
+			await ReadLogAsync(log, str1.Length + str2.Length, str3.Length + str4.Length));
+			await AssertLineOffsetAsync(LogCollection, log.Id, 0, 0, 0);
+			await AssertLineOffsetAsync(LogCollection, log.Id, 1, 1, 6);
+			await AssertLineOffsetAsync(LogCollection, log.Id, 2, 2, 10);
+			await AssertLineOffsetAsync(LogCollection, log.Id, 3, 3, 14);
+			await AssertLineOffsetAsync(LogCollection, log.Id, 4, 4, 18);
+			await AssertLineOffsetAsync(LogCollection, log.Id, 5, 5, 22);
+			await AssertLineOffsetAsync(LogCollection, log.Id, 6, 6, 27);
+			await AssertLineOffsetAsync(LogCollection, log.Id, 7, 7, 32);
+
+			// Fifth write with flush and data that will span more than chunk. Will become chunk #3
+			string a = "Lorem ipsum dolor sit amet\n";
+			string b = "consectetur adipiscing\n";
+			string str5 = a + b;
+
+			offset += str4.Length;
+			lineIndex += str5.Count(f => f == '\n');
+
+			// Using this single write below will fail the ReadLogFile assert below. A bug?
+			//await LogFileService.WriteLogDataAsync(LogFile, Offset, LineIndex, Encoding.ASCII.GetBytes(Str5), true);
+
+			// Dividing it in two like this will work however
+			log = await logWriter.WriteDataAsync(Encoding.ASCII.GetBytes(a));
+			log = await logWriter.WriteDataAsync(Encoding.ASCII.GetBytes(b));
+
+			Assert.AreEqual(str5, await ReadLogAsync(log, offset, str5.Length));
+		}
+
+		[TestMethod]
+		public async Task GetLogFileTestAsync()
+		{
+			// Will implicitly test GetLogFileAsync(), AddCachedLogFile()
+			JobId jobId = JobIdUtils.GenerateNewId();
+			SessionId sessionId = SessionIdUtils.GenerateNewId();
+			ILog a = await LogCollection.AddAsync(jobId, null, sessionId, LogType.Text, logId: null, cancellationToken: CancellationToken.None);
+			ILog b = (await LogCollection.GetAsync(a.Id, CancellationToken.None))!;
+			Assert.AreEqual(a.JobId, b.JobId);
+			Assert.AreEqual(a.SessionId, b.SessionId);
+			Assert.AreEqual(a.Type, b.Type);
+
+			ILog? notFound = await LogCollection.GetAsync(LogIdUtils.GenerateNewId(), CancellationToken.None);
+			Assert.IsNull(notFound);
+		}
+
+		[TestMethod]
+		public async Task AuthorizeForSessionAsync()
+		{
+			JobId jobId = JobIdUtils.GenerateNewId();
+			SessionId sessionId = SessionIdUtils.GenerateNewId();
+			ILog log = await LogCollection.AddAsync(jobId, null, sessionId, LogType.Text, logId: null, cancellationToken: CancellationToken.None);
+			ILog logNoSession = await LogCollection.AddAsync(jobId, null, null, LogType.Text, logId: null, cancellationToken: CancellationToken.None);
+
+			ClaimsPrincipal hasClaim = new ClaimsPrincipal(new ClaimsIdentity(new List<Claim>
+			{
+				new Claim(HordeClaimTypes.AgentSessionId, sessionId.ToString()),
+			}, "TestAuthType"));
+			ClaimsPrincipal hasNoClaim = new ClaimsPrincipal(new ClaimsIdentity(new List<Claim>
+			{
+				new Claim(HordeClaimTypes.AgentSessionId, "invalid-session-id"),
+			}, "TestAuthType"));
+
+			Assert.IsTrue(log.AuthorizeForSession(hasClaim));
+			Assert.IsFalse(log.AuthorizeForSession(hasNoClaim));
+			Assert.IsFalse(logNoSession.AuthorizeForSession(hasClaim));
+		}
+
+		[TestMethod]
+		public async Task ChunkSplittingAsync()
+		{
+			JobId jobId = JobIdUtils.GenerateNewId();
+			ILog logFile = await LogCollection.AddAsync(jobId, null, null, LogType.Text, logId: null, cancellationToken: CancellationToken.None);
+			await using TestLogWriter logWriter = new TestLogWriter(logFile, StorageService);
+
+			byte[] line1 = Encoding.UTF8.GetBytes("hello world\n");
+			await logWriter.WriteDataAsync(line1);
+
+			byte[] line2 = Encoding.UTF8.GetBytes("ab\n");
+			await logWriter.WriteDataAsync(line2);
+
+			byte[] line3 = Encoding.UTF8.GetBytes("a\n");
+			await logWriter.WriteDataAsync(line3);
+
+			byte[] line4 = Encoding.UTF8.GetBytes("b\n");
+			await logWriter.WriteDataAsync(line4);
+
+			byte[] line5 = Encoding.UTF8.GetBytes("a\nb\n");
+			await logWriter.WriteDataAsync(line5);
+
+			logFile = await logWriter.FlushAsync();
+			LogMetadata metadata = await logFile.GetMetadataAsync(CancellationToken.None);
+			Assert.AreEqual(6, metadata.MaxLineIndex);
+		}
+	}
+}
\ No newline at end of file
diff --git a/Engine/Source/Programs/Horde/Plugins/Compute/HordeServer.Compute.Tests/Logs/LogTailServiceTests.cs b/Engine/Source/Programs/Horde/Plugins/Compute/HordeServer.Compute.Tests/Logs/LogTailServiceTests.cs
new file mode 100644
--- /dev/null
+++ b/Engine/Source/Programs/Horde/Plugins/Compute/HordeServer.Compute.Tests/Logs/LogTailServiceTests.cs
@@ -0,0 +1,188 @@
+// Copyright Epic Games, Inc. All Rights Reserved.
+
+using System;
+using System.Collections.Generic;
+using System.Text;
+using System.Threading;
+using System.Threading.Tasks;
+using EpicGames.Core;
+using EpicGames.Horde.Logs;
+using EpicGames.Redis;
+using HordeServer.Logs;
+using HordeServer.Server;
+using Microsoft.Extensions.DependencyInjection;
+using Microsoft.Extensions.Logging;
+using Microsoft.VisualStudio.TestTools.UnitTesting;
+
+namespace HordeServer.Tests.Logs
+{
+	[TestClass]
+	public class LogTailServiceTests : ComputeTestSetup
+	{
+		readonly LogId _logId = LogId.Parse("637809e571f04e1d1ba311f1");
+
+		[TestMethod]
+		public async Task TruncatedAsync()
+		{
+			LogTailService tailService = ServiceProvider.GetRequiredService<LogTailService>();
+			await tailService.EnableTailingAsync(_logId, 30);
+			await tailService.AppendAsync(_logId, 30, Encoding.UTF8.GetBytes("foo\nbar\nbaz\n"));
+
+			for (int idx = 25; idx <= 29; idx++)
+			{
+				List<Utf8String> lines = await tailService.ReadAsync(_logId, idx, 10);
+				Assert.AreEqual(0, lines.Count);
+			}
+
+			{
+				List<Utf8String> lines = await tailService.ReadAsync(_logId, 30, 10);
+				Assert.AreEqual(3, lines.Count);
+				Assert.AreEqual("foo", lines[0].ToString());
+				Assert.AreEqual("bar", lines[1].ToString());
+				Assert.AreEqual("baz", lines[2].ToString());
+			}
+		}
+
+		[TestMethod]
+		public async Task TruncatedMultiLineAsync()
+		{
+			LogTailService tailService = ServiceProvider.GetRequiredService<LogTailService>();
+			await tailService.EnableTailingAsync(_logId, 30);
+			await tailService.AppendAsync(_logId, 30, Encoding.UTF8.GetBytes("foo\n"));
+			await tailService.AppendAsync(_logId, 31, Encoding.UTF8.GetBytes("bar\n"));
+			await tailService.AppendAsync(_logId, 32, Encoding.UTF8.GetBytes("baz\n"));
+
+			for (int idx = 25; idx <= 29; idx++)
+			{
+				List<Utf8String> lines = await tailService.ReadAsync(_logId, idx, 10);
+				Assert.AreEqual(0, lines.Count);
+			}
+
+			{
+				List<Utf8String> lines = await tailService.ReadAsync(_logId, 30, 10);
+				Assert.AreEqual(3, lines.Count);
+				Assert.AreEqual("foo", lines[0].ToString());
+				Assert.AreEqual("bar", lines[1].ToString());
+				Assert.AreEqual("baz", lines[2].ToString());
+			}
+		}
+
+		[TestMethod]
+		public async Task SplittingAsync()
+		{
+			IRedisService redisService = ServiceProvider.GetRequiredService<IRedisService>();
+
+			LogTailService tailService = new LogTailService(redisService, ServiceProvider.GetRequiredService<IClock>(), 4, ServiceProvider.GetRequiredService<IServerInfo>(), ServiceProvider.GetRequiredService<ILogger<LogTailService>>());
+			await tailService.StartAsync(CancellationToken.None);
+
+			await tailService.EnableTailingAsync(_logId, 30);
+			await tailService.AppendAsync(_logId, 30, Encoding.UTF8.GetBytes("line30\nline31\nline32\nline33\nline34\nline35\nline36\nline37\nline38\nline39\nline40\n"));
+
+			Assert.AreEqual(4, await redisService.GetDatabase().SortedSetLengthAsync(LogTailService.TailDataKey(_logId)));
+
+			await tailService.FlushAsync(_logId, 33);
+
+			Assert.AreEqual(4, await redisService.GetDatabase().SortedSetLengthAsync(LogTailService.TailDataKey(_logId)));
+
+			await Clock.AdvanceAsync(TimeSpan.FromMinutes(1.0));
+
+			Assert.AreEqual(3, await redisService.GetDatabase().SortedSetLengthAsync(LogTailService.TailDataKey(_logId)));
+
+			{
+				List<Utf8String> lines = await tailService.ReadAsync(_logId, 25, 20);
+				Assert.AreEqual(0, lines.Count);
+			}
+			{
+				List<Utf8String> lines = await tailService.ReadAsync(_logId, 31, 20);
+				Assert.AreEqual(0, lines.Count);
+			}
+			{
+				List<Utf8String> lines = await tailService.ReadAsync(_logId, 32, 20);
+				Assert.AreEqual("line32", lines[0].ToString());
+				Assert.AreEqual("line33", lines[1].ToString());
+				Assert.AreEqual("line34", lines[2].ToString());
+				Assert.AreEqual("line35", lines[3].ToString());
+				Assert.AreEqual("line36", lines[4].ToString());
+				Assert.AreEqual("line37", lines[5].ToString());
+				Assert.AreEqual("line38", lines[6].ToString());
+				Assert.AreEqual("line39", lines[7].ToString());
+				Assert.AreEqual("line40", lines[8].ToString());
+			}
+		}
+
+		[TestMethod]
+		public async Task ExpiryAsync()
+		{
+			IRedisService redisService = GetRedisServiceSingleton();
+
+			LogTailService tailService = ServiceProvider.GetRequiredService<LogTailService>();
+
+			await tailService.StartAsync(CancellationToken.None);
+			await tailService.EnableTailingAsync(_logId, 30);
+			await tailService.AppendAsync(_logId, 30, Encoding.UTF8.GetBytes("foo\n"));
+
+			Assert.IsTrue(await redisService.GetDatabase().KeyExistsAsync(LogTailService.TailNextKey(_logId)));
+			Assert.IsTrue(await redisService.GetDatabase().KeyExistsAsync(LogTailService.TailDataKey(_logId)));
+
+			await Clock.AdvanceAsync(TimeSpan.FromMinutes(5.0));
+
+			Assert.IsFalse(await redisService.GetDatabase().KeyExistsAsync(LogTailService.TailNextKey(_logId)));
+			Assert.IsFalse(await redisService.GetDatabase().KeyExistsAsync(LogTailService.TailDataKey(_logId)));
+		}
+
+		[TestMethod]
+		public async Task MissingDataAsync()
+		{
+			LogTailService tailService = ServiceProvider.GetRequiredService<LogTailService>();
+
+			await tailService.StartAsync(CancellationToken.None);
+			await tailService.EnableTailingAsync(_logId, 30);
+			await tailService.AppendAsync(_logId, 31, Encoding.UTF8.GetBytes("line31\nline32\nline33\nline34\nline35\n"));
+
+			{
+				List<Utf8String> lines = await tailService.ReadAsync(_logId, 29, 4);
+				Assert.AreEqual(0, lines.Count);
+			}
+			{
+				List<Utf8String> lines = await tailService.ReadAsync(_logId, 29, 4);
+				Assert.AreEqual(0, lines.Count);
+			}
+
+			await tailService.AppendAsync(_logId, 30, Encoding.UTF8.GetBytes("line30\n"));
+			await tailService.AppendAsync(_logId, 32, Encoding.UTF8.GetBytes("line32\nline33\nline34\nline35\n"));
+
+			{
+				List<Utf8String> lines = await tailService.ReadAsync(_logId, 29, 4);
+				Assert.AreEqual(0, lines.Count);
+			}
+			{
+				List<Utf8String> lines = await tailService.ReadAsync(_logId, 30, 4);
+				Assert.AreEqual(1, lines.Count);
+				Assert.AreEqual("line30", lines[0].ToString());
+			}
+			{
+				List<Utf8String> lines = await tailService.ReadAsync(_logId, 32, 4);
+				Assert.AreEqual(0, lines.Count);
+			}
+		}
+
+		[TestMethod]
+		public async Task AsyncEnableAsync()
+		{
+			LogTailService tailService = ServiceProvider.GetRequiredService<LogTailService>();
+			await tailService.StartAsync(CancellationToken.None);
+
+			LogTailService tailService2 = new LogTailService(ServiceProvider.GetRequiredService<IRedisService>(), ServiceProvider.GetRequiredService<IClock>(), 4, ServiceProvider.GetRequiredService<IServerInfo>(), ServiceProvider.GetRequiredService<ILogger<LogTailService>>());
+
+			Task<int> task = tailService.WaitForTailNextAsync(_logId, CancellationToken.None);
+			Assert.IsTrue(!task.IsCompleted);
+
+			await tailService2.EnableTailingAsync(_logId, 30);
+
+			await Task.WhenAny(task, Task.Delay(1000));
+			Assert.IsTrue(task.IsCompletedSuccessfully);
+
+			Assert.AreEqual(30, await task);
+		}
+	}
+}
diff --git a/Engine/Source/Programs/Horde/Plugins/Compute/HordeServer.Compute.Tests/Properties/AssemblyInfo.cs b/Engine/Source/Programs/Horde/Plugins/Compute/HordeServer.Compute.Tests/Properties/AssemblyInfo.cs
new file mode 100644
--- /dev/null
+++ b/Engine/Source/Programs/Horde/Plugins/Compute/HordeServer.Compute.Tests/Properties/AssemblyInfo.cs
@@ -0,0 +1,5 @@
+// Copyright Epic Games, Inc. All Rights Reserved.
+
+using System.Resources;
+
+[assembly: NeutralResourcesLanguage("en-US")]
diff --git a/Engine/Source/Programs/Horde/Plugins/Compute/HordeServer.Compute/Properties/AssemblyInfo.cs b/Engine/Source/Programs/Horde/Plugins/Compute/HordeServer.Compute/Properties/AssemblyInfo.cs
--- a/Engine/Source/Programs/Horde/Plugins/Compute/HordeServer.Compute/Properties/AssemblyInfo.cs
+++ b/Engine/Source/Programs/Horde/Plugins/Compute/HordeServer.Compute/Properties/AssemblyInfo.cs
@@ -4,4 +4,5 @@ using System.Resources;
 using System.Runtime.CompilerServices;
 
 [assembly: InternalsVisibleTo("HordeServer.Tests")]
+[assembly: InternalsVisibleTo("HordeServer.Compute.Tests")]
 [assembly: NeutralResourcesLanguage("en-US")]
