Horde: Move tests into HordeServer.Build.Tests.

#rnx

diff --git a/Engine/Source/Programs/Horde/HordeServer.Tests/Accounts/AccountCollectionTests.cs b/Engine/Source/Programs/Horde/HordeServer.Tests/Accounts/AccountCollectionTests.cs
--- a/Engine/Source/Programs/Horde/HordeServer.Tests/Accounts/AccountCollectionTests.cs
+++ b/Engine/Source/Programs/Horde/HordeServer.Tests/Accounts/AccountCollectionTests.cs
@@ -9,7 +9,7 @@ using Microsoft.VisualStudio.TestTools.UnitTesting;
 namespace HordeServer.Tests.Accounts
 {
 	[TestClass]
-	public class AccountCollectionTests : TestSetup
+	public class AccountCollectionTests : ServerTestSetup
 	{
 		[TestMethod]
 		public async Task CaseInsensitiveAsync()
diff --git a/Engine/Source/Programs/Horde/HordeServer.Tests/Acls/AclTests.cs b/Engine/Source/Programs/Horde/HordeServer.Tests/Acls/AclTests.cs
--- a/Engine/Source/Programs/Horde/HordeServer.Tests/Acls/AclTests.cs
+++ b/Engine/Source/Programs/Horde/HordeServer.Tests/Acls/AclTests.cs
@@ -16,7 +16,7 @@ using Microsoft.VisualStudio.TestTools.UnitTesting;
 namespace HordeServer.Tests.Acls
 {
 	[TestClass]
-	public class AclTests : TestSetup
+	public class AclTests : ServerTestSetup
 	{
 		[TestMethod]
 		public void TestCustomProfile()
diff --git a/Engine/Source/Programs/Horde/HordeServer.Tests/Server/ServerStatusServiceTests.cs b/Engine/Source/Programs/Horde/HordeServer.Tests/Server/ServerStatusServiceTests.cs
--- a/Engine/Source/Programs/Horde/HordeServer.Tests/Server/ServerStatusServiceTests.cs
+++ b/Engine/Source/Programs/Horde/HordeServer.Tests/Server/ServerStatusServiceTests.cs
@@ -13,7 +13,7 @@ using Microsoft.VisualStudio.TestTools.UnitTesting;
 namespace HordeServer.Tests.Server
 {
 	[TestClass]
-	public class ServerStatusServiceTest : TestSetup
+	public class ServerStatusServiceTest : ServerTestSetup
 	{
 		[TestMethod]
 		public async Task UpdatesAreStoredNewToOldAsync()
diff --git a/Engine/Source/Programs/Horde/Plugins/Analytics/HordeServer.Analytics.Tests/Telemetry/TelemetryTests.cs b/Engine/Source/Programs/Horde/Plugins/Analytics/HordeServer.Analytics.Tests/Telemetry/TelemetryTests.cs
--- a/Engine/Source/Programs/Horde/Plugins/Analytics/HordeServer.Analytics.Tests/Telemetry/TelemetryTests.cs
+++ b/Engine/Source/Programs/Horde/Plugins/Analytics/HordeServer.Analytics.Tests/Telemetry/TelemetryTests.cs
@@ -21,7 +21,7 @@ using Microsoft.VisualStudio.TestTools.UnitTesting;
 namespace HordeServer.Analytics.Tests.Telemetry
 {
 	[TestClass]
-	public class TelemetryTests : ServerServiceTest
+	public class TelemetryTests : ServerTestSetup
 	{
 		static readonly TelemetryRecordMeta s_metadata = new TelemetryRecordMeta();
 
diff --git a/Engine/Source/Programs/Horde/Plugins/Build/HordeServer.Build.Tests/Agents/Relay/AgentRelayClientTest.cs b/Engine/Source/Programs/Horde/Plugins/Build/HordeServer.Build.Tests/Agents/Relay/AgentRelayClientTest.cs
new file mode 100644
--- /dev/null
+++ b/Engine/Source/Programs/Horde/Plugins/Build/HordeServer.Build.Tests/Agents/Relay/AgentRelayClientTest.cs
@@ -0,0 +1,82 @@
+// Copyright Epic Games, Inc. All Rights Reserved.
+
+using System;
+using System.Collections.Generic;
+using System.Threading;
+using System.Threading.Tasks;
+using Grpc.Core;
+using Horde.Common.Rpc;
+using HordeServer.Agents.Relay;
+//using HordeServer.Tests.Server;
+using HordeServer.Utilities;
+using Microsoft.Extensions.Logging.Abstractions;
+using Microsoft.VisualStudio.TestTools.UnitTesting;
+
+namespace HordeServer.Tests.Agents.Relay;
+
+public class TestRelayRpcClient : RelayRpc.RelayRpcClient
+{
+	private static readonly ServerCallContext s_adminContext = new ServerCallContextStub(HordeClaims.AdminClaim.ToClaim());
+	public TestAsyncStreamReader<GetPortMappingsResponse> GetPortMappingsResponses { get; } = new(s_adminContext);
+
+	public override AsyncServerStreamingCall<GetPortMappingsResponse> GetPortMappings(
+		GetPortMappingsRequest request, Metadata headers = null!,
+		DateTime? deadline = null, CancellationToken cancellationToken = default)
+	{
+		return new AsyncServerStreamingCall<GetPortMappingsResponse>(
+			GetPortMappingsResponses, null!, null!, null!, () => { });
+	}
+}
+
+[TestClass]
+public class AgentRelayClientTest
+{
+	private readonly TestRelayRpcClient _grpcClient = new();
+	private readonly AgentRelayClient _relayClient;
+
+	public AgentRelayClientTest()
+	{
+		_relayClient = new AgentRelayClient("myCluster", "myAgent", new List<string> { "192.168.1.99" }, Nftables.CreateNull(), _grpcClient, NullLogger<AgentRelayClient>.Instance);
+		_relayClient.CooldownOnException = TimeSpan.Zero;
+	}
+
+	[TestMethod]
+	public async Task GetPortMappingsLongPoll_Results_Async()
+	{
+		using CancellationTokenSource cts = new(3000);
+		_grpcClient.GetPortMappingsResponses.AddMessage(new GetPortMappingsResponse { PortMappings = { NftablesTests.LeaseMap1 }, RevisionCount = 1 });
+		_grpcClient.GetPortMappingsResponses.AddMessage(new GetPortMappingsResponse { PortMappings = { NftablesTests.LeaseMap2, NftablesTests.LeaseMap1 }, RevisionCount = 2 });
+
+		List<PortMapping>? pm1 = await _relayClient.GetPortMappingsLongPollAsync(cts.Token);
+		Assert.AreEqual(1, pm1!.Count);
+		Assert.AreEqual("lease1", pm1[0].LeaseId);
+		Assert.AreEqual(1, _relayClient.RevisionNumber);
+
+		List<PortMapping>? pm2 = await _relayClient.GetPortMappingsLongPollAsync(cts.Token);
+		Assert.AreEqual(2, pm2!.Count);
+		Assert.AreEqual("lease2", pm2[0].LeaseId);
+		Assert.AreEqual("lease1", pm2[1].LeaseId);
+		Assert.AreEqual(2, _relayClient.RevisionNumber);
+	}
+
+	[TestMethod]
+	public async Task GetPortMappingsLongPoll_TimesOut_Async()
+	{
+		using CancellationTokenSource cts = new(3000);
+		await Assert.ThrowsExceptionAsync<OperationCanceledException>(() => _relayClient.GetPortMappingsLongPollAsync(cts.Token));
+	}
+
+	[TestMethod]
+	public async Task ListenForPortMappings_TimesOut_Async()
+	{
+		using CancellationTokenSource cts = new(3000);
+
+		Task t = _relayClient.ListenForPortMappingsAsync(cts.Token);
+		await Task.Delay(100, cts.Token);
+		_grpcClient.GetPortMappingsResponses.AddMessage(new GetPortMappingsResponse { PortMappings = { NftablesTests.LeaseMap1 } });
+		await Task.Delay(100, cts.Token);
+		_grpcClient.GetPortMappingsResponses.AddMessage(new GetPortMappingsResponse { PortMappings = { NftablesTests.LeaseMap2 } });
+
+		await t;
+	}
+}
diff --git a/Engine/Source/Programs/Horde/Plugins/Build/HordeServer.Build.Tests/Fixture.cs b/Engine/Source/Programs/Horde/Plugins/Build/HordeServer.Build.Tests/Fixture.cs
new file mode 100644
--- /dev/null
+++ b/Engine/Source/Programs/Horde/Plugins/Build/HordeServer.Build.Tests/Fixture.cs
@@ -0,0 +1,132 @@
+// Copyright Epic Games, Inc. All Rights Reserved.
+
+using System.Collections.Generic;
+using System.Threading.Tasks;
+using EpicGames.Core;
+using EpicGames.Horde.Agents.Pools;
+using EpicGames.Horde.Jobs;
+using EpicGames.Horde.Jobs.Templates;
+using EpicGames.Horde.Projects;
+using EpicGames.Horde.Streams;
+using HordeServer.Agents;
+using HordeServer.Configuration;
+using HordeServer.Jobs;
+using HordeServer.Jobs.Graphs;
+using HordeServer.Jobs.Templates;
+using HordeServer.Plugins;
+using HordeServer.Projects;
+using HordeServer.Server;
+using HordeServer.Streams;
+
+namespace HordeServer.Tests
+{
+	public class Fixture
+	{
+		public IJob Job1 { get; private set; } = null!;
+		public IJob Job2 { get; private set; } = null!;
+		public ITemplate Template { get; private set; } = null!;
+		public IGraph Graph { get; private set; } = null!;
+		public StreamId StreamId { get; private set; }
+		public StreamConfig? StreamConfig { get; private set; }
+		public TemplateId TemplateRefId1 { get; private set; }
+		public TemplateId TemplateRefId2 { get; private set; }
+		public IAgent Agent1 { get; private set; } = null!;
+		public string Agent1Name { get; private set; } = null!;
+		public const string PoolName = "TestingPool";
+
+		public static async Task<Fixture> CreateAsync(ConfigService configService, IGraphCollection graphCollection, ITemplateCollection templateCollection, JobService jobService, AgentService agentService, IPluginCollection pluginCollection, ServerSettings serverSettings)
+		{
+			Fixture fixture = new Fixture();
+			await fixture.PopulateAsync(configService, graphCollection, templateCollection, jobService, agentService, pluginCollection, serverSettings);
+
+			//			(PerforceService as PerforceServiceStub)?.AddChange("//UE5/Main", 112233, "leet.coder", "Did stuff", new []{"file.cpp"});
+			//			(PerforceService as PerforceServiceStub)?.AddChange("//UE5/Main", 1111, "swarm", "A shelved CL here", new []{"renderer.cpp"});
+
+			return fixture;
+		}
+
+		private async Task PopulateAsync(ConfigService configService, IGraphCollection graphCollection, ITemplateCollection templateCollection, JobService jobService, AgentService agentService, IPluginCollection pluginCollection, ServerSettings serverSettings)
+		{
+			FixtureGraph fg = new FixtureGraph();
+			fg.Id = ContentHash.Empty;
+			fg.Schema = 1122;
+			fg.Groups = new List<INodeGroup>();
+			fg.Aggregates = new List<IAggregate>();
+			fg.Labels = new List<ILabel>();
+
+			Template = await templateCollection.GetOrAddAsync(new TemplateConfig { Name = "Test template" });
+			Graph = await graphCollection.AddAsync(Template, null);
+
+			TemplateRefId1 = new TemplateId("template1");
+			TemplateRefId2 = new TemplateId("template2");
+
+			List<TemplateRefConfig> templates = new List<TemplateRefConfig>();
+			templates.Add(new TemplateRefConfig { Id = TemplateRefId1, Name = "Test Template" });
+			templates.Add(new TemplateRefConfig { Id = TemplateRefId2, Name = "Test Template" });
+
+			List<TabConfig> tabs = new List<TabConfig>();
+			tabs.Add(new TabConfig { Title = "foo", Templates = new List<TemplateId> { TemplateRefId1, TemplateRefId2 } });
+
+			Dictionary<string, AgentConfig> agentTypes = new()
+			{
+				{ "Win64", new() { Pool = new PoolId(PoolName) } }
+			};
+
+			StreamId streamId = new StreamId("ue5-main");
+			StreamConfig streamConfig = new StreamConfig { Id = streamId, Name = "//UE5/Main", Tabs = tabs, Templates = templates, AgentTypes = agentTypes };
+
+			ProjectId projectId = new ProjectId("ue5");
+			ProjectConfig projectConfig = new ProjectConfig { Id = projectId, Name = "UE5", Streams = new List<StreamConfig> { streamConfig } };
+
+			BuildConfig buildConfig = new BuildConfig();
+			buildConfig.Projects = new List<ProjectConfig> { projectConfig };
+
+			GlobalConfig globalConfig = new GlobalConfig();
+			globalConfig.Plugins.AddBuildConfig(buildConfig);
+			globalConfig.PostLoad(serverSettings, pluginCollection.LoadedPlugins);
+			configService.OverrideConfig(globalConfig);
+
+			StreamId = streamId;
+			StreamConfig = streamConfig;
+
+			Job1 = await jobService.CreateJobAsync(
+				jobId: JobId.Parse("5f283932841e7fdbcafb6ab5"),
+				streamConfig: streamConfig,
+				templateRefId: TemplateRefId1,
+				templateHash: Template.Hash,
+				graph: Graph,
+				name: "hello1",
+				change: 1000001,
+				codeChange: 1000002,
+				new CreateJobOptions { PreflightChange = 1001 }
+			);
+			Job1 = (await jobService.GetJobAsync(Job1.Id))!;
+
+			Job2 = await jobService.CreateJobAsync(
+				jobId: JobId.Parse("5f69ea1b68423e921b035106"),
+				streamConfig: streamConfig,
+				templateRefId: new TemplateId("template-id-1"),
+				templateHash: ContentHash.MD5("made-up-template-hash"),
+				graph: fg,
+				name: "hello2",
+				change: 2000001,
+				codeChange: 2000002,
+				new CreateJobOptions()
+			);
+			Job2 = (await jobService.GetJobAsync(Job2.Id))!;
+
+			Agent1Name = "testAgent1";
+			Agent1 = await agentService.CreateAgentAsync(Agent1Name, false, "");
+		}
+
+		private class FixtureGraph : IGraph
+		{
+			public ContentHash Id { get; set; } = ContentHash.Empty;
+			public int Schema { get; set; }
+			public IReadOnlyList<INodeGroup> Groups { get; set; } = null!;
+			public IReadOnlyList<IAggregate> Aggregates { get; set; } = null!;
+			public IReadOnlyList<ILabel> Labels { get; set; } = null!;
+			public IReadOnlyList<IGraphArtifact> Artifacts { get; set; } = null!;
+		}
+	}
+}
\ No newline at end of file
diff --git a/Engine/Source/Programs/Horde/Plugins/Build/HordeServer.Build.Tests/HordeServer.Build.Tests.csproj b/Engine/Source/Programs/Horde/Plugins/Build/HordeServer.Build.Tests/HordeServer.Build.Tests.csproj
new file mode 100644
--- /dev/null
+++ b/Engine/Source/Programs/Horde/Plugins/Build/HordeServer.Build.Tests/HordeServer.Build.Tests.csproj
@@ -0,0 +1,40 @@
+<Project Sdk="Microsoft.NET.Sdk">
+
+  <PropertyGroup>
+    <TargetFramework>net8.0</TargetFramework>
+    <ImplicitUsings>enable</ImplicitUsings>
+    <Nullable>enable</Nullable>
+
+    <IsPackable>false</IsPackable>
+    <IsTestProject>true</IsTestProject>
+    <RootNamespace>HordeServer.Tests</RootNamespace>
+  </PropertyGroup>
+
+  <ItemGroup>
+    <PackageReference Include="coverlet.collector" Version="6.0.0" />
+    <PackageReference Include="Microsoft.NET.Test.Sdk" Version="17.8.0" />
+    <PackageReference Include="MSTest.TestAdapter" Version="3.1.1" />
+    <PackageReference Include="MSTest.TestFramework" Version="3.1.1" />
+    <PackageReference Include="Microsoft.AspNetCore.Mvc.Testing" Version="6.0.25" />
+    <PackageReference Include="Moq" Version="4.20.70" />
+    <PackageReference Include="Microsoft.Extensions.Caching.Memory" Version="8.0.0" />
+  </ItemGroup>
+
+  <ItemGroup>
+    <ProjectReference Include="..\HordeServer.Build\HordeServer.Build.csproj" />
+    <ProjectReference Include="..\..\Compute\HordeServer.Compute\HordeServer.Compute.csproj" />
+    <ProjectReference Include="..\..\..\HordeAgent.Tests\HordeAgent.Tests.csproj" />
+    <ProjectReference Include="..\..\..\HordeServer.Tests.Shared\HordeServer.Tests.Shared.csproj" />
+    <ProjectReference Include="..\..\..\HordeAgent\HordeAgent.csproj">
+      <Aliases>HordeAgent</Aliases>
+    </ProjectReference>
+    <ProjectReference Include="..\..\..\Drivers\JobDriver\JobDriver.csproj">
+      <Aliases>JobDriver</Aliases>
+    </ProjectReference>
+  </ItemGroup>
+
+  <ItemGroup>
+    <Using Include="Microsoft.VisualStudio.TestTools.UnitTesting" />
+  </ItemGroup>
+
+</Project>
diff --git a/Engine/Source/Programs/Horde/Plugins/Build/HordeServer.Build.Tests/Server/RpcServiceTest.cs b/Engine/Source/Programs/Horde/Plugins/Build/HordeServer.Build.Tests/Server/RpcServiceTest.cs
new file mode 100644
--- /dev/null
+++ b/Engine/Source/Programs/Horde/Plugins/Build/HordeServer.Build.Tests/Server/RpcServiceTest.cs
@@ -0,0 +1,417 @@
+// Copyright Epic Games, Inc. All Rights Reserved.
+
+using System;
+using System.Collections.Generic;
+using System.Linq;
+using System.Reflection;
+using System.Security.Claims;
+using System.Threading;
+using System.Threading.Channels;
+using System.Threading.Tasks;
+using EpicGames.Horde.Agents;
+using EpicGames.Horde.Agents.Pools;
+using Grpc.Core;
+using HordeCommon.Rpc;
+using HordeCommon.Rpc.Messages;
+using HordeServer.Agents;
+using HordeServer.Server;
+using HordeServer.Utilities;
+using Microsoft.AspNetCore.Http;
+using Microsoft.AspNetCore.Http.Features;
+using Microsoft.VisualStudio.TestTools.UnitTesting;
+
+namespace HordeServer.Tests.Server
+{
+	using ISession = Microsoft.AspNetCore.Http.ISession;
+
+	sealed class HttpContextStub : HttpContext
+	{
+		public override ConnectionInfo Connection { get; } = null!;
+		public override IFeatureCollection Features { get; } = null!;
+		public override IDictionary<object, object?> Items { get; set; } = null!;
+		public override HttpRequest Request { get; } = null!;
+		public override CancellationToken RequestAborted { get; set; }
+		public override IServiceProvider RequestServices { get; set; } = null!;
+		public override HttpResponse Response { get; } = null!;
+		public override ISession Session { get; set; } = null!;
+		public override string TraceIdentifier { get; set; } = null!;
+		public override ClaimsPrincipal User { get; set; }
+		public override WebSocketManager WebSockets { get; } = null!;
+
+		public HttpContextStub(Claim roleClaimType)
+		{
+			User = new ClaimsPrincipal(new ClaimsIdentity(new List<Claim>
+			{
+				roleClaimType
+			}, "TestAuthType"));
+		}
+
+		public HttpContextStub(ClaimsPrincipal user)
+		{
+			User = user;
+		}
+
+		public override void Abort()
+		{
+			throw new NotImplementedException();
+		}
+	}
+
+	public class ServerCallContextStub : ServerCallContext
+	{
+		// Copied from ServerCallContextExtensions.cs in Grpc.Core
+		const string HttpContextKey = "__HttpContext";
+
+		protected override string MethodCore { get; } = null!;
+		protected override string HostCore { get; } = null!;
+		protected override string PeerCore { get; } = null!;
+		protected override DateTime DeadlineCore { get; } = DateTime.Now.AddHours(24);
+		protected override Metadata RequestHeadersCore { get; } = null!;
+		protected override CancellationToken CancellationTokenCore => _cancellationToken;
+		protected override Metadata ResponseTrailersCore { get; } = null!;
+		protected override Status StatusCore { get; set; }
+		protected override WriteOptions? WriteOptionsCore { get; set; } = null!;
+		protected override AuthContext AuthContextCore { get; } = null!;
+
+		private CancellationToken _cancellationToken;
+
+		public static ServerCallContext ForAdminWithAgentSessionId(string agentSessionId)
+		{
+			return new ServerCallContextStub(new ClaimsPrincipal(new ClaimsIdentity(new List<Claim>
+			{
+				HordeClaims.AdminClaim.ToClaim(),
+				new Claim(HordeClaimTypes.AgentSessionId, agentSessionId),
+			}, "TestAuthType")));
+		}
+
+		public static ServerCallContext ForAdmin()
+		{
+			return new ServerCallContextStub(new ClaimsPrincipal(new ClaimsIdentity(new List<Claim>
+			{
+				HordeClaims.AdminClaim.ToClaim()
+			}, "TestAuthType")));
+		}
+
+		public ServerCallContextStub(Claim roleClaimType)
+		{
+			// The GetHttpContext extension falls back to getting the HttpContext from UserState
+			// We can piggyback on that behavior during tests
+			UserState[HttpContextKey] = new HttpContextStub(roleClaimType);
+		}
+
+		public ServerCallContextStub(ClaimsPrincipal user)
+		{
+			// The GetHttpContext extension falls back to getting the HttpContext from UserState
+			// We can piggyback on that behavior during tests
+			UserState[HttpContextKey] = new HttpContextStub(user);
+		}
+
+		public void SetCancellationToken(CancellationToken cancellationToken)
+		{
+			_cancellationToken = cancellationToken;
+		}
+
+		protected override Task WriteResponseHeadersAsyncCore(Metadata responseHeaders)
+		{
+			throw new NotImplementedException();
+		}
+
+		protected override ContextPropagationToken CreatePropagationTokenCore(ContextPropagationOptions? options)
+		{
+			throw new NotImplementedException();
+		}
+	}
+
+	[TestClass]
+	public class RpcServiceTest : TestSetup
+	{
+		private readonly ServerCallContext _adminContext = new ServerCallContextStub(HordeClaims.AdminClaim.ToClaim());
+
+		class RpcServiceInvoker : CallInvoker
+		{
+			private readonly RpcService _rpcService;
+			private readonly ServerCallContext _serverCallContext;
+
+			public RpcServiceInvoker(RpcService rpcService, ServerCallContext serverCallContext)
+			{
+				_rpcService = rpcService;
+				_serverCallContext = serverCallContext;
+			}
+
+			public override TResponse BlockingUnaryCall<TRequest, TResponse>(Method<TRequest, TResponse> method, string? host, CallOptions options, TRequest request)
+			{
+				throw new NotImplementedException("Blocking calls are not supported! Method " + method.FullName);
+			}
+
+			public override AsyncUnaryCall<TResponse> AsyncUnaryCall<TRequest, TResponse>(Method<TRequest, TResponse> method, string? host, CallOptions options, TRequest request)
+			{
+				MethodInfo methodInfo = GetMethod(method.Name);
+				Task<TResponse> res = (methodInfo.Invoke(_rpcService, new object[] { request, _serverCallContext }) as Task<TResponse>)!;
+				return new AsyncUnaryCall<TResponse>(res, null!, null!, null!, null!, null!);
+			}
+
+			public override AsyncServerStreamingCall<TResponse> AsyncServerStreamingCall<TRequest, TResponse>(Method<TRequest, TResponse> method, string? host, CallOptions options,
+				TRequest request)
+			{
+				Console.WriteLine($"RpcServiceInvoker.AsyncServerStreamingCall(method={method.FullName} request={request})");
+				throw new NotImplementedException();
+			}
+
+			public override AsyncClientStreamingCall<TRequest, TResponse> AsyncClientStreamingCall<TRequest, TResponse>(Method<TRequest, TResponse> method, string? host, CallOptions options)
+			{
+				Console.WriteLine($"RpcServiceInvoker.AsyncClientStreamingCall(method={method.FullName})");
+				throw new NotImplementedException();
+			}
+
+			public override AsyncDuplexStreamingCall<TRequest, TResponse> AsyncDuplexStreamingCall<TRequest, TResponse>(Method<TRequest, TResponse> method, string? host, CallOptions options)
+			{
+				Console.WriteLine($"RpcServiceInvoker.AsyncDuplexStreamingCall(method={method.FullName})");
+
+				GrpcDuplexStreamHandler<TRequest> requestStream = new GrpcDuplexStreamHandler<TRequest>(_serverCallContext);
+				GrpcDuplexStreamHandler<TResponse> responseStream = new GrpcDuplexStreamHandler<TResponse>(_serverCallContext);
+
+				MethodInfo methodInfo = GetMethod(method.Name);
+				Task methodTask = (methodInfo.Invoke(_rpcService, new object[] { requestStream, responseStream, _serverCallContext }) as Task)!;
+				methodTask.ContinueWith(t => { Console.Error.WriteLine($"Uncaught exception in {method.Name}: {t}"); }, CancellationToken.None, TaskContinuationOptions.OnlyOnFaulted, TaskScheduler.Default);
+
+				return new AsyncDuplexStreamingCall<TRequest, TResponse>(requestStream, responseStream, null!, null!, null!, null!);
+			}
+
+			private MethodInfo GetMethod(string methodName)
+			{
+				MethodInfo? method = _rpcService.GetType().GetMethod(methodName);
+				if (method == null)
+				{
+					throw new ArgumentException($"Method {methodName} not found in RpcService");
+				}
+
+				return method;
+			}
+		}
+
+		/// <summary>
+		/// Combines and cross-writes streams for a duplex streaming call in gRPC
+		/// </summary>
+		/// <typeparam name="T"></typeparam>
+		class GrpcDuplexStreamHandler<T> : IServerStreamWriter<T>, IClientStreamWriter<T>, IAsyncStreamReader<T> where T : class
+		{
+			private readonly ServerCallContext _serverCallContext;
+			private readonly Channel<T> _channel;
+
+			public WriteOptions? WriteOptions { get; set; }
+
+			public GrpcDuplexStreamHandler(ServerCallContext serverCallContext)
+			{
+				_channel = System.Threading.Channels.Channel.CreateUnbounded<T>();
+
+				_serverCallContext = serverCallContext;
+			}
+
+			public void Complete()
+			{
+				_channel.Writer.Complete();
+			}
+
+			public IAsyncEnumerable<T> ReadAllAsync()
+			{
+				return _channel.Reader.ReadAllAsync();
+			}
+
+			public async Task<T?> ReadNextAsync()
+			{
+				if (await _channel.Reader.WaitToReadAsync())
+				{
+					_channel.Reader.TryRead(out T? message);
+					return message;
+				}
+				else
+				{
+					return null;
+				}
+			}
+
+			public Task WriteAsync(T message)
+			{
+				if (_serverCallContext.CancellationToken.IsCancellationRequested)
+				{
+					return Task.FromCanceled(_serverCallContext.CancellationToken);
+				}
+
+				if (!_channel.Writer.TryWrite(message))
+				{
+					throw new InvalidOperationException("Unable to write message.");
+				}
+
+				return Task.CompletedTask;
+			}
+
+			public Task CompleteAsync()
+			{
+				Complete();
+				return Task.CompletedTask;
+			}
+
+			public async Task<bool> MoveNext(CancellationToken cancellationToken)
+			{
+				_serverCallContext.CancellationToken.ThrowIfCancellationRequested();
+
+				if (await _channel.Reader.WaitToReadAsync(cancellationToken))
+				{
+					if (_channel.Reader.TryRead(out T? message))
+					{
+						Current = message;
+						return true;
+					}
+				}
+
+				Current = null!;
+				return false;
+			}
+
+			public T Current { get; private set; } = null!;
+		}
+
+		public RpcServiceTest()
+		{
+			UpdateConfig(x => x.Plugins.GetComputeConfig().Pools.Clear());
+		}
+
+		[TestMethod]
+		public async Task CreateSessionTestAsync()
+		{
+			RpcCreateSessionRequest req = new RpcCreateSessionRequest();
+			await Assert.ThrowsExceptionAsync<StructuredRpcException>(() => RpcService.CreateSession(req, _adminContext));
+
+			req.Id = new AgentId("MyName").ToString();
+			await Assert.ThrowsExceptionAsync<StructuredRpcException>(() => RpcService.CreateSession(req, _adminContext));
+
+			req.Capabilities = new RpcAgentCapabilities();
+			RpcCreateSessionResponse res = await RpcService.CreateSession(req, _adminContext);
+
+			Assert.AreEqual("MYNAME", res.AgentId);
+			// TODO: Check Token, ExpiryTime, SessionId 
+		}
+
+		[TestMethod]
+		public async Task AgentJoinsPoolThroughPropertiesAsync()
+		{
+			RpcCreateSessionRequest req = new() { Id = new AgentId("bogusAgentName").ToString(), Capabilities = new RpcAgentCapabilities() };
+			req.Capabilities.Properties.Add($"{KnownPropertyNames.RequestedPools}=fooPool,barPool");
+			RpcCreateSessionResponse res = await RpcService.CreateSession(req, _adminContext);
+
+			IAgent agent = (await AgentService.GetAgentAsync(new AgentId(res.AgentId)))!;
+			CollectionAssert.AreEquivalent(new List<PoolId> { new("fooPool"), new("barPool") }, agent.Pools.ToList());
+
+			// Connect a second time, when the agent has already been created
+			req = new() { Id = new AgentId("bogusAgentName").ToString(), Capabilities = new RpcAgentCapabilities() };
+			req.Capabilities.Properties.Add($"{KnownPropertyNames.RequestedPools}=bazPool");
+			res = await RpcService.CreateSession(req, _adminContext);
+
+			agent = (await AgentService.GetAgentAsync(new AgentId(res.AgentId)))!;
+			CollectionAssert.AreEquivalent(new List<PoolId> { new("fooPool"), new("barPool"), new("bazPool") }, agent.Pools.ToList());
+		}
+
+		[TestMethod]
+		public async Task PropertiesFromAgentCapabilitiesAsync()
+		{
+			RpcCreateSessionRequest req = new() { Id = new AgentId("bogusAgentName").ToString(), Capabilities = new RpcAgentCapabilities() };
+			req.Capabilities.Properties.Add("fooKey=barValue");
+			RpcCreateSessionResponse res = await RpcService.CreateSession(req, _adminContext);
+			IAgent agent = (await AgentService.GetAgentAsync(new AgentId(res.AgentId)))!;
+			Assert.IsTrue(agent.Properties.Contains("fooKey=barValue"));
+		}
+
+		[TestMethod]
+		public async Task PropertiesFromDeviceCapabilitiesAsync()
+		{
+			RpcCreateSessionRequest req = new() { Id = new AgentId("bogusAgentName").ToString(), Capabilities = new RpcAgentCapabilities() };
+			req.Capabilities.Devices.Add(new RpcDeviceCapabilities { Handle = "someHandle", Properties = { "foo=bar" } });
+			RpcCreateSessionResponse res = await RpcService.CreateSession(req, _adminContext);
+			IAgent agent = (await AgentService.GetAgentAsync(new AgentId(res.AgentId)))!;
+			Assert.IsTrue(agent.Properties.Contains("foo=bar"));
+		}
+
+		[TestMethod]
+		public async Task KnownPropertiesAreSetAsResourcesAsync()
+		{
+			RpcCreateSessionRequest req = new() { Id = new AgentId("bogusAgentName").ToString(), Capabilities = new RpcAgentCapabilities() };
+			req.Capabilities.Devices.Add(new RpcDeviceCapabilities { Handle = "someHandle", Properties = { $"{KnownPropertyNames.LogicalCores}=10" } });
+			RpcCreateSessionResponse res = await RpcService.CreateSession(req, _adminContext);
+			IAgent agent = (await AgentService.GetAgentAsync(new AgentId(res.AgentId)))!;
+			Assert.AreEqual(10, agent.Resources[KnownPropertyNames.LogicalCores]);
+		}
+
+		[TestMethod]
+		public async Task UpdateSessionTestAsync()
+		{
+			RpcCreateSessionRequest createReq = new RpcCreateSessionRequest
+			{
+				Id = new AgentId("UpdateSessionTest1").ToString(),
+				Capabilities = new RpcAgentCapabilities()
+			};
+			RpcCreateSessionResponse createRes = await RpcService.CreateSession(createReq, _adminContext);
+			string agentId = createRes.AgentId;
+			string sessionId = createRes.SessionId;
+
+			TestAsyncStreamReader<RpcUpdateSessionRequest> requestStream =
+				new TestAsyncStreamReader<RpcUpdateSessionRequest>(_adminContext);
+			TestServerStreamWriter<RpcUpdateSessionResponse> responseStream =
+				new TestServerStreamWriter<RpcUpdateSessionResponse>(_adminContext);
+			Task call = RpcService.UpdateSession(requestStream, responseStream, _adminContext);
+
+			requestStream.AddMessage(new RpcUpdateSessionRequest { AgentId = "does-not-exist", SessionId = sessionId });
+			StructuredRpcException re = await Assert.ThrowsExceptionAsync<StructuredRpcException>(() => call);
+			Assert.AreEqual(StatusCode.NotFound, re.StatusCode);
+			Assert.IsTrue(re.Message.Contains("Invalid agent name", StringComparison.OrdinalIgnoreCase));
+		}
+
+		[TestMethod]
+		public async Task FinishBatchTestAsync()
+		{
+			RpcCreateSessionRequest createReq = new RpcCreateSessionRequest
+			{
+				Id = new AgentId("UpdateSessionTest1").ToString(),
+				Capabilities = new RpcAgentCapabilities()
+			};
+			RpcCreateSessionResponse createRes = await RpcService.CreateSession(createReq, _adminContext);
+			string agentId = createRes.AgentId;
+			string sessionId = createRes.SessionId;
+
+			TestAsyncStreamReader<RpcUpdateSessionRequest> requestStream =
+				new TestAsyncStreamReader<RpcUpdateSessionRequest>(_adminContext);
+			TestServerStreamWriter<RpcUpdateSessionResponse> responseStream =
+				new TestServerStreamWriter<RpcUpdateSessionResponse>(_adminContext);
+			Task call = RpcService.UpdateSession(requestStream, responseStream, _adminContext);
+
+			requestStream.AddMessage(new RpcUpdateSessionRequest { AgentId = "does-not-exist", SessionId = sessionId });
+			StructuredRpcException re = await Assert.ThrowsExceptionAsync<StructuredRpcException>(() => call);
+			Assert.AreEqual(StatusCode.NotFound, re.StatusCode);
+			Assert.IsTrue(re.Message.Contains("Invalid agent name", StringComparison.OrdinalIgnoreCase));
+		}
+
+		/*
+		[TestMethod]
+		public async Task UploadSoftwareAsync()
+		{
+			TestSetup TestSetup = await GetTestSetup();
+			
+			MemoryStream OutputStream = new MemoryStream();
+			using (ZipArchive ZipFile = new ZipArchive(OutputStream, ZipArchiveMode.Create, false))
+			{
+				string TempFilename = Path.GetTempFileName();
+				File.WriteAllText(TempFilename, "{\"Horde\": {\"Version\": \"myVersion\"}}");
+				ZipFile.CreateEntryFromFile(TempFilename, "appsettings.json");
+			}
+
+			ByteString Data = ByteString.CopyFrom(OutputStream.ToArray());
+			UploadSoftwareRequest Req = new UploadSoftwareRequest { Channel = "boguschannel", Data = Data };
+
+			UploadSoftwareResponse Res1 = await TestSetup.RpcService.UploadSoftware(Req, AdminContext);
+			Assert.AreEqual("r1", Res1.Version);
+			
+			UploadSoftwareResponse Res2 = await TestSetup.RpcService.UploadSoftware(Req, AdminContext);
+			Assert.AreEqual("r2", Res2.Version);
+		}
+		*/
+	}
+}
\ No newline at end of file
diff --git a/Engine/Source/Programs/Horde/Plugins/Build/HordeServer.Build.Tests/Streams/StreamServiceTest.cs b/Engine/Source/Programs/Horde/Plugins/Build/HordeServer.Build.Tests/Streams/StreamServiceTest.cs
new file mode 100644
--- /dev/null
+++ b/Engine/Source/Programs/Horde/Plugins/Build/HordeServer.Build.Tests/Streams/StreamServiceTest.cs
@@ -0,0 +1,39 @@
+// Copyright Epic Games, Inc. All Rights Reserved.
+
+using System;
+using System.Globalization;
+using System.Threading.Tasks;
+using HordeServer.Streams;
+using Microsoft.VisualStudio.TestTools.UnitTesting;
+
+namespace HordeServer.Tests.Streams
+{
+	[TestClass]
+	public class StreamServiceTests : TestSetup
+	{
+		[TestMethod]
+		public async Task PausingAsync()
+		{
+			Fixture fixture = await CreateFixtureAsync();
+
+			IStream stream = (await StreamCollection.GetAsync(fixture!.StreamConfig!))!;
+			Assert.IsFalse(stream.IsPaused(DateTime.UtcNow));
+			Assert.IsNull(stream.PausedUntil);
+			Assert.IsNull(stream.PauseComment);
+
+			DateTime pausedUntil = DateTime.UtcNow.AddHours(1);
+			await StreamCollection.UpdatePauseStateAsync(stream, newPausedUntil: pausedUntil, newPauseComment: "mycomment");
+			stream = (await StreamCollection.GetAsync(fixture!.StreamConfig!))!;
+			// Comparing by string to avoid comparing exact milliseconds as those are not persisted in MongoDB fields
+			Assert.IsTrue(stream.IsPaused(DateTime.UtcNow));
+			Assert.AreEqual(pausedUntil.ToString(CultureInfo.InvariantCulture), stream.PausedUntil!.Value.ToString(CultureInfo.InvariantCulture));
+			Assert.AreEqual("mycomment", stream.PauseComment);
+
+			await StreamCollection.UpdatePauseStateAsync(stream, newPausedUntil: null, newPauseComment: null);
+			stream = (await StreamCollection.GetAsync(fixture!.StreamConfig!))!;
+			Assert.IsFalse(stream.IsPaused(DateTime.UtcNow));
+			Assert.IsNull(stream.PausedUntil);
+			Assert.IsNull(stream.PauseComment);
+		}
+	}
+}
diff --git a/Engine/Source/Programs/Horde/Plugins/Build/HordeServer.Build.Tests/TestSetup.cs b/Engine/Source/Programs/Horde/Plugins/Build/HordeServer.Build.Tests/TestSetup.cs
new file mode 100644
--- /dev/null
+++ b/Engine/Source/Programs/Horde/Plugins/Build/HordeServer.Build.Tests/TestSetup.cs
@@ -0,0 +1,349 @@
+// Copyright Epic Games, Inc. All Rights Reserved.
+
+using System;
+using System.Collections.Generic;
+using System.Net;
+using System.Net.Sockets;
+using System.Security.Claims;
+using System.Threading.Tasks;
+using Amazon.AutoScaling;
+using Amazon.CloudWatch;
+using Amazon.EC2;
+using EpicGames.Horde.Agents;
+using EpicGames.Horde.Agents.Pools;
+using HordeServer.Accounts;
+using HordeServer.Agents;
+using HordeServer.Agents.Enrollment;
+using HordeServer.Agents.Fleet;
+using HordeServer.Agents.Leases;
+using HordeServer.Agents.Pools;
+using HordeServer.Agents.Relay;
+using HordeServer.Agents.Sessions;
+using HordeServer.Agents.Telemetry;
+using HordeServer.Agents.Utilization;
+using HordeServer.Artifacts;
+using HordeServer.Auditing;
+using HordeServer.Commits;
+using HordeServer.Compute;
+using HordeServer.Dashboard;
+using HordeServer.Devices;
+using HordeServer.Issues;
+using HordeServer.Jobs;
+using HordeServer.Jobs.Bisect;
+using HordeServer.Jobs.Graphs;
+using HordeServer.Jobs.Schedules;
+using HordeServer.Jobs.Templates;
+using HordeServer.Jobs.TestData;
+using HordeServer.Jobs.Timing;
+using HordeServer.Logs;
+using HordeServer.Notifications;
+using HordeServer.Perforce;
+using HordeServer.Server;
+using HordeServer.Storage;
+using HordeServer.Streams;
+using HordeServer.Tasks;
+using HordeServer.Tests.Stubs.Services;
+using HordeServer.Tools;
+using HordeServer.Ugs;
+using HordeServer.Users;
+using HordeServer.Utilities;
+using Microsoft.AspNetCore.Http;
+using Microsoft.AspNetCore.Mvc;
+using Microsoft.Extensions.Configuration;
+using Microsoft.Extensions.DependencyInjection;
+using Microsoft.Extensions.Logging;
+using Microsoft.VisualStudio.TestTools.UnitTesting;
+using Moq;
+
+namespace HordeServer.Tests
+{
+	/// <summary>
+	/// Handles set up of collections, services, fixtures etc during testing
+	///
+	/// Easier to pass all these things around in a single object.
+	/// </summary>
+	public class TestSetup : ServerTestSetup
+	{
+		public IGraphCollection GraphCollection => ServiceProvider.GetRequiredService<IGraphCollection>();
+		public INotificationTriggerCollection NotificationTriggerCollection => ServiceProvider.GetRequiredService<INotificationTriggerCollection>();
+		public IStreamCollection StreamCollection => ServiceProvider.GetRequiredService<IStreamCollection>();
+		public IJobCollection JobCollection => ServiceProvider.GetRequiredService<IJobCollection>();
+		public IAgentCollection AgentCollection => ServiceProvider.GetRequiredService<IAgentCollection>();
+		public IJobStepRefCollection JobStepRefCollection => ServiceProvider.GetRequiredService<IJobStepRefCollection>();
+		public IJobTimingCollection JobTimingCollection => ServiceProvider.GetRequiredService<IJobTimingCollection>();
+		public IUgsMetadataCollection UgsMetadataCollection => ServiceProvider.GetRequiredService<IUgsMetadataCollection>();
+		public IIssueCollection IssueCollection => ServiceProvider.GetRequiredService<IIssueCollection>();
+		public IPoolCollection PoolCollection => ServiceProvider.GetRequiredService<IPoolCollection>();
+		public ILeaseCollection LeaseCollection => ServiceProvider.GetRequiredService<ILeaseCollection>();
+		public ILogCollection LogCollection => ServiceProvider.GetRequiredService<ILogCollection>();
+		public ISessionCollection SessionCollection => ServiceProvider.GetRequiredService<ISessionCollection>();
+		public ITestDataCollection TestDataCollection => ServiceProvider.GetRequiredService<ITestDataCollection>();
+		public IUserCollection UserCollection => ServiceProvider.GetRequiredService<IUserCollection>();
+		public IDeviceCollection DeviceCollection => ServiceProvider.GetRequiredService<IDeviceCollection>();
+		public IDashboardPreviewCollection DashboardPreviewCollection => ServiceProvider.GetRequiredService<IDashboardPreviewCollection>();
+		public IBisectTaskCollection BisectTaskCollection => ServiceProvider.GetRequiredService<IBisectTaskCollection>();
+
+		public FleetService FleetService => ServiceProvider.GetRequiredService<FleetService>();
+		public AgentService AgentService => ServiceProvider.GetRequiredService<AgentService>();
+		public AgentRelayService AgentRelayService => ServiceProvider.GetRequiredService<AgentRelayService>();
+		public ICommitService CommitService => ServiceProvider.GetRequiredService<ICommitService>();
+		public GlobalsService GlobalsService => ServiceProvider.GetRequiredService<GlobalsService>();
+		public ITemplateCollection TemplateCollection => ServiceProvider.GetRequiredService<ITemplateCollection>();
+		internal PerforceServiceStub PerforceService => (PerforceServiceStub)ServiceProvider.GetRequiredService<IPerforceService>();
+		public ISubscriptionCollection SubscriptionCollection => ServiceProvider.GetRequiredService<ISubscriptionCollection>();
+		public INotificationService NotificationService => ServiceProvider.GetRequiredService<INotificationService>();
+		public IssueService IssueService => ServiceProvider.GetRequiredService<IssueService>();
+		public JobTaskSource JobTaskSource => ServiceProvider.GetRequiredService<JobTaskSource>();
+		public JobService JobService => ServiceProvider.GetRequiredService<JobService>();
+		public RpcService RpcService => ServiceProvider.GetRequiredService<RpcService>();
+		public PoolService PoolService => ServiceProvider.GetRequiredService<PoolService>();
+		public ScheduleService ScheduleService => ServiceProvider.GetRequiredService<ScheduleService>();
+		public DeviceService DeviceService => ServiceProvider.GetRequiredService<DeviceService>();
+		public StorageService StorageService => ServiceProvider.GetRequiredService<StorageService>();
+		public TestDataService TestDataService => ServiceProvider.GetRequiredService<TestDataService>();
+		public ComputeService ComputeService => ServiceProvider.GetRequiredService<ComputeService>();
+
+		public JobsController JobsController => GetJobsController();
+		public AgentsController AgentsController => GetAgentsController();
+		public PoolsController PoolsController => GetPoolsController();
+		public LeasesController LeasesController => GetLeasesController();
+		public DevicesController DevicesController => GetDevicesController();
+		public DashboardController DashboardController => GetDashboardController();
+		public TestDataController TestDataController => GetTestDataController();
+		public BisectTasksController BisectTasksController => GetBisectTasksController();
+
+		public TestSetup()
+		{
+			AddPlugin<AnalyticsPlugin>();
+			AddPlugin<BuildPlugin>();
+			AddPlugin<ComputePlugin>();
+			AddPlugin<StoragePlugin>();
+		}
+
+		protected override void ConfigureServices(IServiceCollection services)
+		{
+			base.ConfigureServices(services);
+
+			IConfiguration config = new ConfigurationBuilder().Build();
+			services.Configure<ServerSettings>(ConfigureSettings);
+			services.AddSingleton<IConfiguration>(config);
+
+			services.AddHttpClient<RpcService>();
+
+			services.AddSingleton<IAuditLog<AgentId>>(sp => sp.GetRequiredService<IAuditLogFactory<AgentId>>().Create("Agents.Log", "AgentId"));
+
+			services.AddSingleton<IAccountCollection, AccountCollection>();
+			services.AddSingleton<IAgentCollection, AgentCollection>();
+			services.AddSingleton<IAgentTelemetryCollection, AgentTelemetryCollection>();
+			services.AddSingleton<IArtifactCollection, ArtifactCollection>();
+			services.AddSingleton<ICommitService, CommitService>();
+			services.AddSingleton<IGraphCollection, GraphCollection>();
+			services.AddSingleton<IIssueCollection, IssueCollection>();
+			services.AddSingleton<IJobCollection, JobCollection>();
+			services.AddSingleton<IJobStepRefCollection, JobStepRefCollection>();
+			services.AddSingleton<IJobTimingCollection, JobTimingCollection>();
+			services.AddSingleton<ILeaseCollection, LeaseCollection>();
+			services.AddSingleton<ILogCollection, LogCollection>();
+			services.AddSingleton<INotificationTriggerCollection, NotificationTriggerCollection>();
+			services.AddSingleton<IPoolCollection, PoolCollection>();
+			services.AddSingleton<IBisectTaskCollection, BisectTaskCollection>();
+			services.AddSingleton<ISessionCollection, SessionCollection>();
+			services.AddSingleton<ISubscriptionCollection, SubscriptionCollection>();
+			services.AddSingleton<IStreamCollection, StreamCollection>();
+			services.AddSingleton<ITemplateCollection, TemplateCollection>();
+			services.AddSingleton<ITestDataCollection, TestDataCollection>();
+			services.AddSingleton<IUtilizationDataCollection, UtilizationDataCollection>();
+			services.AddSingleton<ITemplateCollection, TemplateCollection>();
+			services.AddSingleton<IToolCollection, ToolCollection>();
+			services.AddSingleton<IUgsMetadataCollection, UgsMetadataCollection>();
+			services.AddSingleton<IDeviceCollection, DeviceCollection>();
+
+			// Empty mocked object to satisfy basic test runs
+			services.AddSingleton<IAmazonEC2>(sp => new Mock<IAmazonEC2>().Object);
+			services.AddSingleton<IAmazonAutoScaling>(sp => new Mock<IAmazonAutoScaling>().Object);
+			services.AddSingleton<IAmazonCloudWatch>(sp => new Mock<IAmazonCloudWatch>().Object);
+			services.AddSingleton<IFleetManagerFactory, FleetManagerFactory>();
+
+			services.AddSingleton<IPoolSizeStrategyFactory, NoOpPoolSizeStrategyFactory>();
+			services.AddSingleton<IPoolSizeStrategyFactory, JobQueueStrategyFactory>();
+			services.AddSingleton<IPoolSizeStrategyFactory, LeaseUtilizationStrategyFactory>();
+			services.AddSingleton<IPoolSizeStrategyFactory, LeaseUtilizationAwsMetricStrategyFactory>();
+
+			services.AddSingleton<AgentService>();
+			services.AddSingleton(provider => new Lazy<AgentService>(provider.GetRequiredService<AgentService>));
+			services.AddSingleton<AgentRelayService>();
+			services.AddSingleton<AwsAutoScalingLifecycleService>();
+			services.AddSingleton<ArtifactExpirationService>();
+			services.AddSingleton<FleetService>();
+			services.AddSingleton<ConsistencyService>();
+			services.AddSingleton<RequestTrackerService>();
+			services.AddSingleton<GlobalsService>();
+			services.AddSingleton<JobTaskSource>();
+			services.AddSingleton<IssueService>();
+			services.AddSingleton<JobService>();
+			services.AddSingleton<JobExpirationService>();
+			services.AddSingleton<LogTailService>();
+			services.AddSingleton<INotificationService, NotificationService>();
+			services.AddSingleton<IPerforceService, PerforceServiceStub>();
+			services.AddSingleton<PerforceLoadBalancer>();
+			services.AddSingleton<PoolService>();
+			services.AddSingleton<BisectService>();
+			services.AddSingleton<RpcService>();
+			services.AddSingleton<ScheduleService>();
+			services.AddSingleton<DeviceService>();
+			services.AddSingleton<TestDataService>();
+			services.AddSingleton<ComputeService>();
+			services.AddSingleton<EnrollmentService>();
+
+			services.AddSingleton<ConformTaskSource>();
+			services.AddSingleton<ICommitService, CommitService>();
+		}
+
+		public Task<Fixture> CreateFixtureAsync()
+		{
+			return Fixture.CreateAsync(ConfigService, GraphCollection, TemplateCollection, JobService, AgentService, PluginCollection, ServerSettings);
+		}
+
+		private JobsController GetJobsController()
+		{
+			JobsController jobsCtrl = ActivatorUtilities.CreateInstance<JobsController>(ServiceProvider);
+			jobsCtrl.ControllerContext = GetControllerContext();
+			return jobsCtrl;
+		}
+
+		private DevicesController GetDevicesController()
+		{
+			DevicesController devicesCtrl = ActivatorUtilities.CreateInstance<DevicesController>(ServiceProvider);
+			devicesCtrl.ControllerContext = GetControllerContext();
+			return devicesCtrl;
+		}
+
+		private DashboardController GetDashboardController()
+		{
+			DashboardController dashboardCtrl = ActivatorUtilities.CreateInstance<DashboardController>(ServiceProvider);
+			dashboardCtrl.ControllerContext = GetControllerContext();
+			return dashboardCtrl;
+		}
+
+		private TestDataController GetTestDataController()
+		{
+			TestDataController dataCtrl = ActivatorUtilities.CreateInstance<TestDataController>(ServiceProvider);
+			dataCtrl.ControllerContext = GetControllerContext();
+			return dataCtrl;
+		}
+
+		private BisectTasksController GetBisectTasksController()
+		{
+			BisectTasksController bisectCtrl = ActivatorUtilities.CreateInstance<BisectTasksController>(ServiceProvider);
+			bisectCtrl.ControllerContext = GetControllerContext();
+			return bisectCtrl;
+		}
+
+		private AgentsController GetAgentsController()
+		{
+			AgentsController agentCtrl = ActivatorUtilities.CreateInstance<AgentsController>(ServiceProvider);
+			agentCtrl.ControllerContext = GetControllerContext();
+			return agentCtrl;
+		}
+
+		private PoolsController GetPoolsController()
+		{
+			PoolsController controller = ActivatorUtilities.CreateInstance<PoolsController>(ServiceProvider);
+			controller.ControllerContext = GetControllerContext();
+			return controller;
+		}
+
+		private LeasesController GetLeasesController()
+		{
+			LeasesController controller = ActivatorUtilities.CreateInstance<LeasesController>(ServiceProvider);
+			controller.ControllerContext = GetControllerContext();
+			return controller;
+		}
+
+		private static ControllerContext GetControllerContext()
+		{
+			ControllerContext controllerContext = new ControllerContext();
+			controllerContext.HttpContext = new DefaultHttpContext();
+			controllerContext.HttpContext.User = new ClaimsPrincipal(new ClaimsIdentity(
+				new List<Claim> { HordeClaims.AdminClaim.ToClaim() }, "TestAuthType"));
+			return controllerContext;
+		}
+
+		private static int s_agentIdCounter = 1;
+		public Task<IAgent> CreateAgentAsync(IPool pool, bool enabled = true, bool requestShutdown = false, List<string>? properties = null, List<AgentWorkspaceInfo>? workspaces = null, TimeSpan? adjustClockBy = null)
+		{
+			return CreateAgentAsync(pool.Id, enabled, requestShutdown, properties, workspaces, adjustClockBy);
+		}
+
+		/// <summary>
+		/// Helper function for setting up agents to be used in tests
+		/// </summary>
+		/// <param name="poolId">Pool ID which the agent should belong to</param>
+		/// <param name="enabled">Whether set the agent as enabled</param>
+		/// <param name="requestShutdown">Mark it with a request for shutdown</param>
+		/// <param name="properties">Any properties to assign</param>
+		/// <param name="workspaces">Any workspaces to assign</param>
+		/// <param name="adjustClockBy">Time span to temporarily skew the clock when creating the agent</param>
+		/// <param name="awsInstanceId">AWS instance ID for the agent (will be set in properties)</param>
+		/// <param name="lease">A lease to assign the agent</param>
+		/// <param name="ephemeral">Whether the agent is ephemeral</param>
+		/// <returns>A new agent</returns>
+		public async Task<IAgent> CreateAgentAsync(
+			PoolId? poolId,
+			bool enabled = true,
+			bool requestShutdown = false,
+			List<string>? properties = null,
+			List<AgentWorkspaceInfo>? workspaces = null,
+			TimeSpan? adjustClockBy = null,
+			string? awsInstanceId = null,
+			AgentLease? lease = null,
+			bool ephemeral = false)
+		{
+			DateTime now = Clock.UtcNow;
+			if (adjustClockBy != null)
+			{
+				Clock.UtcNow = now + adjustClockBy.Value;
+			}
+
+			Dictionary<string, int> resources = new();
+			List<string> tempProps = new(properties ?? new List<string>());
+			if (awsInstanceId != null)
+			{
+				tempProps.Add(KnownPropertyNames.AwsInstanceId + "=" + awsInstanceId);
+			}
+
+			IAgent? agent = await AgentService.CreateAgentAsync("TestAgent" + s_agentIdCounter++, ephemeral, "");
+			Assert.IsNotNull(agent);
+
+			agent = await agent.TryUpdateAsync(new UpdateAgentOptions { Enabled = enabled, ExplicitPools = poolId != null ? [poolId.Value] : [] });
+			Assert.IsNotNull(agent);
+
+			agent = await AgentService.CreateSessionAsync(agent, AgentStatus.Ok, tempProps, resources, null);
+			Assert.IsNotNull(agent);
+
+			if (workspaces is { Count: > 0 })
+			{
+				await agent.TryUpdateWorkspacesAsync(workspaces, false);
+			}
+
+			if (requestShutdown)
+			{
+				await agent.TryUpdateAsync(new UpdateAgentOptions { RequestShutdown = true });
+			}
+
+			if (lease != null)
+			{
+				await agent.TryAddLeaseAsync(lease);
+			}
+
+			Clock.UtcNow = now;
+			return agent;
+		}
+
+		protected async Task<IPool> CreatePoolAsync(PoolConfig poolConfig)
+		{
+			UpdateConfig(config => config.Plugins.GetComputeConfig().Pools.Add(poolConfig));
+			return await PoolCollection.GetAsync(poolConfig.Id) ?? throw new NotImplementedException();
+		}
+	}
+}
\ No newline at end of file
diff --git a/Engine/Source/Programs/Horde/Plugins/Build/HordeServer.Build/Properties/AssemblyInfo.cs b/Engine/Source/Programs/Horde/Plugins/Build/HordeServer.Build/Properties/AssemblyInfo.cs
--- a/Engine/Source/Programs/Horde/Plugins/Build/HordeServer.Build/Properties/AssemblyInfo.cs
+++ b/Engine/Source/Programs/Horde/Plugins/Build/HordeServer.Build/Properties/AssemblyInfo.cs
@@ -4,4 +4,5 @@ using System.Resources;
 using System.Runtime.CompilerServices;
 
 [assembly: InternalsVisibleTo("HordeServer.Tests")]
+[assembly: InternalsVisibleTo("HordeServer.Build.Tests")]
 [assembly: NeutralResourcesLanguage("en-US")]
diff --git a/Engine/Source/Programs/Horde/Plugins/Compute/HordeServer.Compute/Properties/AssemblyInfo.cs b/Engine/Source/Programs/Horde/Plugins/Compute/HordeServer.Compute/Properties/AssemblyInfo.cs
--- a/Engine/Source/Programs/Horde/Plugins/Compute/HordeServer.Compute/Properties/AssemblyInfo.cs
+++ b/Engine/Source/Programs/Horde/Plugins/Compute/HordeServer.Compute/Properties/AssemblyInfo.cs
@@ -4,5 +4,6 @@ using System.Resources;
 using System.Runtime.CompilerServices;
 
 [assembly: InternalsVisibleTo("HordeServer.Tests")]
+[assembly: InternalsVisibleTo("HordeServer.Build.Tests")]
 [assembly: InternalsVisibleTo("HordeServer.Compute.Tests")]
 [assembly: NeutralResourcesLanguage("en-US")]
diff --git a/Engine/Source/Programs/Horde/Plugins/Storage/HordeServer.Storage.Tests/Storage/BlobStoreTests.cs b/Engine/Source/Programs/Horde/Plugins/Storage/HordeServer.Storage.Tests/Storage/BlobStoreTests.cs
--- a/Engine/Source/Programs/Horde/Plugins/Storage/HordeServer.Storage.Tests/Storage/BlobStoreTests.cs
+++ b/Engine/Source/Programs/Horde/Plugins/Storage/HordeServer.Storage.Tests/Storage/BlobStoreTests.cs
@@ -15,7 +15,7 @@ using Microsoft.VisualStudio.TestTools.UnitTesting;
 namespace HordeServer.Tests.Storage
 {
 	[TestClass]
-	public class BlobStoreTests : ServerServiceTest
+	public class BlobStoreTests : ServerTestSetup
 	{
 		static readonly BlobType s_blobType = new BlobType("{AFDF76A7-4DEE-5333-F5B5-37B8451251CA}", 1);
 
diff --git a/Engine/Source/Programs/Horde/Plugins/Storage/HordeServer.Storage.Tests/Storage/GcServiceTests.cs b/Engine/Source/Programs/Horde/Plugins/Storage/HordeServer.Storage.Tests/Storage/GcServiceTests.cs
--- a/Engine/Source/Programs/Horde/Plugins/Storage/HordeServer.Storage.Tests/Storage/GcServiceTests.cs
+++ b/Engine/Source/Programs/Horde/Plugins/Storage/HordeServer.Storage.Tests/Storage/GcServiceTests.cs
@@ -15,7 +15,7 @@ using Microsoft.VisualStudio.TestTools.UnitTesting;
 namespace HordeServer.Tests.Storage
 {
 	[TestClass]
-	public sealed class GcServiceTests : ServerServiceTest
+	public sealed class GcServiceTests : ServerTestSetup
 	{
 		public StorageService StorageService => ServiceProvider.GetRequiredService<StorageService>();
 
diff --git a/Engine/Source/Programs/Horde/Plugins/Storage/HordeServer.Storage.Tests/Storage/StorageServiceTests.cs b/Engine/Source/Programs/Horde/Plugins/Storage/HordeServer.Storage.Tests/Storage/StorageServiceTests.cs
--- a/Engine/Source/Programs/Horde/Plugins/Storage/HordeServer.Storage.Tests/Storage/StorageServiceTests.cs
+++ b/Engine/Source/Programs/Horde/Plugins/Storage/HordeServer.Storage.Tests/Storage/StorageServiceTests.cs
@@ -11,7 +11,7 @@ using Microsoft.VisualStudio.TestTools.UnitTesting;
 namespace HordeServer.Tests.Storage
 {
 	[TestClass]
-	public class StorageServiceTests : ServerServiceTest
+	public class StorageServiceTests : ServerTestSetup
 	{
 		public StorageService StorageService => ServiceProvider.GetRequiredService<StorageService>();
 
diff --git a/Engine/Source/Programs/Horde/Plugins/Tools/HordeServer.Tools.Tests/Tools/ToolTests.cs b/Engine/Source/Programs/Horde/Plugins/Tools/HordeServer.Tools.Tests/Tools/ToolTests.cs
--- a/Engine/Source/Programs/Horde/Plugins/Tools/HordeServer.Tools.Tests/Tools/ToolTests.cs
+++ b/Engine/Source/Programs/Horde/Plugins/Tools/HordeServer.Tools.Tests/Tools/ToolTests.cs
@@ -19,7 +19,7 @@ using HordeServer.Plugins;
 namespace HordeServer.Tests.Tools
 {
 	[TestClass]
-	public class ToolTests : ServerServiceTest
+	public class ToolTests : ServerTestSetup
 	{
 		readonly ToolId _toolId = new ToolId("ugs");
 
diff --git a/Engine/Source/Programs/Horde/Scripts/test.sh b/Engine/Source/Programs/Horde/Scripts/test.sh
--- a/Engine/Source/Programs/Horde/Scripts/test.sh
+++ b/Engine/Source/Programs/Horde/Scripts/test.sh
@@ -16,6 +16,8 @@ testProjects=(
 	"Source/Programs/Shared/EpicGames.Serialization.Tests/EpicGames.Serialization.Tests.csproj"
 	"Source/Programs/Horde/HordeServer.Tests/HordeServer.Tests.csproj"
 	"Source/Programs/Horde/Plugins/Analytics/HordeServer.Analytics.Tests/HordeServer.Analytics.Tests.csproj"
+	"Source/Programs/Horde/Plugins/Build/HordeServer.Build.Tests/HordeServer.Build.Tests.csproj"
+	"Source/Programs/Horde/Plugins/Compute/HordeServer.Compute.Tests/HordeServer.Compute.Tests.csproj"
 	"Source/Programs/Horde/Plugins/Ddc/HordeServer.Ddc.Tests/HordeServer.Ddc.Tests.csproj"
 	"Source/Programs/Horde/Plugins/Storage/HordeServer.Storage.Tests/HordeServer.Storage.Tests.csproj"
 	"Source/Programs/Horde/Plugins/Tools/HordeServer.Tools.Tests/HordeServer.Tools.Tests.csproj"
