Horde: Fix ACL tests.

diff --git a/Engine/Source/Programs/Horde/HordeServer.Tests.Shared/ServerTestSetup.cs b/Engine/Source/Programs/Horde/HordeServer.Tests.Shared/ServerTestSetup.cs
--- a/Engine/Source/Programs/Horde/HordeServer.Tests.Shared/ServerTestSetup.cs
+++ b/Engine/Source/Programs/Horde/HordeServer.Tests.Shared/ServerTestSetup.cs
@@ -9,6 +9,7 @@ using System.Text;
 using System.Threading.Tasks;
 using EpicGames.Core;
 using HordeCommon;
+using HordeServer.Accounts;
 using HordeServer.Acls;
 using HordeServer.Auditing;
 using HordeServer.Configuration;
@@ -102,6 +103,8 @@ namespace HordeServer.Tests
 			services.AddSingleton<IServerInfo>(serverInfo);
 			services.AddSingleton<IPluginCollection>(_pluginCollection);
 
+			services.AddSingleton<IAccountCollection, AccountCollection>();
+
 			services.AddLogging(builder => { builder.AddConsole().SetMinimumLevel(LogLevel.Debug); });
 			services.AddSingleton<IMemoryCache>(sp => new MemoryCache(new MemoryCacheOptions { }));
 
diff --git a/Engine/Source/Programs/Horde/HordeServer.Tests/Acls/AclTests.cs b/Engine/Source/Programs/Horde/HordeServer.Tests/Acls/AclTests.cs
deleted file mode 100644
diff --git a/Engine/Source/Programs/Horde/Plugins/Build/HordeServer.Build.Tests/Acls/AclTests.cs b/Engine/Source/Programs/Horde/Plugins/Build/HordeServer.Build.Tests/Acls/AclTests.cs
new file mode 100644
--- /dev/null
+++ b/Engine/Source/Programs/Horde/Plugins/Build/HordeServer.Build.Tests/Acls/AclTests.cs
@@ -0,0 +1,125 @@
+// Copyright Epic Games, Inc. All Rights Reserved.
+
+using System.Collections.Generic;
+using System.Security.Claims;
+using EpicGames.Horde.Acls;
+using EpicGames.Horde.Projects;
+using EpicGames.Horde.Streams;
+using HordeServer.Acls;
+using HordeServer.Jobs;
+using HordeServer.Logs;
+using HordeServer.Projects;
+using HordeServer.Server;
+using HordeServer.Streams;
+using Microsoft.VisualStudio.TestTools.UnitTesting;
+
+namespace HordeServer.Tests.Acls
+{
+	[TestClass]
+	public class AclTests : TestSetup
+	{
+		[TestMethod]
+		public void TestCustomProfile()
+		{
+			UpdateConfig(globalConfig =>
+			{
+				globalConfig.Acl.Profiles = new List<AclProfileConfig>
+				{
+					new AclProfileConfig
+					{
+						Id = new AclProfileId("global-profile"),
+						Actions = new List<AclAction>{ JobAclAction.CreateJob, JobAclAction.RetryJobStep }
+					},
+				};
+
+				ProjectConfig projectConfig = new ProjectConfig();
+				projectConfig.Id = new ProjectId("ue5");
+				projectConfig.Acl.Profiles = new List<AclProfileConfig>
+				{
+					new AclProfileConfig
+					{
+						Id = new AclProfileId("project-profile"),
+						Actions = new List<AclAction>{ JobAclAction.ViewJob },
+						Extends = new List<AclProfileId>{ new AclProfileId("global-profile") },
+						ExcludeActions = new List<AclAction>{ JobAclAction.RetryJobStep }
+					},
+				};
+				globalConfig.Plugins.GetBuildConfig().Projects.Add(projectConfig);
+
+				StreamConfig streamConfig = new StreamConfig();
+				streamConfig.Id = new StreamId("ue5-main");
+				streamConfig.Acl.Profiles = new List<AclProfileConfig>
+				{
+					new AclProfileConfig
+					{
+						Id = new AclProfileId("stream-profile"),
+						Actions = new List<AclAction>{ JobAclAction.DeleteJob }
+					},
+				};
+				streamConfig.Acl.Entries = new List<AclEntryConfig>
+				{
+					new AclEntryConfig
+					{
+						Claim = new AclClaimConfig{ Type = "foo", Value = "bar" },
+						Profiles = new List<AclProfileId>{ new AclProfileId("project-profile"), new AclProfileId("stream-profile") }
+					}
+				};
+				projectConfig.Streams.Add(streamConfig);
+			});
+
+			GlobalConfig globalConfig = GlobalConfig.CurrentValue;
+			globalConfig.Plugins.GetBuildConfig().TryGetStream(new StreamId("ue5-main"), out StreamConfig? streamConfig);
+
+			ClaimsIdentity identity = new ClaimsIdentity(null, new[] { new Claim("foo", "bar") });
+			ClaimsPrincipal principal = new ClaimsPrincipal(identity);
+			Assert.IsTrue(streamConfig!.Acl.Authorize(JobAclAction.CreateJob, principal));
+			Assert.IsTrue(streamConfig!.Acl.Authorize(JobAclAction.ViewJob, principal));
+			Assert.IsTrue(streamConfig!.Acl.Authorize(JobAclAction.DeleteJob, principal));
+			Assert.IsFalse(streamConfig!.Acl.Authorize(JobAclAction.ExecuteJob, principal));
+			Assert.IsFalse(streamConfig!.Acl.Authorize(JobAclAction.RetryJobStep, principal));
+		}
+
+		[TestMethod]
+		public void TestDefaultProfile()
+		{
+			UpdateConfig(globalConfig =>
+			{
+				ProjectConfig projectConfig = new ProjectConfig();
+				projectConfig.Id = new ProjectId("ue5");
+				globalConfig.Plugins.GetBuildConfig().Projects.Add(projectConfig);
+
+				StreamConfig streamConfig = new StreamConfig();
+				streamConfig.Id = new StreamId("ue5-main");
+				streamConfig.Acl.Entries = new List<AclEntryConfig>
+				{
+					new AclEntryConfig
+					{
+						Claim = new AclClaimConfig{ Type = "foo", Value = "bar" },
+						Profiles = new List<AclProfileId>{ new AclProfileId("default-read") }
+					}
+				};
+				projectConfig.Streams.Add(streamConfig);
+			});
+
+			GlobalConfig globalConfig = GlobalConfig.CurrentValue;
+			globalConfig.Plugins.GetBuildConfig().TryGetStream(new StreamId("ue5-main"), out StreamConfig? streamConfig);
+
+			{
+				ClaimsIdentity identity = new ClaimsIdentity(null, new[] { new Claim("foo", "bar") });
+				ClaimsPrincipal principal = new ClaimsPrincipal(identity);
+				Assert.IsFalse(streamConfig!.Acl.Authorize(JobAclAction.CreateJob, principal));
+				Assert.IsTrue(streamConfig!.Acl.Authorize(JobAclAction.ViewJob, principal));
+				Assert.IsTrue(streamConfig!.Acl.Authorize(LogAclAction.ViewLog, principal));
+				Assert.IsFalse(streamConfig!.Acl.Authorize(JobAclAction.DeleteJob, principal));
+			}
+			{
+				ClaimsIdentity identity = new ClaimsIdentity(null, new[] { new Claim("foo", "bar2") });
+				ClaimsPrincipal principal = new ClaimsPrincipal(identity);
+				Assert.IsFalse(streamConfig!.Acl.Authorize(JobAclAction.CreateJob, principal));
+				Assert.IsFalse(streamConfig!.Acl.Authorize(JobAclAction.ViewJob, principal));
+				Assert.IsFalse(streamConfig!.Acl.Authorize(LogAclAction.ViewLog, principal));
+				Assert.IsFalse(streamConfig!.Acl.Authorize(JobAclAction.DeleteJob, principal));
+			}
+		}
+	}
+}
diff --git a/Engine/Source/Programs/Horde/Plugins/Build/HordeServer.Build.Tests/TestSetup.cs b/Engine/Source/Programs/Horde/Plugins/Build/HordeServer.Build.Tests/TestSetup.cs
--- a/Engine/Source/Programs/Horde/Plugins/Build/HordeServer.Build.Tests/TestSetup.cs
+++ b/Engine/Source/Programs/Horde/Plugins/Build/HordeServer.Build.Tests/TestSetup.cs
@@ -132,7 +132,6 @@ namespace HordeServer.Tests
 
 			services.AddSingleton<IAuditLog<AgentId>>(sp => sp.GetRequiredService<IAuditLogFactory<AgentId>>().Create("Agents.Log", "AgentId"));
 
-			services.AddSingleton<IAccountCollection, AccountCollection>();
 			services.AddSingleton<IAgentCollection, AgentCollection>();
 			services.AddSingleton<IAgentTelemetryCollection, AgentTelemetryCollection>();
 			services.AddSingleton<IArtifactCollection, ArtifactCollection>();
