Horde: Missing file.

#rnx

diff --git a/Engine/Source/Programs/Horde/Plugins/Compute/HordeServer.Compute.Tests/ComputeTestSetup.cs b/Engine/Source/Programs/Horde/Plugins/Compute/HordeServer.Compute.Tests/ComputeTestSetup.cs
new file mode 100644
--- /dev/null
+++ b/Engine/Source/Programs/Horde/Plugins/Compute/HordeServer.Compute.Tests/ComputeTestSetup.cs
@@ -0,0 +1,33 @@
+// Copyright Epic Games, Inc. All Rights Reserved.
+
+using System;
+using System.Collections.Generic;
+using System.Linq;
+using System.Text;
+using System.Threading.Tasks;
+using HordeServer.Compute;
+using HordeServer.Logs;
+using HordeServer.Storage;
+using Microsoft.Extensions.DependencyInjection;
+using Microsoft.Extensions.Logging;
+
+namespace HordeServer.Tests
+{
+	public class ComputeTestSetup : ServerTestSetup
+	{
+		public ComputeService ComputeService => ServiceProvider.GetRequiredService<ComputeService>();
+		public StorageService StorageService => ServiceProvider.GetRequiredService<StorageService>();
+		public ILogCollection LogCollection => ServiceProvider.GetRequiredService<ILogCollection>();
+
+		public ComputeTestSetup()
+		{
+			AddPlugin<StoragePlugin>();
+			AddPlugin<ComputePlugin>();
+		}
+
+		protected override void ConfigureServices(IServiceCollection services)
+		{
+			base.ConfigureServices(services);
+		}
+	}
+}
