Horde: Fix static analysis warnings.

#rnx

diff --git a/Engine/Source/Programs/Horde/HordeServer/Startup.cs b/Engine/Source/Programs/Horde/HordeServer/Startup.cs
--- a/Engine/Source/Programs/Horde/HordeServer/Startup.cs
+++ b/Engine/Source/Programs/Horde/HordeServer/Startup.cs
@@ -29,26 +29,14 @@ using Grpc.Core.Interceptors;
 using HordeCommon;
 using HordeServer.Accounts;
 using HordeServer.Acls;
-using HordeServer.Agents;
 using HordeServer.Agents.Enrollment;
-using HordeServer.Agents.Leases;
-using HordeServer.Agents.Pools;
 using HordeServer.Agents.Relay;
-using HordeServer.Agents.Sessions;
-using HordeServer.Agents.Telemetry;
-using HordeServer.Agents.Utilization;
-using HordeServer.Artifacts;
 using HordeServer.Auditing;
 using HordeServer.Authentication;
-using HordeServer.Compute;
 using HordeServer.Configuration;
 using HordeServer.Dashboard;
 using HordeServer.Devices;
-using HordeServer.Issues;
-using HordeServer.Issues.External;
 using HordeServer.Jobs;
-using HordeServer.Jobs.Schedules;
-using HordeServer.Jobs.TestData;
 using HordeServer.Logs;
 using HordeServer.Notifications;
 using HordeServer.Perforce;
@@ -57,7 +45,6 @@ using HordeServer.Server;
 using HordeServer.Server.Notices;
 using HordeServer.ServiceAccounts;
 using HordeServer.Storage;
-using HordeServer.Tasks;
 using HordeServer.Users;
 using HordeServer.Utilities;
 using Microsoft.AspNetCore.Authentication;
@@ -579,7 +566,6 @@ namespace HordeServer
 			// Hosted service that needs to run no matter the run mode of the process (server vs worker)
 			services.AddHostedService(provider => (DowntimeService)provider.GetRequiredService<IDowntimeService>());
 
-
 			if (settings.IsRunModeActive(RunMode.Worker) && !settings.MongoReadOnlyMode)
 			{
 				services.AddHostedService<MetricService>();
