Horde: Remove project reference causing a warning at Docker build time.

#rnx

diff --git a/Engine/Source/Programs/Horde/Plugins/Build/HordeServer.Build.Tests/HordeServer.Build.Tests.csproj b/Engine/Source/Programs/Horde/Plugins/Build/HordeServer.Build.Tests/HordeServer.Build.Tests.csproj
--- a/Engine/Source/Programs/Horde/Plugins/Build/HordeServer.Build.Tests/HordeServer.Build.Tests.csproj
+++ b/Engine/Source/Programs/Horde/Plugins/Build/HordeServer.Build.Tests/HordeServer.Build.Tests.csproj
@@ -23,7 +23,6 @@
   <ItemGroup>
     <ProjectReference Include="..\HordeServer.Build\HordeServer.Build.csproj" />
     <ProjectReference Include="..\..\Compute\HordeServer.Compute\HordeServer.Compute.csproj" />
-    <ProjectReference Include="..\..\..\HordeAgent.Tests\HordeAgent.Tests.csproj" />
     <ProjectReference Include="..\..\..\HordeServer.Tests.Shared\HordeServer.Tests.Shared.csproj" />
     <ProjectReference Include="..\..\..\HordeAgent\HordeAgent.csproj">
       <Aliases>HordeAgent</Aliases>
