Horde: Fix another couple of resource leaks.

#rnx

diff --git a/Engine/Source/Programs/Horde/Plugins/Ddc/HordeServer.Ddc/Ddc/RefService.cs b/Engine/Source/Programs/Horde/Plugins/Ddc/HordeServer.Ddc/Ddc/RefService.cs
--- a/Engine/Source/Programs/Horde/Plugins/Ddc/HordeServer.Ddc/Ddc/RefService.cs
+++ b/Engine/Source/Programs/Horde/Plugins/Ddc/HordeServer.Ddc/Ddc/RefService.cs
@@ -138,7 +138,7 @@ namespace HordeServer.Ddc
 				throw new RefNotFoundException(ns, bucket, key);
 			}
 
-			BlobData data = await node.References.First(x => x.Hash == node.RootHash).ReadBlobDataAsync(cancellationToken);
+			using BlobData data = await node.References.First(x => x.Hash == node.RootHash).ReadBlobDataAsync(cancellationToken);
 			BlobContents contents = new BlobContents(data.Data.ToArray());
 
 			RefRecord record = new RefRecord(ns, bucket, key, DateTime.UtcNow, null, BlobId.FromIoHash(node.RootHash), true);
diff --git a/Engine/Source/Programs/Horde/Plugins/Storage/HordeServer.Storage.Tests/Storage/GcServiceTests.cs b/Engine/Source/Programs/Horde/Plugins/Storage/HordeServer.Storage.Tests/Storage/GcServiceTests.cs
--- a/Engine/Source/Programs/Horde/Plugins/Storage/HordeServer.Storage.Tests/Storage/GcServiceTests.cs
+++ b/Engine/Source/Programs/Horde/Plugins/Storage/HordeServer.Storage.Tests/Storage/GcServiceTests.cs
@@ -80,7 +80,7 @@ namespace HordeServer.Tests.Storage
 				BlobLocator locator = root.GetLocator();
 				if (nodes.Add(locator))
 				{
-					BlobData data = await root.ReadBlobDataAsync();
+					using BlobData data = await root.ReadBlobDataAsync();
 					await FindNodesAsync(store, data.Imports, nodes);
 				}
 			}
