Horde: Add a debug endpoint for writing memory mapped file cache stats to the log.

diff --git a/Engine/Source/Programs/Horde/HordeServer/Server/DebugController.cs b/Engine/Source/Programs/Horde/HordeServer/Server/DebugController.cs
--- a/Engine/Source/Programs/Horde/HordeServer/Server/DebugController.cs
+++ b/Engine/Source/Programs/Horde/HordeServer/Server/DebugController.cs
@@ -14,6 +14,7 @@ using EpicGames.Core;
 using EpicGames.Horde.Agents.Leases;
 using EpicGames.Horde.Compute;
 using EpicGames.Horde.Logs;
+using EpicGames.Horde.Storage;
 using Google.Protobuf;
 using Horde.Common.Rpc;
 using HordeServer.Agents.Relay;
@@ -45,6 +46,7 @@ namespace HordeServer.Server
 		private static readonly Random s_random = new();
 
 		private readonly IMongoService _mongoService;
+		private readonly IServiceProvider _serviceProvider;
 		private readonly ConfigService _configService;
 		private readonly AgentRelayService _agentRelayService;
 		private readonly ILogCollection _logCollection;
@@ -56,6 +58,7 @@ namespace HordeServer.Server
 		/// </summary>
 		public DebugController(
 			IMongoService mongoService,
+			IServiceProvider serviceProvider,
 			ConfigService configService,
 			AgentRelayService agentRelayService,
 			ILogCollection logCollection,
@@ -63,6 +66,7 @@ namespace HordeServer.Server
 			ILogger<DebugController> logger)
 		{
 			_mongoService = mongoService;
+			_serviceProvider = serviceProvider;
 			_configService = configService;
 			_agentRelayService = agentRelayService;
 			_logCollection = logCollection;
@@ -463,5 +467,27 @@ namespace HordeServer.Server
 			string stringArg = "hello";
 			throw new Exception($"Message: numberArg:{numberArg}, stringArg:{stringArg}");
 		}
+
+		/// <summary>
+		/// Writes stats for the storage backend cache
+		/// </summary>
+		[HttpGet]
+		[Route("/api/v1/debug/writecacherefstats")]
+		public ActionResult WriteCacheRefStats()
+		{
+			if (!_globalConfig.Value.Authorize(ServerAclAction.Debug, User))
+			{
+				return Forbid(ServerAclAction.Debug);
+			}
+
+			StorageBackendCache? storageBackendCache = _serviceProvider.GetService(typeof(StorageBackendCache)) as StorageBackendCache;
+			if (storageBackendCache == null)
+			{
+				return NotFound();
+			}
+
+			storageBackendCache.WriteRefStats(_logger);
+			return Ok();
+		}
 	}
 }
diff --git a/Engine/Source/Programs/Shared/EpicGames.Core/MemoryMappedFileCache.cs b/Engine/Source/Programs/Shared/EpicGames.Core/MemoryMappedFileCache.cs
--- a/Engine/Source/Programs/Shared/EpicGames.Core/MemoryMappedFileCache.cs
+++ b/Engine/Source/Programs/Shared/EpicGames.Core/MemoryMappedFileCache.cs
@@ -4,7 +4,9 @@ using System;
 using System.Collections.Generic;
 using System.IO;
 using System.IO.MemoryMappedFiles;
+using System.Linq;
 using System.Threading;
+using Microsoft.Extensions.Logging;
 
 namespace EpicGames.Core
 {
@@ -285,5 +287,27 @@ namespace EpicGames.Core
 				listNode = nextListNode;
 			}
 		}
+
+		/// <summary>
+		/// Prints stats about the number of outstanding references to items in the cache
+		/// </summary>
+		public void WriteRefStats(ILogger logger)
+		{
+			lock (_lockObject)
+			{
+				logger.LogInformation("Memory mapped file cache using {NumEntries:n0}/{MaxEntries:n0} ({Size:n0}mb/{MaxSize:n0}mb)", _mappedFiles.Count, _maxMappedCount, _mappedSize, _maxMappedSize);
+
+				int referencedCount = 0;
+				ulong referencedSize = 0;
+				foreach (MappedFile mappedFile in _mappedFiles.OrderByDescending(x => x.RefCount).ThenBy(x => x.Key).TakeWhile(x => x.RefCount > 1))
+				{
+					logger.LogInformation("Mapped file {File} has {NumRefs} references", mappedFile.Key, mappedFile.RefCount);
+					referencedCount++;
+					referencedSize += mappedFile.MappedSize;
+				}
+
+				logger.LogInformation("{NumFiles:n0} files have more than one reference count ({SizeMb:n1}mb)", referencedCount, referencedSize / (1024.0 * 1024.0));
+			}
+		}
 	}
 }
diff --git a/Engine/Source/Programs/Shared/EpicGames.Horde/Storage/StorageBackendCache.cs b/Engine/Source/Programs/Shared/EpicGames.Horde/Storage/StorageBackendCache.cs
--- a/Engine/Source/Programs/Shared/EpicGames.Horde/Storage/StorageBackendCache.cs
+++ b/Engine/Source/Programs/Shared/EpicGames.Horde/Storage/StorageBackendCache.cs
@@ -379,5 +379,9 @@ namespace EpicGames.Horde.Storage
 			stats.Add("backend.cache.write_time_ms", (_writeTimeTicks * 1000) / Stopwatch.Frequency);
 			stats.Add("backend.cache.fetch_bytes", _fetchBytes);
 		}
+
+		/// <inheritdoc cref="MemoryMappedFileCache.WriteRefStats(ILogger)"/>
+		public void WriteRefStats(ILogger logger)
+			=> _memoryMappedFileCache.WriteRefStats(logger);
 	}
 }
