Horde: Change ref stats endpoint to a POST query.

#rnx

diff --git a/Engine/Source/Programs/Horde/HordeServer/Server/DebugController.cs b/Engine/Source/Programs/Horde/HordeServer/Server/DebugController.cs
--- a/Engine/Source/Programs/Horde/HordeServer/Server/DebugController.cs
+++ b/Engine/Source/Programs/Horde/HordeServer/Server/DebugController.cs
@@ -480,12 +480,7 @@ namespace HordeServer.Server
 				return Forbid(ServerAclAction.Debug);
 			}
 
-			StorageBackendCache? storageBackendCache = _serviceProvider.GetService(typeof(StorageBackendCache)) as StorageBackendCache;
-			if (storageBackendCache == null)
-			{
-				return NotFound();
-			}
-
+			StorageBackendCache storageBackendCache = _serviceProvider.GetRequiredService<StorageBackendCache>();
 			storageBackendCache.WriteRefStats(_logger);
 			return Ok();
 		}
