Horde: Fix schema generation for server docs.

#rnx

diff --git a/Engine/Source/Programs/Horde/HordeServer/Commands/Generate/ConfigDocsCommand.cs b/Engine/Source/Programs/Horde/HordeServer/Commands/Generate/ConfigDocsCommand.cs
--- a/Engine/Source/Programs/Horde/HordeServer/Commands/Generate/ConfigDocsCommand.cs
+++ b/Engine/Source/Programs/Horde/HordeServer/Commands/Generate/ConfigDocsCommand.cs
@@ -13,10 +13,13 @@ using System.Xml;
 using EpicGames.Core;
 using EpicGames.Horde.Acls;
 using HordeServer.Dashboard;
+using HordeServer.Plugins;
 using HordeServer.Projects;
 using HordeServer.Server;
+using HordeServer.Storage;
 using HordeServer.Streams;
 using HordeServer.Telemetry.Metrics;
+using HordeServer.Utilities;
 using Microsoft.Extensions.Logging;
 
 namespace HordeServer.Commands.Generate
@@ -40,14 +43,23 @@ namespace HordeServer.Commands.Generate
 		{
 			DirectoryReference.CreateDirectory(OutputDir);
 
-			XmlDocReader xmlDocReader = new XmlDocReader();
+			PluginCollection pluginCollection = new PluginCollection();
+			pluginCollection.Add<AnalyticsPlugin>();
+			pluginCollection.Add<BuildPlugin>();
+			pluginCollection.Add<ComputePlugin>();
+			pluginCollection.Add<DdcPlugin>();
+			pluginCollection.Add<SecretsPlugin>();
+			pluginCollection.Add<StoragePlugin>();
+			pluginCollection.Add<ToolsPlugin>();
 
-			JsonSchema serverSchema = JsonSchema.FromType(typeof(ServerSettings), xmlDocReader);
-			JsonSchema globalSchema = JsonSchema.FromType(typeof(GlobalConfig), xmlDocReader);
-			JsonSchema projectSchema = JsonSchema.FromType(typeof(ProjectConfig), xmlDocReader);
-			JsonSchema streamSchema = JsonSchema.FromType(typeof(StreamConfig), xmlDocReader);
-			JsonSchema telemetryConfigSchema = JsonSchema.FromType(typeof(TelemetryStoreConfig), xmlDocReader);
-			JsonSchema dashboardConfigSchema = JsonSchema.FromType(typeof(DashboardConfig), xmlDocReader);
+			JsonSchemaCache schemaCache = new JsonSchemaCache(pluginCollection);
+
+			JsonSchema serverSchema = schemaCache.CreateSchema(typeof(ServerSettings));
+			JsonSchema globalSchema = schemaCache.CreateSchema(typeof(GlobalConfig));
+			JsonSchema projectSchema = schemaCache.CreateSchema(typeof(ProjectConfig));
+			JsonSchema streamSchema = schemaCache.CreateSchema(typeof(StreamConfig));
+			JsonSchema telemetryConfigSchema = schemaCache.CreateSchema(typeof(TelemetryStoreConfig));
+			JsonSchema dashboardConfigSchema = schemaCache.CreateSchema(typeof(DashboardConfig));
 
 			Dictionary<JsonSchemaType, PageInfo> typeToPageInfo = new Dictionary<JsonSchemaType, PageInfo>
 			{
@@ -68,7 +80,7 @@ namespace HordeServer.Commands.Generate
 				Assembly agentAssembly = Assembly.LoadFile(Agent.FullName);
 				Type agentSettingsType = agentAssembly.GetType("HordeAgent.AgentSettings")!;
 
-				JsonSchema agentSchema = JsonSchema.FromType(agentSettingsType, xmlDocReader);
+				JsonSchema agentSchema = schemaCache.CreateSchema(agentSettingsType);
 				await WriteDocAsync(agentSchema.RootType, "Agent.json (Agent)", "Deployment/AgentSettings.md", "[Horde](../../README.md) > [Deployment](../Deployment.md) > [Agent](Agent.md)", AppSettingsIntro, new Dictionary<string, string>(), logger);
 			}
 
diff --git a/Engine/Source/Programs/Horde/HordeServer/Server/DebugController.cs b/Engine/Source/Programs/Horde/HordeServer/Server/DebugController.cs
--- a/Engine/Source/Programs/Horde/HordeServer/Server/DebugController.cs
+++ b/Engine/Source/Programs/Horde/HordeServer/Server/DebugController.cs
@@ -27,6 +27,7 @@ using Microsoft.AspNetCore.Authorization;
 using Microsoft.AspNetCore.Http;
 using Microsoft.AspNetCore.Http.Features;
 using Microsoft.AspNetCore.Mvc;
+using Microsoft.Extensions.DependencyInjection;
 using Microsoft.Extensions.Logging;
 using Microsoft.Extensions.Options;
 using MongoDB.Bson;
@@ -471,7 +472,7 @@ namespace HordeServer.Server
 		/// <summary>
 		/// Writes stats for the storage backend cache
 		/// </summary>
-		[HttpGet]
+		[HttpPost]
 		[Route("/api/v1/debug/writecacherefstats")]
 		public ActionResult WriteCacheRefStats()
 		{
