Horde: Fix threading issue causing block cache to always attempt to access element zero, causing tests to get stuck in an infinite loop.

diff --git a/Engine/Source/Programs/Horde/Plugins/Build/HordeServer.Build.Tests/Artifacts/BlockCacheTests.cs b/Engine/Source/Programs/Horde/Plugins/Build/HordeServer.Build.Tests/Artifacts/BlockCacheTests.cs
--- a/Engine/Source/Programs/Horde/Plugins/Build/HordeServer.Build.Tests/Artifacts/BlockCacheTests.cs
+++ b/Engine/Source/Programs/Horde/Plugins/Build/HordeServer.Build.Tests/Artifacts/BlockCacheTests.cs
@@ -83,10 +83,10 @@ namespace HordeServer.Tests.Artifacts
 			}
 
 			using BlockCache blockCache = BlockCache.CreateInMemory(5, 16, 4096);
-			Parallel.For(0, 4, threadIdx =>
+			Parallel.For(0, 16, threadIdx =>
 			{
 				Random rng = new Random(threadIdx);
-				for (int idx = 0; idx < 10000; idx++)
+				for (int idx = 0; idx < 40000; idx++)
 				{
 					int itemIdx = rng.Next(items.Count);
 					(string name, byte[] data) = items[itemIdx];
diff --git a/Engine/Source/Programs/Horde/Plugins/Build/HordeServer.Build/Artifacts/BlockCache.cs b/Engine/Source/Programs/Horde/Plugins/Build/HordeServer.Build/Artifacts/BlockCache.cs
--- a/Engine/Source/Programs/Horde/Plugins/Build/HordeServer.Build/Artifacts/BlockCache.cs
+++ b/Engine/Source/Programs/Horde/Plugins/Build/HordeServer.Build/Artifacts/BlockCache.cs
@@ -336,14 +336,17 @@ namespace HordeServer.Artifacts
 			while (!_freeBlocks.TryDequeue(out blockIdx))
 			{
 				PriorityQueue<(IoHash, BlockState), long> priorityQueue = new PriorityQueue<(IoHash, BlockState), long>(NumEvictionSamples);
-				for (int sampleIdx = 0; sampleIdx < NumEvictionSamples; sampleIdx++)
+				lock (_rng)
 				{
-					blockIdx = _rng.Next(_numBlocks);
-
-					BlockHeader header = GetBlockHeader(blockIdx);
-					if (_lookup.TryGetValue(header.Key, out BlockState? state))
+					for (int sampleIdx = 0; sampleIdx < NumEvictionSamples; sampleIdx++)
 					{
-						priorityQueue.Enqueue((header.Key, state), state.LastAccessTicks);
+						blockIdx = _rng.Next(_numBlocks);
+
+						BlockHeader header = GetBlockHeader(blockIdx);
+						if (_lookup.TryGetValue(header.Key, out BlockState? state))
+						{
+							priorityQueue.Enqueue((header.Key, state), state.LastAccessTicks);
+						}
 					}
 				}
 
