Horde: Remove dependency on Horde.Agent.Tests to fix static analysis warnings.

#rnx

diff --git a/Engine/Source/Programs/Horde/Plugins/Compute/HordeServer.Compute.Tests/HordeServer.Compute.Tests.csproj b/Engine/Source/Programs/Horde/Plugins/Compute/HordeServer.Compute.Tests/HordeServer.Compute.Tests.csproj
--- a/Engine/Source/Programs/Horde/Plugins/Compute/HordeServer.Compute.Tests/HordeServer.Compute.Tests.csproj
+++ b/Engine/Source/Programs/Horde/Plugins/Compute/HordeServer.Compute.Tests/HordeServer.Compute.Tests.csproj
@@ -20,7 +20,6 @@
   <ItemGroup>
     <ProjectReference Include="..\HordeServer.Compute\HordeServer.Compute.csproj" />
     <ProjectReference Include="..\..\Storage\HordeServer.Storage\HordeServer.Storage.csproj" />
-    <ProjectReference Include="..\..\..\HordeAgent.Tests\HordeAgent.Tests.csproj" />
     <ProjectReference Include="..\..\..\HordeServer.Tests.Shared\HordeServer.Tests.Shared.csproj" />
     <ProjectReference Include="..\..\..\HordeAgent\HordeAgent.csproj">
       <Aliases>HordeAgent</Aliases>
