Horde: Updating schema documentation

#rnx

diff --git a/Engine/Source/Programs/Horde/Docs/Config/Schema/AclActions.md b/Engine/Source/Programs/Horde/Docs/Config/Schema/AclActions.md
--- a/Engine/Source/Programs/Horde/Docs/Config/Schema/AclActions.md
+++ b/Engine/Source/Programs/Horde/Docs/Config/Schema/AclActions.md
@@ -11,89 +11,6 @@
 | `DeleteAccount` | Delete an account from the server |
 | `ViewAccount` | Ability to view account information |
 
-## Agents
-
-| Name | Description |
-| ---- | ----------- |
-| `CreateAgent` | Ability to create an agent. This may be done explicitly, or granted to agents to allow them to self-register. |
-| `UpdateAgent` | Update an agent's name, pools, etc... |
-| `DeleteAgent` | Soft-delete an agent |
-| `ViewAgent` | View an agent |
-| `ListAgents` | List the available agents |
-
-## Artifacts
-
-| Name | Description |
-| ---- | ----------- |
-| `ReadArtifact` | Permission to read from an artifact |
-| `WriteArtifact` | Permission to write to an artifact |
-| `UploadArtifact` | Ability to create an artifact. Typically just for debugging; agents have this access for a particular session. |
-| `DownloadArtifact` | Ability to download an artifact |
-
-## Bisect
-
-| Name | Description |
-| ---- | ----------- |
-| `CreateBisectTask` | Ability to start new bisect tasks |
-| `UpdateBisectTask` | Ability to update a bisect task |
-| `ViewBisectTask` | Ability to view a bisect task |
-
-## Compute
-
-| Name | Description |
-| ---- | ----------- |
-| `AddComputeTasks` | User can add tasks to the compute cluster |
-| `GetComputeTasks` | User can get and list tasks from the compute cluster |
-
-## Ddc
-
-| Name | Description |
-| ---- | ----------- |
-| `DdcReadObject` | General read access to refs / blobs and so on |
-| `DdcWriteObject` | General write access to upload refs / blobs etc |
-| `DdcDeleteObject` | Access to delete blobs / refs etc |
-| `DdcDeleteBucket` | Access to delete a particular bucket |
-| `DdcDeleteNamespace` | Access to delete a whole namespace |
-| `DdcReadTransactionLog` | Access to read the transaction log |
-| `DdcWriteTransactionLog` | Access to write the transaction log |
-| `DdcAdminAction` | Access to perform administrative task |
-
-## Devices
-
-| Name | Description |
-| ---- | ----------- |
-| `DeviceRead` | Ability to read devices |
-| `DeviceWrite` | Ability to write devices |
-
-## Jobs
-
-| Name | Description |
-| ---- | ----------- |
-| `CreateJob` | Ability to start new jobs |
-| `UpdateJob` | Rename a job, modify its priority, etc... |
-| `DeleteJob` | Delete a job properties |
-| `ExecuteJob` | Allows updating a job metadata (name, changelist number, step properties, new groups, job states, etc...). Typically granted to agents. Not user facing. |
-| `RetryJobStep` | Ability to retry a failed job step |
-| `ViewJob` | Ability to view a job |
-
-## Leases
-
-| Name | Description |
-| ---- | ----------- |
-| `ViewLeases` | View all the leases that an agent has worked on |
-| `ViewLeaseTasks` | View the task data for a lease |
-
-## Logs
-
-| Name | Description |
-| ---- | ----------- |
-| `CreateLog` | Ability to create a log. Implicitly granted to agents. |
-| `UpdateLog` | Ability to update log metadata |
-| `ViewLog` | Ability to view a log contents |
-| `WriteLogData` | Ability to write log data |
-| `CreateEvent` | Ability to create events |
-| `ViewEvent` | Ability to view events |
-
 ## Notices
 
 | Name | Description |
@@ -102,55 +19,6 @@
 | `UpdateNotice` | Ability to modify notices on the server |
 | `DeleteNotice` | Ability to delete notices |
 
-## Notifications
-
-| Name | Description |
-| ---- | ----------- |
-| `CreateSubscription` | Ability to subscribe to notifications |
-
-## Pools
-
-| Name | Description |
-| ---- | ----------- |
-| `CreatePool` | Create a global pool of agents |
-| `UpdatePool` | Modify an agent pool |
-| `DeletePool` | Delete an agent pool |
-| `ViewPool` | Ability to view a pool |
-| `ListPools` | View all the available agent pools |
-
-## Projects
-
-| Name | Description |
-| ---- | ----------- |
-| `CreateProject` | Allows the creation of new projects |
-| `DeleteProject` | Allows deletion of projects. |
-| `UpdateProject` | Modify attributes of a project (name, categories, etc...) |
-| `ViewProject` | View information about a project |
-
-## Replicators
-
-| Name | Description |
-| ---- | ----------- |
-| `UpdateReplicator` | Allows deletion of projects. |
-| `ViewReplicator` | Allows the creation of new projects |
-
-## Secrets
-
-| Name | Description |
-| ---- | ----------- |
-| `ViewSecret` | View a credential |
-
-## Server
-
-| Name | Description |
-| ---- | ----------- |
-| `AdminRead` | Ability to read any data from the server. Always inherited. |
-| `AdminWrite` | Ability to write any data to the server. |
-| `Debug` | Access to the debug endpoints |
-| `Impersonate` | Ability to impersonate another user |
-| `ViewCosts` | View estimated costs for particular operations |
-| `IssueBearerToken` | Issue bearer token for the current user |
-
 ## ServiceAccounts
 
 | Name | Description |
@@ -159,52 +27,3 @@
 | `UpdateAccount` | Update an account settings |
 | `DeleteAccount` | Delete an account from the server |
 | `ViewAccount` | Ability to view account information |
-
-## Sessions
-
-| Name | Description |
-| ---- | ----------- |
-| `CreateSession` | Granted to agents to call CreateSession, which returns a bearer token identifying themselves valid to call UpdateSesssion via gRPC. |
-| `ViewSession` | Allows viewing information about an agent session |
-
-## Software
-
-| Name | Description |
-| ---- | ----------- |
-| `UploadSoftware` | Ability to upload new versions of the agent software |
-| `DownloadSoftware` | Ability to download the agent software |
-| `DeleteSoftware` | Ability to delete agent software |
-
-## Storage
-
-| Name | Description |
-| ---- | ----------- |
-| `ReadBlobs` | Ability to read blobs from the storage service |
-| `WriteBlobs` | Ability to write blobs to the storage service |
-| `ReadRefs` | Ability to read refs from the storage service |
-| `WriteRefs` | Ability to write refs to the storage service |
-| `DeleteRefs` | Ability to delete refs |
-
-## Streams
-
-| Name | Description |
-| ---- | ----------- |
-| `CreateStream` | Allows the creation of new streams within a project |
-| `UpdateStream` | Allows updating a stream (agent types, templates, schedules) |
-| `DeleteStream` | Allows deleting a stream |
-| `ViewStream` | Ability to view a stream |
-| `ViewChanges` | View changes submitted to a stream. NOTE: this returns responses from the server's Perforce account, which may be a priviledged user. |
-| `ViewTemplate` | View template associated with a stream |
-
-## Telemetry
-
-| Name | Description |
-| ---- | ----------- |
-| `QueryMetrics` | Ability to search for various metrics |
-
-## Tools
-
-| Name | Description |
-| ---- | ----------- |
-| `DownloadTool` | Ability to download a tool |
-| `UploadTool` | Ability to upload new tool versions |
diff --git a/Engine/Source/Programs/Horde/Docs/Config/Schema/Globals.md b/Engine/Source/Programs/Horde/Docs/Config/Schema/Globals.md
--- a/Engine/Source/Programs/Horde/Docs/Config/Schema/Globals.md
+++ b/Engine/Source/Programs/Horde/Docs/Config/Schema/Globals.md
@@ -10,25 +10,10 @@ Name | Description
 `include` | [`ConfigInclude`](#configinclude)`[]`<br>Other paths to include
 `macros` | [`ConfigMacro`](#configmacro)`[]`<br>Macros within the global scope
 `dashboard` | [`DashboardConfig`](Dashboard.md)<br>Settings for the dashboard
-`projects` | [`ProjectConfig`](Projects.md)`[]`<br>List of projects
-`pools` | [`PoolConfig`](#poolconfig)`[]`<br>List of pools
 `downtime` | [`ScheduledDowntime`](#scheduleddowntime)`[]`<br>List of scheduled downtime
-`perforceClusters` | [`PerforceCluster`](#perforcecluster)`[]`<br>List of Perforce clusters
-`software` | [`AgentSoftwareConfig`](#agentsoftwareconfig)`[]`<br>List of costs of a particular agent type
-`rates` | [`AgentRateConfig`](#agentrateconfig)`[]`<br>List of costs of a particular agent type
-`networks` | [`NetworkConfig`](#networkconfig)`[]`<br>List of networks
-`compute` | [`ComputeClusterConfig`](#computeclusterconfig)`[]`<br>List of compute profiles
-`secrets` | [`SecretConfig`](#secretconfig)`[]`<br>List of secrets
-`devices` | [`DeviceConfig`](#deviceconfig)<br>Device configuration
-`tools` | [`ToolConfig`](#toolconfig)`[]`<br>List of tools hosted by the server
-`maxConformCount` | `integer`<br>Maximum number of conforms to run at once
-`agentShutdownIfDisabledGracePeriod` | `string`<br>Time to wait before shutting down an agent that has been disabled Used if no value is set on the actual pool.
-`storage` | [`StorageConfig`](#storageconfig)<br>Storage configuration
-`artifactTypes` | [`ArtifactTypeConfig`](#artifacttypeconfig)`[]`<br>Configuration for different artifact types
-`telemetryStores` | [`TelemetryStoreConfig`](Telemetry.md)`[]`<br>Metrics to aggregate on the Horde server
+`plugins` | `object`<br>Plugin config objects
 `parameters` | `object`<br>General parameters for other tools. Can be queried through the api/v1/parameters endpoint.
 `acl` | [`AclConfig`](#aclconfig)<br>Access control list
-`issueFixedTag` | `string`<br>Commit tag to use for marking issues as fixed
 
 ## ConfigInclude
 
@@ -43,430 +28,425 @@ Name | Description
 `name` | `string`<br>
 `value` | `string`<br>
 
-## PoolConfig
+## ScheduledDowntime
 
-Mutable configuration for a pool
+Settings for the maintenance window
 
 Name | Description
 ---- | -----------
-`id` | `string`<br>Unique id for this pool
-`base` | `string`<br>Base pool config to copy settings from
-`name` | `string`<br>Name of the pool
-`condition` | `string`<br>Condition for agents to automatically be included in this pool
-`properties` | `string` `->` `string`<br>Arbitrary properties related to this pool
-`color` | [`PoolColor`](#poolcolor-enum)<br>Color to use for this pool on the dashboard
-`enableAutoscaling` | `boolean`<br>Whether to enable autoscaling for this pool
-`minAgents` | `integer`<br>The minimum number of agents to keep in the pool
-`numReserveAgents` | `integer`<br>The minimum number of idle agents to hold in reserve
-`conformInterval` | `string`<br>Interval between conforms. If zero, the pool will not conform on a schedule.
-`scaleOutCooldown` | `string`<br>Cooldown time between scale-out events
-`scaleInCooldown` | `string`<br>Cooldown time between scale-in events
-`shutdownIfDisabledGracePeriod` | `string`<br>Time to wait before shutting down an agent that has been disabled
-`sizeStrategy` | [`PoolSizeStrategy`](#poolsizestrategy-enum)<br>
-`sizeStrategies` | [`PoolSizeStrategyInfo`](#poolsizestrategyinfo)`[]`<br>List of pool sizing strategies for this pool. The first strategy with a matching condition will be picked.
-`fleetManagers` | [`FleetManagerInfo`](#fleetmanagerinfo)`[]`<br>List of fleet managers for this pool. The first strategy with a matching condition will be picked. If empty or no conditions match, a default fleet manager will be used.
-`leaseUtilizationSettings` | [`LeaseUtilizationSettings`](#leaseutilizationsettings)<br>Settings for lease utilization pool sizing strategy (if used)
-`jobQueueSettings` | [`JobQueueSettings`](#jobqueuesettings)<br>Settings for job queue pool sizing strategy (if used)
-`computeQueueAwsMetricSettings` | [`ComputeQueueAwsMetricSettings`](#computequeueawsmetricsettings)<br>Settings for job queue pool sizing strategy (if used)
+`startTime` | `string`<br>Start time
+`finishTime` | `string`<br>Finish time
+`frequency` | [`ScheduledDowntimeFrequency`](#scheduleddowntimefrequency-enum)<br>Frequency that the window repeats
 
-## PoolColor (Enum)
+## ScheduledDowntimeFrequency (Enum)
 
-Color to use for labels of this pool
+How frequently the maintence window repeats
 
 Name | Description
 ---- | -----------
-`Default` | 
-`Blue` | 
-`Orange` | 
-`Green` | 
-`Gray` | 
-
-## PoolSizeStrategy (Enum)
+`Once` | Once
+`Daily` | Every day
+`Weekly` | Every week
 
-Available pool sizing strategies
+## AnalyticsConfig
 
 Name | Description
 ---- | -----------
-`LeaseUtilization` | Strategy based on lease utilization
-`JobQueue` | Strategy based on size of job build queue
-`NoOp` | No-op strategy used as fallback/default behavior
-`ComputeQueueAwsMetric` | A no-op strategy that reports metrics to let an external AWS auto-scaling policy scale the fleet
-`LeaseUtilizationAwsMetric` | A no-op strategy that reports metrics to let an external AWS auto-scaling policy scale the fleet
-
-## PoolSizeStrategyInfo
+`stores` | [`TelemetryStoreConfig`](Telemetry.md)`[]`<br>
 
-Metadata for configuring and picking a pool sizing strategy
+## BuildConfig
 
 Name | Description
 ---- | -----------
-`type` | [`PoolSizeStrategy`](#poolsizestrategy-enum)<br>Strategy implementation to use
-`condition` | `string`<br>Condition if this strategy should be enabled (right now, using date/time as a distinguishing factor)
-`config` | `string`<br>Configuration for the strategy, serialized as JSON
-`extraAgentCount` | `integer`<br>Integer to add after pool size has been calculated. Can also be negative.
-
-## FleetManagerInfo
+`perforceClusters` | [`PerforceCluster`](#perforcecluster)`[]`<br>
+`devices` | [`DeviceConfig`](#deviceconfig)<br>
+`maxConformCount` | `integer`<br>
+`agentShutdownIfDisabledGracePeriod` | `string`<br>
+`artifactTypes` | [`ArtifactTypeConfig`](#artifacttypeconfig)`[]`<br>
+`projects` | [`ProjectConfig`](Projects.md)`[]`<br>
+`issueFixedTag` | `string`<br>
 
-Metadata for configuring and picking a fleet manager
+## PerforceCluster
 
 Name | Description
 ---- | -----------
-`type` | [`FleetManagerType`](#fleetmanagertype-enum)<br>Fleet manager type implementation to use
-`condition` | `string`<br>Condition if this strategy should be enabled (right now, using date/time as a distinguishing factor)
-`config` | `string`<br>Configuration for the strategy, serialized as JSON
-
-## FleetManagerType (Enum)
+`name` | `string`<br>
+`serviceAccount` | `string`<br>
+`canImpersonate` | `boolean`<br>
+`supportsPartitionedWorkspaces` | `boolean`<br>
+`servers` | [`PerforceServer`](#perforceserver)`[]`<br>
+`credentials` | [`PerforceCredentials`](#perforcecredentials)`[]`<br>
+`autoSdk` | [`AutoSdkWorkspace`](#autosdkworkspace)`[]`<br>
 
-Available fleet managers
+## PerforceServer
 
 Name | Description
 ---- | -----------
-`Default` | Default fleet manager
-`NoOp` | No-op fleet manager.
-`Aws` | Fleet manager for handling AWS EC2 instances. Will create and/or terminate instances from scratch.
-`AwsReuse` | Fleet manager for handling AWS EC2 instances. Will start already existing but stopped instances to reuse existing EBS disks.
-`AwsRecycle` | Fleet manager for handling AWS EC2 instances. Will start already existing but stopped instances to reuse existing EBS disks.
-`AwsAsg` | Fleet manager for handling AWS EC2 instances. Uses an EC2 auto-scaling group for controlling the number of running instances.
-
-## LeaseUtilizationSettings
+`serverAndPort` | `string`<br>
+`healthCheck` | `boolean`<br>
+`resolveDns` | `boolean`<br>
+`maxConformCount` | `integer`<br>
+`condition` | `string`<br>
+`properties` | `string[]`<br>
 
-Lease utilization sizing settings for a pool
+## PerforceCredentials
 
 Name | Description
 ---- | -----------
-`sampleTimeSec` | `integer`<br>Time period for each sample
-`numSamples` | `integer`<br>Number of samples to collect for calculating lease utilization
-`numSamplesForResult` | `integer`<br>Min number of samples for a valid result
-`minAgents` | `integer`<br>The minimum number of agents to keep in the pool
-`numReserveAgents` | `integer`<br>The minimum number of idle agents to hold in reserve
-
-## JobQueueSettings
+`userName` | `string`<br>
+`password` | `string`<br>
+`ticket` | `string`<br>
 
-Job queue sizing settings for a pool
+## AutoSdkWorkspace
 
 Name | Description
 ---- | -----------
-`scaleOutFactor` | `number`<br>Factor translating queue size to additional agents to grow the pool with The result is always rounded up to nearest integer. Example: if there are 20 jobs in queue, a factor 0.25 will result in 5 new agents being added (20 * 0.25)
-`scaleInFactor` | `number`<br>Factor by which to shrink the pool size with when queue is empty The result is always rounded up to nearest integer. Example: when the queue size is zero, a default value of 0.9 will shrink the pool by 10% (current agent count * 0.9)
-`samplePeriodMin` | `integer`<br>How far back in time to look for job batches (that potentially are in the queue)
-`readyTimeThresholdSec` | `integer`<br>Time spent in ready state before considered truly waiting for an agent<br>A job batch can be in ready state before getting picked up and executed. This threshold will help ensure only batches that have been waiting longer than this value will be considered.
-
-## ComputeQueueAwsMetricSettings
+`name` | `string`<br>
+`properties` | `string[]`<br>
+`userName` | `string`<br>
+`stream` | `string`<br>
 
-Settings for
+## DeviceConfig
 
 Name | Description
 ---- | -----------
-`computeClusterId` | `string`<br>Compute cluster ID to observe
-`namespace` | `string`<br>AWS CloudWatch namespace to write metrics in
+`platforms` | [`DevicePlatformConfig`](#deviceplatformconfig)`[]`<br>
+`pools` | [`DevicePoolConfig`](#devicepoolconfig)`[]`<br>
 
-## ScheduledDowntime
-
-Settings for the maintenance window
+## DevicePlatformConfig
 
 Name | Description
 ---- | -----------
-`startTime` | `string`<br>Start time
-`finishTime` | `string`<br>Finish time
-`frequency` | [`ScheduledDowntimeFrequency`](#scheduleddowntimefrequency-enum)<br>Frequency that the window repeats
-
-## ScheduledDowntimeFrequency (Enum)
+`id` | `string`<br>
+`name` | `string`<br>
+`models` | `string[]`<br>
+`legacyNames` | `string[]`<br>
+`legacyPerfSpecHighModel` | `string`<br>
 
-How frequently the maintence window repeats
+## DevicePoolConfig
 
 Name | Description
 ---- | -----------
-`Once` | Once
-`Daily` | Every day
-`Weekly` | Every week
+`id` | `string`<br>
+`name` | `string`<br>
+`poolType` | [`DevicePoolType`](#devicepooltype-enum)<br>
+`projectIds` | `string[]`<br>
 
-## PerforceCluster
+## DevicePoolType (Enum)
 
-Information about a cluster of Perforce servers.
+The type of device pool
 
 Name | Description
 ---- | -----------
-`name` | `string`<br>Name of the cluster
-`serviceAccount` | `string`<br>Username for Horde to log in to this server. Will use the first account specified below if not overridden.
-`canImpersonate` | `boolean`<br>Whether the service account can impersonate other users
-`supportsPartitionedWorkspaces` | `boolean`<br>Whether to use partitioned workspaces on this server
-`servers` | [`PerforceServer`](#perforceserver)`[]`<br>List of servers
-`credentials` | [`PerforceCredentials`](#perforcecredentials)`[]`<br>List of server credentials
-`autoSdk` | [`AutoSdkWorkspace`](#autosdkworkspace)`[]`<br>List of autosdk streams
-
-## PerforceServer
+`Automation` | Available to CIS jobs
+`Shared` | Shared by users with remote checking and checkouts
 
-Information about an individual Perforce server
+## ArtifactTypeConfig
 
 Name | Description
 ---- | -----------
-`serverAndPort` | `string`<br>The server and port. The server may be a DNS entry with multiple records, in which case it will be actively load balanced.
-`healthCheck` | `boolean`<br>Whether to query the healthcheck address under each server
-`resolveDns` | `boolean`<br>Whether to resolve the DNS entries and load balance between different hosts
-`maxConformCount` | `integer`<br>Maximum number of simultaneous conforms on this server
-`condition` | `string`<br>Optional condition for a machine to be eligable to use this server
-`properties` | `string[]`<br>List of properties for an agent to be eligable to use this server
-
-## PerforceCredentials
+`type` | `string`<br>
+`name` | `string`<br>
+`keepCount` | `integer`<br>
+`keepDays` | `integer`<br>
 
-Credentials for a Perforce user
+## ComputeConfig
 
 Name | Description
 ---- | -----------
-`userName` | `string`<br>The username
-`password` | `string`<br>Password for the user
-`ticket` | `string`<br>Login ticket for the user (will be used instead of password if set)
-
-## AutoSdkWorkspace
+`acl` | [`AclConfig`](#aclconfig)<br>
+`versionEnum` | [`ConfigVersion`](#configversion-enum)<br>
+`rates` | [`AgentRateConfig`](#agentrateconfig)`[]`<br>
+`clusters` | [`ComputeClusterConfig`](#computeclusterconfig)`[]`<br>
+`pools` | [`PoolConfig`](#poolconfig)`[]`<br>
+`software` | [`AgentSoftwareConfig`](#agentsoftwareconfig)`[]`<br>
+`networks` | [`NetworkConfig`](#networkconfig)`[]`<br>
 
-Path to a platform and stream to use for syncing AutoSDK
+## AclConfig
 
 Name | Description
 ---- | -----------
-`name` | `string`<br>Name of this workspace
-`properties` | `string[]`<br>The agent properties to check (eg. "OSFamily=Windows")
-`userName` | `string`<br>Username for logging in to the server
-`stream` | `string`<br>Stream to use
-
-## AgentSoftwareConfig
+`entries` | [`AclEntryConfig`](#aclentryconfig)`[]`<br>
+`profiles` | [`AclProfileConfig`](#aclprofileconfig)`[]`<br>
+`inherit` | `boolean`<br>
+`exceptions` | `string[]`<br>
 
-Selects different agent software versions by evaluating a condition
+## AclEntryConfig
 
 Name | Description
 ---- | -----------
-`toolId` | `string`<br>Tool identifier
-`condition` | `string`<br>Condition for using this channel
+`claim` | [`AclClaimConfig`](#aclclaimconfig)<br>
+`actions` | `string[]`<br>
+`profiles` | `string[]`<br>
 
-## AgentRateConfig
+## AclClaimConfig
 
-Describes the monetary cost of agents matching a particular criteria
+Name | Description
+---- | -----------
+`type` | `string`<br>
+`value` | `string`<br>
+
+## AclProfileConfig
 
 Name | Description
 ---- | -----------
-`condition` | `string`<br>Condition string
-`rate` | `number`<br>Rate for this agent
+`id` | `string`<br>
+`actions` | `string[]`<br>
+`excludeActions` | `string[]`<br>
+`extends` | `string[]`<br>
 
-## NetworkConfig
+## ConfigVersion (Enum)
+
+Name | Description
+---- | -----------
+`None` | 
+`Initial` | 
+`PoolsInConfigFiles` | 
+`Latest` | 
+`LatestPlusOne` | 
 
-Describes a network The ID describes any logical grouping, such as region, availability zone, rack or office location.
+## AgentRateConfig
 
 Name | Description
 ---- | -----------
-`id` | `string`<br>ID for this network
-`cidrBlock` | `string`<br>CIDR block
-`description` | `string`<br>Human-readable description
-`computeId` | `string`<br>Compute ID for this network (used when allocating compute resources)
+`condition` | `string`<br>
+`rate` | `number`<br>
 
 ## ComputeClusterConfig
 
-Profile for executing compute requests
-
 Name | Description
 ---- | -----------
-`id` | `string`<br>Name of the partition
-`namespaceId` | `string`<br>Name of the namespace to use
-`requestBucketId` | `string`<br>Name of the input bucket
-`responseBucketId` | `string`<br>Name of the output bucket
-`condition` | `string`<br>Filter for agents to include
-`acl` | [`AclConfig`](#aclconfig)<br>Access control list
+`id` | `string`<br>
+`namespaceId` | `string`<br>
+`requestBucketId` | `string`<br>
+`responseBucketId` | `string`<br>
+`condition` | `string`<br>
+`acl` | [`AclConfig`](#aclconfig)<br>
 
-## AclConfig
+## PoolConfig
 
 Name | Description
 ---- | -----------
-`entries` | [`AclEntryConfig`](#aclentryconfig)`[]`<br>
-`profiles` | [`AclProfileConfig`](#aclprofileconfig)`[]`<br>
-`inherit` | `boolean`<br>
-`exceptions` | `string[]`<br>
+`id` | `string`<br>
+`base` | `string`<br>
+`name` | `string`<br>
+`condition` | `string`<br>
+`properties` | `string` `->` `string`<br>
+`color` | [`PoolColor`](#poolcolor-enum)<br>
+`enableAutoscaling` | `boolean`<br>
+`minAgents` | `integer`<br>
+`numReserveAgents` | `integer`<br>
+`conformInterval` | `string`<br>
+`scaleOutCooldown` | `string`<br>
+`scaleInCooldown` | `string`<br>
+`shutdownIfDisabledGracePeriod` | `string`<br>
+`sizeStrategy` | [`PoolSizeStrategy`](#poolsizestrategy-enum)<br>
+`sizeStrategies` | [`PoolSizeStrategyInfo`](#poolsizestrategyinfo)`[]`<br>
+`fleetManagers` | [`FleetManagerInfo`](#fleetmanagerinfo)`[]`<br>
+`leaseUtilizationSettings` | [`LeaseUtilizationSettings`](#leaseutilizationsettings)<br>
+`jobQueueSettings` | [`JobQueueSettings`](#jobqueuesettings)<br>
+`computeQueueAwsMetricSettings` | [`ComputeQueueAwsMetricSettings`](#computequeueawsmetricsettings)<br>
 
-## AclEntryConfig
+## PoolColor (Enum)
 
 Name | Description
 ---- | -----------
-`claim` | [`AclClaimConfig`](#aclclaimconfig)<br>
-`actions` | `string[]`<br>
-`profiles` | `string[]`<br>
+`Default` | 
+`Blue` | 
+`Orange` | 
+`Green` | 
+`Gray` | 
 
-## AclClaimConfig
+## PoolSizeStrategy (Enum)
 
 Name | Description
 ---- | -----------
-`type` | `string`<br>
-`value` | `string`<br>
+`LeaseUtilization` | 
+`JobQueue` | 
+`NoOp` | 
+`ComputeQueueAwsMetric` | 
+`LeaseUtilizationAwsMetric` | 
 
-## AclProfileConfig
+## PoolSizeStrategyInfo
 
 Name | Description
 ---- | -----------
-`id` | `string`<br>
-`actions` | `string[]`<br>
-`excludeActions` | `string[]`<br>
-`extends` | `string[]`<br>
+`type` | [`PoolSizeStrategy`](#poolsizestrategy-enum)<br>
+`condition` | `string`<br>
+`config` | `string`<br>
+`extraAgentCount` | `integer`<br>
 
-## SecretConfig
-
-Configuration for a secret value
+## FleetManagerInfo
 
 Name | Description
 ---- | -----------
-`id` | `string`<br>Identifier for this secret
-`data` | `string` `->` `string`<br>Key/value pairs associated with this secret
-`sources` | [`ExternalSecretConfig`](#externalsecretconfig)`[]`<br>Providers to source key/value pairs from
-`acl` | [`AclConfig`](#aclconfig)<br>Defines access to this particular secret
+`type` | [`FleetManagerType`](#fleetmanagertype-enum)<br>
+`condition` | `string`<br>
+`config` | `string`<br>
 
-## ExternalSecretConfig
-
-Configuration for an external secret provider
+## FleetManagerType (Enum)
 
 Name | Description
 ---- | -----------
-`provider` | `string`<br>Name of the provider to use
-`format` | [`ExternalSecretFormat`](#externalsecretformat-enum)<br>Format of the secret
-`key` | `string`<br>Optional key indicating the parameter to set in the resulting data array. Required if if  is .
-`path` | `string`<br>Optional value indicating what to fetch from the provider
-
-## ExternalSecretFormat (Enum)
+`Default` | 
+`NoOp` | 
+`Aws` | 
+`AwsReuse` | 
+`AwsRecycle` | 
+`AwsAsg` | 
 
-Format describing how to parse external secret values
+## LeaseUtilizationSettings
 
 Name | Description
 ---- | -----------
-`Text` | Secret is a plain text value which will be stored using the external secret key
-`Json` | Secret is a JSON formatted string containing key/value pairs
+`sampleTimeSec` | `integer`<br>
+`numSamples` | `integer`<br>
+`numSamplesForResult` | `integer`<br>
+`minAgents` | `integer`<br>
+`numReserveAgents` | `integer`<br>
 
-## DeviceConfig
+## JobQueueSettings
+
+Name | Description
+---- | -----------
+`scaleOutFactor` | `number`<br>
+`scaleInFactor` | `number`<br>
+`samplePeriodMin` | `integer`<br>
+`readyTimeThresholdSec` | `integer`<br>
 
-Configuration for devices
+## ComputeQueueAwsMetricSettings
 
 Name | Description
 ---- | -----------
-`platforms` | [`DevicePlatformConfig`](#deviceplatformconfig)`[]`<br>List of device platforms
-`pools` | [`DevicePoolConfig`](#devicepoolconfig)`[]`<br>List of device pools
+`computeClusterId` | `string`<br>
+`namespace` | `string`<br>
 
-## DevicePlatformConfig
+## AgentSoftwareConfig
 
-Configuration for a device platform
+Name | Description
+---- | -----------
+`toolId` | `string`<br>
+`condition` | `string`<br>
+
+## NetworkConfig
 
 Name | Description
 ---- | -----------
-`id` | `string`<br>The id for this platform
-`name` | `string`<br>Name of the platform
-`models` | `string[]`<br>A list of platform models
-`legacyNames` | `string[]`<br>Legacy names which older versions of Gauntlet may be using
-`legacyPerfSpecHighModel` | `string`<br>Model name for the high perf spec, which may be requested by Gauntlet
+`id` | `string`<br>
+`cidrBlock` | `string`<br>
+`description` | `string`<br>
+`computeId` | `string`<br>
+
+## EmptyPluginConfig
 
-## DevicePoolConfig
 
-Configuration for a device pool
+## SecretsConfig
 
 Name | Description
 ---- | -----------
-`id` | `string`<br>The id for this platform
-`name` | `string`<br>The name of the pool
-`poolType` | [`DevicePoolType`](#devicepooltype-enum)<br>The type of the pool
-`projectIds` | `string[]`<br>List of project ids associated with pool
+`secrets` | [`SecretConfig`](#secretconfig)`[]`<br>
 
-## DevicePoolType (Enum)
-
-The type of device pool
+## SecretConfig
 
 Name | Description
 ---- | -----------
-`Automation` | Available to CIS jobs
-`Shared` | Shared by users with remote checking and checkouts
+`id` | `string`<br>
+`data` | `string` `->` `string`<br>
+`sources` | [`ExternalSecretConfig`](#externalsecretconfig)`[]`<br>
+`acl` | [`AclConfig`](#aclconfig)<br>
 
-## ToolConfig
+## ExternalSecretConfig
+
+Name | Description
+---- | -----------
+`provider` | `string`<br>
+`format` | [`ExternalSecretFormat`](#externalsecretformat-enum)<br>
+`key` | `string`<br>
+`path` | `string`<br>
 
-Options for configuring a tool
+## ExternalSecretFormat (Enum)
 
 Name | Description
 ---- | -----------
-`id` | `string`<br>Unique identifier for the tool
-`name` | `string`<br>Name of the tool
-`description` | `string`<br>Description for the tool
-`category` | `string`<br>Category for the tool. Will cause the tool to be shown in a different tab in the dashboard.
-`group` | `string`<br>Grouping key for different variations of the same tool. The dashboard will show these together.
-`platforms` | `string[]`<br>Platforms for this tool. Takes the form of a NET RID (https://learn.microsoft.com/en-us/dotnet/core/rid-catalog).
-`public` | `boolean`<br>Whether this tool should be exposed for download on a public endpoint without authentication
-`showInUgs` | `boolean`<br>Whether to show this tool for download in the UGS tools menu
-`showInDashboard` | `boolean`<br>Whether to show this tool for download in the dashboard
-`namespaceId` | `string`<br>Default namespace for new deployments of this tool
-`acl` | [`AclConfig`](#aclconfig)<br>Permissions for the tool
+`Text` | 
+`Json` | 
 
 ## StorageConfig
 
-Configuration for storage
-
 Name | Description
 ---- | -----------
-`enableGc` | `boolean`<br>Whether to enable garbage collection
-`enableGcVerification` | `boolean`<br>Whether to enable garbage collection in verification mode (nothing deleted, just logging on access to deleted blobs)
-`backends` | [`BackendConfig`](#backendconfig)`[]`<br>List of storage backends
-`namespaces` | [`NamespaceConfig`](#namespaceconfig)`[]`<br>List of namespaces for storage
+`enableGc` | `boolean`<br>
+`enableGcVerification` | `boolean`<br>
+`backends` | [`BackendConfig`](#backendconfig)`[]`<br>
+`namespaces` | [`NamespaceConfig`](#namespaceconfig)`[]`<br>
 
 ## BackendConfig
 
-Common settings object for different providers
-
 Name | Description
 ---- | -----------
-`id` | `string`<br>The storage backend ID
-`base` | `string`<br>Base backend to copy default settings from
-`type` | [`StorageBackendType`](#storagebackendtype-enum)<br>The type of storage backend to use
-`baseDir` | `string`<br>Base directory for filesystem storage
-`awsBucketName` | `string`<br>Name of the bucket to use
-`awsBucketPath` | `string`<br>Base path within the bucket
-`awsCredentials` | [`AwsCredentialsType`](#awscredentialstype-enum)<br>Type of credentials to use
-`awsRole` | `string`<br>ARN of a role to assume
-`awsProfile` | `string`<br>The AWS profile to read credentials form
-`awsRegion` | `string`<br>Region to connect to
-`azureConnectionString` | `string`<br>Connection string for Azure
-`azureContainerName` | `string`<br>Name of the container
+`id` | `string`<br>
+`base` | `string`<br>
+`type` | [`StorageBackendType`](#storagebackendtype-enum)<br>
+`baseDir` | `string`<br>
+`awsBucketName` | `string`<br>
+`awsBucketPath` | `string`<br>
+`awsCredentials` | [`AwsCredentialsType`](#awscredentialstype-enum)<br>
+`awsRole` | `string`<br>
+`awsProfile` | `string`<br>
+`awsRegion` | `string`<br>
+`azureConnectionString` | `string`<br>
+`azureContainerName` | `string`<br>
 `relayServer` | `string`<br>
 `relayToken` | `string`<br>
 
 ## StorageBackendType (Enum)
 
-Types of storage backend to use
-
 Name | Description
 ---- | -----------
-`FileSystem` | Local filesystem
-`Aws` | AWS S3
-`Azure` | Azure blob store
-`Memory` | In-memory only (for testing)
+`FileSystem` | 
+`Aws` | 
+`Azure` | 
+`Memory` | 
 
 ## AwsCredentialsType (Enum)
 
-Credentials to use for AWS
-
 Name | Description
 ---- | -----------
-`Default` | Use default credentials from the AWS SDK
-`Profile` | Read credentials from the  profile in the AWS config file
-`AssumeRole` | Assume a particular role. Should specify ARN in
-`AssumeRoleWebIdentity` | Assume a particular role using the current environment variables.
+`Default` | 
+`Profile` | 
+`AssumeRole` | 
+`AssumeRoleWebIdentity` | 
 
 ## NamespaceConfig
 
-Configuration of a particular namespace
-
 Name | Description
 ---- | -----------
-`id` | `string`<br>Identifier for this namespace
-`backend` | `string`<br>Backend to use for this namespace
-`prefix` | `string`<br>Prefix for items within this namespace
-`gcFrequencyHrs` | `number`<br>How frequently to run garbage collection, in hours.
-`gcDelayHrs` | `number`<br>How long to keep newly uploaded orphanned objects before allowing them to be deleted, in hours.
-`enableAliases` | `boolean`<br>Support querying exports by their aliases
-`acl` | [`AclConfig`](#aclconfig)<br>Access list for this namespace
+`id` | `string`<br>
+`backend` | `string`<br>
+`prefix` | `string`<br>
+`gcFrequencyHrs` | `number`<br>
+`gcDelayHrs` | `number`<br>
+`enableAliases` | `boolean`<br>
+`acl` | [`AclConfig`](#aclconfig)<br>
 
-## ArtifactTypeConfig
+## ToolsConfig
 
-Configuration for an artifact
+Name | Description
+---- | -----------
+`tools` | [`ToolConfig`](#toolconfig)`[]`<br>
+
+## ToolConfig
 
 Name | Description
 ---- | -----------
-`type` | `string`<br>Name of the artifact type
-`name` | `string`<br>Legacy 'Name' property
-`keepCount` | `integer`<br>Number of artifacts to retain
-`keepDays` | `integer`<br>Number of days to retain artifacts of this type
+`id` | `string`<br>
+`name` | `string`<br>
+`description` | `string`<br>
+`category` | `string`<br>
+`group` | `string`<br>
+`platforms` | `string[]`<br>
+`public` | `boolean`<br>
+`showInUgs` | `boolean`<br>
+`showInDashboard` | `boolean`<br>
+`namespaceId` | `string`<br>
+`acl` | [`AclConfig`](#aclconfig)<br>
diff --git a/Engine/Source/Programs/Horde/Docs/Config/Schema/Projects.md b/Engine/Source/Programs/Horde/Docs/Config/Schema/Projects.md
--- a/Engine/Source/Programs/Horde/Docs/Config/Schema/Projects.md
+++ b/Engine/Source/Programs/Horde/Docs/Config/Schema/Projects.md
@@ -2,25 +2,23 @@
 
 # *.project.json
 
-Stores configuration for a project
-
-Name | Description
----- | -----------
-`id` | `string`<br>The project id
-`name` | `string`<br>Name for the new project
-`path` | `string`<br>Direct include path for the project config. For backwards compatibility with old config files when including from a GlobalConfig object.
-`include` | [`ConfigInclude`](#configinclude)`[]`<br>Includes for other configuration files
-`macros` | [`ConfigMacro`](#configmacro)`[]`<br>Macros within the global scope
-`order` | `integer`<br>Order of this project on the dashboard
-`logo` | `string`<br>Path to the project logo
-`logoDarkTheme` | `string`<br>Optional path to the project logo for the dark theme
-`pools` | [`PoolConfig`](#poolconfig)`[]`<br>List of pools for this project
-`categories` | [`ProjectCategoryConfig`](#projectcategoryconfig)`[]`<br>Categories to include in this project
-`jobOptions` | [`JobOptions`](#joboptions)<br>Default settings for executing jobs
-`workspaceTypes` | `string` `->` [`WorkspaceConfig`](#workspaceconfig)<br>Default workspace types for streams These are added to the list of each stream's workspace types.
-`telemetryStoreId` | `string`<br>Telemetry store for Horde data for this project
-`streams` | [`StreamConfig`](Streams.md)`[]`<br>List of streams
-`acl` | [`AclConfig`](#aclconfig)<br>Acl entries
+Name | Description
+---- | -----------
+`id` | `string`<br>
+`name` | `string`<br>
+`path` | `string`<br>
+`include` | [`ConfigInclude`](#configinclude)`[]`<br>
+`macros` | [`ConfigMacro`](#configmacro)`[]`<br>
+`order` | `integer`<br>
+`logo` | `string`<br>
+`logoDarkTheme` | `string`<br>
+`pools` | [`PoolConfig`](#poolconfig)`[]`<br>
+`categories` | [`ProjectCategoryConfig`](#projectcategoryconfig)`[]`<br>
+`jobOptions` | [`JobOptions`](#joboptions)<br>
+`workspaceTypes` | `string` `->` [`WorkspaceConfig`](#workspaceconfig)<br>
+`telemetryStoreId` | `string`<br>
+`streams` | [`StreamConfig`](Streams.md)`[]`<br>
+`acl` | [`AclConfig`](#aclconfig)<br>
 
 ## ConfigInclude
 
@@ -37,34 +35,30 @@ Name | Description
 
 ## PoolConfig
 
-Mutable configuration for a pool
-
 Name | Description
 ---- | -----------
-`id` | `string`<br>Unique id for this pool
-`base` | `string`<br>Base pool config to copy settings from
-`name` | `string`<br>Name of the pool
-`condition` | `string`<br>Condition for agents to automatically be included in this pool
-`properties` | `string` `->` `string`<br>Arbitrary properties related to this pool
-`color` | [`PoolColor`](#poolcolor-enum)<br>Color to use for this pool on the dashboard
-`enableAutoscaling` | `boolean`<br>Whether to enable autoscaling for this pool
-`minAgents` | `integer`<br>The minimum number of agents to keep in the pool
-`numReserveAgents` | `integer`<br>The minimum number of idle agents to hold in reserve
-`conformInterval` | `string`<br>Interval between conforms. If zero, the pool will not conform on a schedule.
-`scaleOutCooldown` | `string`<br>Cooldown time between scale-out events
-`scaleInCooldown` | `string`<br>Cooldown time between scale-in events
-`shutdownIfDisabledGracePeriod` | `string`<br>Time to wait before shutting down an agent that has been disabled
+`id` | `string`<br>
+`base` | `string`<br>
+`name` | `string`<br>
+`condition` | `string`<br>
+`properties` | `string` `->` `string`<br>
+`color` | [`PoolColor`](#poolcolor-enum)<br>
+`enableAutoscaling` | `boolean`<br>
+`minAgents` | `integer`<br>
+`numReserveAgents` | `integer`<br>
+`conformInterval` | `string`<br>
+`scaleOutCooldown` | `string`<br>
+`scaleInCooldown` | `string`<br>
+`shutdownIfDisabledGracePeriod` | `string`<br>
 `sizeStrategy` | [`PoolSizeStrategy`](#poolsizestrategy-enum)<br>
-`sizeStrategies` | [`PoolSizeStrategyInfo`](#poolsizestrategyinfo)`[]`<br>List of pool sizing strategies for this pool. The first strategy with a matching condition will be picked.
-`fleetManagers` | [`FleetManagerInfo`](#fleetmanagerinfo)`[]`<br>List of fleet managers for this pool. The first strategy with a matching condition will be picked. If empty or no conditions match, a default fleet manager will be used.
-`leaseUtilizationSettings` | [`LeaseUtilizationSettings`](#leaseutilizationsettings)<br>Settings for lease utilization pool sizing strategy (if used)
-`jobQueueSettings` | [`JobQueueSettings`](#jobqueuesettings)<br>Settings for job queue pool sizing strategy (if used)
-`computeQueueAwsMetricSettings` | [`ComputeQueueAwsMetricSettings`](#computequeueawsmetricsettings)<br>Settings for job queue pool sizing strategy (if used)
+`sizeStrategies` | [`PoolSizeStrategyInfo`](#poolsizestrategyinfo)`[]`<br>
+`fleetManagers` | [`FleetManagerInfo`](#fleetmanagerinfo)`[]`<br>
+`leaseUtilizationSettings` | [`LeaseUtilizationSettings`](#leaseutilizationsettings)<br>
+`jobQueueSettings` | [`JobQueueSettings`](#jobqueuesettings)<br>
+`computeQueueAwsMetricSettings` | [`ComputeQueueAwsMetricSettings`](#computequeueawsmetricsettings)<br>
 
 ## PoolColor (Enum)
 
-Color to use for labels of this pool
-
 Name | Description
 ---- | -----------
 `Default` | 
@@ -75,93 +69,77 @@ Name | Description
 
 ## PoolSizeStrategy (Enum)
 
-Available pool sizing strategies
-
 Name | Description
 ---- | -----------
-`LeaseUtilization` | Strategy based on lease utilization
-`JobQueue` | Strategy based on size of job build queue
-`NoOp` | No-op strategy used as fallback/default behavior
-`ComputeQueueAwsMetric` | A no-op strategy that reports metrics to let an external AWS auto-scaling policy scale the fleet
-`LeaseUtilizationAwsMetric` | A no-op strategy that reports metrics to let an external AWS auto-scaling policy scale the fleet
+`LeaseUtilization` | 
+`JobQueue` | 
+`NoOp` | 
+`ComputeQueueAwsMetric` | 
+`LeaseUtilizationAwsMetric` | 
 
 ## PoolSizeStrategyInfo
 
-Metadata for configuring and picking a pool sizing strategy
-
 Name | Description
 ---- | -----------
-`type` | [`PoolSizeStrategy`](#poolsizestrategy-enum)<br>Strategy implementation to use
-`condition` | `string`<br>Condition if this strategy should be enabled (right now, using date/time as a distinguishing factor)
-`config` | `string`<br>Configuration for the strategy, serialized as JSON
-`extraAgentCount` | `integer`<br>Integer to add after pool size has been calculated. Can also be negative.
+`type` | [`PoolSizeStrategy`](#poolsizestrategy-enum)<br>
+`condition` | `string`<br>
+`config` | `string`<br>
+`extraAgentCount` | `integer`<br>
 
 ## FleetManagerInfo
 
-Metadata for configuring and picking a fleet manager
-
 Name | Description
 ---- | -----------
-`type` | [`FleetManagerType`](#fleetmanagertype-enum)<br>Fleet manager type implementation to use
-`condition` | `string`<br>Condition if this strategy should be enabled (right now, using date/time as a distinguishing factor)
-`config` | `string`<br>Configuration for the strategy, serialized as JSON
+`type` | [`FleetManagerType`](#fleetmanagertype-enum)<br>
+`condition` | `string`<br>
+`config` | `string`<br>
 
 ## FleetManagerType (Enum)
 
-Available fleet managers
-
 Name | Description
 ---- | -----------
-`Default` | Default fleet manager
-`NoOp` | No-op fleet manager.
-`Aws` | Fleet manager for handling AWS EC2 instances. Will create and/or terminate instances from scratch.
-`AwsReuse` | Fleet manager for handling AWS EC2 instances. Will start already existing but stopped instances to reuse existing EBS disks.
-`AwsRecycle` | Fleet manager for handling AWS EC2 instances. Will start already existing but stopped instances to reuse existing EBS disks.
-`AwsAsg` | Fleet manager for handling AWS EC2 instances. Uses an EC2 auto-scaling group for controlling the number of running instances.
+`Default` | 
+`NoOp` | 
+`Aws` | 
+`AwsReuse` | 
+`AwsRecycle` | 
+`AwsAsg` | 
 
 ## LeaseUtilizationSettings
 
-Lease utilization sizing settings for a pool
-
 Name | Description
 ---- | -----------
-`sampleTimeSec` | `integer`<br>Time period for each sample
-`numSamples` | `integer`<br>Number of samples to collect for calculating lease utilization
-`numSamplesForResult` | `integer`<br>Min number of samples for a valid result
-`minAgents` | `integer`<br>The minimum number of agents to keep in the pool
-`numReserveAgents` | `integer`<br>The minimum number of idle agents to hold in reserve
+`sampleTimeSec` | `integer`<br>
+`numSamples` | `integer`<br>
+`numSamplesForResult` | `integer`<br>
+`minAgents` | `integer`<br>
+`numReserveAgents` | `integer`<br>
 
 ## JobQueueSettings
 
-Job queue sizing settings for a pool
-
 Name | Description
 ---- | -----------
-`scaleOutFactor` | `number`<br>Factor translating queue size to additional agents to grow the pool with The result is always rounded up to nearest integer. Example: if there are 20 jobs in queue, a factor 0.25 will result in 5 new agents being added (20 * 0.25)
-`scaleInFactor` | `number`<br>Factor by which to shrink the pool size with when queue is empty The result is always rounded up to nearest integer. Example: when the queue size is zero, a default value of 0.9 will shrink the pool by 10% (current agent count * 0.9)
-`samplePeriodMin` | `integer`<br>How far back in time to look for job batches (that potentially are in the queue)
-`readyTimeThresholdSec` | `integer`<br>Time spent in ready state before considered truly waiting for an agent<br>A job batch can be in ready state before getting picked up and executed. This threshold will help ensure only batches that have been waiting longer than this value will be considered.
+`scaleOutFactor` | `number`<br>
+`scaleInFactor` | `number`<br>
+`samplePeriodMin` | `integer`<br>
+`readyTimeThresholdSec` | `integer`<br>
 
 ## ComputeQueueAwsMetricSettings
 
-Settings for
-
 Name | Description
 ---- | -----------
-`computeClusterId` | `string`<br>Compute cluster ID to observe
-`namespace` | `string`<br>AWS CloudWatch namespace to write metrics in
+`computeClusterId` | `string`<br>
+`namespace` | `string`<br>
 
 ## ProjectCategoryConfig
 
-Information about a category to display for a stream
-
 Name | Description
 ---- | -----------
-`name` | `string`<br>Name of this category
-`row` | `integer`<br>Index of the row to display this category on
-`showOnNavMenu` | `boolean`<br>Whether to show this category on the nav menu
-`includePatterns` | `string[]`<br>Patterns for stream names to include
-`excludePatterns` | `string[]`<br>Patterns for stream names to exclude
+`name` | `string`<br>
+`row` | `integer`<br>
+`showOnNavMenu` | `boolean`<br>
+`includePatterns` | `string[]`<br>
+`excludePatterns` | `string[]`<br>
 
 ## JobOptions
 
@@ -190,24 +168,22 @@ Name | Description
 
 ## WorkspaceConfig
 
-Information about a workspace type
-
 Name | Description
 ---- | -----------
-`base` | `string`<br>Base workspace to derive from
-`cluster` | `string`<br>Name of the Perforce server cluster to use
-`serverAndPort` | `string`<br>The Perforce server and port (eg. perforce:1666)
-`userName` | `string`<br>User to log into Perforce with (defaults to buildmachine)
-`password` | `string`<br>Password to use to log into the workspace
-`identifier` | `string`<br>Identifier to distinguish this workspace from other workspaces. Defaults to the workspace type name.
-`stream` | `string`<br>Override for the stream to sync
-`view` | `string[]`<br>Custom view for the workspace
-`incremental` | `boolean`<br>Whether to use an incrementally synced workspace
-`useAutoSdk` | `boolean`<br>Whether to use the AutoSDK
-`autoSdkView` | `string[]`<br>View for the AutoSDK paths to sync. If null, the whole thing will be synced.
-`method` | `string`<br>Method to use when syncing/materializing data from Perforce
-`minScratchSpace` | `integer`<br>Minimum disk space that must be available *after* syncing this workspace (in megabytes) If not available, the job will be aborted.
-`conformDiskFreeSpace` | `integer`<br>Threshold for when to trigger an automatic conform of agent. Measured in megabytes free on disk. Set to null or 0 to disable.
+`base` | `string`<br>
+`cluster` | `string`<br>
+`serverAndPort` | `string`<br>
+`userName` | `string`<br>
+`password` | `string`<br>
+`identifier` | `string`<br>
+`stream` | `string`<br>
+`view` | `string[]`<br>
+`incremental` | `boolean`<br>
+`useAutoSdk` | `boolean`<br>
+`autoSdkView` | `string[]`<br>
+`method` | `string`<br>
+`minScratchSpace` | `integer`<br>
+`conformDiskFreeSpace` | `integer`<br>
 
 ## AclConfig
 
diff --git a/Engine/Source/Programs/Horde/Docs/Config/Schema/Streams.md b/Engine/Source/Programs/Horde/Docs/Config/Schema/Streams.md
--- a/Engine/Source/Programs/Horde/Docs/Config/Schema/Streams.md
+++ b/Engine/Source/Programs/Horde/Docs/Config/Schema/Streams.md
@@ -2,38 +2,36 @@
 
 # *.stream.json
 
-Config for a stream
-
-Name | Description
----- | -----------
-`id` | `string`<br>Identifier for the stream
-`path` | `string`<br>Direct include path for the stream config. For backwards compatibility with old config files when including from a ProjectConfig object.
-`include` | [`ConfigInclude`](#configinclude)`[]`<br>Includes for other configuration files
-`macros` | [`ConfigMacro`](#configmacro)`[]`<br>Macros within this stream
-`name` | `string`<br>Name of the stream
-`clusterName` | `string`<br>The perforce cluster containing the stream
-`order` | `integer`<br>Order for this stream
-`initialAgentType` | `string`<br>Default initial agent type for templates
-`notificationChannel` | `string`<br>Notification channel for all jobs in this stream
-`notificationChannelFilter` | `string`<br>Notification channel filter for this template. Can be Success, Failure, or Warnings.
-`triageChannel` | `string`<br>Channel to post issue triage notifications
-`jobOptions` | [`JobOptions`](#joboptions)<br>Default settings for executing jobs
-`telemetryStoreId` | `string`<br>Telemetry store for Horde data for this stream
-`autoSdkView` | `string[]`<br>View for the AutoSDK paths to sync. If null, the whole thing will be synced.
-`defaultPreflightTemplate` | `string`<br>Legacy name for the default preflight template
-`defaultPreflight` | [`DefaultPreflightConfig`](#defaultpreflightconfig)<br>Default template for running preflights
-`commitTags` | [`CommitTagConfig`](#committagconfig)`[]`<br>List of tags to apply to commits. Allows fast searching and classification of different commit types (eg. code vs content).
-`tabs` | [`TabConfig`](#tabconfig)`[]`<br>List of tabs to show for the new stream
-`environment` | `string` `->` `string`<br>Global environment variables for all agents in this stream
-`agentTypes` | `string` `->` [`AgentConfig`](#agentconfig)<br>Map of agent name to type
-`workspaceTypes` | `string` `->` [`WorkspaceConfig`](#workspaceconfig)<br>Map of workspace name to type
-`templates` | [`TemplateRefConfig`](#templaterefconfig)`[]`<br>List of templates to create
-`acl` | [`AclConfig`](#aclconfig)<br>Custom permissions for this object
-`pausedUntil` | `string`<br>Pause stream builds until specified date
-`pauseComment` | `string`<br>Reason for pausing builds of the stream
-`replicators` | [`ReplicatorConfig`](#replicatorconfig)`[]`<br>Configuration for workers to replicate commit data into Horde Storage.
-`workflows` | [`WorkflowConfig`](#workflowconfig)`[]`<br>Workflows for dealing with new issues
-`tokens` | [`TokenConfig`](#tokenconfig)`[]`<br>Tokens to create for each job step
+Name | Description
+---- | -----------
+`id` | `string`<br>
+`path` | `string`<br>
+`include` | [`ConfigInclude`](#configinclude)`[]`<br>
+`macros` | [`ConfigMacro`](#configmacro)`[]`<br>
+`name` | `string`<br>
+`clusterName` | `string`<br>
+`order` | `integer`<br>
+`initialAgentType` | `string`<br>
+`notificationChannel` | `string`<br>
+`notificationChannelFilter` | `string`<br>
+`triageChannel` | `string`<br>
+`jobOptions` | [`JobOptions`](#joboptions)<br>
+`telemetryStoreId` | `string`<br>
+`autoSdkView` | `string[]`<br>
+`defaultPreflightTemplate` | `string`<br>
+`defaultPreflight` | [`DefaultPreflightConfig`](#defaultpreflightconfig)<br>
+`commitTags` | [`CommitTagConfig`](#committagconfig)`[]`<br>
+`tabs` | [`TabConfig`](#tabconfig)`[]`<br>
+`environment` | `string` `->` `string`<br>
+`agentTypes` | `string` `->` [`AgentConfig`](#agentconfig)<br>
+`workspaceTypes` | `string` `->` [`WorkspaceConfig`](#workspaceconfig)<br>
+`templates` | [`TemplateRefConfig`](#templaterefconfig)`[]`<br>
+`acl` | [`AclConfig`](#aclconfig)<br>
+`pausedUntil` | `string`<br>
+`pauseComment` | `string`<br>
+`replicators` | [`ReplicatorConfig`](#replicatorconfig)`[]`<br>
+`workflows` | [`WorkflowConfig`](#workflowconfig)`[]`<br>
+`tokens` | [`TokenConfig`](#tokenconfig)`[]`<br>
 
 ## ConfigInclude
 
@@ -75,12 +73,10 @@ Name | Description
 
 ## DefaultPreflightConfig
 
-Specifies defaults for running a preflight
-
 Name | Description
 ---- | -----------
-`templateId` | `string`<br>The template id to query
-`change` | [`ChangeQueryConfig`](#changequeryconfig)<br>Query for the change to use
+`templateId` | `string`<br>
+`change` | [`ChangeQueryConfig`](#changequeryconfig)<br>
 
 ## ChangeQueryConfig
 
@@ -108,153 +104,133 @@ Name | Description
 
 ## CommitTagConfig
 
-Configuration for custom commit filters
-
 Name | Description
 ---- | -----------
-`name` | `string`<br>Name of the tag
-`base` | `string`<br>Base tag to copy settings from
-`filter` | `string[]`<br>List of files to be included in this filter
+`name` | `string`<br>
+`base` | `string`<br>
+`filter` | `string[]`<br>
 
 ## TabConfig
 
-Information about a page to display in the dashboard for a stream
-
 Name | Description
 ---- | -----------
-`title` | `string`<br>Title of this page
-`type` | `string`<br>Type of this tab
-`style` | [`TabStyle`](#tabstyle-enum)<br>Presentation style for this page
-`showNames` | `boolean`<br>Whether to show job names on this page
-`showPreflights` | `boolean`<br>Whether to show all user preflights
-`jobNames` | `string[]`<br>Names of jobs to include on this page. If there is only one name specified, the name column does not need to be displayed.
-`templates` | `string[]`<br>List of job template names to show on this page.
-`columns` | [`TabColumnConfig`](#tabcolumnconfig)`[]`<br>Columns to display for different types of aggregates
+`title` | `string`<br>
+`type` | `string`<br>
+`style` | [`TabStyle`](#tabstyle-enum)<br>
+`showNames` | `boolean`<br>
+`showPreflights` | `boolean`<br>
+`jobNames` | `string[]`<br>
+`templates` | `string[]`<br>
+`columns` | [`TabColumnConfig`](#tabcolumnconfig)`[]`<br>
 
 ## TabStyle (Enum)
 
-Style for rendering a tab
-
 Name | Description
 ---- | -----------
-`Normal` | Regular job list
-`Compact` | Omit job names, show condensed view
+`Normal` | 
+`Compact` | 
 
 ## TabColumnConfig
 
-Describes a column to display on the jobs page
-
 Name | Description
 ---- | -----------
-`type` | [`TabColumnType`](#tabcolumntype-enum)<br>The type of column
-`heading` | `string`<br>Heading for this column
-`category` | `string`<br>Category of aggregates to display in this column. If null, includes any aggregate not matched by another column.
-`parameter` | `string`<br>Parameter to show in this column
-`relativeWidth` | `integer`<br>Relative width of this column.
+`type` | [`TabColumnType`](#tabcolumntype-enum)<br>
+`heading` | `string`<br>
+`category` | `string`<br>
+`parameter` | `string`<br>
+`relativeWidth` | `integer`<br>
 
 ## TabColumnType (Enum)
 
-Type of a column in a jobs tab
-
 Name | Description
 ---- | -----------
-`Labels` | Contains labels
-`Parameter` | Contains parameters
+`Labels` | 
+`Parameter` | 
 
 ## AgentConfig
 
-Mapping from a BuildGraph agent type to a set of machines on the farm
-
 Name | Description
 ---- | -----------
-`base` | `string`<br>Base agent config to inherit settings from
-`pool` | `string`<br>Pool of agents to use for this agent type
-`workspace` | `string`<br>Name of the workspace to sync
-`tempStorageDir` | `string`<br>Path to the temporary storage dir
-`environment` | `string` `->` `string`<br>Environment variables to be set when executing the job
-`tokens` | [`TokenConfig`](#tokenconfig)`[]`<br>Tokens to allocate for this agent type
+`base` | `string`<br>
+`pool` | `string`<br>
+`workspace` | `string`<br>
+`tempStorageDir` | `string`<br>
+`environment` | `string` `->` `string`<br>
+`tokens` | [`TokenConfig`](#tokenconfig)`[]`<br>
 
 ## TokenConfig
 
-Configuration for allocating access tokens for each job
-
 Name | Description
 ---- | -----------
-`url` | `string`<br>URL to request tokens from
-`clientId` | `string`<br>Client id to use to request a new token
-`clientSecret` | `string`<br>Client secret to request a new access token
-`envVar` | `string`<br>Environment variable to set with the access token
+`url` | `string`<br>
+`clientId` | `string`<br>
+`clientSecret` | `string`<br>
+`envVar` | `string`<br>
 
 ## WorkspaceConfig
 
-Information about a workspace type
-
 Name | Description
 ---- | -----------
-`base` | `string`<br>Base workspace to derive from
-`cluster` | `string`<br>Name of the Perforce server cluster to use
-`serverAndPort` | `string`<br>The Perforce server and port (eg. perforce:1666)
-`userName` | `string`<br>User to log into Perforce with (defaults to buildmachine)
-`password` | `string`<br>Password to use to log into the workspace
-`identifier` | `string`<br>Identifier to distinguish this workspace from other workspaces. Defaults to the workspace type name.
-`stream` | `string`<br>Override for the stream to sync
-`view` | `string[]`<br>Custom view for the workspace
-`incremental` | `boolean`<br>Whether to use an incrementally synced workspace
-`useAutoSdk` | `boolean`<br>Whether to use the AutoSDK
-`autoSdkView` | `string[]`<br>View for the AutoSDK paths to sync. If null, the whole thing will be synced.
-`method` | `string`<br>Method to use when syncing/materializing data from Perforce
-`minScratchSpace` | `integer`<br>Minimum disk space that must be available *after* syncing this workspace (in megabytes) If not available, the job will be aborted.
-`conformDiskFreeSpace` | `integer`<br>Threshold for when to trigger an automatic conform of agent. Measured in megabytes free on disk. Set to null or 0 to disable.
+`base` | `string`<br>
+`cluster` | `string`<br>
+`serverAndPort` | `string`<br>
+`userName` | `string`<br>
+`password` | `string`<br>
+`identifier` | `string`<br>
+`stream` | `string`<br>
+`view` | `string[]`<br>
+`incremental` | `boolean`<br>
+`useAutoSdk` | `boolean`<br>
+`autoSdkView` | `string[]`<br>
+`method` | `string`<br>
+`minScratchSpace` | `integer`<br>
+`conformDiskFreeSpace` | `integer`<br>
 
 ## TemplateRefConfig
 
-Parameters to create a template within a stream
-
-Name | Description
----- | -----------
-`id` | `string`<br>Optional identifier for this ref. If not specified, an id will be generated from the name.
-`base` | `string`<br>Base template id to copy from
-`showUgsBadges` | `boolean`<br>Whether to show badges in UGS for these jobs
-`showUgsAlerts` | `boolean`<br>Whether to show alerts in UGS for these jobs
-`notificationChannel` | `string`<br>Notification channel for this template. Overrides the stream channel if set.
-`notificationChannelFilter` | `string`<br>Notification channel filter for this template. Can be a combination of "Success", "Failure" and "Warnings" separated by pipe characters.
-`triageChannel` | `string`<br>Triage channel for this template. Overrides the stream channel if set.
-`workflowId` | `string`<br>Workflow to user for this stream
-`annotations` | `string` `->` `string`<br>Default annotations to apply to nodes in this template
-`schedule` | [`ScheduleConfig`](#scheduleconfig)<br>Schedule to execute this template
-`chainedJobs` | [`ChainedJobTemplateConfig`](#chainedjobtemplateconfig)`[]`<br>List of chained job triggers
-`acl` | [`AclConfig`](#aclconfig)<br>The ACL for this template
-`name` | `string`<br>Name for the new template
-`description` | `string`<br>Description for the template
-`priority` | [`Priority`](#priority-enum)<br>Default priority for this job
-`allowPreflights` | `boolean`<br>Whether to allow preflights of this template
-`updateIssues` | `boolean`<br>Whether issues should be updated for all jobs using this template
-`promoteIssuesByDefault` | `boolean`<br>Whether issues should be promoted by default for this template, promoted issues will generate user notifications
-`initialAgentType` | `string`<br>Initial agent type to parse the buildgraph script on
-`submitNewChange` | `string`<br>Path to a file within the stream to submit to generate a new changelist for jobs
-`submitDescription` | `string`<br>Description for new changelists
-`defaultChange` | [`ChangeQueryConfig`](#changequeryconfig)`[]`<br>Default change to build at. Each object has a condition parameter which can evaluated by the server to determine which change to use.
-`arguments` | `string[]`<br>Fixed arguments for the new job
-`parameters` | [`TextParameterData`](#textparameterdata)/[`ListParameterData`](#listparameterdata)/[`BoolParameterData`](#boolparameterdata)`[]`<br>Parameters for this template
-`jobOptions` | [`JobOptions`](#joboptions)<br>Default settings for jobs
+Name | Description
+---- | -----------
+`id` | `string`<br>
+`base` | `string`<br>
+`showUgsBadges` | `boolean`<br>
+`showUgsAlerts` | `boolean`<br>
+`notificationChannel` | `string`<br>
+`notificationChannelFilter` | `string`<br>
+`triageChannel` | `string`<br>
+`workflowId` | `string`<br>
+`annotations` | `string` `->` `string`<br>
+`schedule` | [`ScheduleConfig`](#scheduleconfig)<br>
+`chainedJobs` | [`ChainedJobTemplateConfig`](#chainedjobtemplateconfig)`[]`<br>
+`acl` | [`AclConfig`](#aclconfig)<br>
+`name` | `string`<br>
+`description` | `string`<br>
+`priority` | [`Priority`](#priority-enum)<br>
+`allowPreflights` | `boolean`<br>
+`updateIssues` | `boolean`<br>
+`promoteIssuesByDefault` | `boolean`<br>
+`initialAgentType` | `string`<br>
+`submitNewChange` | `string`<br>
+`submitDescription` | `string`<br>
+`defaultChange` | [`ChangeQueryConfig`](#changequeryconfig)`[]`<br>
+`arguments` | `string[]`<br>
+`parameters` | [`TextParameterData`](#textparameterdata)/[`ListParameterData`](#listparameterdata)/[`BoolParameterData`](#boolparameterdata)`[]`<br>
+`jobOptions` | [`JobOptions`](#joboptions)<br>
 
 ## ScheduleConfig
 
-Parameters to create a new schedule
-
 Name | Description
 ---- | -----------
-`claims` | [`AclClaimConfig`](#aclclaimconfig)`[]`<br>Roles to impersonate for this schedule
-`enabled` | `boolean`<br>Whether the schedule should be enabled
-`maxActive` | `integer`<br>Maximum number of builds that can be active at once
-`maxChanges` | `integer`<br>Maximum number of changes the schedule can fall behind head revision. If greater than zero, builds will be triggered for every submitted changelist until the backlog is this size.
-`requireSubmittedChange` | `boolean`<br>Whether the build requires a change to be submitted
-`gate` | [`ScheduleGateConfig`](#schedulegateconfig)<br>Gate allowing the schedule to trigger
-`commits` | `string[]`<br>Commit tags for this schedule
-`filter` | [`ChangeContentFlags`](#changecontentflags-enum)`[]`<br>The types of changes to run for
-`files` | `string[]`<br>Files that should cause the job to trigger
-`templateParameters` | `string` `->` `string`<br>Parameters for the template
-`patterns` | [`SchedulePatternConfig`](#schedulepatternconfig)`[]`<br>New patterns for the schedule
+`claims` | [`AclClaimConfig`](#aclclaimconfig)`[]`<br>
+`enabled` | `boolean`<br>
+`maxActive` | `integer`<br>
+`maxChanges` | `integer`<br>
+`requireSubmittedChange` | `boolean`<br>
+`gate` | [`ScheduleGateConfig`](#schedulegateconfig)<br>
+`commits` | `string[]`<br>
+`filter` | [`ChangeContentFlags`](#changecontentflags-enum)`[]`<br>
+`files` | `string[]`<br>
+`templateParameters` | `string` `->` `string`<br>
+`patterns` | [`SchedulePatternConfig`](#schedulepatternconfig)`[]`<br>
 
 ## AclClaimConfig
 
@@ -265,32 +241,26 @@ Name | Description
 
 ## ScheduleGateConfig
 
-Gate allowing a schedule to trigger.
-
 Name | Description
 ---- | -----------
-`templateId` | `string`<br>The template containing the dependency
-`target` | `string`<br>Target to wait for
+`templateId` | `string`<br>
+`target` | `string`<br>
 
 ## ChangeContentFlags (Enum)
 
-Flags identifying content of a changelist
-
 Name | Description
 ---- | -----------
-`ContainsCode` | The change contains code
-`ContainsContent` | The change contains content
+`ContainsCode` | 
+`ContainsContent` | 
 
 ## SchedulePatternConfig
 
-Parameters to create a new schedule
-
 Name | Description
 ---- | -----------
-`daysOfWeek` | [`DayOfWeek`](#dayofweek-enum)`[]`<br>Days of the week to run this schedule on. If null, the schedule will run every day.
-`minTime` | `string`<br>Time during the day for the first schedule to trigger. Measured in minutes from midnight.
-`maxTime` | `string`<br>Time during the day for the last schedule to trigger. Measured in minutes from midnight.
-`interval` | `string`<br>Interval between each schedule triggering
+`daysOfWeek` | [`DayOfWeek`](#dayofweek-enum)`[]`<br>
+`minTime` | `string`<br>
+`maxTime` | `string`<br>
+`interval` | `string`<br>
 
 ## DayOfWeek (Enum)
 
@@ -306,13 +276,11 @@ Name | Description
 
 ## ChainedJobTemplateConfig
 
-Trigger for another template
-
 Name | Description
 ---- | -----------
-`trigger` | `string`<br>Name of the target that needs to complete before starting the other template
-`templateId` | `string`<br>Id of the template to trigger
-`useDefaultChangeForTemplate` | `boolean`<br>Whether to use the default change for the template rather than the change for the parent job.
+`trigger` | `string`<br>
+`templateId` | `string`<br>
+`useDefaultChangeForTemplate` | `boolean`<br>
 
 ## AclConfig
 
@@ -356,121 +324,105 @@ Name | Description
 
 ## TextParameterData
 
-Free-form text entry parameter
-
 Name | Description
 ---- | -----------
 `type` | Text<br>Type discriminator
-`id` | `string`<br>Identifier for this parameter
-`label` | `string`<br>Name of the parameter associated with this parameter.
-`argument` | `string`<br>Argument to pass to the executor
-`default` | `string`<br>Default value for this argument
-`scheduleOverride` | `string`<br>Override for the default value for this parameter when running a scheduled build
-`hint` | `string`<br>Hint text for this parameter
-`validation` | `string`<br>Regex used to validate this parameter
-`validationError` | `string`<br>Message displayed if validation fails, informing user of valid values.
-`toolTip` | `string`<br>Tool-tip text to display
+`id` | `string`<br>
+`label` | `string`<br>
+`argument` | `string`<br>
+`default` | `string`<br>
+`scheduleOverride` | `string`<br>
+`hint` | `string`<br>
+`validation` | `string`<br>
+`validationError` | `string`<br>
+`toolTip` | `string`<br>
 
 ## ListParameterData
 
-Allows the user to select a value from a constrained list of choices
-
 Name | Description
 ---- | -----------
 `type` | List<br>Type discriminator
-`label` | `string`<br>Label to display next to this parameter. Defaults to the parameter name.
-`style` | [`ListParameterStyle`](#listparameterstyle-enum)<br>The type of list parameter
-`items` | [`ListParameterItemData`](#listparameteritemdata)`[]`<br>List of values to display in the list
-`toolTip` | `string`<br>Tool tip text to display
+`label` | `string`<br>
+`style` | [`ListParameterStyle`](#listparameterstyle-enum)<br>
+`items` | [`ListParameterItemData`](#listparameteritemdata)`[]`<br>
+`toolTip` | `string`<br>
 
 ## ListParameterStyle (Enum)
 
-Style of list parameter
-
 Name | Description
 ---- | -----------
-`List` | Regular drop-down list. One item is always selected.
-`MultiList` | Drop-down list with checkboxes
-`TagPicker` | Tag picker from list of options
+`List` | 
+`MultiList` | 
+`TagPicker` | 
 
 ## ListParameterItemData
 
-Possible option for a list parameter
-
 Name | Description
 ---- | -----------
-`id` | `string`<br>Identifier for this parameter
-`group` | `string`<br>Optional group heading to display this entry under, if the picker style supports it.
-`text` | `string`<br>Name of the parameter associated with this list.
-`argumentIfEnabled` | `string`<br>Argument to pass with this parameter.
-`argumentsIfEnabled` | `string[]`<br>Arguments to pass with this parameter.
-`argumentIfDisabled` | `string`<br>Argument to pass with this parameter.
-`argumentsIfDisabled` | `string[]`<br>Arguments to pass if this parameter is disabled.
-`default` | `boolean`<br>Whether this item is selected by default
-`scheduleOverride` | `boolean`<br>Overridden value for this property in schedule builds
+`id` | `string`<br>
+`group` | `string`<br>
+`text` | `string`<br>
+`argumentIfEnabled` | `string`<br>
+`argumentsIfEnabled` | `string[]`<br>
+`argumentIfDisabled` | `string`<br>
+`argumentsIfDisabled` | `string[]`<br>
+`default` | `boolean`<br>
+`scheduleOverride` | `boolean`<br>
 
 ## BoolParameterData
 
-Allows the user to toggle an option on or off
-
 Name | Description
 ---- | -----------
 `type` | Bool<br>Type discriminator
-`id` | `string`<br>Identifier for this parameter
-`label` | `string`<br>Name of the parameter associated with this parameter.
-`argumentIfEnabled` | `string`<br>Argument to add if this parameter is enabled
-`argumentsIfEnabled` | `string[]`<br>Argument to add if this parameter is enabled
-`argumentIfDisabled` | `string`<br>Argument to add if this parameter is enabled
-`argumentsIfDisabled` | `string[]`<br>Arguments to add if this parameter is disabled
-`default` | `boolean`<br>Whether this argument is enabled by default
-`scheduleOverride` | `boolean`<br>Override for this parameter in scheduled builds
-`toolTip` | `string`<br>Tool tip text to display
+`id` | `string`<br>
+`label` | `string`<br>
+`argumentIfEnabled` | `string`<br>
+`argumentsIfEnabled` | `string[]`<br>
+`argumentIfDisabled` | `string`<br>
+`argumentsIfDisabled` | `string[]`<br>
+`default` | `boolean`<br>
+`scheduleOverride` | `boolean`<br>
+`toolTip` | `string`<br>
 
 ## ReplicatorConfig
 
-Configuration for a stream replicator
-
 Name | Description
 ---- | -----------
-`id` | `string`<br>Identifier for the replicator within the current stream
-`enabled` | `boolean`<br>Whether the replicator is enabled
-`minChange` | `integer`<br>Minimum change number to replicate
-`maxChange` | `integer`<br>Maximum change number to replicate
+`id` | `string`<br>
+`enabled` | `boolean`<br>
+`minChange` | `integer`<br>
+`maxChange` | `integer`<br>
 
 ## WorkflowConfig
 
-Configuration for an issue workflow
-
-Name | Description
----- | -----------
-`id` | `string`<br>Identifier for this workflow
-`reportTimes` | `string[]`<br>Times of day at which to send a report
-`summaryTab` | `string`<br>Name of the tab to post summary data to
-`reportChannel` | `string`<br>Channel to post summary information for these templates.
-`reportWarnings` | `boolean`<br>Whether to include issues with a warning status in the summary
-`groupIssuesByTemplate` | `boolean`<br>Whether to group issues by template in the report
-`triageChannel` | `string`<br>Channel to post threads for triaging new issues
-`triagePrefix` | `string`<br>Prefix for all triage messages
-`triageSuffix` | `string`<br>Suffix for all triage messages
-`triageInstructions` | `string`<br>Instructions posted to triage threads
-`triageAlias` | `string`<br>User id of a Slack user/alias to ping if there is nobody assigned to an issue by default.
-`triageTypeAliases` | `string` `->` `string`<br>Slack user/alias to ping for specific issue types (such as Systemic), if there is nobody assigned to an issue by default.
-`escalateAlias` | `string`<br>Alias to ping if an issue has not been resolved for a certain amount of time
-`escalateTimes` | `integer[]`<br>Times after an issue has been opened to escalate to the alias above, in minutes. Continues to notify on the last interval once reaching the end of the list.
-`maxMentions` | `integer`<br>Maximum number of people to mention on a triage thread
-`allowMentions` | `boolean`<br>Whether to mention people on this thread. Useful to disable for testing.
-`inviteRestrictedUsers` | `boolean`<br>Uses the admin.conversations.invite API to invite users to the channel
-`skipWhenEmpty` | `boolean`<br>Skips sending reports when there are no active issues.
-`annotations` | `string` `->` `string`<br>Additional node annotations implicit in this workflow
-`externalIssues` | [`ExternalIssueConfig`](#externalissueconfig)<br>External issue tracking configuration for this workflow
-`issueHandlers` | `string[]`<br>Additional issue handlers enabled for this workflow
+Name | Description
+---- | -----------
+`id` | `string`<br>
+`reportTimes` | `string[]`<br>
+`summaryTab` | `string`<br>
+`reportChannel` | `string`<br>
+`reportWarnings` | `boolean`<br>
+`groupIssuesByTemplate` | `boolean`<br>
+`triageChannel` | `string`<br>
+`triagePrefix` | `string`<br>
+`triageSuffix` | `string`<br>
+`triageInstructions` | `string`<br>
+`triageAlias` | `string`<br>
+`triageTypeAliases` | `string` `->` `string`<br>
+`escalateAlias` | `string`<br>
+`escalateTimes` | `integer[]`<br>
+`maxMentions` | `integer`<br>
+`allowMentions` | `boolean`<br>
+`inviteRestrictedUsers` | `boolean`<br>
+`skipWhenEmpty` | `boolean`<br>
+`annotations` | `string` `->` `string`<br>
+`externalIssues` | [`ExternalIssueConfig`](#externalissueconfig)<br>
+`issueHandlers` | `string[]`<br>
 
 ## ExternalIssueConfig
 
-External issue tracking configuration for a workflow
-
 Name | Description
 ---- | -----------
-`projectKey` | `string`<br>Project key in external issue tracker
-`defaultComponentId` | `string`<br>Default component id for issues using workflow
-`defaultIssueTypeId` | `string`<br>Default issue type id for issues using workflow
+`projectKey` | `string`<br>
+`defaultComponentId` | `string`<br>
+`defaultIssueTypeId` | `string`<br>
diff --git a/Engine/Source/Programs/Horde/Docs/Config/Schema/Telemetry.md b/Engine/Source/Programs/Horde/Docs/Config/Schema/Telemetry.md
--- a/Engine/Source/Programs/Horde/Docs/Config/Schema/Telemetry.md
+++ b/Engine/Source/Programs/Horde/Docs/Config/Schema/Telemetry.md
@@ -2,15 +2,13 @@
 
 # *.telemetry.json
 
-Config for metrics
-
 Name | Description
 ---- | -----------
-`id` | `string`<br>Identifier for this store
-`acl` | [`AclConfig`](#aclconfig)<br>Permissions for this store
-`metrics` | [`MetricConfig`](#metricconfig)`[]`<br>Metrics to aggregate on the Horde server
-`include` | [`ConfigInclude`](#configinclude)`[]`<br>Includes for other configuration files
-`macros` | [`ConfigMacro`](#configmacro)`[]`<br>Macros within this configuration
+`id` | `string`<br>
+`acl` | [`AclConfig`](#aclconfig)<br>
+`metrics` | [`MetricConfig`](#metricconfig)`[]`<br>
+`include` | [`ConfigInclude`](#configinclude)`[]`<br>
+`macros` | [`ConfigMacro`](#configmacro)`[]`<br>
 
 ## AclConfig
 
@@ -47,30 +45,26 @@ Name | Description
 
 ## MetricConfig
 
-Configures a metric to aggregate on the server
-
 Name | Description
 ---- | -----------
-`id` | `string`<br>Identifier for this metric
-`filter` | `string`<br>Filter expression to evaluate to determine which events to include. This query is evaluated against an array.
-`property` | `string`<br>Property to aggregate
-`groupBy` | `string`<br>Property to group by. Specified as a comma-separated list of JSON path expressions.
-`function` | [`AggregationFunction`](#aggregationfunction-enum)<br>How to aggregate samples for this metric
-`percentile` | `integer`<br>For the percentile function, specifies the percentile to measure
-`interval` | `string`<br>Interval for each metric. Supports times such as "2d", "1h", "1h30m", "20s".
+`id` | `string`<br>
+`filter` | `string`<br>
+`property` | `string`<br>
+`groupBy` | `string`<br>
+`function` | [`AggregationFunction`](#aggregationfunction-enum)<br>
+`percentile` | `integer`<br>
+`interval` | `string`<br>
 
 ## AggregationFunction (Enum)
 
-Method for aggregating samples into a metric
-
 Name | Description
 ---- | -----------
-`Count` | Count the number of matching elements
-`Min` | Take the minimum value of all samples
-`Max` | Take the maximum value of all samples
-`Sum` | Sum all the reported values
-`Average` | Average all the samples
-`Percentile` | Estimates the value at a certain percentile
+`Count` | 
+`Min` | 
+`Max` | 
+`Sum` | 
+`Average` | 
+`Percentile` | 
 
 ## ConfigInclude
 
diff --git a/Engine/Source/Programs/Horde/Docs/Deployment/ServerSettings.md b/Engine/Source/Programs/Horde/Docs/Deployment/ServerSettings.md
--- a/Engine/Source/Programs/Horde/Docs/Deployment/ServerSettings.md
+++ b/Engine/Source/Programs/Horde/Docs/Deployment/ServerSettings.md
@@ -12,8 +12,6 @@ Name | Description
 `httpPort` | `integer`<br>Main port for serving HTTP.
 `httpsPort` | `integer`<br>Port for serving HTTP with TLS enabled. Disabled by default.
 `http2Port` | `integer`<br>Dedicated port for serving only HTTP/2.
-`computeTunnelPort` | `integer`<br>Port to listen on for tunneling compute sockets to agents
-`computeTunnelAddress` | `string`<br>What address (host:port) clients should connect to for compute socket tunneling Port may differ from  if Horde server is behind a reverse proxy/firewall
 `mongoConnectionString` | `string`<br>Connection string for the Mongo database
 `databaseConnectionString` | `string`<br>MongoDB connection string
 `mongoDatabaseName` | `string`<br>MongoDB database name
@@ -25,7 +23,6 @@ Name | Description
 `shutdownMemoryThreshold` | `integer`<br>Shutdown the current server process if memory usage reaches this threshold (specified in MB)<br>Usually set to 80-90% of available memory to avoid CLR heap using all of it. If a memory leak was to occur, it's usually better to restart the process rather than to let the GC work harder and harder trying to recoup memory.<br>Should only be used when multiple server processes are running behind a load balancer and one can be safely restarted automatically by the underlying process handler (Docker, Kubernetes, AWS ECS, Supervisor etc). The shutdown behaves similar to receiving a SIGTERM and will wait for outstanding requests to finish.
 `serverPrivateCert` | `string`<br>Optional PFX certificate to use for encrypting agent SSL traffic. This can be a self-signed certificate, as long as it's trusted by agents.
 `authMethod` | [`AuthMethod`](#authmethod-enum)<br>Issuer for tokens from the auth provider
-`autoEnrollAgents` | `boolean`<br>Whether to automatically enroll agents in the farm
 `oidcProfileName` | `string`<br>Optional profile name to report through the /api/v1/server/auth endpoint. Allows sharing auth tokens between providers configured through the same profile name in OidcToken.exe config files.
 `oidcAuthority` | `string`<br>Issuer for tokens from the auth provider
 `oidcAudience` | `string`<br>Audience for validating externally issued tokens
@@ -47,7 +44,6 @@ Name | Description
 `corsOrigin` | `string`<br>Allowed Cors origin
 `enableDebugEndpoints` | `boolean`<br>Whether to enable debug/administrative REST API endpoints
 `enableNewAgentsByDefault` | `boolean`<br>Whether to automatically enable new agents by default. If false, new agents must manually be enabled before they can take on work.
-`testDataRetainMonths` | `integer`<br>The number of months to retain test data
 `schedulePollingInterval` | `string`<br>Interval between rebuilding the schedule queue with a DB query.
 `noResourceBackOffTime` | `string`<br>Interval between polling for new jobs
 `initiateJobBackOffTime` | `string`<br>Interval between attempting to assign agents to take on jobs
@@ -55,74 +51,28 @@ Name | Description
 `redisConnectionString` | `string`<br>Config for connecting to Redis server(s). Setting it to null will disable Redis use and connection See format at https://stackexchange.github.io/StackExchange.Redis/Configuration.html
 `redisConnectionConfig` | `string`<br>
 `redisReadOnlyMode` | `boolean`<br>Whether to disable writes to Redis.
-`logServiceWriteCacheType` | `string`<br>Type of write cache to use in log service Currently Supported: "InMemory" or "Redis"
-`storageBackends` | [`BackendConfig`](#backendconfig)`[]`<br>Overridden settings for storage backends. Useful for running against a production server with custom backends.
-`logStorage` | [`StorageBackendOptions`](#storagebackendoptions)<br>Settings for artifact storage
-`artifactStorage` | [`StorageBackendOptions`](#storagebackendoptions)<br>Settings for artifact storage
-`commitStorage` | [`TreeStoreOptions`](#treestoreoptions)<br>Configuration of tree storage
+`logServiceWriteCacheType` | `string`<br>Overridden settings for storage backends. Useful for running against a production server with custom backends.
 `logJsonToStdOut` | `boolean`<br>Whether to log json to stdout
 `logSessionRequests` | `boolean`<br>Whether to log requests to the UpdateSession and QueryServerState RPC endpoints
-`enableLogService` | `boolean`<br>Whether to enable the hosted LogService running background jobs
-`fleetManagerV2` | [`FleetManagerType`](#fleetmanagertype-enum)<br>Default fleet manager to use (when not specified by pool)
-`fleetManagerV2Config` | `string`<br>Config for the fleet manager (serialized JSON)
-`awsAutoScalingQueueUrls` | `string[]`<br>AWS SQS queue URLs where lifecycle events from EC2 auto-scaling are received
-`disableSchedules` | `boolean`<br>Whether to run scheduled jobs.
 `scheduleTimeZone` | `string`<br>Timezone for evaluating schedules
-`slackToken` | `string`<br>Bot token for interacting with Slack (xoxb-*)
-`slackSocketToken` | `string`<br>Token for opening a socket to slack (xapp-*)
-`slackAdminToken` | `string`<br>Admin user token for Slack (xoxp-*). This is only required when using the admin endpoints to invite users.
-`slackUsers` | `string`<br>Filtered list of slack users to send notifications to. Should be Slack user ids, separated by commas.
-`slackErrorPrefix` | `string`<br>Prefix to use when reporting errors
-`slackWarningPrefix` | `string`<br>Prefix to use when reporting warnings
-`configNotificationChannel` | `string`<br>Channel for sending messages related to config update failures
-`updateStreamsNotificationChannel` | `string`<br>Channel to send stream notification update failures to
-`jobNotificationChannel` | `string`<br>Slack channel to send job related notifications to. Multiple channels can be specified, separated by ;
-`agentNotificationChannel` | `string`<br>Slack channel to send agent related notifications to.
 `dashboardUrl` | `string`<br>The URl to use for generating links back to the dashboard.
 `helpEmailAddress` | `string`<br>Help email address that users can contact with issues
 `helpSlackChannel` | `string`<br>Help slack channel that users can use for issues
-`p4SwarmUrl` | `string`<br>Url of P4 Swarm installation
-`jiraUsername` | `string`<br>The Jira service account user name
-`jiraApiToken` | `string`<br>The Jira service account API token
-`jiraUrl` | `string`<br>The Uri for the Jira installation
-`sharedDeviceCheckoutDays` | `integer`<br>The number of days shared device checkouts are held
-`deviceProblemCooldownMinutes` | `integer`<br>The number of cooldown minutes for device problems
-`deviceReportChannel` | `string`<br>Channel to send device reports to
-`defaultAgentPoolSizeStrategy` | [`PoolSizeStrategy`](#poolsizestrategy-enum)<br>Default agent pool sizing strategy for pools that doesn't have one explicitly configured
-`agentPoolScaleOutCooldownSeconds` | `integer`<br>Scale-out cooldown for auto-scaling agent pools (in seconds). Can be overridden by per-pool settings.
-`agentPoolScaleInCooldownSeconds` | `integer`<br>Scale-in cooldown for auto-scaling agent pools (in seconds). Can be overridden by per-pool settings.
 `globalThreadPoolMinSize` | `integer`<br>Set the minimum size of the global thread pool This value has been found in need of tweaking to avoid timeouts with the Redis client during bursts of traffic. Default is 16 for .NET Core CLR. The correct value is dependent on the traffic the Horde Server is receiving. For Epic's internal deployment, this is set to 40.
 `withDatadog` | `boolean`<br>Whether to enable Datadog integration for tracing
-`withAws` | `boolean`<br>Whether to enable Amazon Web Services (AWS) specific features
 `configPath` | `string`<br>Path to the root config file. Relative to the server.json file by default.
-`perforce` | [`PerforceConnectionSettings`](#perforceconnectionsettings)`[]`<br>Perforce connections for use by the Horde server (not agents)
-`useLocalPerforceEnv` | `boolean`<br>Whether to use the local Perforce environment
-`perforceConnectionPoolSize` | `integer`<br>Number of pooled perforce connections to keep
-`enableUpgradeTasks` | `boolean`<br>Whether to enable the upgrade task source.
-`enableConformTasks` | `boolean`<br>Whether to enable the conform task source.
 `forceConfigUpdateOnStartup` | `boolean`<br>Forces configuration data to be read and updated as part of appplication startup, rather than on a schedule. Useful when running locally.
 `openBrowser` | `boolean`<br>Whether to open a browser on startup
-`bundleCacheDir` | `string`<br>Directory to use for the coarse-grained backend cache. This caches full bundles downloaded from the upstream object store.
-`bundleCacheSize` | `string`<br>Maximum size of the storage cache on disk. Accepts standard binary suffixes (kb, mb, gb, tb, etc...)
-`bundleCacheSizeBytes` | `integer`<br>Accessor for the bundle cache size in bytes
-`blockCacheDir` | `string`<br>Directory to store the fine-grained block cache. This caches individual exports embedded in bundles.
-`blockCacheSize` | `string`<br>Maximum size of the block cache. Accepts standard binary suffixes. Currently only allocates in multiples of 1024mb.
-`blockCacheSizeBytes` | `integer`<br>Accessor for the block cache size in bytes
 `featureFlags` | [`FeatureFlagSettings`](#featureflagsettings)<br>Experimental features to enable on the server.
-`commits` | [`CommitSettings`](#commitsettings)<br>Options for the commit service
-`telemetry` | [`BaseTelemetryConfig`](#basetelemetryconfig)`[]`<br>Settings for sending telemetry events to external services (for example Snowflake, ClickHouse etc)
-`bundledTools` | [`BundledToolConfig`](#bundledtoolconfig)`[]`<br>Tools bundled along with the server. Data for each tool can be produced using the 'bundle create' command, and should be stored in the Tools directory.
 `openTelemetry` | [`OpenTelemetrySettings`](#opentelemetrysettings)<br>Options for OpenTelemetry
 
 ## RunMode (Enum)
 
-Type of run mode this process should use. Each carry different types of workloads. More than one mode can be active. But not all modes are not guaranteed to be compatible with each other and will raise an error if combined in such a way.
-
 Name | Description
 ---- | -----------
-`None` | Default no-op value (ASP.NET config will default to this for enums that cannot be parsed)
-`Server` | Handle and respond to incoming external requests, such as HTTP REST and gRPC calls. These requests are time-sensitive and short-lived, typically less than 5 secs. If processes handling requests are unavailable, it will be very visible for users.
-`Worker` | Run non-request facing workloads. Such as background services, processing queues, running work based on timers etc. Short periods of downtime or high CPU usage due to bursts are fine for this mode. No user requests will be impacted directly. If auto-scaling is used, a much more aggressive policy can be applied (tighter process packing, higher avg CPU usage).
+`None` | 
+`Server` | 
+`Worker` | 
 
 ## AuthMethod (Enum)
 
@@ -135,292 +85,12 @@ Name | Description
 `OpenIdConnect` | Generic OpenID Connect authentication, recommended for most
 `Horde` | Authenticate using username and password credentials stored in Horde OpenID Connect (OIDC) is first and foremost recommended. But if you have a small installation (less than ~10 users) or lacking an OIDC provider, this is an option.
 
-## BackendConfig
-
-Common settings object for different providers
-
-Name | Description
----- | -----------
-`id` | `string`<br>The storage backend ID
-`base` | `string`<br>Base backend to copy default settings from
-`type` | [`StorageBackendType`](#storagebackendtype-enum)<br>The type of storage backend to use
-`baseDir` | `string`<br>Base directory for filesystem storage
-`awsBucketName` | `string`<br>Name of the bucket to use
-`awsBucketPath` | `string`<br>Base path within the bucket
-`awsCredentials` | [`AwsCredentialsType`](#awscredentialstype-enum)<br>Type of credentials to use
-`awsRole` | `string`<br>ARN of a role to assume
-`awsProfile` | `string`<br>The AWS profile to read credentials form
-`awsRegion` | `string`<br>Region to connect to
-`azureConnectionString` | `string`<br>Connection string for Azure
-`azureContainerName` | `string`<br>Name of the container
-`relayServer` | `string`<br>
-`relayToken` | `string`<br>
-
-## StorageBackendType (Enum)
-
-Types of storage backend to use
-
-Name | Description
----- | -----------
-`FileSystem` | Local filesystem
-`Aws` | AWS S3
-`Azure` | Azure blob store
-`Memory` | In-memory only (for testing)
-
-## AwsCredentialsType (Enum)
-
-Credentials to use for AWS
-
-Name | Description
----- | -----------
-`Default` | Use default credentials from the AWS SDK
-`Profile` | Read credentials from the  profile in the AWS config file
-`AssumeRole` | Assume a particular role. Should specify ARN in
-`AssumeRoleWebIdentity` | Assume a particular role using the current environment variables.
-
-## StorageBackendOptions
-
-Common settings object for different providers
-
-Name | Description
----- | -----------
-`type` | [`StorageBackendType`](#storagebackendtype-enum)<br>The type of storage backend to use
-`baseDir` | `string`<br>Base directory for filesystem storage
-`awsBucketName` | `string`<br>Name of the bucket to use
-`awsBucketPath` | `string`<br>Base path within the bucket
-`awsCredentials` | [`AwsCredentialsType`](#awscredentialstype-enum)<br>Type of credentials to use
-`awsRole` | `string`<br>ARN of a role to assume
-`awsProfile` | `string`<br>The AWS profile to read credentials form
-`awsRegion` | `string`<br>Region to connect to
-`azureConnectionString` | `string`<br>Connection string for Azure
-`azureContainerName` | `string`<br>Name of the container
-
-## TreeStoreOptions
-
-Options for storing trees
-
-Name | Description
----- | -----------
-`bundle` | [`BundleOptions`](#bundleoptions)<br>Options for creating bundles
-`chunking` | [`ChunkingOptions`](#chunkingoptions)<br>Options for chunking content
-`type` | [`StorageBackendType`](#storagebackendtype-enum)<br>The type of storage backend to use
-`baseDir` | `string`<br>Base directory for filesystem storage
-`awsBucketName` | `string`<br>Name of the bucket to use
-`awsBucketPath` | `string`<br>Base path within the bucket
-`awsCredentials` | [`AwsCredentialsType`](#awscredentialstype-enum)<br>Type of credentials to use
-`awsRole` | `string`<br>ARN of a role to assume
-`awsProfile` | `string`<br>The AWS profile to read credentials form
-`awsRegion` | `string`<br>Region to connect to
-`azureConnectionString` | `string`<br>Connection string for Azure
-`azureContainerName` | `string`<br>Name of the container
-
-## BundleOptions
-
-Options for configuring a bundle serializer
-
-Name | Description
----- | -----------
-`maxVersion` | [`BundleVersion`](#bundleversion-enum)<br>Maximum version number of bundles to write
-`maxBlobSize` | `integer`<br>Maximum payload size fo a blob
-`compressionFormat` | [`BundleCompressionFormat`](#bundlecompressionformat-enum)<br>Compression format to use
-`minCompressionPacketSize` | `integer`<br>Minimum size of a block to be compressed
-`maxWriteQueueLength` | `integer`<br>Maximum amount of data to store in memory. This includes any background writes as well as bundles being built.
-
-## BundleVersion (Enum)
-
-Bundle version number
-
-Name | Description
----- | -----------
-`Initial` | Initial version number
-`ExportAliases` | Added the BundleExport.Alias property
-`RemoveAliases` | Back out change to include aliases. Will likely do this through an API rather than baked into the data.
-`InPlace` | Use data structures which support in-place reading and writing.
-`ImportHashes` | Add import hashes to imported nodes
-`LatestV1` | Last version using the V1 pipeline
-`PacketSequence` | Structure bundles as a sequence of self-contained packets (uses V2 code)
-`Latest` | The current version number
-`LatestV2` | Last version using the V2 pipeline
-`LatestPlusOne` | Last item in the enum. Used for
-
-## BundleCompressionFormat (Enum)
-
-Indicates the compression format in the bundle
-
-Name | Description
----- | -----------
-`None` | Packets are uncompressed
-`LZ4` | LZ4 compression
-`Gzip` | Gzip compression
-`Oodle` | Oodle compression (Selkie)
-`Brotli` | Brotli compression
-`Zstd` | ZStandard compression
-
-## ChunkingOptions
-
-Options for creating file nodes
-
-Name | Description
----- | -----------
-`leafOptions` | [`LeafChunkedDataNodeOptions`](#leafchunkeddatanodeoptions)<br>Options for creating leaf nodes
-`interiorOptions` | [`InteriorChunkedDataNodeOptions`](#interiorchunkeddatanodeoptions)<br>Options for creating interior nodes
-
-## LeafChunkedDataNodeOptions
-
-Options for creating a specific type of file nodes
-
-Name | Description
----- | -----------
-`minSize` | `integer`<br>Minimum chunk size
-`maxSize` | `integer`<br>Maximum chunk size. Chunks will be split on this boundary if another match is not found.
-`targetSize` | `integer`<br>Target chunk size for content-slicing
-
-## InteriorChunkedDataNodeOptions
-
-Options for creating interior nodes
-
-Name | Description
----- | -----------
-`minChildCount` | `integer`<br>Minimum number of children in each node
-`targetChildCount` | `integer`<br>Target number of children in each node
-`maxChildCount` | `integer`<br>Maximum number of children in each node
-`sliceThreshold` | `integer`<br>Threshold hash value for splitting interior nodes
-
-## FleetManagerType (Enum)
-
-Available fleet managers
-
-Name | Description
----- | -----------
-`Default` | Default fleet manager
-`NoOp` | No-op fleet manager.
-`Aws` | Fleet manager for handling AWS EC2 instances. Will create and/or terminate instances from scratch.
-`AwsReuse` | Fleet manager for handling AWS EC2 instances. Will start already existing but stopped instances to reuse existing EBS disks.
-`AwsRecycle` | Fleet manager for handling AWS EC2 instances. Will start already existing but stopped instances to reuse existing EBS disks.
-`AwsAsg` | Fleet manager for handling AWS EC2 instances. Uses an EC2 auto-scaling group for controlling the number of running instances.
-
-## PoolSizeStrategy (Enum)
-
-Available pool sizing strategies
-
-Name | Description
----- | -----------
-`LeaseUtilization` | Strategy based on lease utilization
-`JobQueue` | Strategy based on size of job build queue
-`NoOp` | No-op strategy used as fallback/default behavior
-`ComputeQueueAwsMetric` | A no-op strategy that reports metrics to let an external AWS auto-scaling policy scale the fleet
-`LeaseUtilizationAwsMetric` | A no-op strategy that reports metrics to let an external AWS auto-scaling policy scale the fleet
-
-## PerforceConnectionSettings
-
-Perforce connection information for use by the Horde server (for reading config files, etc...)
-
-Name | Description
----- | -----------
-`id` | `string`<br>Identifier for this server
-`serverAndPort` | `string`<br>Server and port
-`credentials` | [`PerforceCredentials`](#perforcecredentials)<br>Credentials for the server
-
-## PerforceCredentials
-
-Credentials for a Perforce user
-
-Name | Description
----- | -----------
-`userName` | `string`<br>The username
-`password` | `string`<br>Password for the user
-`ticket` | `string`<br>Login ticket for the user (will be used instead of password if set)
-
 ## FeatureFlagSettings
 
 Feature flags to aid rollout of new features.
 Once a feature is running in its intended state and is stable, the flag should be removed. A name and date of when the flag was created is noted next to it to help encourage this behavior. Try having them be just a flag, a boolean.
 
 
-## CommitSettings
-
-Options for the commit service
-
-Name | Description
----- | -----------
-`replicateMetadata` | `boolean`<br>Whether to mirror commit metadata to the database
-`replicateContent` | `boolean`<br>Whether to mirror commit data to storage
-`bundle` | [`BundleOptions`](#bundleoptions)<br>Options for how objects are packed together
-`chunking` | [`ChunkingOptions`](#chunkingoptions)<br>Options for how objects are sliced
-
-## BaseTelemetryConfig
-
-Configuration for the telemetry sink
-
-Name | Description
----- | -----------
-`type` | [`TelemetrySinkType`](#telemetrysinktype-enum)<br>Type of telemetry sink
-
-## TelemetrySinkType (Enum)
-
-Type of telemetry provider to use
-
-Name | Description
----- | -----------
-`None` | No telemetry sink (default)
-`Epic` | Use the Epic telemetry sink
-`ClickHouse` | Use the ClickHouse telemetry sink
-`Mongo` | Mongo telemetry
-
-## BundledToolConfig
-
-Configuration for a tool bundled alongsize the server
-
-Name | Description
----- | -----------
-`version` | `string`<br>Version string for the current tool data
-`refName` | `string`<br>Ref name in the tools directory
-`dataDir` | `string`<br>Directory containing blob data for this tool. If empty, the tools/{id} folder next to the server will be used.
-`id` | `string`<br>Unique identifier for the tool
-`name` | `string`<br>Name of the tool
-`description` | `string`<br>Description for the tool
-`category` | `string`<br>Category for the tool. Will cause the tool to be shown in a different tab in the dashboard.
-`group` | `string`<br>Grouping key for different variations of the same tool. The dashboard will show these together.
-`platforms` | `string[]`<br>Platforms for this tool. Takes the form of a NET RID (https://learn.microsoft.com/en-us/dotnet/core/rid-catalog).
-`public` | `boolean`<br>Whether this tool should be exposed for download on a public endpoint without authentication
-`showInUgs` | `boolean`<br>Whether to show this tool for download in the UGS tools menu
-`showInDashboard` | `boolean`<br>Whether to show this tool for download in the dashboard
-`namespaceId` | `string`<br>Default namespace for new deployments of this tool
-`acl` | [`AclConfig`](#aclconfig)<br>Permissions for the tool
-
-## AclConfig
-
-Name | Description
----- | -----------
-`entries` | [`AclEntryConfig`](#aclentryconfig)`[]`<br>
-`profiles` | [`AclProfileConfig`](#aclprofileconfig)`[]`<br>
-`inherit` | `boolean`<br>
-`exceptions` | `string[]`<br>
-
-## AclEntryConfig
-
-Name | Description
----- | -----------
-`claim` | [`AclClaimConfig`](#aclclaimconfig)<br>
-`actions` | `string[]`<br>
-`profiles` | `string[]`<br>
-
-## AclClaimConfig
-
-Name | Description
----- | -----------
-`type` | `string`<br>
-`value` | `string`<br>
-
-## AclProfileConfig
-
-Name | Description
----- | -----------
-`id` | `string`<br>
-`actions` | `string[]`<br>
-`excludeActions` | `string[]`<br>
-`extends` | `string[]`<br>
-
 ## OpenTelemetrySettings
 
 OpenTelemetry configuration for collection and sending of traces and metrics.
