Horde: Use the NET 6 web application builder rather than the legacy convention based Startup-class.

diff --git a/Engine/Source/Programs/Horde/HordeServer/Commands/ServerCommand.cs b/Engine/Source/Programs/Horde/HordeServer/Commands/ServerCommand.cs
--- a/Engine/Source/Programs/Horde/HordeServer/Commands/ServerCommand.cs
+++ b/Engine/Source/Programs/Horde/HordeServer/Commands/ServerCommand.cs
@@ -8,6 +8,7 @@ using System.Security.Cryptography.X509Certificates;
 using System.Threading;
 using System.Threading.Tasks;
 using EpicGames.Core;
+using Microsoft.AspNetCore.Builder;
 using Microsoft.AspNetCore.Hosting;
 using Microsoft.AspNetCore.Server.Kestrel.Core;
 using Microsoft.Extensions.Configuration;
@@ -48,74 +49,22 @@ namespace HordeServer.Commands
 			logger.LogInformation("Data directory: {DataDir}", ServerApp.DataDir);
 			logger.LogInformation("Server config: {ConfigFile}", ServerApp.ServerConfigFile);
 
-			using (X509Certificate2? grpcCertificate = ReadGrpcCertificate(_hordeSettings))
-			{
-				using IHost host = CreateHostBuilderWithCert(_args, _config, _hordeSettings, grpcCertificate).Build();
-
-				await host.RunAsync();
-				return 0;
-			}
-		}
+			using X509Certificate2? grpcCertificate = ReadGrpcCertificate(_hordeSettings);
 
-		static IHostBuilder CreateHostBuilderWithCert(string[] args, IConfiguration config, ServerSettings serverSettings, X509Certificate2? sslCert)
-		{
 			AppContext.SetSwitch("System.Net.Http.SocketsHttpHandler.Http2UnencryptedSupport", true);
 
-			IHostBuilder hostBuilder = Host.CreateDefaultBuilder(args)
-				.UseSerilog()
-				.ConfigureHostOptions(options => options.ShutdownTimeout = TimeSpan.FromSeconds(30.0))
-				.ConfigureAppConfiguration(builder => builder.AddConfiguration(config))
-				.ConfigureWebHostDefaults(webBuilder =>
-				{
-					webBuilder.UseShutdownTimeout(TimeSpan.FromSeconds(30.0));
-
-					webBuilder.UseUrls(); // Disable default URLs; we will configure each port directly.
-
-					webBuilder.UseWebRoot("DashboardApp");
-
-					webBuilder.ConfigureKestrel(options =>
-					{
-						options.Limits.MaxRequestBodySize = 256 * 1024 * 1024;
-
-						// When agents are saturated with work (CPU or I/O), slow sending of gRPC data can happen.
-						// Kestrel protects against this behavior by default as it's commonly used for malicious attacks.
-						// Setting a more generous data rate should prevent incoming HTTP connections from being closed prematurely.
-						options.Limits.MinRequestBodyDataRate = new MinDataRate(10, TimeSpan.FromSeconds(60));
-
-						options.Limits.KeepAliveTimeout = TimeSpan.FromSeconds(220); // 10 seconds more than agent's timeout
+			WebApplicationOptions appOptions = new WebApplicationOptions
+			{
+				WebRootPath = "DashboardApp",
+				Args = _args
+			};
 
-						int httpPort = serverSettings.HttpPort;
-						if (httpPort != 0)
-						{
-							options.ListenAnyIP(httpPort, configure => { configure.Protocols = HttpProtocols.Http1; });
-						}
+			WebApplicationBuilder appBuilder = WebApplication.CreateBuilder(appOptions);
 
-						int httpsPort = serverSettings.HttpsPort;
-						if (httpsPort != 0)
-						{
-							options.ListenAnyIP(httpsPort, configure =>
-							{
-								if (sslCert != null)
-								{
-									configure.UseHttps(sslCert);
-								}
-								else
-								{
-									configure.UseHttps();
-								}
-							});
-						}
-
-						// To serve HTTP/2 with gRPC *without* TLS enabled, a separate port for HTTP/2 must be used.
-						// This is useful when having a load balancer in front that terminates TLS.
-						int http2Port = serverSettings.Http2Port;
-						if (http2Port != 0)
-						{
-							options.ListenAnyIP(http2Port, configure => { configure.Protocols = HttpProtocols.Http2; });
-						}
-					});
-					webBuilder.UseStartup<Startup>();
-				});
+			appBuilder.Host
+				.UseSerilog()
+				.ConfigureHostOptions(options => options.ShutdownTimeout = TimeSpan.FromSeconds(30.0))
+				.ConfigureAppConfiguration(builder => builder.AddConfiguration(_config));
 
 			if (OperatingSystem.IsWindows() && WindowsServiceHelpers.IsWindowsService())
 			{
@@ -124,17 +73,14 @@ namespace HordeServer.Commands
 				try
 				{
 					// Register the default WindowsServiceLifetime
-					hostBuilder = hostBuilder.UseWindowsService();
+					appBuilder.Host.UseWindowsService();
 
 #pragma warning disable CA1416
 					// Replace the default WindowsServiceLifetime (if there is one; we may not be running as a service) with a custom one
 					// that waits for all application startup before the service enters the running state. See [REDACTED_PROFILE_URL]/runtime/issues/50019
-					hostBuilder = hostBuilder.ConfigureServices(services =>
-					{
-						ServiceDescriptor descriptor = services.First(x => x.ImplementationType == typeof(WindowsServiceLifetime));
-						services.Remove(descriptor);
-						services.AddSingleton<IHostLifetime, CustomWindowsServiceLifetime>();
-					});
+					ServiceDescriptor descriptor = appBuilder.Services.First(x => x.ImplementationType == typeof(WindowsServiceLifetime));
+					appBuilder.Services.Remove(descriptor);
+					appBuilder.Services.AddSingleton<IHostLifetime, CustomWindowsServiceLifetime>();
 #pragma warning restore CA1416
 				}
 				catch (InvalidOperationException)
@@ -142,7 +88,61 @@ namespace HordeServer.Commands
 				}
 			}
 
-			return hostBuilder;
+			appBuilder.WebHost.UseShutdownTimeout(TimeSpan.FromSeconds(30.0));
+
+			appBuilder.WebHost.UseUrls(); // Disable default URLs; we will configure each port directly.
+
+			appBuilder.WebHost.ConfigureKestrel(options =>
+			{
+				options.Limits.MaxRequestBodySize = 256 * 1024 * 1024;
+
+				// When agents are saturated with work (CPU or I/O), slow sending of gRPC data can happen.
+				// Kestrel protects against this behavior by default as it's commonly used for malicious attacks.
+				// Setting a more generous data rate should prevent incoming HTTP connections from being closed prematurely.
+				options.Limits.MinRequestBodyDataRate = new MinDataRate(10, TimeSpan.FromSeconds(60));
+
+				options.Limits.KeepAliveTimeout = TimeSpan.FromSeconds(220); // 10 seconds more than agent's timeout
+
+				int httpPort = _hordeSettings.HttpPort;
+				if (httpPort != 0)
+				{
+					options.ListenAnyIP(httpPort, configure => { configure.Protocols = HttpProtocols.Http1; });
+				}
+
+				int httpsPort = _hordeSettings.HttpsPort;
+				if (httpsPort != 0)
+				{
+					options.ListenAnyIP(httpsPort, configure =>
+					{
+						if (grpcCertificate != null)
+						{
+							configure.UseHttps(grpcCertificate);
+						}
+						else
+						{
+							configure.UseHttps();
+						}
+					});
+				}
+
+				// To serve HTTP/2 with gRPC *without* TLS enabled, a separate port for HTTP/2 must be used.
+				// This is useful when having a load balancer in front that terminates TLS.
+				int http2Port = _hordeSettings.Http2Port;
+				if (http2Port != 0)
+				{
+					options.ListenAnyIP(http2Port, configure => { configure.Protocols = HttpProtocols.Http2; });
+				}
+			});
+
+			Startup.ConfigureServices(appBuilder.Configuration, appBuilder.Services);
+
+			using (WebApplication app = appBuilder.Build())
+			{
+				Startup.Configure(app, app.Environment, app.Lifetime, app.Services.GetRequiredService<IOptions<ServerSettings>>());
+				await app.RunAsync();
+			}
+
+			return 0;
 		}
 
 		// Custom service lifetime to wait for startup before continuing
@@ -225,11 +225,5 @@ namespace HordeServer.Commands
 				return new X509Certificate2(FileReference.ReadAllBytes(serverPrivateCert));
 			}
 		}
-
-		public static IHostBuilder CreateHostBuilderForTesting(string[] args)
-		{
-			ServerSettings hordeSettings = new ServerSettings();
-			return CreateHostBuilderWithCert(args, new ConfigurationBuilder().Build(), hordeSettings, null);
-		}
 	}
 }
diff --git a/Engine/Source/Programs/Horde/HordeServer/ServerApp.cs b/Engine/Source/Programs/Horde/HordeServer/ServerApp.cs
--- a/Engine/Source/Programs/Horde/HordeServer/ServerApp.cs
+++ b/Engine/Source/Programs/Horde/HordeServer/ServerApp.cs
@@ -15,7 +15,6 @@ using HordeServer.Utilities;
 using Microsoft.Extensions.Configuration;
 using Microsoft.Extensions.Configuration.Memory;
 using Microsoft.Extensions.DependencyInjection;
-using Microsoft.Extensions.Hosting;
 using Microsoft.Win32;
 using Serilog;
 using Serilog.Configuration;
@@ -186,9 +185,6 @@ namespace HordeServer
 #pragma warning restore ASP0000 // Do not call 'IServiceCollection.BuildServiceProvider' in 'ConfigureServices'
 		}
 
-		// Used by WebApplicationFactory in controller tests. Uses reflection to call this exact function signature.
-		public static IHostBuilder CreateHostBuilder(string[] args) => ServerCommand.CreateHostBuilderForTesting(args);
-
 		/// <summary>
 		/// Gets the current environment
 		/// </summary>
diff --git a/Engine/Source/Programs/Horde/HordeServer/Startup.cs b/Engine/Source/Programs/Horde/HordeServer/Startup.cs
--- a/Engine/Source/Programs/Horde/HordeServer/Startup.cs
+++ b/Engine/Source/Programs/Horde/HordeServer/Startup.cs
@@ -215,13 +215,6 @@ namespace HordeServer
 			}
 		}
 
-		public Startup(IConfiguration configuration)
-		{
-			Configuration = configuration;
-		}
-
-		public IConfiguration Configuration { get; }
-
 		public static void BindServerSettings(IConfiguration configuration, ServerSettings settings)
 		{
 			IConfigurationSection hordeSection = configuration.GetSection("Horde");
@@ -259,13 +252,12 @@ namespace HordeServer
 			}
 		}
 
-		// This method gets called *multiple times* by the runtime. Use this method to add services to the container.
-		public void ConfigureServices(IServiceCollection services)
+		public static void ConfigureServices(IConfiguration configuration, IServiceCollection services)
 		{
 			services.AddSingleton<JsonSchemaCache>();
 
 			// IOptionsMonitor pattern for live updating of configuration settings
-			services.Configure<ServerSettings>(x => BindServerSettings(Configuration, x));
+			services.Configure<ServerSettings>(x => BindServerSettings(configuration, x));
 
 			// We may upload a large number of references with posts to storage endpoints. Increase the max number of form values to allow this (the default is 1024).
 			services.Configure<FormOptions>(options =>
@@ -275,9 +267,9 @@ namespace HordeServer
 
 			// Bind the settings again for local variable access in this method
 			ServerSettings settings = new();
-			BindServerSettings(Configuration, settings);
+			BindServerSettings(configuration, settings);
 
-			ServerInfo serverInfo = new ServerInfo(Configuration, Options.Create(settings));
+			ServerInfo serverInfo = new ServerInfo(configuration, Options.Create(settings));
 			services.AddSingleton<IServerInfo>(serverInfo);
 
 			// Register the plugin collection
@@ -292,7 +284,7 @@ namespace HordeServer
 			services.AddSingleton<IPluginCollection>(pluginCollection);
 
 			// Register all the plugin services
-			IConfigurationSection pluginsConfig = Configuration.GetSection("Horde").GetSection("Plugins");
+			IConfigurationSection pluginsConfig = configuration.GetSection("Horde").GetSection("Plugins");
 			foreach (ILoadedPlugin plugin in pluginCollection.LoadedPlugins)
 			{
 				IConfigurationSection pluginConfig = pluginsConfig.GetSection(plugin.Name.ToString());
@@ -1080,8 +1072,7 @@ namespace HordeServer
 
 		public static void AddServices(IServiceCollection serviceCollection, IConfiguration configuration)
 		{
-			Startup startup = new Startup(configuration);
-			startup.ConfigureServices(serviceCollection);
+			Startup.ConfigureServices(configuration, serviceCollection);
 			serviceCollection.AddSingleton<IHostApplicationLifetime, HostApplicationLifetime>();
 		}
 
