Horde: Support registering commands from plugin modules. Move tool data command into tool plugin.

#rnx

diff --git a/Engine/Source/Programs/Horde/HordeServer.Shared/Utilities/JsonConfigFile.cs b/Engine/Source/Programs/Horde/HordeServer.Shared/Utilities/JsonConfigFile.cs
new file mode 100644
--- /dev/null
+++ b/Engine/Source/Programs/Horde/HordeServer.Shared/Utilities/JsonConfigFile.cs
@@ -0,0 +1,84 @@
+﻿// Copyright Epic Games, Inc. All Rights Reserved.
+
+using System;
+using System.Buffers;
+using System.Text.Json;
+using System.Text.Json.Nodes;
+using System.Threading;
+using System.Threading.Tasks;
+using EpicGames.Core;
+
+namespace HordeServer.Utilities
+{
+	using JsonObject = System.Text.Json.Nodes.JsonObject;
+
+	/// <summary>
+	/// Allows manipulating JSON config files without fully deserializing them
+	/// </summary>
+	public class JsonConfigFile
+	{
+		public JsonObject Root { get; }
+
+		public JsonConfigFile(JsonObject root)
+		{
+			Root = root;
+		}
+
+		public static async Task<JsonConfigFile> ReadAsync(FileReference file, CancellationToken cancellationToken = default)
+		{
+			byte[] data = await FileReference.ReadAllBytesAsync(file, cancellationToken);
+			JsonObject? obj = JsonNode.Parse(data, new JsonNodeOptions { PropertyNameCaseInsensitive = true }, new JsonDocumentOptions { AllowTrailingCommas = true, CommentHandling = JsonCommentHandling.Skip }) as JsonObject;
+			return new JsonConfigFile(obj ?? new JsonObject());
+		}
+
+		public async Task WriteAsync(FileReference file, CancellationToken cancellationToken = default)
+		{
+			ArrayBufferWriter<byte> buffer = new ArrayBufferWriter<byte>();
+			using (Utf8JsonWriter writer = new Utf8JsonWriter(buffer, new JsonWriterOptions { Indented = true }))
+			{
+				Root.WriteTo(writer);
+			}
+			await FileReference.WriteAllBytesAsync(file, buffer.WrittenMemory.ToArray(), cancellationToken);
+		}
+
+		public static T FindOrAddNode<T>(JsonObject obj, string name, Func<T> factory) where T : JsonNode
+		{
+			JsonNode? node = obj[name];
+			if (node != null)
+			{
+				if (node is T existingTypedNode)
+				{
+					return existingTypedNode;
+				}
+				else
+				{
+					obj.Remove(name);
+				}
+			}
+
+			T newTypedNode = factory();
+			obj.Add(name, newTypedNode);
+			return newTypedNode;
+		}
+
+		public static JsonObject FindOrAddElementByKey(JsonArray array, string key, string name)
+		{
+			foreach (JsonNode? element in array)
+			{
+				if (element is JsonObject obj)
+				{
+					JsonNode? node = obj[key];
+					if (node != null && (string?)node.AsValue() == name)
+					{
+						return obj;
+					}
+				}
+			}
+
+			JsonObject newObj = new JsonObject();
+			newObj[key] = name;
+			array.Add(newObj);
+			return newObj;
+		}
+	}
+}
diff --git a/Engine/Source/Programs/Horde/HordeServer/Commands/Generate/ToolDataCommand.cs b/Engine/Source/Programs/Horde/HordeServer/Commands/Generate/ToolDataCommand.cs
deleted file mode 100644
diff --git a/Engine/Source/Programs/Horde/HordeServer/ServerApp.cs b/Engine/Source/Programs/Horde/HordeServer/ServerApp.cs
--- a/Engine/Source/Programs/Horde/HordeServer/ServerApp.cs
+++ b/Engine/Source/Programs/Horde/HordeServer/ServerApp.cs
@@ -11,11 +11,14 @@ using System.Threading;
 using System.Threading.Tasks;
 using EpicGames.Core;
 using HordeServer.Commands;
+using HordeServer.Plugins;
+using HordeServer.Storage;
 using HordeServer.Utilities;
 using Microsoft.Extensions.Configuration;
 using Microsoft.Extensions.Configuration.Memory;
 using Microsoft.Extensions.DependencyInjection;
 using Microsoft.Extensions.Hosting;
+using Microsoft.Extensions.Options;
 using Microsoft.Win32;
 using Serilog;
 using Serilog.Configuration;
@@ -178,6 +181,15 @@ namespace HordeServer
 			services.AddSingleton<ServerSettings>(serverSettings);
 			services.Configure<ServerSettings>(x => Startup.BindServerSettings(config, x));
 
+			ServerInfo serverInfo = new ServerInfo(config, Options.Create(serverSettings));
+			services.AddSingleton<IServerInfo>(serverInfo);
+
+			IPluginCollection pluginCollection = ConfigurePlugins();
+			foreach (Assembly pluginAssembly in pluginCollection.LoadedPlugins.Select(x => x.Assembly).Distinct())
+			{
+				services.AddCommandsFromAssembly(pluginAssembly);
+			}
+
 #pragma warning disable ASP0000 // Do not call 'IServiceCollection.BuildServiceProvider' in 'ConfigureServices'
 			await using (ServiceProvider serviceProvider = services.BuildServiceProvider())
 			{
@@ -186,6 +198,19 @@ namespace HordeServer
 #pragma warning restore ASP0000 // Do not call 'IServiceCollection.BuildServiceProvider' in 'ConfigureServices'
 		}
 
+		public static IPluginCollection ConfigurePlugins()
+		{
+			PluginCollection pluginCollection = new PluginCollection();
+			pluginCollection.Add<AnalyticsPlugin>();
+			pluginCollection.Add<BuildPlugin>();
+			pluginCollection.Add<ComputePlugin>();
+			pluginCollection.Add<DdcPlugin>();
+			pluginCollection.Add<SecretsPlugin>();
+			pluginCollection.Add<StoragePlugin>();
+			pluginCollection.Add<ToolsPlugin>();
+			return pluginCollection;
+		}
+
 		// Used by WebApplicationFactory in controller tests. Uses reflection to call this exact function signature.
 		public static IHostBuilder CreateHostBuilder(string[] args) => ServerCommand.CreateHostBuilderForTesting(args);
 
diff --git a/Engine/Source/Programs/Horde/HordeServer/Startup.cs b/Engine/Source/Programs/Horde/HordeServer/Startup.cs
--- a/Engine/Source/Programs/Horde/HordeServer/Startup.cs
+++ b/Engine/Source/Programs/Horde/HordeServer/Startup.cs
@@ -44,7 +44,6 @@ using HordeServer.Plugins;
 using HordeServer.Server;
 using HordeServer.Server.Notices;
 using HordeServer.ServiceAccounts;
-using HordeServer.Storage;
 using HordeServer.Users;
 using HordeServer.Utilities;
 using Microsoft.AspNetCore.Authentication;
@@ -281,14 +280,7 @@ namespace HordeServer
 			services.AddSingleton<IServerInfo>(serverInfo);
 
 			// Register the plugin collection
-			PluginCollection pluginCollection = new PluginCollection();
-			pluginCollection.Add<AnalyticsPlugin>();
-			pluginCollection.Add<BuildPlugin>();
-			pluginCollection.Add<ComputePlugin>();
-			pluginCollection.Add<DdcPlugin>();
-			pluginCollection.Add<SecretsPlugin>();
-			pluginCollection.Add<StoragePlugin>();
-			pluginCollection.Add<ToolsPlugin>();
+			IPluginCollection pluginCollection = ServerApp.ConfigurePlugins();
 			services.AddSingleton<IPluginCollection>(pluginCollection);
 
 			// Register all the plugin services
diff --git a/Engine/Source/Programs/Horde/HordeServer/Utilities/JsonConfigFile.cs b/Engine/Source/Programs/Horde/HordeServer/Utilities/JsonConfigFile.cs
deleted file mode 100644
diff --git a/Engine/Source/Programs/Horde/Plugins/Tools/HordeServer.Tools/Commands/Generate/ToolDataCommand.cs b/Engine/Source/Programs/Horde/Plugins/Tools/HordeServer.Tools/Commands/Generate/ToolDataCommand.cs
new file mode 100644
--- /dev/null
+++ b/Engine/Source/Programs/Horde/Plugins/Tools/HordeServer.Tools/Commands/Generate/ToolDataCommand.cs
@@ -0,0 +1,161 @@
+// Copyright Epic Games, Inc. All Rights Reserved.
+
+using System;
+using System.ComponentModel;
+using System.Text.Json.Nodes;
+using System.Threading;
+using System.Threading.Tasks;
+using EpicGames.Core;
+using EpicGames.Horde.Storage;
+using EpicGames.Horde.Storage.Backends;
+using EpicGames.Horde.Storage.Bundles;
+using EpicGames.Horde.Storage.Nodes;
+using HordeServer.Tools;
+using HordeServer.Utilities;
+using Microsoft.Extensions.Logging;
+
+namespace HordeServer.Commands.Generate
+{
+	using JsonObject = System.Text.Json.Nodes.JsonObject;
+
+	[Command("generate", "tooldata", "Creates bundled tool data, and registers it with the server config file")]
+	class ToolDataCommand : Command
+	{
+		[CommandLine(Required = true)]
+		[Description("Identifier for the tool")]
+		public string Id { get; set; } = String.Empty;
+
+		[CommandLine]
+		[Description("Name of the tool")]
+		public string? Name { get; set; }
+
+		[CommandLine]
+		[Description("Description for the tool")]
+		public string? Description { get; set; }
+
+		[CommandLine]
+		[Description("Category for the tool")]
+		public string? Category { get; set; }
+
+		[CommandLine]
+		[Description("Version string for the tool")]
+		public string? Version { get; set; }
+
+		[CommandLine]
+		[Description("If true, the tool may be downloaded by any user without authentication")]
+		public bool Public { get; set; }
+
+		[CommandLine]
+		[Description("Shows the tool for download in UGS")]
+		public bool ShowInUgs { get; set; }
+
+		[CommandLine("-ShowInDashboard=")]
+		[Description("Shows the tool for download on the dashboard")]
+		public bool ShowInDashboard { get; set; } = true;
+
+		[CommandLine(Required = true)]
+		[Description("Source directory for tool data")]
+		public DirectoryReference InputDir { get; set; } = null!;
+
+		[CommandLine]
+		[Description("Directory containing the server to modify")]
+		public DirectoryReference ServerDir { get; set; }
+
+		public ToolDataCommand(IServerInfo serverInfo)
+		{
+			ServerDir = serverInfo.AppDir;
+		}
+
+		public override async Task<int> ExecuteAsync(ILogger logger)
+		{
+			// Create the local agent bundle
+			DirectoryReference bundleDir = DirectoryReference.Combine(ServerDir, "Tools");
+			DirectoryReference.CreateDirectory(bundleDir);
+
+			RefName refName = new RefName(Id);
+			await using (BundleCache bundleCache = new BundleCache())
+			{
+				using (IStorageClient client = BundleStorageClient.CreateFromDirectory(bundleDir, bundleCache, logger))
+				{
+					IBlobRef<DirectoryNode> dirNodeRef;
+					await using (DedupeBlobWriter writer = client.CreateDedupeBlobWriter(refName))
+					{
+						logger.LogInformation("Populating cache with existing refs...");
+						await PopulateCacheAsync(client, writer, bundleDir, CancellationToken.None);
+
+						logger.LogInformation("");
+						logger.LogInformation("Writing tool data for {ToolId}", refName);
+						DirectoryNode dirNode = new DirectoryNode();
+						await dirNode.AddFilesAsync(InputDir.ToDirectoryInfo(), writer);
+						dirNodeRef = await writer.WriteBlobAsync(dirNode);
+
+						logger.LogInformation("");
+						writer.GetStats().Print(logger);
+
+						logger.LogInformation("");
+					}
+					await client.WriteRefAsync(refName, dirNodeRef);
+				}
+			}
+
+			// Update the server config to include the bundled tool
+			FileReference serverConfigFile = FileReference.Combine(ServerDir, "appsettings.json");
+			{
+				JsonConfigFile serverConfig = await JsonConfigFile.ReadAsync(serverConfigFile);
+
+				JsonObject hordeConfig = JsonConfigFile.FindOrAddNode(serverConfig.Root, "Horde", () => new JsonObject());
+				JsonObject pluginsConfig = JsonConfigFile.FindOrAddNode(hordeConfig, "Plugins", () => new JsonObject());
+				JsonObject toolsConfig = JsonConfigFile.FindOrAddNode(serverConfig.Root, "Tools", () => new JsonObject());
+				JsonArray bundledTools = JsonConfigFile.FindOrAddNode(toolsConfig, nameof(ToolsServerConfig.BundledTools), () => new JsonArray());
+
+				JsonObject bundledTool = JsonConfigFile.FindOrAddElementByKey(bundledTools, nameof(BundledToolConfig.Id), Id);
+				bundledTool[nameof(BundledToolConfig.Name)] = Name ?? Id.ToString();
+				if (!String.IsNullOrEmpty(Description))
+				{
+					bundledTool[nameof(BundledToolConfig.Description)] = Description;
+				}
+				if (!String.IsNullOrEmpty(Category))
+				{
+					bundledTool[nameof(BundledToolConfig.Category)] = Category;
+				}
+				if (!String.IsNullOrEmpty(Version))
+				{
+					bundledTool[nameof(BundledToolConfig.Version)] = Version;
+				}
+
+				bundledTool[nameof(BundledToolConfig.RefName)] = refName.ToString();
+
+				if (Public)
+				{
+					bundledTool[nameof(BundledToolConfig.Public)] = true;
+				}
+				if (ShowInUgs)
+				{
+					bundledTool[nameof(BundledToolConfig.ShowInUgs)] = true;
+				}
+
+				logger.LogInformation("Updating {File}", serverConfigFile);
+				await serverConfig.WriteAsync(serverConfigFile);
+			}
+
+			return 0;
+		}
+
+		static async Task PopulateCacheAsync(IStorageClient client, DedupeBlobWriter writer, DirectoryReference searchDir, CancellationToken cancellationToken)
+		{
+			foreach (RefName refName in FileStorageBackend.EnumerateRefs(searchDir))
+			{
+				IBlobRef? blobRef = await client.TryReadRefAsync(refName, cancellationToken: cancellationToken);
+				if (blobRef != null)
+				{
+					using BlobData blobData = await blobRef.ReadBlobDataAsync(cancellationToken);
+					if (blobData.Type.Guid == DirectoryNode.BlobTypeGuid)
+					{
+						IBlobRef<DirectoryNode> directoryRef = BlobRef.Create<DirectoryNode>(blobRef);
+						await writer.AddToCacheAsync(directoryRef, cancellationToken);
+					}
+				}
+			}
+		}
+	}
+}
diff --git a/Engine/Source/Programs/Horde/Plugins/Tools/HordeServer.Tools/ToolsPlugin.cs b/Engine/Source/Programs/Horde/Plugins/Tools/HordeServer.Tools/ToolsPlugin.cs
--- a/Engine/Source/Programs/Horde/Plugins/Tools/HordeServer.Tools/ToolsPlugin.cs
+++ b/Engine/Source/Programs/Horde/Plugins/Tools/HordeServer.Tools/ToolsPlugin.cs
@@ -1,5 +1,7 @@
 // Copyright Epic Games, Inc. All Rights Reserved.
 
+using System.Reflection;
+using EpicGames.Core;
 using HordeServer.Plugins;
 using HordeServer.Tools;
 using Microsoft.Extensions.DependencyInjection;
@@ -16,6 +18,7 @@ namespace HordeServer
 		public void ConfigureServices(IServiceCollection services)
 		{
 			services.AddSingleton<IToolCollection, ToolCollection>();
+			services.AddCommandsFromAssembly(Assembly.GetExecutingAssembly());
 		}
 	}
 
