Horde: Use a [ConfigDoc] attribute to indicate types that need auto-generated documentation.

#rnx

diff --git a/Engine/Source/Programs/Horde/HordeAgent/AgentSettings.cs b/Engine/Source/Programs/Horde/HordeAgent/AgentSettings.cs
--- a/Engine/Source/Programs/Horde/HordeAgent/AgentSettings.cs
+++ b/Engine/Source/Programs/Horde/HordeAgent/AgentSettings.cs
@@ -3,6 +3,7 @@
 using System.ComponentModel.DataAnnotations;
 using EpicGames.Core;
 using EpicGames.Horde;
+using HordeServer.Utilities;
 using Microsoft.Extensions.Configuration;
 
 namespace HordeAgent
@@ -82,6 +83,8 @@ namespace HordeAgent
 	/// <summary>
 	/// Global settings for the agent
 	/// </summary>
+	[ConfigDoc("Agent.json (Agent)", "Deployment/AgentSettings.md", "[Horde](../../README.md) > [Deployment](../Deployment.md) > [Agent](Agent.md)",
+		Introduction = "All Horde-specific settings are stored in a root object called `Horde`. Other .NET functionality may be configured using properties in the root of this file.")]
 	public class AgentSettings
 	{
 		/// <summary>
diff --git a/Engine/Source/Programs/Horde/HordeServer/Commands/Generate/ConfigDocsCommand.cs b/Engine/Source/Programs/Horde/HordeServer/Commands/Generate/ConfigDocsCommand.cs
--- a/Engine/Source/Programs/Horde/HordeServer/Commands/Generate/ConfigDocsCommand.cs
+++ b/Engine/Source/Programs/Horde/HordeServer/Commands/Generate/ConfigDocsCommand.cs
@@ -12,13 +12,7 @@ using System.Threading.Tasks;
 using System.Xml;
 using EpicGames.Core;
 using EpicGames.Horde.Acls;
-using HordeServer.Dashboard;
 using HordeServer.Plugins;
-using HordeServer.Projects;
-using HordeServer.Server;
-using HordeServer.Storage;
-using HordeServer.Streams;
-using HordeServer.Telemetry.Metrics;
 using HordeServer.Utilities;
 using Microsoft.Extensions.Logging;
 
@@ -37,39 +31,20 @@ namespace HordeServer.Commands.Generate
 
 		record class PageInfo(string Title, string LinkRail, string FileName, string? Introduction = null);
 
-		const string AppSettingsIntro = "All Horde-specific settings are stored in a root object called `Horde`. Other .NET functionality may be configured using properties in the root of this file.";
+		readonly IPluginCollection _pluginCollection;
+
+		public DocsCommand(IPluginCollection pluginCollection)
+			=> _pluginCollection = pluginCollection;
 
 		public override async Task<int> ExecuteAsync(ILogger logger)
 		{
 			DirectoryReference.CreateDirectory(OutputDir);
 
-			PluginCollection pluginCollection = new PluginCollection();
-			pluginCollection.Add<AnalyticsPlugin>();
-			pluginCollection.Add<BuildPlugin>();
-			pluginCollection.Add<ComputePlugin>();
-			pluginCollection.Add<DdcPlugin>();
-			pluginCollection.Add<SecretsPlugin>();
-			pluginCollection.Add<StoragePlugin>();
-			pluginCollection.Add<ToolsPlugin>();
-
-			JsonSchemaCache schemaCache = new JsonSchemaCache(pluginCollection);
-
-			JsonSchema serverSchema = schemaCache.CreateSchema(typeof(ServerSettings));
-			JsonSchema globalSchema = schemaCache.CreateSchema(typeof(GlobalConfig));
-			JsonSchema projectSchema = schemaCache.CreateSchema(typeof(ProjectConfig));
-			JsonSchema streamSchema = schemaCache.CreateSchema(typeof(StreamConfig));
-			JsonSchema telemetryConfigSchema = schemaCache.CreateSchema(typeof(TelemetryStoreConfig));
-			JsonSchema dashboardConfigSchema = schemaCache.CreateSchema(typeof(DashboardConfig));
-
-			Dictionary<JsonSchemaType, PageInfo> typeToPageInfo = new Dictionary<JsonSchemaType, PageInfo>
-			{
-				[serverSchema.RootType] = new PageInfo("Server.json", "[Horde](../../README.md) > [Deployment](../Deployment.md) > [Server](Server.md)", "Deployment/ServerSettings.md", AppSettingsIntro),
-				[globalSchema.RootType] = new PageInfo("Globals.json", "[Horde](../../../README.md) > [Configuration](../../Config.md)", "Config/Schema/Globals.md"),
-				[projectSchema.RootType] = new PageInfo("*.project.json", "[Horde](../../../README.md) > [Configuration](../../Config.md)", "Config/Schema/Projects.md"),
-				[streamSchema.RootType] = new PageInfo("*.stream.json", "[Horde](../../../README.md) > [Configuration](../../Config.md)", "Config/Schema/Streams.md"),
-				[telemetryConfigSchema.RootType] = new PageInfo("*.telemetry.json", "[Horde](../../../README.md) > [Configuration](../../Config.md)", "Config/Schema/Telemetry.md"),
-				[dashboardConfigSchema.RootType] = new PageInfo("*.dashboard.json", "[Horde](../../../README.md) > [Configuration](../../Config.md)", "Config/Schema/Dashboard.md"),
-			};
+			JsonSchemaCache schemaCache = new JsonSchemaCache(_pluginCollection);
+
+			HashSet<Assembly> assemblies = new HashSet<Assembly>();
+			assemblies.Add(Assembly.GetExecutingAssembly());
+			assemblies.UnionWith(_pluginCollection.LoadedPlugins.Select(x => x.Assembly));
 
 			if (Agent == null)
 			{
@@ -77,11 +52,18 @@ namespace HordeServer.Commands.Generate
 			}
 			else
 			{
-				Assembly agentAssembly = Assembly.LoadFile(Agent.FullName);
-				Type agentSettingsType = agentAssembly.GetType("HordeAgent.AgentSettings")!;
+				assemblies.Add(Assembly.LoadFile(Agent.FullName));
+			}
 
-				JsonSchema agentSchema = schemaCache.CreateSchema(agentSettingsType);
-				await WriteDocAsync(agentSchema.RootType, "Agent.json (Agent)", "Deployment/AgentSettings.md", "[Horde](../../README.md) > [Deployment](../Deployment.md) > [Agent](Agent.md)", AppSettingsIntro, new Dictionary<string, string>(), logger);
+			Dictionary<JsonSchemaType, PageInfo> typeToPageInfo = new Dictionary<JsonSchemaType, PageInfo>();
+			foreach (Type type in assemblies.SelectMany(x => x.GetExportedTypes()))
+			{
+				ConfigDocAttribute? attribute = type.GetCustomAttribute<ConfigDocAttribute>();
+				if (attribute != null)
+				{
+					JsonSchema schema = schemaCache.CreateSchema(type);
+					typeToPageInfo.Add(schema.RootType, new PageInfo(attribute.Title, attribute.LinkRail, attribute.FileName, attribute.Introduction));
+				}
 			}
 
 			Dictionary<string, string> typeNameToLink = typeToPageInfo.ToDictionary(x => x.Key.Name!, x => Path.GetFileName(x.Value.FileName), StringComparer.Ordinal);
@@ -90,7 +72,7 @@ namespace HordeServer.Commands.Generate
 				await WriteDocAsync(type, pageInfo.Title, pageInfo.FileName, pageInfo.LinkRail, pageInfo.Introduction, typeNameToLink, logger);
 			}
 
-			await WriteAclDocAsync(Assembly.GetExecutingAssembly(), logger);
+			await WriteAclDocAsync(assemblies, logger);
 			return 0;
 		}
 
@@ -335,7 +317,7 @@ namespace HordeServer.Commands.Generate
 			return "#" + anchor.Trim('-');
 		}
 
-		async Task WriteAclDocAsync(Assembly assembly, ILogger logger)
+		async Task WriteAclDocAsync(IEnumerable<Assembly> assemblies, ILogger logger)
 		{
 			FileReference file = FileReference.Combine(OutputDir, "Config/Schema/AclActions.md");
 			DirectoryReference.CreateDirectory(file.Directory);
@@ -351,70 +333,73 @@ namespace HordeServer.Commands.Generate
 					Dictionary<string, List<PropertyInfo>> categoryToProperties = new Dictionary<string, List<PropertyInfo>>();
 
 					List<PropertyInfo> newProperties = new List<PropertyInfo>();
-					foreach (Type type in assembly.GetTypes())
+					foreach (Assembly assembly in assemblies)
 					{
-						if (type.IsClass)
+						foreach (Type type in assembly.GetTypes())
 						{
-							foreach (PropertyInfo propertyInfo in type.GetProperties(BindingFlags.Static | BindingFlags.Public))
-							{
-								if (propertyInfo.PropertyType == typeof(AclAction))
-								{
-									newProperties.Add(propertyInfo);
-								}
-							}
-							if (newProperties.Count > 0)
+							if (type.IsClass)
 							{
-								string category = type.Namespace ?? "Default";
-
-								int topNamespaceIdx = category.LastIndexOf('.');
-								if (topNamespaceIdx != -1)
+								foreach (PropertyInfo propertyInfo in type.GetProperties(BindingFlags.Static | BindingFlags.Public))
 								{
-									category = category.Substring(topNamespaceIdx + 1);
+									if (propertyInfo.PropertyType == typeof(AclAction))
+									{
+										newProperties.Add(propertyInfo);
+									}
 								}
-
-								List<PropertyInfo>? existingProperties;
-								if (categoryToProperties.TryGetValue(category, out existingProperties))
-								{
-									existingProperties.AddRange(newProperties);
-									newProperties.Clear();
-								}
-								else
+								if (newProperties.Count > 0)
 								{
-									categoryToProperties.Add(category, newProperties);
-									newProperties = new List<PropertyInfo>();
+									string category = type.Namespace ?? "Default";
+
+									int topNamespaceIdx = category.LastIndexOf('.');
+									if (topNamespaceIdx != -1)
+									{
+										category = category.Substring(topNamespaceIdx + 1);
+									}
+
+									List<PropertyInfo>? existingProperties;
+									if (categoryToProperties.TryGetValue(category, out existingProperties))
+									{
+										existingProperties.AddRange(newProperties);
+										newProperties.Clear();
+									}
+									else
+									{
+										categoryToProperties.Add(category, newProperties);
+										newProperties = new List<PropertyInfo>();
+									}
 								}
 							}
 						}
-					}
 
-					XmlDocument documentation = new XmlDocument();
+						XmlDocument documentation = new XmlDocument();
 
-					FileReference inputDocumentationFile = new FileReference(assembly.Location).ChangeExtension(".xml");
-					if (FileReference.Exists(inputDocumentationFile))
-					{
-						documentation.Load(inputDocumentationFile.FullName);
-					}
-
-					foreach ((string category, List<PropertyInfo> properties) in categoryToProperties.OrderBy(x => x.Key))
-					{
-						await writer.WriteLineAsync();
-						await writer.WriteLineAsync($"## {category}");
-						await writer.WriteLineAsync();
-						await writer.WriteLineAsync("| Name | Description |");
-						await writer.WriteLineAsync("| ---- | ----------- |");
-						foreach (PropertyInfo property in properties)
+						FileReference inputDocumentationFile = new FileReference(assembly.Location).ChangeExtension(".xml");
+						if (FileReference.Exists(inputDocumentationFile))
 						{
-							AclAction action = (AclAction)property.GetValue(null)!;
-							string description = String.Empty;
+							documentation.Load(inputDocumentationFile.FullName);
+						}
 
-							string selector = $"//member[@name='P:{property.DeclaringType!.FullName}.{property.Name}']/summary";
-							XmlNode? node = documentation.SelectSingleNode(selector);
-							if (node != null)
+						foreach ((string category, List<PropertyInfo> properties) in categoryToProperties.OrderBy(x => x.Key))
+						{
+							await writer.WriteLineAsync();
+							await writer.WriteLineAsync($"## {category}");
+							await writer.WriteLineAsync();
+							await writer.WriteLineAsync("| Name | Description |");
+							await writer.WriteLineAsync("| ---- | ----------- |");
+							foreach (PropertyInfo property in properties)
 							{
-								description = node.InnerText.Trim().Replace("\r\n", "\n", StringComparison.Ordinal);
-							}
+								AclAction action = (AclAction)property.GetValue(null)!;
+								string description = String.Empty;
 
-							await writer.WriteLineAsync($"| `{action.Name}` | {description} |");
+								string selector = $"//member[@name='P:{property.DeclaringType!.FullName}.{property.Name}']/summary";
+								XmlNode? node = documentation.SelectSingleNode(selector);
+								if (node != null)
+								{
+									description = node.InnerText.Trim().Replace("\r\n", "\n", StringComparison.Ordinal);
+								}
+
+								await writer.WriteLineAsync($"| `{action.Name}` | {description} |");
+							}
 						}
 					}
 				}
diff --git a/Engine/Source/Programs/Horde/HordeServer/Commands/Generate/ConfigSchemasCommand.cs b/Engine/Source/Programs/Horde/HordeServer/Commands/Generate/ConfigSchemasCommand.cs
--- a/Engine/Source/Programs/Horde/HordeServer/Commands/Generate/ConfigSchemasCommand.cs
+++ b/Engine/Source/Programs/Horde/HordeServer/Commands/Generate/ConfigSchemasCommand.cs
@@ -18,11 +18,16 @@ namespace HordeServer.Commands.Generate
 		[Description("Output directory to write schemas to. Defaults to the 'Schemas' subfolder of the application directory.")]
 		DirectoryReference? _outputDir = null!;
 
+		readonly IPluginCollection _pluginCollection;
+
+		public SchemasCommand(IPluginCollection pluginCollection)
+			=> _pluginCollection = pluginCollection;
+
 		public override Task<int> ExecuteAsync(ILogger logger)
 		{
 			_outputDir ??= DirectoryReference.Combine(ServerApp.AppDir, "Schemas");
 
-			JsonSchemaCache schemaCache = new JsonSchemaCache(new PluginCollection());
+			JsonSchemaCache schemaCache = new JsonSchemaCache(_pluginCollection);
 
 			DirectoryReference.CreateDirectory(_outputDir);
 			foreach (Type schemaType in SchemaController.ConfigSchemas)
diff --git a/Engine/Source/Programs/Horde/HordeServer/Dashboard/DashboardConfig.cs b/Engine/Source/Programs/Horde/HordeServer/Dashboard/DashboardConfig.cs
--- a/Engine/Source/Programs/Horde/HordeServer/Dashboard/DashboardConfig.cs
+++ b/Engine/Source/Programs/Horde/HordeServer/Dashboard/DashboardConfig.cs
@@ -4,6 +4,7 @@ using System.Collections.Generic;
 using EpicGames.Core;
 using EpicGames.Horde.Common;
 using HordeServer.Configuration;
+using HordeServer.Utilities;
 
 namespace HordeServer.Dashboard
 {
@@ -12,6 +13,7 @@ namespace HordeServer.Dashboard
 	/// </summary>
 	[JsonSchema("https://unrealengine.com/horde/dashboard")]
 	[JsonSchemaCatalog("Horde Dashboard", "Horde dashboard configuration file", new[] { "*.dashboard.json", "Dashboard/*.json" })]
+	[ConfigDoc("*.dashboard.json", "[Horde](../../../README.md) > [Configuration](../../Config.md)", "Config/Schema/Dashboard.md")]
 	[ConfigIncludeRoot]
 	[ConfigMacroScope]
 	public class DashboardConfig
diff --git a/Engine/Source/Programs/Horde/HordeServer/Server/GlobalConfig.cs b/Engine/Source/Programs/Horde/HordeServer/Server/GlobalConfig.cs
--- a/Engine/Source/Programs/Horde/HordeServer/Server/GlobalConfig.cs
+++ b/Engine/Source/Programs/Horde/HordeServer/Server/GlobalConfig.cs
@@ -38,6 +38,7 @@ namespace HordeServer.Server
 	/// </summary>
 	[JsonSchema("https://unrealengine.com/horde/global")]
 	[JsonSchemaCatalog("Horde Globals", "Horde global configuration file", new[] { "globals.json", "*.global.json" })]
+	[ConfigDoc("Globals.json", "[Horde](../../../README.md) > [Configuration](../../Config.md)", "Config/Schema/Globals.md")]
 	[ConfigIncludeRoot]
 	[ConfigMacroScope]
 	public class GlobalConfig
diff --git a/Engine/Source/Programs/Horde/HordeServer/ServerApp.cs b/Engine/Source/Programs/Horde/HordeServer/ServerApp.cs
--- a/Engine/Source/Programs/Horde/HordeServer/ServerApp.cs
+++ b/Engine/Source/Programs/Horde/HordeServer/ServerApp.cs
@@ -185,6 +185,8 @@ namespace HordeServer
 			services.AddSingleton<IServerInfo>(serverInfo);
 
 			IPluginCollection pluginCollection = ConfigurePlugins();
+			services.AddSingleton<IPluginCollection>(pluginCollection);
+
 			foreach (Assembly pluginAssembly in pluginCollection.LoadedPlugins.Select(x => x.Assembly).Distinct())
 			{
 				services.AddCommandsFromAssembly(pluginAssembly);
diff --git a/Engine/Source/Programs/Horde/HordeServer/ServerSettings.cs b/Engine/Source/Programs/Horde/HordeServer/ServerSettings.cs
--- a/Engine/Source/Programs/Horde/HordeServer/ServerSettings.cs
+++ b/Engine/Source/Programs/Horde/HordeServer/ServerSettings.cs
@@ -6,6 +6,7 @@ using System.Linq;
 using System.Net;
 using EpicGames.Core;
 using EpicGames.Horde.Server;
+using HordeServer.Utilities;
 
 namespace HordeServer
 {
@@ -85,6 +86,9 @@ namespace HordeServer
 	/// <summary>
 	/// Global settings for the application
 	/// </summary>
+	[ConfigDoc("Server.json", "[Horde](../../README.md) > [Deployment](../Deployment.md) > [Server](Server.md)", "Deployment/ServerSettings.md", 
+		Introduction = "All Horde-specific settings are stored in a root object called `Horde`. Other .NET functionality may be configured using properties in the root of this file."
+	)]
 	public class ServerSettings
 	{
 		/// <summary>
diff --git a/Engine/Source/Programs/Horde/Plugins/Analytics/HordeServer.Analytics/Telemetry/Metrics/MetricConfig.cs b/Engine/Source/Programs/Horde/Plugins/Analytics/HordeServer.Analytics/Telemetry/Metrics/MetricConfig.cs
--- a/Engine/Source/Programs/Horde/Plugins/Analytics/HordeServer.Analytics/Telemetry/Metrics/MetricConfig.cs
+++ b/Engine/Source/Programs/Horde/Plugins/Analytics/HordeServer.Analytics/Telemetry/Metrics/MetricConfig.cs
@@ -128,6 +128,7 @@ namespace HordeServer.Telemetry.Metrics
 	/// </summary>
 	[JsonSchema("https://unrealengine.com/horde/telemetry")]
 	[JsonSchemaCatalog("Horde Telemetry", "Horde telemetry configuration file", new[] { "*.telemetry.json", "*.metrics.json", "Metrics/*.json" })]
+	[ConfigDoc("*.telemetry.json", "[Horde](../../../README.md) > [Configuration](../../Config.md)", "Config/Schema/Telemetry.md")]
 	[ConfigIncludeRoot]
 	[ConfigMacroScope]
 	public class TelemetryStoreConfig
diff --git a/Engine/Source/Programs/Horde/Plugins/Build/HordeServer.Build/Projects/ProjectConfig.cs b/Engine/Source/Programs/Horde/Plugins/Build/HordeServer.Build/Projects/ProjectConfig.cs
--- a/Engine/Source/Programs/Horde/Plugins/Build/HordeServer.Build/Projects/ProjectConfig.cs
+++ b/Engine/Source/Programs/Horde/Plugins/Build/HordeServer.Build/Projects/ProjectConfig.cs
@@ -16,6 +16,7 @@ using HordeServer.Agents.Pools;
 using HordeServer.Configuration;
 using HordeServer.Server;
 using HordeServer.Streams;
+using HordeServer.Utilities;
 
 #pragma warning disable CA2227 // Change x to be read-only by removing the property setter
 
@@ -26,6 +27,7 @@ namespace HordeServer.Projects
 	/// </summary>
 	[JsonSchema("https://unrealengine.com/horde/project")]
 	[JsonSchemaCatalog("Horde Project", "Horde project configuration file", new[] { "*.project.json", "Projects/*.json" })]
+	[ConfigDoc("*.project.json", "[Horde](../../../README.md) > [Configuration](../../Config.md)", "Config/Schema/Projects.md")]
 	[ConfigIncludeRoot]
 	[ConfigMacroScope]
 	[DebuggerDisplay("{Id}")]
diff --git a/Engine/Source/Programs/Horde/Plugins/Build/HordeServer.Build/Streams/StreamConfig.cs b/Engine/Source/Programs/Horde/Plugins/Build/HordeServer.Build/Streams/StreamConfig.cs
--- a/Engine/Source/Programs/Horde/Plugins/Build/HordeServer.Build/Streams/StreamConfig.cs
+++ b/Engine/Source/Programs/Horde/Plugins/Build/HordeServer.Build/Streams/StreamConfig.cs
@@ -81,6 +81,7 @@ namespace HordeServer.Streams
 	/// </summary>
 	[JsonSchema("https://unrealengine.com/horde/stream")]
 	[JsonSchemaCatalog("Horde Stream", "Horde stream configuration file", new[] { "*.stream.json", "Streams/*.json" })]
+	[ConfigDoc("*.stream.json", "[Horde](../../../README.md) > [Configuration](../../Config.md)", "Config/Schema/Streams.md")]
 	[ConfigIncludeRoot]
 	[ConfigMacroScope]
 	[DebuggerDisplay("{Id}")]
diff --git a/Engine/Source/Programs/Shared/EpicGames.Horde/Utilities/ConfigDocAttribute.cs b/Engine/Source/Programs/Shared/EpicGames.Horde/Utilities/ConfigDocAttribute.cs
new file mode 100644
--- /dev/null
+++ b/Engine/Source/Programs/Shared/EpicGames.Horde/Utilities/ConfigDocAttribute.cs
@@ -0,0 +1,46 @@
+// Copyright Epic Games, Inc. All Rights Reserved.
+
+using System;
+using System.Collections.Generic;
+using System.Linq;
+using System.Text;
+using System.Threading.Tasks;
+
+namespace HordeServer.Utilities
+{
+	/// <summary>
+	/// Attribute indicating that an object should generate a schema doc page
+	/// </summary>
+	public sealed class ConfigDocAttribute : Attribute
+	{
+		/// <summary>
+		/// Page title
+		/// </summary>
+		public string Title { get; }
+
+		/// <summary>
+		/// Rail to show with breadcrumbs at the top of the page
+		/// </summary>
+		public string LinkRail { get; }
+
+		/// <summary>
+		/// Output filename
+		/// </summary>
+		public string FileName { get; }
+
+		/// <summary>
+		/// Optional introductory text on the page
+		/// </summary>
+		public string? Introduction { get; set; }
+
+		/// <summary>
+		/// Constructor
+		/// </summary>
+		public ConfigDocAttribute(string title, string linkRail, string fileName)
+		{
+			Title = title;
+			LinkRail = linkRail;
+			FileName = fileName;
+		}
+	}
+}
