Horde: Allow stripping resources from config files without needing to access project config options explicitly.

#rnx

diff --git a/Engine/Source/Programs/Horde/HordeServer.Shared/Configuration/ConfigResource.cs b/Engine/Source/Programs/Horde/HordeServer.Shared/Configuration/ConfigResource.cs
--- a/Engine/Source/Programs/Horde/HordeServer.Shared/Configuration/ConfigResource.cs
+++ b/Engine/Source/Programs/Horde/HordeServer.Shared/Configuration/ConfigResource.cs
@@ -29,8 +29,23 @@ namespace HordeServer.Configuration
 	/// <summary>
 	/// Serializer for <see cref="ConfigResource"/> objects
 	/// </summary>
-	class ConfigResourceSerializer : JsonConverter<ConfigResource>
+	public class ConfigResourceSerializer : JsonConverter<ConfigResource>
 	{
+		readonly bool _withData;
+
+		/// <summary>
+		/// Explicit default constructor (cannot use default argument with JsonSerializer)
+		/// </summary>
+		public ConfigResourceSerializer()
+			: this(true) { }
+
+		/// <summary>
+		/// Constructor
+		/// </summary>
+		/// <param name="withData">Whether to include data in serialized resource objects</param>
+		public ConfigResourceSerializer(bool withData)
+			=> _withData = withData;
+
 		/// <inheritdoc/>
 		public override ConfigResource Read(ref Utf8JsonReader reader, Type typeToConvert, JsonSerializerOptions options)
 		{
@@ -50,7 +65,7 @@ namespace HordeServer.Configuration
 							reader.Read();
 							resource.Path = reader.GetString();
 						}
-						else if (reader.ValueTextEquals(nameof(ConfigResource.Data)))
+						else if (reader.ValueTextEquals(nameof(ConfigResource.Data)) && _withData)
 						{
 							reader.Read();
 							resource.Data = reader.GetBytesFromBase64();
@@ -68,7 +83,7 @@ namespace HordeServer.Configuration
 		/// <inheritdoc/>
 		public override void Write(Utf8JsonWriter writer, ConfigResource resource, JsonSerializerOptions options)
 		{
-			if (resource.Data.Length == 0)
+			if (resource.Data.Length == 0 || !_withData)
 			{
 				writer.WriteStringValue(resource.Path);
 			}
diff --git a/Engine/Source/Programs/Horde/HordeServer/Configuration/ConfigService.cs b/Engine/Source/Programs/Horde/HordeServer/Configuration/ConfigService.cs
--- a/Engine/Source/Programs/Horde/HordeServer/Configuration/ConfigService.cs
+++ b/Engine/Source/Programs/Horde/HordeServer/Configuration/ConfigService.cs
@@ -723,9 +723,12 @@ namespace HordeServer.Configuration
 			return JsonSerializer.SerializeToUtf8Bytes(config, _jsonOptions);
 		}
 
-		internal GlobalConfig? Deserialize(byte[] data)
+		internal GlobalConfig? Deserialize(byte[] data, bool withData)
 		{
-			return JsonSerializer.Deserialize<GlobalConfig>(data, _jsonOptions);
+			JsonSerializerOptions options = new JsonSerializerOptions(_jsonOptions);
+			options.Converters.Add(new ConfigResourceSerializer(withData));
+
+			return JsonSerializer.Deserialize<GlobalConfig>(data, options);
 		}
 
 		/// <summary>
diff --git a/Engine/Source/Programs/Horde/HordeServer/Server/DebugController.cs b/Engine/Source/Programs/Horde/HordeServer/Server/DebugController.cs
--- a/Engine/Source/Programs/Horde/HordeServer/Server/DebugController.cs
+++ b/Engine/Source/Programs/Horde/HordeServer/Server/DebugController.cs
@@ -20,7 +20,6 @@ using Horde.Common.Rpc;
 using HordeServer.Agents.Relay;
 using HordeServer.Configuration;
 using HordeServer.Logs;
-using HordeServer.Projects;
 using HordeServer.Utilities;
 using JetBrains.Profiler.SelfApi;
 using Microsoft.AspNetCore.Authorization;
@@ -145,13 +144,7 @@ namespace HordeServer.Server
 
 			// Duplicate the config, so we can redact stuff that we don't want to return through the browser
 			byte[] data = _configService.Serialize(_globalConfig.Value);
-			GlobalConfig config = _configService.Deserialize(data)!;
-
-			foreach (ProjectConfig project in config.Plugins.GetBuildConfig().Projects)
-			{
-				project.Logo = null;
-			}
-
+			GlobalConfig config = _configService.Deserialize(data, false)!;
 			return config;
 		}
 
