Horde: Move debug endpoints for compute/storage system into their respective plugins.

#rnx

diff --git a/Engine/Source/Programs/Horde/HordeServer/Server/DebugController.cs b/Engine/Source/Programs/Horde/HordeServer/Server/DebugController.cs
--- a/Engine/Source/Programs/Horde/HordeServer/Server/DebugController.cs
+++ b/Engine/Source/Programs/Horde/HordeServer/Server/DebugController.cs
@@ -11,31 +11,21 @@ using System.Threading;
 using System.Threading.Tasks;
 using System.Web;
 using EpicGames.Core;
-using EpicGames.Horde.Agents.Leases;
-using EpicGames.Horde.Compute;
-using EpicGames.Horde.Logs;
-using EpicGames.Horde.Storage;
-using Google.Protobuf;
-using Horde.Common.Rpc;
-using HordeServer.Agents.Relay;
 using HordeServer.Configuration;
-using HordeServer.Logs;
 using HordeServer.Utilities;
 using JetBrains.Profiler.SelfApi;
 using Microsoft.AspNetCore.Authorization;
 using Microsoft.AspNetCore.Http;
 using Microsoft.AspNetCore.Http.Features;
 using Microsoft.AspNetCore.Mvc;
-using Microsoft.Extensions.DependencyInjection;
 using Microsoft.Extensions.Logging;
 using Microsoft.Extensions.Options;
-using MongoDB.Bson;
 using MongoDB.Driver;
 
 namespace HordeServer.Server
 {
 	/// <summary>
-	/// Controller managing account status
+	/// Controller providing some debug functionality
 	/// </summary>
 	[ApiController]
 	[Authorize]
@@ -46,10 +36,7 @@ namespace HordeServer.Server
 		private static readonly Random s_random = new();
 
 		private readonly IMongoService _mongoService;
-		private readonly IServiceProvider _serviceProvider;
 		private readonly ConfigService _configService;
-		private readonly AgentRelayService _agentRelayService;
-		private readonly ILogCollection _logCollection;
 		private readonly IOptionsSnapshot<GlobalConfig> _globalConfig;
 		private readonly ILogger<DebugController> _logger;
 
@@ -58,18 +45,12 @@ namespace HordeServer.Server
 		/// </summary>
 		public DebugController(
 			IMongoService mongoService,
-			IServiceProvider serviceProvider,
 			ConfigService configService,
-			AgentRelayService agentRelayService,
-			ILogCollection logCollection,
 			IOptionsSnapshot<GlobalConfig> globalConfig,
 			ILogger<DebugController> logger)
 		{
 			_mongoService = mongoService;
-			_serviceProvider = serviceProvider;
 			_configService = configService;
-			_agentRelayService = agentRelayService;
-			_logCollection = logCollection;
 			_globalConfig = globalConfig;
 			_logger = logger;
 		}
@@ -197,43 +178,6 @@ namespace HordeServer.Server
 			return Empty;
 		}
 
-		/// <summary>
-		/// Add a port mapping for agent relay
-		/// </summary>
-		[HttpGet]
-		[Route("/api/v1/debug/relay/add-port")]
-		public async Task<ActionResult<object>> AddPortMappingAsync([FromQuery] string? clientIpStr = null, [FromQuery] string? agentIpStr = null, [FromQuery] int? agentPort = null)
-		{
-			if (!_globalConfig.Value.Authorize(ServerAclAction.Debug, User))
-			{
-				return Forbid(ServerAclAction.Debug);
-			}
-
-			if (clientIpStr == null || !IPAddress.TryParse(clientIpStr, out IPAddress? clientIp))
-			{
-				return BadRequest("Unable to read or convert query parameter 'clientIp'");
-			}
-
-			if (agentIpStr == null || !IPAddress.TryParse(agentIpStr, out IPAddress? agentIp))
-			{
-				return BadRequest("Unable to read or convert query parameter 'agentIp'");
-			}
-
-			if (agentPort == null)
-			{
-				return BadRequest("Bad query parameter 'agentPort'");
-			}
-
-			string bogusLeaseId = ObjectId.GenerateNewId().ToString();
-			List<Port> ports = new()
-			{
-				new Port { RelayPort = -1, AgentPort = agentPort.Value, Protocol = PortProtocol.Tcp }
-			};
-
-			PortMapping portMapping = await _agentRelayService.AddPortMappingAsync(new ClusterId("default"), LeaseId.Parse(bogusLeaseId), clientIp, agentIp, ports);
-			return JsonFormatter.Default.Format(portMapping);
-		}
-
 		/// <summary>
 		/// Generate log message of varying size
 		/// </summary>
@@ -291,30 +235,6 @@ namespace HordeServer.Server
 			return Ok($"Log message generated logLevel={logLevelInternal} messageLen={messageLen} exceptionMessageLen={exceptionMessageLen} argCount={argCount} argLen={argLen}");
 		}
 
-		/// <summary>
-		/// Retrieve metadata about a specific log file
-		/// </summary>
-		/// <param name="logId">Id of the log file to get information about</param>
-		/// <param name="filter">Filter for the properties to return</param>
-		/// <returns>Information about the requested project</returns>
-		[HttpGet]
-		[Route("/api/v1/debug/logs/{logId}")]
-		public async Task<ActionResult<object>> GetLogAsync(LogId logId, [FromQuery] PropertyFilter? filter = null)
-		{
-			if (!_globalConfig.Value.Authorize(ServerAclAction.Debug, User))
-			{
-				return Forbid(ServerAclAction.Debug);
-			}
-
-			ILog? log = await _logCollection.GetAsync(logId, CancellationToken.None);
-			if (log == null)
-			{
-				return NotFound();
-			}
-
-			return log.ApplyFilter(filter);
-		}
-
 		/// <summary>
 		/// Populate the database with test data
 		/// </summary>
@@ -461,22 +381,5 @@ namespace HordeServer.Server
 			string stringArg = "hello";
 			throw new Exception($"Message: numberArg:{numberArg}, stringArg:{stringArg}");
 		}
-
-		/// <summary>
-		/// Writes stats for the storage backend cache
-		/// </summary>
-		[HttpPost]
-		[Route("/api/v1/debug/writecacherefstats")]
-		public ActionResult WriteCacheRefStats()
-		{
-			if (!_globalConfig.Value.Authorize(ServerAclAction.Debug, User))
-			{
-				return Forbid(ServerAclAction.Debug);
-			}
-
-			StorageBackendCache storageBackendCache = _serviceProvider.GetRequiredService<StorageBackendCache>();
-			storageBackendCache.WriteRefStats(_logger);
-			return Ok();
-		}
 	}
 }
diff --git a/Engine/Source/Programs/Horde/Plugins/Compute/HordeServer.Compute/Debug/ComputeDebugController.cs b/Engine/Source/Programs/Horde/Plugins/Compute/HordeServer.Compute/Debug/ComputeDebugController.cs
--- a/Engine/Source/Programs/Horde/Plugins/Compute/HordeServer.Compute/Debug/ComputeDebugController.cs
+++ b/Engine/Source/Programs/Horde/Plugins/Compute/HordeServer.Compute/Debug/ComputeDebugController.cs
@@ -6,23 +6,76 @@ using System.Linq;
 using System.Net;
 using System.Text;
 using System.Threading.Tasks;
+using EpicGames.Horde.Agents.Leases;
+using EpicGames.Horde.Compute;
+using Google.Protobuf;
+using Horde.Common.Rpc;
+using HordeServer.Agents.Relay;
 using HordeServer.Server;
 using HordeServer.Utilities;
+using Microsoft.AspNetCore.Authorization;
 using Microsoft.AspNetCore.Http;
 using Microsoft.AspNetCore.Mvc;
 using Microsoft.Extensions.Options;
+using MongoDB.Bson;
 
 namespace HordeServer.Debug
 {
+	/// <summary>
+	/// Debug functionality for the compute system
+	/// </summary>
+	[ApiController]
+	[Authorize]
+	[DebugEndpoint]
+	[Tags("Debug")]
 	class ComputeDebugController : HordeControllerBase
 	{
+		readonly AgentRelayService _agentRelayService;
 		readonly IOptionsSnapshot<ComputeConfig> _computeConfig;
 
-		public ComputeDebugController(IOptionsSnapshot<ComputeConfig> computeConfig)
+		public ComputeDebugController(AgentRelayService agentRelayService, IOptionsSnapshot<ComputeConfig> computeConfig)
 		{
+			_agentRelayService = agentRelayService;
 			_computeConfig = computeConfig;
 		}
 
+		/// <summary>
+		/// Add a port mapping for agent relay
+		/// </summary>
+		[HttpGet]
+		[Route("/api/v1/debug/relay/add-port")]
+		public async Task<ActionResult<object>> AddPortMappingAsync([FromQuery] string? clientIpStr = null, [FromQuery] string? agentIpStr = null, [FromQuery] int? agentPort = null)
+		{
+			if (!_computeConfig.Value.Authorize(ServerAclAction.Debug, User))
+			{
+				return Forbid(ServerAclAction.Debug);
+			}
+
+			if (clientIpStr == null || !IPAddress.TryParse(clientIpStr, out IPAddress? clientIp))
+			{
+				return BadRequest("Unable to read or convert query parameter 'clientIp'");
+			}
+
+			if (agentIpStr == null || !IPAddress.TryParse(agentIpStr, out IPAddress? agentIp))
+			{
+				return BadRequest("Unable to read or convert query parameter 'agentIp'");
+			}
+
+			if (agentPort == null)
+			{
+				return BadRequest("Bad query parameter 'agentPort'");
+			}
+
+			string bogusLeaseId = ObjectId.GenerateNewId().ToString();
+			List<Port> ports = new()
+			{
+				new Port { RelayPort = -1, AgentPort = agentPort.Value, Protocol = PortProtocol.Tcp }
+			};
+
+			PortMapping portMapping = await _agentRelayService.AddPortMappingAsync(new ClusterId("default"), LeaseId.Parse(bogusLeaseId), clientIp, agentIp, ports);
+			return JsonFormatter.Default.Format(portMapping);
+		}
+
 		/// <summary>
 		/// Get the network ID for a given IP address
 		/// </summary>
diff --git a/Engine/Source/Programs/Horde/Plugins/Storage/HordeServer.Storage/Debug/StorageDebugController.cs b/Engine/Source/Programs/Horde/Plugins/Storage/HordeServer.Storage/Debug/StorageDebugController.cs
new file mode 100644
--- /dev/null
+++ b/Engine/Source/Programs/Horde/Plugins/Storage/HordeServer.Storage/Debug/StorageDebugController.cs
@@ -0,0 +1,56 @@
+// Copyright Epic Games, Inc. All Rights Reserved.
+
+using EpicGames.Horde.Storage;
+using HordeServer.Server;
+using HordeServer.Storage;
+using HordeServer.Utilities;
+using Microsoft.AspNetCore.Authorization;
+using Microsoft.AspNetCore.Http;
+using Microsoft.AspNetCore.Mvc;
+using Microsoft.Extensions.DependencyInjection;
+using Microsoft.Extensions.Logging;
+using Microsoft.Extensions.Options;
+
+namespace HordeServer.Debug
+{
+	/// <summary>
+	/// Debug functionality for the storage system
+	/// </summary>
+	[ApiController]
+	[Authorize]
+	[DebugEndpoint]
+	[Tags("Debug")]
+	class StorageDebugController : HordeControllerBase
+	{
+		readonly IServiceProvider _serviceProvider;
+		readonly IOptionsSnapshot<StorageConfig> _storageConfig;
+		readonly ILogger _logger;
+
+		/// <summary>
+		/// Constructor
+		/// </summary>
+		public StorageDebugController(IServiceProvider serviceProvider, IOptionsSnapshot<StorageConfig> storageConfig, ILogger<StorageDebugController> logger)
+		{
+			_serviceProvider = serviceProvider;
+			_storageConfig = storageConfig;
+			_logger = logger;
+		}
+
+		/// <summary>
+		/// Writes stats for the storage backend cache
+		/// </summary>
+		[HttpPost]
+		[Route("/api/v1/debug/writecacherefstats")]
+		public ActionResult WriteCacheRefStats()
+		{
+			if (!_storageConfig.Value.Acl.Authorize(ServerAclAction.Debug, User))
+			{
+				return Forbid(ServerAclAction.Debug);
+			}
+
+			StorageBackendCache storageBackendCache = _serviceProvider.GetRequiredService<StorageBackendCache>();
+			storageBackendCache.WriteRefStats(_logger);
+			return Ok();
+		}
+	}
+}
diff --git a/Engine/Source/Programs/Horde/Plugins/Storage/HordeServer.Storage/StorageConfig.cs b/Engine/Source/Programs/Horde/Plugins/Storage/HordeServer.Storage/StorageConfig.cs
--- a/Engine/Source/Programs/Horde/Plugins/Storage/HordeServer.Storage/StorageConfig.cs
+++ b/Engine/Source/Programs/Horde/Plugins/Storage/HordeServer.Storage/StorageConfig.cs
@@ -55,6 +55,12 @@ namespace HordeServer.Storage
 	/// </summary>
 	public class StorageConfig : IPluginConfig
 	{
+		/// <summary>
+		/// ACL for the storage system
+		/// </summary>
+		[JsonIgnore]
+		public AclConfig Acl { get; set; } = null!;
+
 		/// <summary>
 		/// Hash of the current object revision to allow detecting changes
 		/// </summary>
@@ -89,6 +95,8 @@ namespace HordeServer.Storage
 		/// </summary>
 		public void PostLoad(PluginConfigOptions configOptions)
 		{
+			Acl = configOptions.ParentAcl;
+
 			// Create a lookup for backend configs, allowing later entries to override previous ones
 			_backendLookup.Clear();
 			foreach (BackendConfig backendConfig in Backends)
