Horde: Add an interface to allow the /api/v1/server/info endpoint to query the currently deployed agent version without directly referencing the compute plugin.

#rnx

diff --git a/Engine/Source/Programs/Horde/HordeServer.Shared/Server/IAgentVersionProvider.cs b/Engine/Source/Programs/Horde/HordeServer.Shared/Server/IAgentVersionProvider.cs
new file mode 100644
--- /dev/null
+++ b/Engine/Source/Programs/Horde/HordeServer.Shared/Server/IAgentVersionProvider.cs
@@ -0,0 +1,22 @@
+// Copyright Epic Games, Inc. All Rights Reserved.
+
+using System;
+using System.Collections.Generic;
+using System.Linq;
+using System.Text;
+using System.Threading.Tasks;
+
+namespace HordeServer.Server
+{
+	/// <summary>
+	/// Allows subsystems to advertise the version of bundled tools via the /server/info endpoint
+	/// </summary>
+	public interface IAgentVersionProvider
+	{
+		/// <summary>
+		/// Returns the version number of the Horde agent
+		/// </summary>
+		Task<string?> GetAsync(CancellationToken cancellationToken = default);
+	}
+}
+
diff --git a/Engine/Source/Programs/Horde/HordeServer/Server/ServerController.cs b/Engine/Source/Programs/Horde/HordeServer/Server/ServerController.cs
--- a/Engine/Source/Programs/Horde/HordeServer/Server/ServerController.cs
+++ b/Engine/Source/Programs/Horde/HordeServer/Server/ServerController.cs
@@ -1,14 +1,14 @@
 // Copyright Epic Games, Inc. All Rights Reserved.
 
 using System;
+using System.Collections.Generic;
 using System.Diagnostics;
+using System.Linq;
 using System.Reflection;
 using System.Threading.Tasks;
 using EpicGames.Core;
 using EpicGames.Horde;
 using EpicGames.Horde.Server;
-using HordeServer.Agents;
-using HordeServer.Tools;
 using HordeServer.Utilities;
 using Microsoft.AspNetCore.Authorization;
 using Microsoft.AspNetCore.Mvc;
@@ -24,16 +24,16 @@ namespace HordeServer.Server
 	[Route("[controller]")]
 	public class ServerController : HordeControllerBase
 	{
-		readonly IToolCollection _toolCollection;
+		readonly IAgentVersionProvider? _agentVersionProvider;
 		readonly IClock _clock;
 		readonly IOptionsSnapshot<GlobalConfig> _globalConfig;
 
 		/// <summary>
 		/// Constructor
 		/// </summary>
-		public ServerController(IToolCollection toolCollection, IClock clock, IOptionsSnapshot<GlobalConfig> globalConfig)
+		public ServerController(IEnumerable<IAgentVersionProvider> agentVersionProviders, IClock clock, IOptionsSnapshot<GlobalConfig> globalConfig)
 		{
-			_toolCollection = toolCollection;
+			_agentVersionProvider = agentVersionProviders.FirstOrDefault();
 			_clock = clock;
 			_globalConfig = globalConfig;
 		}
@@ -65,14 +65,9 @@ namespace HordeServer.Server
 			FileVersionInfo versionInfo = FileVersionInfo.GetVersionInfo(Assembly.GetExecutingAssembly().Location);
 			response.ServerVersion = versionInfo.ProductVersion ?? String.Empty;
 
-			ITool? tool = await _toolCollection.GetAsync(AgentExtensions.AgentToolId, HttpContext.RequestAborted);
-			if (tool != null)
+			if (_agentVersionProvider != null)
 			{
-				IToolDeployment? deployment = tool.GetCurrentDeployment(1.0, _clock.UtcNow);
-				if (deployment != null)
-				{
-					response.AgentVersion = deployment.Version;
-				}
+				response.AgentVersion = await _agentVersionProvider.GetAsync(HttpContext.RequestAborted);
 			}
 
 			return response;
diff --git a/Engine/Source/Programs/Horde/Plugins/Compute/HordeServer.Compute/ComputePlugin.cs b/Engine/Source/Programs/Horde/Plugins/Compute/HordeServer.Compute/ComputePlugin.cs
--- a/Engine/Source/Programs/Horde/Plugins/Compute/HordeServer.Compute/ComputePlugin.cs
+++ b/Engine/Source/Programs/Horde/Plugins/Compute/HordeServer.Compute/ComputePlugin.cs
@@ -52,6 +52,8 @@ namespace HordeServer
 			services.AddSingleton<IPoolCollection, PoolCollection>();
 			services.AddSingleton<ISessionCollection, SessionCollection>();
 
+			services.AddSingleton<IAgentVersionProvider, AgentVersionProvider>();
+
 			services.AddSingleton<AgentService>();
 			services.AddSingleton(provider => new Lazy<AgentService>(provider.GetRequiredService<AgentService>));
 
