Horde: Fix incorrect order of arguments to ConfigDoc attribute.

#rnx

diff --git a/Engine/Source/Programs/Horde/HordeAgent/AgentSettings.cs b/Engine/Source/Programs/Horde/HordeAgent/AgentSettings.cs
--- a/Engine/Source/Programs/Horde/HordeAgent/AgentSettings.cs
+++ b/Engine/Source/Programs/Horde/HordeAgent/AgentSettings.cs
@@ -83,7 +83,7 @@ namespace HordeAgent
 	/// <summary>
 	/// Global settings for the agent
 	/// </summary>
-	[ConfigDoc("Agent.json (Agent)", "Deployment/AgentSettings.md", "[Horde](../../README.md) > [Deployment](../Deployment.md) > [Agent](Agent.md)",
+	[ConfigDoc("Agent.json (Agent)", "[Horde](../../README.md) > [Deployment](../Deployment.md) > [Agent](Agent.md)", "Deployment/AgentSettings.md",
 		Introduction = "All Horde-specific settings are stored in a root object called `Horde`. Other .NET functionality may be configured using properties in the root of this file.")]
 	public class AgentSettings
 	{
