Horde: Missing file.

#rnx

diff --git a/Engine/Source/Programs/Horde/Plugins/Compute/HordeServer.Compute/Agents/AgentVersionProvider.cs b/Engine/Source/Programs/Horde/Plugins/Compute/HordeServer.Compute/Agents/AgentVersionProvider.cs
new file mode 100644
--- /dev/null
+++ b/Engine/Source/Programs/Horde/Plugins/Compute/HordeServer.Compute/Agents/AgentVersionProvider.cs
@@ -0,0 +1,46 @@
+// Copyright Epic Games, Inc. All Rights Reserved.
+
+using System;
+using System.Collections.Generic;
+using System.Linq;
+using System.Text;
+using System.Threading.Tasks;
+using EpicGames.Core;
+using HordeServer.Server;
+using HordeServer.Tools;
+
+namespace HordeServer.Agents
+{
+	/// <summary>
+	/// Advertises the version number of the agent tool for the /api/v1/server/info endpoint
+	/// </summary>
+	class AgentVersionProvider : IAgentVersionProvider
+	{
+		readonly IToolCollection _toolCollection;
+		readonly IClock _clock;
+
+		/// <summary>
+		/// Constructor
+		/// </summary>
+		public AgentVersionProvider(IToolCollection toolCollection, IClock clock)
+		{
+			_toolCollection = toolCollection;
+			_clock = clock;
+		}
+
+		/// <inheritdoc/>
+		public async Task<string?> GetAsync(CancellationToken cancellationToken = default)
+		{
+			ITool? tool = await _toolCollection.GetAsync(AgentExtensions.AgentToolId, cancellationToken);
+			if (tool != null)
+			{
+				IToolDeployment? deployment = tool.GetCurrentDeployment(1.0, _clock.UtcNow);
+				if (deployment != null)
+				{
+					return deployment.Version;
+				}
+			}
+			return null;
+		}
+	}
+}
