Horde: Move replication test command into build plugin, using the IServerStartup interface to inject an instance of the Startup class.

#rnx

diff --git a/Engine/Source/Programs/Horde/HordeServer.Shared/IServerStartup.cs b/Engine/Source/Programs/Horde/HordeServer.Shared/IServerStartup.cs
new file mode 100644
--- /dev/null
+++ b/Engine/Source/Programs/Horde/HordeServer.Shared/IServerStartup.cs
@@ -0,0 +1,17 @@
+// Copyright Epic Games, Inc. All Rights Reserved.
+
+using Microsoft.Extensions.DependencyInjection;
+
+namespace HordeServer
+{
+	/// <summary>
+	/// Interface for the main server startup class. Can be used by commands that want to instantiate server services.
+	/// </summary>
+	public interface IServerStartup
+	{
+		/// <summary>
+		/// Configure services for the server
+		/// </summary>
+		void ConfigureServices(IServiceCollection services);
+	}
+}
diff --git a/Engine/Source/Programs/Horde/HordeServer/Commands/Test/ReplicationCommand.cs b/Engine/Source/Programs/Horde/HordeServer/Commands/Test/ReplicationCommand.cs
deleted file mode 100644
diff --git a/Engine/Source/Programs/Horde/HordeServer/ServerApp.cs b/Engine/Source/Programs/Horde/HordeServer/ServerApp.cs
--- a/Engine/Source/Programs/Horde/HordeServer/ServerApp.cs
+++ b/Engine/Source/Programs/Horde/HordeServer/ServerApp.cs
@@ -180,6 +180,7 @@ namespace HordeServer
 			services.AddSingleton<IConfiguration>(config);
 			services.AddSingleton<ServerSettings>(serverSettings);
 			services.Configure<ServerSettings>(x => Startup.BindServerSettings(config, x));
+			services.AddTransient<IServerStartup, Startup>();
 
 			ServerInfo serverInfo = new ServerInfo(config, Options.Create(serverSettings));
 			services.AddSingleton<IServerInfo>(serverInfo);
diff --git a/Engine/Source/Programs/Horde/HordeServer/Startup.cs b/Engine/Source/Programs/Horde/HordeServer/Startup.cs
--- a/Engine/Source/Programs/Horde/HordeServer/Startup.cs
+++ b/Engine/Source/Programs/Horde/HordeServer/Startup.cs
@@ -90,7 +90,7 @@ namespace HordeServer
 	using ContentHash = EpicGames.Core.ContentHash;
 	using ILogger = Microsoft.Extensions.Logging.ILogger;
 
-	class Startup
+	class Startup : IServerStartup
 	{
 		static Startup()
 		{
diff --git a/Engine/Source/Programs/Horde/HordeServer/Utilities/ArtifactChunkStream.cs b/Engine/Source/Programs/Horde/HordeServer/Utilities/ArtifactChunkStream.cs
deleted file mode 100644
diff --git a/Engine/Source/Programs/Horde/Plugins/Build/HordeServer.Build/Commands/Test/ReplicationCommand.cs b/Engine/Source/Programs/Horde/Plugins/Build/HordeServer.Build/Commands/Test/ReplicationCommand.cs
new file mode 100644
--- /dev/null
+++ b/Engine/Source/Programs/Horde/Plugins/Build/HordeServer.Build/Commands/Test/ReplicationCommand.cs
@@ -0,0 +1,93 @@
+// Copyright Epic Games, Inc. All Rights Reserved.
+
+using System;
+using System.Threading.Tasks;
+using EpicGames.Core;
+using EpicGames.Horde.Replicators;
+using EpicGames.Horde.Streams;
+using HordeServer.Configuration;
+using HordeServer.Replicators;
+using HordeServer.Server;
+using HordeServer.Streams;
+using Microsoft.Extensions.Configuration;
+using Microsoft.Extensions.DependencyInjection;
+using Microsoft.Extensions.Logging;
+using Microsoft.Extensions.Options;
+
+namespace HordeServer.Commands.Test
+{
+	[Command("test", "replication", "Replicates commits for a particular change of changes from Perforce")]
+	class TestReplicationCommand : Command
+	{
+		[CommandLine("-Stream=", Required = true)]
+		public string StreamId { get; set; } = String.Empty;
+
+		[CommandLine("-Replicator=", Required = true)]
+		public string ReplicatorId { get; set; } = String.Empty;
+
+		[CommandLine(Required = true)]
+		public int Change { get; set; }
+
+		[CommandLine]
+		public bool Reset { get; set; }
+
+		[CommandLine]
+		public bool Clean { get; set; }
+
+		readonly IConfiguration _configuration;
+		readonly IServerStartup _serverStartup;
+		readonly ILoggerProvider _loggerProvider;
+
+		public TestReplicationCommand(IConfiguration configuration, IServerStartup serverStartup, ILoggerProvider loggerProvider)
+		{
+			_configuration = configuration;
+			_serverStartup = serverStartup;
+			_loggerProvider = loggerProvider;
+		}
+
+		public override async Task<int> ExecuteAsync(ILogger logger)
+		{
+			ServiceCollection serviceCollection = new ServiceCollection();
+			serviceCollection.AddSingleton(_configuration);
+			serviceCollection.AddSingleton(_loggerProvider);
+			_serverStartup.ConfigureServices(serviceCollection);
+
+			await using ServiceProvider serviceProvider = serviceCollection.BuildServiceProvider();
+
+			BuildConfig buildConfig = serviceProvider.GetRequiredService<IOptionsMonitor<BuildConfig>>().CurrentValue;
+
+			PerforceReplicator perforceReplicator = serviceProvider.GetRequiredService<PerforceReplicator>();
+			IStreamCollection streamCollection = serviceProvider.GetRequiredService<IStreamCollection>();
+
+			StreamConfig? streamConfig;
+			if (!buildConfig.TryGetStream(new StreamId(StreamId), out streamConfig))
+			{
+				throw new FatalErrorException($"Stream '{StreamId}' not found");
+			}
+
+			IReplicatorCollection replicatorCollection = serviceProvider.GetRequiredService<IReplicatorCollection>();
+
+			ReplicatorId id = new ReplicatorId(new StreamId(StreamId), new StreamReplicatorId(ReplicatorId));
+			IReplicator replicator = await replicatorCollection.GetOrAddAsync(id);
+
+			while (replicator.Pause || replicator.Reset != Reset || replicator.Clean != Clean || replicator.NextChange != Change)
+			{
+				UpdateReplicatorOptions updateOptions = new UpdateReplicatorOptions { Pause = false, Reset = Reset, Clean = Clean, NextChange = Change };
+
+				IReplicator? nextReplicator = await replicator.TryUpdateAsync(updateOptions);
+				if (nextReplicator != null)
+				{
+					replicator = nextReplicator;
+					break;
+				}
+
+				replicator = await replicatorCollection.GetOrAddAsync(id);
+			}
+
+			PerforceReplicationOptions options = new PerforceReplicationOptions();
+			await perforceReplicator.RunOnceAsync(replicator, buildConfig, streamConfig, options, default);
+
+			return 0;
+		}
+	}
+}
