Horde: Add an IDefaultAclModifier interface to allow plugins to modify the default ACL.

#rnx

diff --git a/Engine/Source/Programs/Horde/HordeServer.Shared/Acls/IDefaultAclModifier.cs b/Engine/Source/Programs/Horde/HordeServer.Shared/Acls/IDefaultAclModifier.cs
new file mode 100644
--- /dev/null
+++ b/Engine/Source/Programs/Horde/HordeServer.Shared/Acls/IDefaultAclModifier.cs
@@ -0,0 +1,17 @@
+// Copyright Epic Games, Inc. All Rights Reserved.
+
+using HordeServer.Utilities;
+
+namespace HordeServer.Acls
+{
+	/// <summary>
+	/// Interface which can be implemented by classes that mutate the default ACL
+	/// </summary>
+	public interface IDefaultAclModifier
+	{
+		/// <summary>
+		/// Modifies the default ACL
+		/// </summary>
+		void Apply(DefaultAclBuilder builder);
+	}
+}
diff --git a/Engine/Source/Programs/Horde/HordeServer.Shared/Utilities/DefaultAclBuilder.cs b/Engine/Source/Programs/Horde/HordeServer.Shared/Utilities/DefaultAclBuilder.cs
new file mode 100644
--- /dev/null
+++ b/Engine/Source/Programs/Horde/HordeServer.Shared/Utilities/DefaultAclBuilder.cs
@@ -0,0 +1,74 @@
+// Copyright Epic Games, Inc. All Rights Reserved.
+
+using System;
+using System.Collections.Generic;
+using System.Linq;
+using System.Text;
+using System.Threading.Tasks;
+using EpicGames.Horde.Acls;
+using HordeServer.Acls;
+
+namespace HordeServer.Utilities
+{
+	/// <summary>
+	/// Utility class for allowing plugins to modify the default ACL used by the server
+	/// </summary>
+	public class DefaultAclBuilder
+	{
+		List<AclEntryConfig> _roles = new List<AclEntryConfig>();
+		HashSet<AclAction> _readActions = new HashSet<AclAction>();
+		HashSet<AclAction> _writeActions = new HashSet<AclAction>();
+
+		/// <summary>
+		/// Adds a custom role with a certain set of entitlements
+		/// </summary>
+		/// <param name="claim">Claim to identify users that should be granted the entitlements</param>
+		/// <param name="actions">Actions to allow the user to perform</param>
+		public void AddCustomRole(AclClaimConfig claim, AclAction[] actions)
+			=> _roles.Add(new AclEntryConfig(claim, actions));
+
+		public void AddCustomRole(AclClaimConfig uploadToolsClaim, object[] values)
+		{
+			throw new NotImplementedException();
+		}
+
+		/// <summary>
+		/// Adds a default read operation that users can perform
+		/// </summary>
+		public void AddDefaultReadAction(AclAction action)
+			=> _readActions.Add(action);
+
+		/// <summary>
+		/// Adds a default write operation that users can perform
+		/// </summary>
+		public void AddDefaultWriteAction(AclAction action)
+			=> _writeActions.Add(action);
+
+		/// <summary>
+		/// Create the new acl config
+		/// </summary>
+		public AclConfig Build()
+		{
+			AclConfig config = new AclConfig();
+			config.Entries = _roles.ToList();
+
+			config.Profiles = new List<AclProfileConfig>();
+			config.Profiles.Add(new AclProfileConfig
+			{
+				Id = new AclProfileId("default-read"),
+				Actions = _readActions.ToList()
+			});
+
+			config.Profiles.Add(new AclProfileConfig
+			{
+				Id = new AclProfileId("default-run"),
+				Extends = new List<AclProfileId>
+				{
+					new AclProfileId("default-read")
+				},
+				Actions = _writeActions.ToList()
+			});
+			return config;
+		}
+	}
+}
diff --git a/Engine/Source/Programs/Horde/HordeServer.Tests.Shared/ServerTestSetup.cs b/Engine/Source/Programs/Horde/HordeServer.Tests.Shared/ServerTestSetup.cs
--- a/Engine/Source/Programs/Horde/HordeServer.Tests.Shared/ServerTestSetup.cs
+++ b/Engine/Source/Programs/Horde/HordeServer.Tests.Shared/ServerTestSetup.cs
@@ -48,6 +48,7 @@ namespace HordeServer.Tests
 		public IOptionsSnapshot<GlobalConfig> GlobalConfigSnapshot => ServiceProvider.GetRequiredService<IOptionsSnapshot<GlobalConfig>>();
 
 		public IPluginCollection PluginCollection => ServiceProvider.GetRequiredService<IPluginCollection>();
+		public IEnumerable<IDefaultAclModifier> DefaultAclModifiers => ServiceProvider.GetRequiredService<IEnumerable<IDefaultAclModifier>>();
 
 		private static bool s_datadogWriterPatched;
 
@@ -62,7 +63,7 @@ namespace HordeServer.Tests
 
 		protected void SetConfig(GlobalConfig globalConfig)
 		{
-			globalConfig.PostLoad(ServerSettings, _pluginCollection.LoadedPlugins);
+			globalConfig.PostLoad(ServerSettings, _pluginCollection.LoadedPlugins, DefaultAclModifiers);
 			ConfigService.OverrideConfig(globalConfig);
 		}
 
@@ -70,7 +71,7 @@ namespace HordeServer.Tests
 		{
 			GlobalConfig globalConfig = GlobalConfig.CurrentValue;
 			action(globalConfig);
-			globalConfig.PostLoad(ServerSettings, _pluginCollection.LoadedPlugins);
+			globalConfig.PostLoad(ServerSettings, _pluginCollection.LoadedPlugins, DefaultAclModifiers);
 			ConfigService.OverrideConfig(globalConfig);
 		}
 
diff --git a/Engine/Source/Programs/Horde/HordeServer.Tests/Configuration/ConfigTests.cs b/Engine/Source/Programs/Horde/HordeServer.Tests/Configuration/ConfigTests.cs
--- a/Engine/Source/Programs/Horde/HordeServer.Tests/Configuration/ConfigTests.cs
+++ b/Engine/Source/Programs/Horde/HordeServer.Tests/Configuration/ConfigTests.cs
@@ -8,6 +8,7 @@ using System.Text.Json;
 using System.Threading;
 using System.Threading.Tasks;
 using EpicGames.Core;
+using HordeServer.Acls;
 using HordeServer.Configuration;
 using HordeServer.Plugins;
 using HordeServer.Projects;
@@ -182,7 +183,7 @@ namespace HordeServer.Tests.Configuration
 
 			GlobalConfig gc = new();
 			gc.Plugins.AddBuildConfig(buildConfig);
-			gc.PostLoad(new ServerSettings(), new List<ILoadedPlugin>());
+			gc.PostLoad(new ServerSettings(), Array.Empty<ILoadedPlugin>(), Array.Empty<IDefaultAclModifier>());
 
 			Dictionary<string,WorkspaceConfig> workspaces = gc.Plugins.GetBuildConfig().Projects[0].Streams[0].WorkspaceTypes;
 			Assert.AreEqual(3, workspaces.Count);
@@ -216,7 +217,7 @@ namespace HordeServer.Tests.Configuration
 
 			GlobalConfig gc = new();
 			gc.Plugins.AddBuildConfig(buildConfig);
-			gc.PostLoad(new ServerSettings(), new List<ILoadedPlugin>());
+			gc.PostLoad(new ServerSettings(), Array.Empty<ILoadedPlugin>(), Array.Empty<IDefaultAclModifier>());
 
 			Dictionary<string, WorkspaceConfig> workspaces = gc.Plugins.GetBuildConfig().Projects[0].Streams[0].WorkspaceTypes;
 			Assert.AreEqual(4, workspaces.Count);
diff --git a/Engine/Source/Programs/Horde/HordeServer/Configuration/ConfigService.cs b/Engine/Source/Programs/Horde/HordeServer/Configuration/ConfigService.cs
--- a/Engine/Source/Programs/Horde/HordeServer/Configuration/ConfigService.cs
+++ b/Engine/Source/Programs/Horde/HordeServer/Configuration/ConfigService.cs
@@ -13,6 +13,7 @@ using System.Threading.Tasks;
 using EpicGames.Core;
 using EpicGames.Horde.Users;
 using EpicGames.Redis;
+using HordeServer.Acls;
 using HordeServer.Plugins;
 using HordeServer.Server;
 using HordeServer.Users;
@@ -118,6 +119,7 @@ namespace HordeServer.Configuration
 		readonly IPluginCollection _pluginCollection;
 		readonly JsonSerializerOptions _jsonOptions;
 		readonly RedisKey _snapshotKey = "config";
+		readonly IEnumerable<IDefaultAclModifier> _defaultAclModifiers;
 		readonly IHealthMonitor _health;
 		readonly ILogger _logger;
 
@@ -141,11 +143,12 @@ namespace HordeServer.Configuration
 		/// <summary>
 		/// Constructor
 		/// </summary>
-		public ConfigService(IRedisService redisService, IOptions<ServerSettings> serverSettings, IEnumerable<IConfigSource> sources, IPluginCollection pluginCollection, IClock clock, IHealthMonitor<ConfigService> health, ILogger<ConfigService> logger)
+		public ConfigService(IRedisService redisService, IOptions<ServerSettings> serverSettings, IEnumerable<IConfigSource> sources, IPluginCollection pluginCollection, IEnumerable<IDefaultAclModifier> aclModifiers, IClock clock, IHealthMonitor<ConfigService> health, ILogger<ConfigService> logger)
 		{
 			_redisService = redisService;
 			_serverSettings = serverSettings.Value;
 			_sources = sources.ToDictionary(x => x.Scheme, x => x, StringComparer.OrdinalIgnoreCase);
+			_defaultAclModifiers = aclModifiers;
 			_pluginCollection = pluginCollection;
 			_health = health;
 			_health.SetName("GlobalConfig");
@@ -279,8 +282,9 @@ namespace HordeServer.Configuration
 			try
 			{
 				Uri globalConfigUri = GetGlobalConfigUri();
+
 				GlobalConfig globalConfig = await context.ReadAsync<GlobalConfig>(globalConfigUri, cancellationToken);
-				globalConfig.PostLoad(_serverSettings, _pluginCollection.LoadedPlugins);
+				globalConfig.PostLoad(_serverSettings, _pluginCollection.LoadedPlugins, _defaultAclModifiers);
 
 				foreach (OverrideConfigFile file in overrideFiles.Values)
 				{
@@ -643,7 +647,7 @@ namespace HordeServer.Configuration
 				}
 
 				// Execute a PostLoad before returning so we can validate that everything is valid
-				globalConfig.PostLoad(_serverSettings, _pluginCollection.LoadedPlugins);
+				globalConfig.PostLoad(_serverSettings, _pluginCollection.LoadedPlugins, _defaultAclModifiers);
 
 				return snapshot;
 			}
@@ -851,7 +855,7 @@ namespace HordeServer.Configuration
 			globalConfig.Revision = IoHash.Compute(data.Span).ToString();
 
 			// Run the postload callbacks on all the config objects
-			globalConfig.PostLoad(_serverSettings, _pluginCollection.LoadedPlugins);
+			globalConfig.PostLoad(_serverSettings, _pluginCollection.LoadedPlugins, _defaultAclModifiers);
 			return globalConfig;
 		}
 	}
diff --git a/Engine/Source/Programs/Horde/HordeServer/Server/GlobalConfig.cs b/Engine/Source/Programs/Horde/HordeServer/Server/GlobalConfig.cs
--- a/Engine/Source/Programs/Horde/HordeServer/Server/GlobalConfig.cs
+++ b/Engine/Source/Programs/Horde/HordeServer/Server/GlobalConfig.cs
@@ -12,21 +12,11 @@ using EpicGames.Horde.Agents.Pools;
 using HordeServer.Acls;
 using HordeServer.Agents;
 using HordeServer.Agents.Pools;
-using HordeServer.Agents.Sessions;
-using HordeServer.Agents.Software;
-using HordeServer.Artifacts;
 using HordeServer.Configuration;
 using HordeServer.Dashboard;
-using HordeServer.Devices;
-using HordeServer.Jobs;
-using HordeServer.Jobs.Bisect;
-using HordeServer.Logs;
-using HordeServer.Notifications;
 using HordeServer.Perforce;
 using HordeServer.Plugins;
-using HordeServer.Projects;
 using HordeServer.Streams;
-using HordeServer.Tools;
 using HordeServer.Utilities;
 
 namespace HordeServer.Server
@@ -116,11 +106,11 @@ namespace HordeServer.Server
 		/// <summary>
 		/// Called after the config file has been read
 		/// </summary>
-		public void PostLoad(ServerSettings serverSettings, IReadOnlyList<ILoadedPlugin> loadedPlugins)
+		public void PostLoad(ServerSettings serverSettings, IReadOnlyList<ILoadedPlugin> loadedPlugins, IEnumerable<IDefaultAclModifier> defaultAclModifiers)
 		{
 			ServerSettings = serverSettings;
 
-			AclConfig defaultAcl = CreateRootAcl();
+			AclConfig defaultAcl = CreateRootAcl(defaultAclModifiers);
 			Acl.PostLoad(defaultAcl, defaultAcl.ScopeName);
 
 			// Ensure that all plugins have an entry in the global config so they can register their ACLs
@@ -148,63 +138,19 @@ namespace HordeServer.Server
 		/// <summary>
 		/// Creates the default root ACL
 		/// </summary>
-		static AclConfig CreateRootAcl()
+		static AclConfig CreateRootAcl(IEnumerable<IDefaultAclModifier> defaultAclModifiers)
 		{
-			AclConfig defaultAcl = new AclConfig();
-
-			defaultAcl.Entries = new List<AclEntryConfig>();
-			defaultAcl.Entries.Add(new AclEntryConfig(new AclClaimConfig(ClaimTypes.Role, "internal:AgentRegistration"), new[] { AgentAclAction.CreateAgent, SessionAclAction.CreateSession }));
-			defaultAcl.Entries.Add(new AclEntryConfig(HordeClaims.AgentRegistrationClaim, new[] { AgentAclAction.CreateAgent, SessionAclAction.CreateSession, AgentAclAction.UpdateAgent, AgentSoftwareAclAction.DownloadSoftware, PoolAclAction.CreatePool, PoolAclAction.UpdatePool, PoolAclAction.ViewPool, PoolAclAction.DeletePool, PoolAclAction.ListPools, StreamAclAction.ViewStream, ProjectAclAction.ViewProject, JobAclAction.ViewJob, ServerAclAction.ViewCosts }));
-			defaultAcl.Entries.Add(new AclEntryConfig(HordeClaims.AgentRoleClaim, new[] { ProjectAclAction.ViewProject, StreamAclAction.ViewStream, LogAclAction.CreateEvent, AgentSoftwareAclAction.DownloadSoftware }));
-			defaultAcl.Entries.Add(new AclEntryConfig(HordeClaims.DownloadSoftwareClaim, new[] { AgentSoftwareAclAction.DownloadSoftware }));
-			defaultAcl.Entries.Add(new AclEntryConfig(HordeClaims.UploadToolsClaim, new[] { AgentSoftwareAclAction.UploadSoftware, ToolAclAction.UploadTool }));
-			defaultAcl.Entries.Add(new AclEntryConfig(HordeClaims.ConfigureProjectsClaim, new[] { ProjectAclAction.CreateProject, ProjectAclAction.UpdateProject, ProjectAclAction.ViewProject, StreamAclAction.CreateStream, StreamAclAction.UpdateStream, StreamAclAction.ViewStream }));
-			defaultAcl.Entries.Add(new AclEntryConfig(HordeClaims.StartChainedJobClaim, new[] { JobAclAction.CreateJob, JobAclAction.ExecuteJob, JobAclAction.UpdateJob, JobAclAction.ViewJob, StreamAclAction.ViewTemplate, StreamAclAction.ViewStream }));
-
-			defaultAcl.Profiles = new List<AclProfileConfig>();
-			defaultAcl.Profiles.Add(new AclProfileConfig
-			{
-				Id = new AclProfileId("default-read"),
-				Actions = new List<AclAction>
-				{
-					AgentAclAction.ListAgents,
-					AgentAclAction.ViewAgent,
-					ArtifactAclAction.DownloadArtifact,
-					ArtifactAclAction.ReadArtifact,
-					BisectTaskAclAction.ViewBisectTask,
-					DeviceAclAction.DeviceRead,
-					JobAclAction.ViewJob,
-					LogAclAction.ViewEvent,
-					LogAclAction.ViewLog,
-					NotificationAclAction.CreateSubscription,
-					PoolAclAction.ListPools,
-					PoolAclAction.ViewPool,
-					ProjectAclAction.ViewProject,
-					ServerAclAction.IssueBearerToken,
-					StreamAclAction.ViewChanges,
-					StreamAclAction.ViewStream,
-					StreamAclAction.ViewTemplate,
-				}
-			});
-			defaultAcl.Profiles.Add(new AclProfileConfig
+			DefaultAclBuilder defaultAclBuilder = new DefaultAclBuilder();
+			defaultAclBuilder.AddDefaultReadAction(ServerAclAction.IssueBearerToken);
+
+			foreach (IDefaultAclModifier defaultAclModifier in defaultAclModifiers)
 			{
-				Id = new AclProfileId("default-run"),
-				Extends = new List<AclProfileId>
-				{
-					new AclProfileId("default-read")
-				},
-				Actions = new List<AclAction>
-				{
-					JobAclAction.CreateJob,
-					JobAclAction.UpdateJob,
-					JobAclAction.RetryJobStep,
-					DeviceAclAction.DeviceWrite,
-					BisectTaskAclAction.CreateBisectTask,
-					BisectTaskAclAction.UpdateBisectTask,
-				}
-			});
+				defaultAclModifier.Apply(defaultAclBuilder);
+			}
 
+			AclConfig defaultAcl = defaultAclBuilder.Build();
 			defaultAcl.PostLoad(null, AclScopeName.Root);
+
 			return defaultAcl;
 		}
 
diff --git a/Engine/Source/Programs/Horde/HordeServer/Server/ServerController.cs b/Engine/Source/Programs/Horde/HordeServer/Server/ServerController.cs
--- a/Engine/Source/Programs/Horde/HordeServer/Server/ServerController.cs
+++ b/Engine/Source/Programs/Horde/HordeServer/Server/ServerController.cs
@@ -6,7 +6,6 @@ using System.Diagnostics;
 using System.Linq;
 using System.Reflection;
 using System.Threading.Tasks;
-using EpicGames.Core;
 using EpicGames.Horde;
 using EpicGames.Horde.Server;
 using HordeServer.Utilities;
@@ -25,16 +24,14 @@ namespace HordeServer.Server
 	public class ServerController : HordeControllerBase
 	{
 		readonly IAgentVersionProvider? _agentVersionProvider;
-		readonly IClock _clock;
 		readonly IOptionsSnapshot<GlobalConfig> _globalConfig;
 
 		/// <summary>
 		/// Constructor
 		/// </summary>
-		public ServerController(IEnumerable<IAgentVersionProvider> agentVersionProviders, IClock clock, IOptionsSnapshot<GlobalConfig> globalConfig)
+		public ServerController(IEnumerable<IAgentVersionProvider> agentVersionProviders, IOptionsSnapshot<GlobalConfig> globalConfig)
 		{
 			_agentVersionProvider = agentVersionProviders.FirstOrDefault();
-			_clock = clock;
 			_globalConfig = globalConfig;
 		}
 
diff --git a/Engine/Source/Programs/Horde/Plugins/Build/HordeServer.Build.Tests/Acls/AclTests.cs b/Engine/Source/Programs/Horde/Plugins/Build/HordeServer.Build.Tests/Acls/AclTests.cs
--- a/Engine/Source/Programs/Horde/Plugins/Build/HordeServer.Build.Tests/Acls/AclTests.cs
+++ b/Engine/Source/Programs/Horde/Plugins/Build/HordeServer.Build.Tests/Acls/AclTests.cs
@@ -16,7 +16,7 @@ using Microsoft.VisualStudio.TestTools.UnitTesting;
 namespace HordeServer.Tests.Acls
 {
 	[TestClass]
-	public class AclTests : TestSetup
+	public class AclTests : BuildTestSetup
 	{
 		[TestMethod]
 		public void TestCustomProfile()
diff --git a/Engine/Source/Programs/Horde/Plugins/Build/HordeServer.Build.Tests/Agents/AgentCollectionTests.cs b/Engine/Source/Programs/Horde/Plugins/Build/HordeServer.Build.Tests/Agents/AgentCollectionTests.cs
--- a/Engine/Source/Programs/Horde/Plugins/Build/HordeServer.Build.Tests/Agents/AgentCollectionTests.cs
+++ b/Engine/Source/Programs/Horde/Plugins/Build/HordeServer.Build.Tests/Agents/AgentCollectionTests.cs
@@ -13,7 +13,7 @@ using Microsoft.VisualStudio.TestTools.UnitTesting;
 namespace HordeServer.Tests.Agents;
 
 [TestClass]
-public class AgentCollectionTests : TestSetup
+public class AgentCollectionTests : BuildTestSetup
 {
 	private IAgent _agent;
 	private readonly AgentLease _lease1;
diff --git a/Engine/Source/Programs/Horde/Plugins/Build/HordeServer.Build.Tests/Agents/AgentControllerTest.cs b/Engine/Source/Programs/Horde/Plugins/Build/HordeServer.Build.Tests/Agents/AgentControllerTest.cs
--- a/Engine/Source/Programs/Horde/Plugins/Build/HordeServer.Build.Tests/Agents/AgentControllerTest.cs
+++ b/Engine/Source/Programs/Horde/Plugins/Build/HordeServer.Build.Tests/Agents/AgentControllerTest.cs
@@ -9,7 +9,7 @@ using Microsoft.VisualStudio.TestTools.UnitTesting;
 namespace HordeServer.Tests.Agents;
 
 [TestClass]
-public class AgentControllerDbTest : TestSetup
+public class AgentControllerDbTest : BuildTestSetup
 {
 	[TestMethod]
 	public async Task UpdateAgentAsync()
diff --git a/Engine/Source/Programs/Horde/Plugins/Build/HordeServer.Build.Tests/Agents/AgentServiceTest.cs b/Engine/Source/Programs/Horde/Plugins/Build/HordeServer.Build.Tests/Agents/AgentServiceTest.cs
--- a/Engine/Source/Programs/Horde/Plugins/Build/HordeServer.Build.Tests/Agents/AgentServiceTest.cs
+++ b/Engine/Source/Programs/Horde/Plugins/Build/HordeServer.Build.Tests/Agents/AgentServiceTest.cs
@@ -24,7 +24,7 @@ namespace HordeServer.Tests.Agents;
 ///     Testing the agent service
 /// </summary>
 [TestClass]
-public class AgentServiceTest : TestSetup
+public class AgentServiceTest : BuildTestSetup
 {
 	[TestMethod]
 	public async Task GetJobsAsync()
diff --git a/Engine/Source/Programs/Horde/Plugins/Build/HordeServer.Build.Tests/Agents/Leases/LeasesControllerDbTest.cs b/Engine/Source/Programs/Horde/Plugins/Build/HordeServer.Build.Tests/Agents/Leases/LeasesControllerDbTest.cs
--- a/Engine/Source/Programs/Horde/Plugins/Build/HordeServer.Build.Tests/Agents/Leases/LeasesControllerDbTest.cs
+++ b/Engine/Source/Programs/Horde/Plugins/Build/HordeServer.Build.Tests/Agents/Leases/LeasesControllerDbTest.cs
@@ -18,7 +18,7 @@ namespace HordeServer.Tests.Agents.Leases
 	/// Database-only integration test for LeasesController
 	/// </summary>
 	[TestClass]
-	public class LeasesControllerDbTest : TestSetup
+	public class LeasesControllerDbTest : BuildTestSetup
 	{
 		[TestMethod]
 		public async Task FindLeasesAsync()
diff --git a/Engine/Source/Programs/Horde/Plugins/Build/HordeServer.Build.Tests/Agents/Pools/PoolUpdateServiceTest.cs b/Engine/Source/Programs/Horde/Plugins/Build/HordeServer.Build.Tests/Agents/Pools/PoolUpdateServiceTest.cs
--- a/Engine/Source/Programs/Horde/Plugins/Build/HordeServer.Build.Tests/Agents/Pools/PoolUpdateServiceTest.cs
+++ b/Engine/Source/Programs/Horde/Plugins/Build/HordeServer.Build.Tests/Agents/Pools/PoolUpdateServiceTest.cs
@@ -16,7 +16,7 @@ using Microsoft.VisualStudio.TestTools.UnitTesting;
 namespace HordeServer.Tests.Agents.Pools;
 
 [TestClass]
-public class PoolUpdateServiceTest : TestSetup
+public class PoolUpdateServiceTest : BuildTestSetup
 {
 	private readonly PoolUpdateService _pus;
 	private IPool _pool = default!;
diff --git a/Engine/Source/Programs/Horde/Plugins/Build/HordeServer.Build.Tests/Agents/Pools/PoolsControllerTest.cs b/Engine/Source/Programs/Horde/Plugins/Build/HordeServer.Build.Tests/Agents/Pools/PoolsControllerTest.cs
--- a/Engine/Source/Programs/Horde/Plugins/Build/HordeServer.Build.Tests/Agents/Pools/PoolsControllerTest.cs
+++ b/Engine/Source/Programs/Horde/Plugins/Build/HordeServer.Build.Tests/Agents/Pools/PoolsControllerTest.cs
@@ -11,7 +11,7 @@ using Microsoft.VisualStudio.TestTools.UnitTesting;
 namespace HordeServer.Tests.Agents.Pools
 {
 	[TestClass]
-	public class PoolsControllerTest : TestSetup
+	public class PoolsControllerTest : BuildTestSetup
 	{
 		public PoolsControllerTest()
 		{
diff --git a/Engine/Source/Programs/Horde/Plugins/Build/HordeServer.Build.Tests/Agents/Registration/RegistrationServiceTests.cs b/Engine/Source/Programs/Horde/Plugins/Build/HordeServer.Build.Tests/Agents/Registration/RegistrationServiceTests.cs
--- a/Engine/Source/Programs/Horde/Plugins/Build/HordeServer.Build.Tests/Agents/Registration/RegistrationServiceTests.cs
+++ b/Engine/Source/Programs/Horde/Plugins/Build/HordeServer.Build.Tests/Agents/Registration/RegistrationServiceTests.cs
@@ -13,7 +13,7 @@ using Microsoft.VisualStudio.TestTools.UnitTesting;
 namespace HordeServer.Tests.Agents.Registration
 {
 	[TestClass]
-	public class RegistrationServiceTests : TestSetup
+	public class RegistrationServiceTests : BuildTestSetup
 	{
 		[TestMethod]
 		public async Task ApproveAgentAsync()
diff --git a/Engine/Source/Programs/Horde/Plugins/Build/HordeServer.Build.Tests/Agents/Relay/AgentRelayServiceTests.cs b/Engine/Source/Programs/Horde/Plugins/Build/HordeServer.Build.Tests/Agents/Relay/AgentRelayServiceTests.cs
--- a/Engine/Source/Programs/Horde/Plugins/Build/HordeServer.Build.Tests/Agents/Relay/AgentRelayServiceTests.cs
+++ b/Engine/Source/Programs/Horde/Plugins/Build/HordeServer.Build.Tests/Agents/Relay/AgentRelayServiceTests.cs
@@ -19,7 +19,7 @@ using Microsoft.VisualStudio.TestTools.UnitTesting;
 namespace HordeServer.Tests.Agents.Relay;
 
 [TestClass]
-public class AgentRelayServiceTests : TestSetup
+public class AgentRelayServiceTests : BuildTestSetup
 {
 	private readonly AgentRelayService _service;
 	private readonly GetPortMappingsRequest _request = new() { ClusterId = "cluster1", AgentId = "agent1", IpAddresses = { "192.168.1.1" } };
diff --git a/Engine/Source/Programs/Horde/Plugins/Build/HordeServer.Build.Tests/Artifacts/ArtifactTests.cs b/Engine/Source/Programs/Horde/Plugins/Build/HordeServer.Build.Tests/Artifacts/ArtifactTests.cs
--- a/Engine/Source/Programs/Horde/Plugins/Build/HordeServer.Build.Tests/Artifacts/ArtifactTests.cs
+++ b/Engine/Source/Programs/Horde/Plugins/Build/HordeServer.Build.Tests/Artifacts/ArtifactTests.cs
@@ -16,7 +16,7 @@ using Microsoft.VisualStudio.TestTools.UnitTesting;
 namespace HordeServer.Tests.Artifacts
 {
 	[TestClass]
-	public class ArtifactTests : TestSetup
+	public class ArtifactTests : BuildTestSetup
 	{
 		[TestMethod]
 		public async Task CreateArtifactAsync()
diff --git a/Engine/Source/Programs/Horde/Plugins/Build/HordeServer.Build.Tests/BuildTestSetup.cs b/Engine/Source/Programs/Horde/Plugins/Build/HordeServer.Build.Tests/BuildTestSetup.cs
new file mode 100644
--- /dev/null
+++ b/Engine/Source/Programs/Horde/Plugins/Build/HordeServer.Build.Tests/BuildTestSetup.cs
@@ -0,0 +1,351 @@
+// Copyright Epic Games, Inc. All Rights Reserved.
+
+using System;
+using System.Collections.Generic;
+using System.Net;
+using System.Net.Sockets;
+using System.Security.Claims;
+using System.Threading.Tasks;
+using Amazon.AutoScaling;
+using Amazon.CloudWatch;
+using Amazon.EC2;
+using EpicGames.Horde.Agents;
+using EpicGames.Horde.Agents.Pools;
+using HordeServer.Accounts;
+using HordeServer.Acls;
+using HordeServer.Agents;
+using HordeServer.Agents.Enrollment;
+using HordeServer.Agents.Fleet;
+using HordeServer.Agents.Leases;
+using HordeServer.Agents.Pools;
+using HordeServer.Agents.Relay;
+using HordeServer.Agents.Sessions;
+using HordeServer.Agents.Telemetry;
+using HordeServer.Agents.Utilization;
+using HordeServer.Artifacts;
+using HordeServer.Auditing;
+using HordeServer.Commits;
+using HordeServer.Compute;
+using HordeServer.Dashboard;
+using HordeServer.Devices;
+using HordeServer.Issues;
+using HordeServer.Jobs;
+using HordeServer.Jobs.Bisect;
+using HordeServer.Jobs.Graphs;
+using HordeServer.Jobs.Schedules;
+using HordeServer.Jobs.Templates;
+using HordeServer.Jobs.TestData;
+using HordeServer.Jobs.Timing;
+using HordeServer.Logs;
+using HordeServer.Notifications;
+using HordeServer.Perforce;
+using HordeServer.Server;
+using HordeServer.Storage;
+using HordeServer.Streams;
+using HordeServer.Tasks;
+using HordeServer.Tests.Stubs.Services;
+using HordeServer.Tools;
+using HordeServer.Ugs;
+using HordeServer.Users;
+using HordeServer.Utilities;
+using Microsoft.AspNetCore.Http;
+using Microsoft.AspNetCore.Mvc;
+using Microsoft.Extensions.Configuration;
+using Microsoft.Extensions.DependencyInjection;
+using Microsoft.Extensions.Logging;
+using Microsoft.VisualStudio.TestTools.UnitTesting;
+using Moq;
+
+namespace HordeServer.Tests
+{
+	/// <summary>
+	/// Handles set up of collections, services, fixtures etc during testing
+	///
+	/// Easier to pass all these things around in a single object.
+	/// </summary>
+	public class BuildTestSetup : ServerTestSetup
+	{
+		public IGraphCollection GraphCollection => ServiceProvider.GetRequiredService<IGraphCollection>();
+		public INotificationTriggerCollection NotificationTriggerCollection => ServiceProvider.GetRequiredService<INotificationTriggerCollection>();
+		public IStreamCollection StreamCollection => ServiceProvider.GetRequiredService<IStreamCollection>();
+		public IJobCollection JobCollection => ServiceProvider.GetRequiredService<IJobCollection>();
+		public IAgentCollection AgentCollection => ServiceProvider.GetRequiredService<IAgentCollection>();
+		public IJobStepRefCollection JobStepRefCollection => ServiceProvider.GetRequiredService<IJobStepRefCollection>();
+		public IJobTimingCollection JobTimingCollection => ServiceProvider.GetRequiredService<IJobTimingCollection>();
+		public IUgsMetadataCollection UgsMetadataCollection => ServiceProvider.GetRequiredService<IUgsMetadataCollection>();
+		public IIssueCollection IssueCollection => ServiceProvider.GetRequiredService<IIssueCollection>();
+		public IPoolCollection PoolCollection => ServiceProvider.GetRequiredService<IPoolCollection>();
+		public ILeaseCollection LeaseCollection => ServiceProvider.GetRequiredService<ILeaseCollection>();
+		public ILogCollection LogCollection => ServiceProvider.GetRequiredService<ILogCollection>();
+		public ISessionCollection SessionCollection => ServiceProvider.GetRequiredService<ISessionCollection>();
+		public ITestDataCollection TestDataCollection => ServiceProvider.GetRequiredService<ITestDataCollection>();
+		public IUserCollection UserCollection => ServiceProvider.GetRequiredService<IUserCollection>();
+		public IDeviceCollection DeviceCollection => ServiceProvider.GetRequiredService<IDeviceCollection>();
+		public IDashboardPreviewCollection DashboardPreviewCollection => ServiceProvider.GetRequiredService<IDashboardPreviewCollection>();
+		public IBisectTaskCollection BisectTaskCollection => ServiceProvider.GetRequiredService<IBisectTaskCollection>();
+
+		public FleetService FleetService => ServiceProvider.GetRequiredService<FleetService>();
+		public AgentService AgentService => ServiceProvider.GetRequiredService<AgentService>();
+		public AgentRelayService AgentRelayService => ServiceProvider.GetRequiredService<AgentRelayService>();
+		public ICommitService CommitService => ServiceProvider.GetRequiredService<ICommitService>();
+		public GlobalsService GlobalsService => ServiceProvider.GetRequiredService<GlobalsService>();
+		public ITemplateCollection TemplateCollection => ServiceProvider.GetRequiredService<ITemplateCollection>();
+		internal PerforceServiceStub PerforceService => (PerforceServiceStub)ServiceProvider.GetRequiredService<IPerforceService>();
+		public ISubscriptionCollection SubscriptionCollection => ServiceProvider.GetRequiredService<ISubscriptionCollection>();
+		public INotificationService NotificationService => ServiceProvider.GetRequiredService<INotificationService>();
+		public IssueService IssueService => ServiceProvider.GetRequiredService<IssueService>();
+		public JobTaskSource JobTaskSource => ServiceProvider.GetRequiredService<JobTaskSource>();
+		public JobService JobService => ServiceProvider.GetRequiredService<JobService>();
+		public RpcService RpcService => ServiceProvider.GetRequiredService<RpcService>();
+		public PoolService PoolService => ServiceProvider.GetRequiredService<PoolService>();
+		public ScheduleService ScheduleService => ServiceProvider.GetRequiredService<ScheduleService>();
+		public DeviceService DeviceService => ServiceProvider.GetRequiredService<DeviceService>();
+		public StorageService StorageService => ServiceProvider.GetRequiredService<StorageService>();
+		public TestDataService TestDataService => ServiceProvider.GetRequiredService<TestDataService>();
+		public ComputeService ComputeService => ServiceProvider.GetRequiredService<ComputeService>();
+
+		public JobsController JobsController => GetJobsController();
+		public AgentsController AgentsController => GetAgentsController();
+		public PoolsController PoolsController => GetPoolsController();
+		public LeasesController LeasesController => GetLeasesController();
+		public DevicesController DevicesController => GetDevicesController();
+		public DashboardController DashboardController => GetDashboardController();
+		public TestDataController TestDataController => GetTestDataController();
+		public BisectTasksController BisectTasksController => GetBisectTasksController();
+
+		public BuildTestSetup()
+		{
+			AddPlugin<AnalyticsPlugin>();
+			AddPlugin<BuildPlugin>();
+			AddPlugin<ComputePlugin>();
+			AddPlugin<StoragePlugin>();
+		}
+
+		protected override void ConfigureServices(IServiceCollection services)
+		{
+			base.ConfigureServices(services);
+
+			IConfiguration config = new ConfigurationBuilder().Build();
+			services.Configure<ServerSettings>(ConfigureSettings);
+			services.AddSingleton<IConfiguration>(config);
+
+			services.AddHttpClient<RpcService>();
+
+			services.AddSingleton<IAuditLog<AgentId>>(sp => sp.GetRequiredService<IAuditLogFactory<AgentId>>().Create("Agents.Log", "AgentId"));
+
+			services.AddSingleton<IAgentCollection, AgentCollection>();
+			services.AddSingleton<IAgentTelemetryCollection, AgentTelemetryCollection>();
+			services.AddSingleton<IArtifactCollection, ArtifactCollection>();
+			services.AddSingleton<ICommitService, CommitService>();
+			services.AddSingleton<IGraphCollection, GraphCollection>();
+			services.AddSingleton<IIssueCollection, IssueCollection>();
+			services.AddSingleton<IJobCollection, JobCollection>();
+			services.AddSingleton<IJobStepRefCollection, JobStepRefCollection>();
+			services.AddSingleton<IJobTimingCollection, JobTimingCollection>();
+			services.AddSingleton<ILeaseCollection, LeaseCollection>();
+			services.AddSingleton<ILogCollection, LogCollection>();
+			services.AddSingleton<INotificationTriggerCollection, NotificationTriggerCollection>();
+			services.AddSingleton<IPoolCollection, PoolCollection>();
+			services.AddSingleton<IBisectTaskCollection, BisectTaskCollection>();
+			services.AddSingleton<ISessionCollection, SessionCollection>();
+			services.AddSingleton<ISubscriptionCollection, SubscriptionCollection>();
+			services.AddSingleton<IStreamCollection, StreamCollection>();
+			services.AddSingleton<ITemplateCollection, TemplateCollection>();
+			services.AddSingleton<ITestDataCollection, TestDataCollection>();
+			services.AddSingleton<IUtilizationDataCollection, UtilizationDataCollection>();
+			services.AddSingleton<ITemplateCollection, TemplateCollection>();
+			services.AddSingleton<IToolCollection, ToolCollection>();
+			services.AddSingleton<IUgsMetadataCollection, UgsMetadataCollection>();
+			services.AddSingleton<IDeviceCollection, DeviceCollection>();
+
+			// Empty mocked object to satisfy basic test runs
+			services.AddSingleton<IAmazonEC2>(sp => new Mock<IAmazonEC2>().Object);
+			services.AddSingleton<IAmazonAutoScaling>(sp => new Mock<IAmazonAutoScaling>().Object);
+			services.AddSingleton<IAmazonCloudWatch>(sp => new Mock<IAmazonCloudWatch>().Object);
+			services.AddSingleton<IFleetManagerFactory, FleetManagerFactory>();
+
+			services.AddSingleton<IPoolSizeStrategyFactory, NoOpPoolSizeStrategyFactory>();
+			services.AddSingleton<IPoolSizeStrategyFactory, JobQueueStrategyFactory>();
+			services.AddSingleton<IPoolSizeStrategyFactory, LeaseUtilizationStrategyFactory>();
+			services.AddSingleton<IPoolSizeStrategyFactory, LeaseUtilizationAwsMetricStrategyFactory>();
+
+			services.AddSingleton<AgentService>();
+			services.AddSingleton(provider => new Lazy<AgentService>(provider.GetRequiredService<AgentService>));
+			services.AddSingleton<AgentRelayService>();
+			services.AddSingleton<AwsAutoScalingLifecycleService>();
+			services.AddSingleton<ArtifactExpirationService>();
+			services.AddSingleton<FleetService>();
+			services.AddSingleton<ConsistencyService>();
+			services.AddSingleton<RequestTrackerService>();
+			services.AddSingleton<GlobalsService>();
+			services.AddSingleton<JobTaskSource>();
+			services.AddSingleton<IssueService>();
+			services.AddSingleton<JobService>();
+			services.AddSingleton<JobExpirationService>();
+			services.AddSingleton<LogTailService>();
+			services.AddSingleton<INotificationService, NotificationService>();
+			services.AddSingleton<IPerforceService, PerforceServiceStub>();
+			services.AddSingleton<PerforceLoadBalancer>();
+			services.AddSingleton<PoolService>();
+			services.AddSingleton<BisectService>();
+			services.AddSingleton<RpcService>();
+			services.AddSingleton<ScheduleService>();
+			services.AddSingleton<DeviceService>();
+			services.AddSingleton<TestDataService>();
+			services.AddSingleton<ComputeService>();
+			services.AddSingleton<EnrollmentService>();
+
+			services.AddSingleton<ConformTaskSource>();
+			services.AddSingleton<ICommitService, CommitService>();
+
+			services.AddSingleton<IDefaultAclModifier, BuildAclModifier>();
+		}
+
+		public Task<Fixture> CreateFixtureAsync()
+		{
+			return Fixture.CreateAsync(ConfigService, GraphCollection, TemplateCollection, JobService, AgentService, PluginCollection, ServerSettings);
+		}
+
+		private JobsController GetJobsController()
+		{
+			JobsController jobsCtrl = ActivatorUtilities.CreateInstance<JobsController>(ServiceProvider);
+			jobsCtrl.ControllerContext = GetControllerContext();
+			return jobsCtrl;
+		}
+
+		private DevicesController GetDevicesController()
+		{
+			DevicesController devicesCtrl = ActivatorUtilities.CreateInstance<DevicesController>(ServiceProvider);
+			devicesCtrl.ControllerContext = GetControllerContext();
+			return devicesCtrl;
+		}
+
+		private DashboardController GetDashboardController()
+		{
+			DashboardController dashboardCtrl = ActivatorUtilities.CreateInstance<DashboardController>(ServiceProvider);
+			dashboardCtrl.ControllerContext = GetControllerContext();
+			return dashboardCtrl;
+		}
+
+		private TestDataController GetTestDataController()
+		{
+			TestDataController dataCtrl = ActivatorUtilities.CreateInstance<TestDataController>(ServiceProvider);
+			dataCtrl.ControllerContext = GetControllerContext();
+			return dataCtrl;
+		}
+
+		private BisectTasksController GetBisectTasksController()
+		{
+			BisectTasksController bisectCtrl = ActivatorUtilities.CreateInstance<BisectTasksController>(ServiceProvider);
+			bisectCtrl.ControllerContext = GetControllerContext();
+			return bisectCtrl;
+		}
+
+		private AgentsController GetAgentsController()
+		{
+			AgentsController agentCtrl = ActivatorUtilities.CreateInstance<AgentsController>(ServiceProvider);
+			agentCtrl.ControllerContext = GetControllerContext();
+			return agentCtrl;
+		}
+
+		private PoolsController GetPoolsController()
+		{
+			PoolsController controller = ActivatorUtilities.CreateInstance<PoolsController>(ServiceProvider);
+			controller.ControllerContext = GetControllerContext();
+			return controller;
+		}
+
+		private LeasesController GetLeasesController()
+		{
+			LeasesController controller = ActivatorUtilities.CreateInstance<LeasesController>(ServiceProvider);
+			controller.ControllerContext = GetControllerContext();
+			return controller;
+		}
+
+		private static ControllerContext GetControllerContext()
+		{
+			ControllerContext controllerContext = new ControllerContext();
+			controllerContext.HttpContext = new DefaultHttpContext();
+			controllerContext.HttpContext.User = new ClaimsPrincipal(new ClaimsIdentity(
+				new List<Claim> { HordeClaims.AdminClaim.ToClaim() }, "TestAuthType"));
+			return controllerContext;
+		}
+
+		private static int s_agentIdCounter = 1;
+		public Task<IAgent> CreateAgentAsync(IPool pool, bool enabled = true, bool requestShutdown = false, List<string>? properties = null, List<AgentWorkspaceInfo>? workspaces = null, TimeSpan? adjustClockBy = null)
+		{
+			return CreateAgentAsync(pool.Id, enabled, requestShutdown, properties, workspaces, adjustClockBy);
+		}
+
+		/// <summary>
+		/// Helper function for setting up agents to be used in tests
+		/// </summary>
+		/// <param name="poolId">Pool ID which the agent should belong to</param>
+		/// <param name="enabled">Whether set the agent as enabled</param>
+		/// <param name="requestShutdown">Mark it with a request for shutdown</param>
+		/// <param name="properties">Any properties to assign</param>
+		/// <param name="workspaces">Any workspaces to assign</param>
+		/// <param name="adjustClockBy">Time span to temporarily skew the clock when creating the agent</param>
+		/// <param name="awsInstanceId">AWS instance ID for the agent (will be set in properties)</param>
+		/// <param name="lease">A lease to assign the agent</param>
+		/// <param name="ephemeral">Whether the agent is ephemeral</param>
+		/// <returns>A new agent</returns>
+		public async Task<IAgent> CreateAgentAsync(
+			PoolId? poolId,
+			bool enabled = true,
+			bool requestShutdown = false,
+			List<string>? properties = null,
+			List<AgentWorkspaceInfo>? workspaces = null,
+			TimeSpan? adjustClockBy = null,
+			string? awsInstanceId = null,
+			AgentLease? lease = null,
+			bool ephemeral = false)
+		{
+			DateTime now = Clock.UtcNow;
+			if (adjustClockBy != null)
+			{
+				Clock.UtcNow = now + adjustClockBy.Value;
+			}
+
+			Dictionary<string, int> resources = new();
+			List<string> tempProps = new(properties ?? new List<string>());
+			if (awsInstanceId != null)
+			{
+				tempProps.Add(KnownPropertyNames.AwsInstanceId + "=" + awsInstanceId);
+			}
+
+			IAgent? agent = await AgentService.CreateAgentAsync("TestAgent" + s_agentIdCounter++, ephemeral, "");
+			Assert.IsNotNull(agent);
+
+			agent = await agent.TryUpdateAsync(new UpdateAgentOptions { Enabled = enabled, ExplicitPools = poolId != null ? [poolId.Value] : [] });
+			Assert.IsNotNull(agent);
+
+			agent = await AgentService.CreateSessionAsync(agent, AgentStatus.Ok, tempProps, resources, null);
+			Assert.IsNotNull(agent);
+
+			if (workspaces is { Count: > 0 })
+			{
+				await agent.TryUpdateWorkspacesAsync(workspaces, false);
+			}
+
+			if (requestShutdown)
+			{
+				await agent.TryUpdateAsync(new UpdateAgentOptions { RequestShutdown = true });
+			}
+
+			if (lease != null)
+			{
+				await agent.TryAddLeaseAsync(lease);
+			}
+
+			Clock.UtcNow = now;
+			return agent;
+		}
+
+		protected async Task<IPool> CreatePoolAsync(PoolConfig poolConfig)
+		{
+			UpdateConfig(config => config.Plugins.GetComputeConfig().Pools.Add(poolConfig));
+			return await PoolCollection.GetAsync(poolConfig.Id) ?? throw new NotImplementedException();
+		}
+	}
+}
\ No newline at end of file
diff --git a/Engine/Source/Programs/Horde/Plugins/Build/HordeServer.Build.Tests/Compute/ComputeServiceTest.cs b/Engine/Source/Programs/Horde/Plugins/Build/HordeServer.Build.Tests/Compute/ComputeServiceTest.cs
--- a/Engine/Source/Programs/Horde/Plugins/Build/HordeServer.Build.Tests/Compute/ComputeServiceTest.cs
+++ b/Engine/Source/Programs/Horde/Plugins/Build/HordeServer.Build.Tests/Compute/ComputeServiceTest.cs
@@ -22,7 +22,7 @@ using OpenTelemetry.Trace;
 namespace HordeServer.Tests.Compute
 {
 	[TestClass]
-	public class ComputeServiceTest : TestSetup
+	public class ComputeServiceTest : BuildTestSetup
 	{
 		private readonly ClusterId _cluster1 = new("cluster1");
 
@@ -264,7 +264,7 @@ namespace HordeServer.Tests.Compute
 				new() { Id = cluster1, Condition = $"pool == '{poolFoo.ToString()}'" },
 				new() { Id = cluster2, Condition = $"pool == '{poolBar.ToString()}'" }
 			};
-			GlobalConfig.CurrentValue.PostLoad(new ServerSettings(), new List<ILoadedPlugin>());
+			GlobalConfig.CurrentValue.PostLoad(new ServerSettings(), Array.Empty<ILoadedPlugin>(), Array.Empty<IDefaultAclModifier>());
 
 			List<string> props = ["ComputeIp=11.0.0.1", "ComputePort=5000"];
 			IAgent agent1 = await CreateAgentAsync(new PoolId("foo"), properties: props);
diff --git a/Engine/Source/Programs/Horde/Plugins/Build/HordeServer.Build.Tests/Dashboard/DashboardTests.cs b/Engine/Source/Programs/Horde/Plugins/Build/HordeServer.Build.Tests/Dashboard/DashboardTests.cs
--- a/Engine/Source/Programs/Horde/Plugins/Build/HordeServer.Build.Tests/Dashboard/DashboardTests.cs
+++ b/Engine/Source/Programs/Horde/Plugins/Build/HordeServer.Build.Tests/Dashboard/DashboardTests.cs
@@ -12,7 +12,7 @@ namespace HordeServer.Tests.Dashboard
 	/// Tests for the dashboard
 	/// </summary>
 	[TestClass]
-	public class DashboardTest : TestSetup
+	public class DashboardTest : BuildTestSetup
 	{
 		[TestMethod]
 		public async Task TestPreviewsAsync()
diff --git a/Engine/Source/Programs/Horde/Plugins/Build/HordeServer.Build.Tests/Devices/DeviceServiceTests.cs b/Engine/Source/Programs/Horde/Plugins/Build/HordeServer.Build.Tests/Devices/DeviceServiceTests.cs
--- a/Engine/Source/Programs/Horde/Plugins/Build/HordeServer.Build.Tests/Devices/DeviceServiceTests.cs
+++ b/Engine/Source/Programs/Horde/Plugins/Build/HordeServer.Build.Tests/Devices/DeviceServiceTests.cs
@@ -28,7 +28,7 @@ namespace HordeServer.Tests.Devices
 	/// Tests for the device service
 	/// </summary>
 	[TestClass]
-	public class DeviceServiceTest : TestSetup
+	public class DeviceServiceTest : BuildTestSetup
 	{
 		private DevicesController? _deviceController;
 
diff --git a/Engine/Source/Programs/Horde/Plugins/Build/HordeServer.Build.Tests/Fixture.cs b/Engine/Source/Programs/Horde/Plugins/Build/HordeServer.Build.Tests/Fixture.cs
--- a/Engine/Source/Programs/Horde/Plugins/Build/HordeServer.Build.Tests/Fixture.cs
+++ b/Engine/Source/Programs/Horde/Plugins/Build/HordeServer.Build.Tests/Fixture.cs
@@ -8,6 +8,7 @@ using EpicGames.Horde.Jobs;
 using EpicGames.Horde.Jobs.Templates;
 using EpicGames.Horde.Projects;
 using EpicGames.Horde.Streams;
+using HordeServer.Acls;
 using HordeServer.Agents;
 using HordeServer.Configuration;
 using HordeServer.Jobs;
@@ -83,7 +84,7 @@ namespace HordeServer.Tests
 
 			GlobalConfig globalConfig = new GlobalConfig();
 			globalConfig.Plugins.AddBuildConfig(buildConfig);
-			globalConfig.PostLoad(serverSettings, pluginCollection.LoadedPlugins);
+			globalConfig.PostLoad(serverSettings, pluginCollection.LoadedPlugins, Array.Empty<IDefaultAclModifier>());
 			configService.OverrideConfig(globalConfig);
 
 			StreamId = streamId;
diff --git a/Engine/Source/Programs/Horde/Plugins/Build/HordeServer.Build.Tests/Fleet/AwsAutoScalingLifecycleTest.cs b/Engine/Source/Programs/Horde/Plugins/Build/HordeServer.Build.Tests/Fleet/AwsAutoScalingLifecycleTest.cs
--- a/Engine/Source/Programs/Horde/Plugins/Build/HordeServer.Build.Tests/Fleet/AwsAutoScalingLifecycleTest.cs
+++ b/Engine/Source/Programs/Horde/Plugins/Build/HordeServer.Build.Tests/Fleet/AwsAutoScalingLifecycleTest.cs
@@ -23,7 +23,7 @@ using Moq;
 namespace HordeServer.Tests.Fleet;
 
 [TestClass]
-public class AwsAutoScalingLifecycleServiceTest : TestSetup
+public class AwsAutoScalingLifecycleServiceTest : BuildTestSetup
 {
 	private AwsAutoScalingLifecycleService _asgLifecycleService = default!;
 	private Mock<IAmazonAutoScaling> _asgMock = default!;
diff --git a/Engine/Source/Programs/Horde/Plugins/Build/HordeServer.Build.Tests/Fleet/AwsRecyclingFleetManagerTest.cs b/Engine/Source/Programs/Horde/Plugins/Build/HordeServer.Build.Tests/Fleet/AwsRecyclingFleetManagerTest.cs
--- a/Engine/Source/Programs/Horde/Plugins/Build/HordeServer.Build.Tests/Fleet/AwsRecyclingFleetManagerTest.cs
+++ b/Engine/Source/Programs/Horde/Plugins/Build/HordeServer.Build.Tests/Fleet/AwsRecyclingFleetManagerTest.cs
@@ -19,7 +19,7 @@ using JsonSerializer = System.Text.Json.JsonSerializer;
 namespace HordeServer.Tests.Fleet
 {
 	[TestClass]
-	public class AwsRecyclingFleetManagerTest : TestSetup
+	public class AwsRecyclingFleetManagerTest : BuildTestSetup
 	{
 		private readonly FakeClock _clock = new();
 
diff --git a/Engine/Source/Programs/Horde/Plugins/Build/HordeServer.Build.Tests/Fleet/AwsReuseFleetManagerTest.cs b/Engine/Source/Programs/Horde/Plugins/Build/HordeServer.Build.Tests/Fleet/AwsReuseFleetManagerTest.cs
--- a/Engine/Source/Programs/Horde/Plugins/Build/HordeServer.Build.Tests/Fleet/AwsReuseFleetManagerTest.cs
+++ b/Engine/Source/Programs/Horde/Plugins/Build/HordeServer.Build.Tests/Fleet/AwsReuseFleetManagerTest.cs
@@ -16,7 +16,7 @@ using Microsoft.VisualStudio.TestTools.UnitTesting;
 namespace HordeServer.Tests.Fleet
 {
 	[TestClass]
-	public class AwsReuseFleetManagerTest : TestSetup
+	public class AwsReuseFleetManagerTest : BuildTestSetup
 	{
 		[TestMethod]
 		public async Task ExpandOneAgentAsync()
diff --git a/Engine/Source/Programs/Horde/Plugins/Build/HordeServer.Build.Tests/Fleet/FleetServiceTest.cs b/Engine/Source/Programs/Horde/Plugins/Build/HordeServer.Build.Tests/Fleet/FleetServiceTest.cs
--- a/Engine/Source/Programs/Horde/Plugins/Build/HordeServer.Build.Tests/Fleet/FleetServiceTest.cs
+++ b/Engine/Source/Programs/Horde/Plugins/Build/HordeServer.Build.Tests/Fleet/FleetServiceTest.cs
@@ -96,7 +96,7 @@ namespace HordeServer.Tests.Fleet
 	}
 
 	[TestClass]
-	public class FleetServiceTest : TestSetup
+	public class FleetServiceTest : BuildTestSetup
 	{
 		readonly FleetManagerSpy _fleetManagerSpy = new();
 
@@ -244,7 +244,7 @@ namespace HordeServer.Tests.Fleet
 	}
 
 	[TestClass]
-	public class PoolSizeStrategyFactoryTest : TestSetup
+	public class PoolSizeStrategyFactoryTest : BuildTestSetup
 	{
 		private static int s_poolCount;
 
@@ -433,7 +433,7 @@ namespace HordeServer.Tests.Fleet
 	}
 
 	[TestClass]
-	public class FleetManagerFactoryTest : TestSetup
+	public class FleetManagerFactoryTest : BuildTestSetup
 	{
 		private static int s_poolCount;
 
diff --git a/Engine/Source/Programs/Horde/Plugins/Build/HordeServer.Build.Tests/Fleet/JobQueueStrategyTest.cs b/Engine/Source/Programs/Horde/Plugins/Build/HordeServer.Build.Tests/Fleet/JobQueueStrategyTest.cs
--- a/Engine/Source/Programs/Horde/Plugins/Build/HordeServer.Build.Tests/Fleet/JobQueueStrategyTest.cs
+++ b/Engine/Source/Programs/Horde/Plugins/Build/HordeServer.Build.Tests/Fleet/JobQueueStrategyTest.cs
@@ -24,7 +24,7 @@ using Microsoft.VisualStudio.TestTools.UnitTesting;
 namespace HordeServer.Tests.Fleet
 {
 	[TestClass]
-	public class JobQueueStrategyTest : TestSetup
+	public class JobQueueStrategyTest : BuildTestSetup
 	{
 		[TestMethod]
 		[Ignore("Flaky test when run through CI / Horde")]
diff --git a/Engine/Source/Programs/Horde/Plugins/Build/HordeServer.Build.Tests/Fleet/LeaseUtilizationAwsMetricStrategyTests.cs b/Engine/Source/Programs/Horde/Plugins/Build/HordeServer.Build.Tests/Fleet/LeaseUtilizationAwsMetricStrategyTests.cs
--- a/Engine/Source/Programs/Horde/Plugins/Build/HordeServer.Build.Tests/Fleet/LeaseUtilizationAwsMetricStrategyTests.cs
+++ b/Engine/Source/Programs/Horde/Plugins/Build/HordeServer.Build.Tests/Fleet/LeaseUtilizationAwsMetricStrategyTests.cs
@@ -22,7 +22,7 @@ using Microsoft.VisualStudio.TestTools.UnitTesting;
 namespace HordeServer.Tests.Fleet
 {
 	[TestClass]
-	public class LeaseUtilizationAwsMetricStrategyTest : TestSetup
+	public class LeaseUtilizationAwsMetricStrategyTest : BuildTestSetup
 	{
 		private readonly IPool _pool;
 		private readonly LeaseUtilizationAwsMetricSettings _settings = new(60, "myNamespace");
diff --git a/Engine/Source/Programs/Horde/Plugins/Build/HordeServer.Build.Tests/Fleet/LeaseUtilizationStrategyTest.cs b/Engine/Source/Programs/Horde/Plugins/Build/HordeServer.Build.Tests/Fleet/LeaseUtilizationStrategyTest.cs
--- a/Engine/Source/Programs/Horde/Plugins/Build/HordeServer.Build.Tests/Fleet/LeaseUtilizationStrategyTest.cs
+++ b/Engine/Source/Programs/Horde/Plugins/Build/HordeServer.Build.Tests/Fleet/LeaseUtilizationStrategyTest.cs
@@ -18,7 +18,7 @@ using Microsoft.VisualStudio.TestTools.UnitTesting;
 namespace HordeServer.Tests.Fleet
 {
 	[TestClass]
-	public class LeaseUtilizationStrategyTest : TestSetup
+	public class LeaseUtilizationStrategyTest : BuildTestSetup
 	{
 		private IAgent _agent1 = null!;
 		private IAgent _agent2 = null!;
diff --git a/Engine/Source/Programs/Horde/Plugins/Build/HordeServer.Build.Tests/Issues/IssueServiceTests.cs b/Engine/Source/Programs/Horde/Plugins/Build/HordeServer.Build.Tests/Issues/IssueServiceTests.cs
--- a/Engine/Source/Programs/Horde/Plugins/Build/HordeServer.Build.Tests/Issues/IssueServiceTests.cs
+++ b/Engine/Source/Programs/Horde/Plugins/Build/HordeServer.Build.Tests/Issues/IssueServiceTests.cs
@@ -41,7 +41,7 @@ using Moq;
 namespace HordeServer.Tests.Issues
 {
 	[TestClass]
-	public class IssueServiceTests : TestSetup
+	public class IssueServiceTests : BuildTestSetup
 	{
 		class TestJsonLogger : ILogger, IAsyncDisposable
 		{
diff --git a/Engine/Source/Programs/Horde/Plugins/Build/HordeServer.Build.Tests/Jobs/Bisect/BisectTests.cs b/Engine/Source/Programs/Horde/Plugins/Build/HordeServer.Build.Tests/Jobs/Bisect/BisectTests.cs
--- a/Engine/Source/Programs/Horde/Plugins/Build/HordeServer.Build.Tests/Jobs/Bisect/BisectTests.cs
+++ b/Engine/Source/Programs/Horde/Plugins/Build/HordeServer.Build.Tests/Jobs/Bisect/BisectTests.cs
@@ -29,7 +29,7 @@ using Moq;
 namespace HordeServer.Tests.Jobs.Bisect
 {
 	[TestClass]
-	public class BisectTests : TestSetup
+	public class BisectTests : BuildTestSetup
 	{
 		[TestMethod]
 		public async Task TestStatesAsync()
diff --git a/Engine/Source/Programs/Horde/Plugins/Build/HordeServer.Build.Tests/Jobs/JobCollectionTests.cs b/Engine/Source/Programs/Horde/Plugins/Build/HordeServer.Build.Tests/Jobs/JobCollectionTests.cs
--- a/Engine/Source/Programs/Horde/Plugins/Build/HordeServer.Build.Tests/Jobs/JobCollectionTests.cs
+++ b/Engine/Source/Programs/Horde/Plugins/Build/HordeServer.Build.Tests/Jobs/JobCollectionTests.cs
@@ -27,7 +27,7 @@ using Moq;
 namespace HordeServer.Tests.Jobs
 {
 	[TestClass]
-	public class JobCollectionTests : TestSetup
+	public class JobCollectionTests : BuildTestSetup
 	{
 		static NewGroup AddGroup(List<NewGroup> groups)
 		{
diff --git a/Engine/Source/Programs/Horde/Plugins/Build/HordeServer.Build.Tests/Jobs/JobExpirationTests.cs b/Engine/Source/Programs/Horde/Plugins/Build/HordeServer.Build.Tests/Jobs/JobExpirationTests.cs
--- a/Engine/Source/Programs/Horde/Plugins/Build/HordeServer.Build.Tests/Jobs/JobExpirationTests.cs
+++ b/Engine/Source/Programs/Horde/Plugins/Build/HordeServer.Build.Tests/Jobs/JobExpirationTests.cs
@@ -17,7 +17,7 @@ using Microsoft.VisualStudio.TestTools.UnitTesting;
 namespace HordeServer.Tests.Jobs
 {
 	[TestClass]
-	public class JobExpirationTests : TestSetup
+	public class JobExpirationTests : BuildTestSetup
 	{
 		async Task<IJob> CreateJobAsync()
 		{
diff --git a/Engine/Source/Programs/Horde/Plugins/Build/HordeServer.Build.Tests/Jobs/JobServiceTests.cs b/Engine/Source/Programs/Horde/Plugins/Build/HordeServer.Build.Tests/Jobs/JobServiceTests.cs
--- a/Engine/Source/Programs/Horde/Plugins/Build/HordeServer.Build.Tests/Jobs/JobServiceTests.cs
+++ b/Engine/Source/Programs/Horde/Plugins/Build/HordeServer.Build.Tests/Jobs/JobServiceTests.cs
@@ -30,7 +30,7 @@ namespace HordeServer.Tests.Jobs
 	using ProjectId = EpicGames.Horde.Projects.ProjectId;
 
 	[TestClass]
-	public class JobServiceTests : TestSetup
+	public class JobServiceTests : BuildTestSetup
 	{
 		[TestMethod]
 		public async Task TestChainedJobsAsync()
diff --git a/Engine/Source/Programs/Horde/Plugins/Build/HordeServer.Build.Tests/Jobs/JobTaskSourceTests.cs b/Engine/Source/Programs/Horde/Plugins/Build/HordeServer.Build.Tests/Jobs/JobTaskSourceTests.cs
--- a/Engine/Source/Programs/Horde/Plugins/Build/HordeServer.Build.Tests/Jobs/JobTaskSourceTests.cs
+++ b/Engine/Source/Programs/Horde/Plugins/Build/HordeServer.Build.Tests/Jobs/JobTaskSourceTests.cs
@@ -20,7 +20,7 @@ using Microsoft.VisualStudio.TestTools.UnitTesting;
 namespace HordeServer.Tests.Jobs
 {
 	[TestClass]
-	public class JobTaskSourceTests : TestSetup
+	public class JobTaskSourceTests : BuildTestSetup
 	{
 		private bool _eventReceived;
 		private bool? _eventPoolHasAgentsOnline;
diff --git a/Engine/Source/Programs/Horde/Plugins/Build/HordeServer.Build.Tests/Jobs/JobsControllerDbTest.cs b/Engine/Source/Programs/Horde/Plugins/Build/HordeServer.Build.Tests/Jobs/JobsControllerDbTest.cs
--- a/Engine/Source/Programs/Horde/Plugins/Build/HordeServer.Build.Tests/Jobs/JobsControllerDbTest.cs
+++ b/Engine/Source/Programs/Horde/Plugins/Build/HordeServer.Build.Tests/Jobs/JobsControllerDbTest.cs
@@ -15,7 +15,7 @@ namespace HordeServer.Tests.Jobs
 	/// the entire ASP.NET chain.
 	/// </summary>
 	[TestClass]
-	public class JobsControllerDbTest : TestSetup
+	public class JobsControllerDbTest : BuildTestSetup
 	{
 		[TestMethod]
 		public async Task GetJobsAsync()
diff --git a/Engine/Source/Programs/Horde/Plugins/Build/HordeServer.Build.Tests/Jobs/Schedules/SchedulerTests.cs b/Engine/Source/Programs/Horde/Plugins/Build/HordeServer.Build.Tests/Jobs/Schedules/SchedulerTests.cs
--- a/Engine/Source/Programs/Horde/Plugins/Build/HordeServer.Build.Tests/Jobs/Schedules/SchedulerTests.cs
+++ b/Engine/Source/Programs/Horde/Plugins/Build/HordeServer.Build.Tests/Jobs/Schedules/SchedulerTests.cs
@@ -24,7 +24,7 @@ using Microsoft.VisualStudio.TestTools.UnitTesting;
 namespace HordeServer.Tests.Jobs.Schedules
 {
 	[TestClass]
-	public class SchedulerTests : TestSetup
+	public class SchedulerTests : BuildTestSetup
 	{
 		ProjectId ProjectId { get; } = new ProjectId("ue5");
 		StreamId StreamId { get; } = new StreamId("ue5-main");
diff --git a/Engine/Source/Programs/Horde/Plugins/Build/HordeServer.Build.Tests/Jobs/TestData/TestDataTests.cs b/Engine/Source/Programs/Horde/Plugins/Build/HordeServer.Build.Tests/Jobs/TestData/TestDataTests.cs
--- a/Engine/Source/Programs/Horde/Plugins/Build/HordeServer.Build.Tests/Jobs/TestData/TestDataTests.cs
+++ b/Engine/Source/Programs/Horde/Plugins/Build/HordeServer.Build.Tests/Jobs/TestData/TestDataTests.cs
@@ -34,7 +34,7 @@ namespace HordeServer.Tests.Jobs.TestData
 	/// Tests for the device service
 	/// </summary>
 	[TestClass]
-	public class TestDataTests : TestSetup
+	public class TestDataTests : BuildTestSetup
 	{
 		const string MainStreamName = "//UE5/Main";
 		readonly StreamId _mainStreamId = new StreamId(StringId.Sanitize(MainStreamName));
diff --git a/Engine/Source/Programs/Horde/Plugins/Build/HordeServer.Build.Tests/Notifications/NotificationServiceTests.cs b/Engine/Source/Programs/Horde/Plugins/Build/HordeServer.Build.Tests/Notifications/NotificationServiceTests.cs
--- a/Engine/Source/Programs/Horde/Plugins/Build/HordeServer.Build.Tests/Notifications/NotificationServiceTests.cs
+++ b/Engine/Source/Programs/Horde/Plugins/Build/HordeServer.Build.Tests/Notifications/NotificationServiceTests.cs
@@ -50,7 +50,7 @@ namespace HordeServer.Tests.Notifications
 	}
 
 	[TestClass]
-	public class NotificationServiceTests : TestSetup
+	public class NotificationServiceTests : BuildTestSetup
 	{
 		protected override void ConfigureServices(IServiceCollection services)
 		{
diff --git a/Engine/Source/Programs/Horde/Plugins/Build/HordeServer.Build.Tests/Perforce/PerforceLoadBalancerTests.cs b/Engine/Source/Programs/Horde/Plugins/Build/HordeServer.Build.Tests/Perforce/PerforceLoadBalancerTests.cs
--- a/Engine/Source/Programs/Horde/Plugins/Build/HordeServer.Build.Tests/Perforce/PerforceLoadBalancerTests.cs
+++ b/Engine/Source/Programs/Horde/Plugins/Build/HordeServer.Build.Tests/Perforce/PerforceLoadBalancerTests.cs
@@ -39,7 +39,7 @@ public class FakeHealthMonitor<T> : IHealthMonitor<T>
 }
 
 [TestClass]
-public class PerforceLoadBalancerTests : TestSetup
+public class PerforceLoadBalancerTests : BuildTestSetup
 {
 	private const string PerforceHostname = "";
 
diff --git a/Engine/Source/Programs/Horde/Plugins/Build/HordeServer.Build.Tests/Server/RpcServiceTest.cs b/Engine/Source/Programs/Horde/Plugins/Build/HordeServer.Build.Tests/Server/RpcServiceTest.cs
--- a/Engine/Source/Programs/Horde/Plugins/Build/HordeServer.Build.Tests/Server/RpcServiceTest.cs
+++ b/Engine/Source/Programs/Horde/Plugins/Build/HordeServer.Build.Tests/Server/RpcServiceTest.cs
@@ -123,7 +123,7 @@ namespace HordeServer.Tests.Server
 	}
 
 	[TestClass]
-	public class RpcServiceTest : TestSetup
+	public class RpcServiceTest : BuildTestSetup
 	{
 		private readonly ServerCallContext _adminContext = new ServerCallContextStub(HordeClaims.AdminClaim.ToClaim());
 
diff --git a/Engine/Source/Programs/Horde/Plugins/Build/HordeServer.Build.Tests/Streams/StreamCollectionTests.cs b/Engine/Source/Programs/Horde/Plugins/Build/HordeServer.Build.Tests/Streams/StreamCollectionTests.cs
--- a/Engine/Source/Programs/Horde/Plugins/Build/HordeServer.Build.Tests/Streams/StreamCollectionTests.cs
+++ b/Engine/Source/Programs/Horde/Plugins/Build/HordeServer.Build.Tests/Streams/StreamCollectionTests.cs
@@ -9,7 +9,7 @@ using Microsoft.VisualStudio.TestTools.UnitTesting;
 namespace HordeServer.Tests.Streams
 {
 	[TestClass]
-	public class StreamCollectionTests : TestSetup
+	public class StreamCollectionTests : BuildTestSetup
 	{
 		private readonly StreamId _streamId = new("bogusStreamId");
 
diff --git a/Engine/Source/Programs/Horde/Plugins/Build/HordeServer.Build.Tests/Streams/StreamServiceTest.cs b/Engine/Source/Programs/Horde/Plugins/Build/HordeServer.Build.Tests/Streams/StreamServiceTest.cs
--- a/Engine/Source/Programs/Horde/Plugins/Build/HordeServer.Build.Tests/Streams/StreamServiceTest.cs
+++ b/Engine/Source/Programs/Horde/Plugins/Build/HordeServer.Build.Tests/Streams/StreamServiceTest.cs
@@ -9,7 +9,7 @@ using Microsoft.VisualStudio.TestTools.UnitTesting;
 namespace HordeServer.Tests.Streams
 {
 	[TestClass]
-	public class StreamServiceTests : TestSetup
+	public class StreamServiceTests : BuildTestSetup
 	{
 		[TestMethod]
 		public async Task PausingAsync()
diff --git a/Engine/Source/Programs/Horde/Plugins/Build/HordeServer.Build.Tests/TestSetup.cs b/Engine/Source/Programs/Horde/Plugins/Build/HordeServer.Build.Tests/TestSetup.cs
deleted file mode 100644
diff --git a/Engine/Source/Programs/Horde/Plugins/Build/HordeServer.Build/BuildAclModifier.cs b/Engine/Source/Programs/Horde/Plugins/Build/HordeServer.Build/BuildAclModifier.cs
new file mode 100644
--- /dev/null
+++ b/Engine/Source/Programs/Horde/Plugins/Build/HordeServer.Build/BuildAclModifier.cs
@@ -0,0 +1,46 @@
+// Copyright Epic Games, Inc. All Rights Reserved.
+
+using HordeServer.Acls;
+using HordeServer.Agents.Software;
+using HordeServer.Artifacts;
+using HordeServer.Devices;
+using HordeServer.Jobs;
+using HordeServer.Jobs.Bisect;
+using HordeServer.Logs;
+using HordeServer.Notifications;
+using HordeServer.Projects;
+using HordeServer.Streams;
+using HordeServer.Utilities;
+
+namespace HordeServer
+{
+	class BuildAclModifier : IDefaultAclModifier
+	{
+		/// <inheritdoc/>
+		public void Apply(DefaultAclBuilder acl)
+		{
+			acl.AddCustomRole(HordeClaims.AgentRegistrationClaim, new[] { StreamAclAction.ViewStream, ProjectAclAction.ViewProject, JobAclAction.ViewJob });
+			acl.AddCustomRole(HordeClaims.AgentRoleClaim, new[] { ProjectAclAction.ViewProject, StreamAclAction.ViewStream, LogAclAction.CreateEvent, AgentSoftwareAclAction.DownloadSoftware });
+			acl.AddCustomRole(HordeClaims.ConfigureProjectsClaim, new[] { ProjectAclAction.CreateProject, ProjectAclAction.UpdateProject, ProjectAclAction.ViewProject, StreamAclAction.CreateStream, StreamAclAction.UpdateStream, StreamAclAction.ViewStream });
+			acl.AddCustomRole(HordeClaims.StartChainedJobClaim, new[] { JobAclAction.CreateJob, JobAclAction.ExecuteJob, JobAclAction.UpdateJob, JobAclAction.ViewJob, StreamAclAction.ViewTemplate, StreamAclAction.ViewStream });
+
+			acl.AddDefaultReadAction(ArtifactAclAction.DownloadArtifact);
+			acl.AddDefaultReadAction(ArtifactAclAction.ReadArtifact);
+			acl.AddDefaultReadAction(BisectTaskAclAction.ViewBisectTask);
+			acl.AddDefaultReadAction(DeviceAclAction.DeviceRead);
+			acl.AddDefaultReadAction(JobAclAction.ViewJob);
+			acl.AddDefaultReadAction(NotificationAclAction.CreateSubscription);
+			acl.AddDefaultReadAction(ProjectAclAction.ViewProject);
+			acl.AddDefaultReadAction(StreamAclAction.ViewChanges);
+			acl.AddDefaultReadAction(StreamAclAction.ViewStream);
+			acl.AddDefaultReadAction(StreamAclAction.ViewTemplate);
+
+			acl.AddDefaultWriteAction(JobAclAction.CreateJob);
+			acl.AddDefaultWriteAction(JobAclAction.UpdateJob);
+			acl.AddDefaultWriteAction(JobAclAction.RetryJobStep);
+			acl.AddDefaultWriteAction(DeviceAclAction.DeviceWrite);
+			acl.AddDefaultWriteAction(BisectTaskAclAction.CreateBisectTask);
+			acl.AddDefaultWriteAction(BisectTaskAclAction.UpdateBisectTask);
+		}
+	}
+}
diff --git a/Engine/Source/Programs/Horde/Plugins/Build/HordeServer.Build/BuildPlugin.cs b/Engine/Source/Programs/Horde/Plugins/Build/HordeServer.Build/BuildPlugin.cs
--- a/Engine/Source/Programs/Horde/Plugins/Build/HordeServer.Build/BuildPlugin.cs
+++ b/Engine/Source/Programs/Horde/Plugins/Build/HordeServer.Build/BuildPlugin.cs
@@ -7,6 +7,7 @@ using System.Linq;
 using System.Text;
 using System.Threading.Tasks;
 using EpicGames.Core;
+using HordeServer.Acls;
 using HordeServer.Agents;
 using HordeServer.Agents.Fleet;
 using HordeServer.Agents.Pools;
@@ -55,6 +56,8 @@ namespace HordeServer
 		/// <inheritdoc/>
 		public void ConfigureServices(IServiceCollection services)
 		{
+			services.AddSingleton<IDefaultAclModifier, BuildAclModifier>();
+
 			services.AddSingleton<IArtifactCollection, ArtifactCollection>();
 			services.AddSingleton<IGraphCollection, GraphCollection>();
 			services.AddSingleton<IssueCollection>();
diff --git a/Engine/Source/Programs/Horde/Plugins/Compute/HordeServer.Compute.Tests/ComputeTestSetup.cs b/Engine/Source/Programs/Horde/Plugins/Compute/HordeServer.Compute.Tests/ComputeTestSetup.cs
--- a/Engine/Source/Programs/Horde/Plugins/Compute/HordeServer.Compute.Tests/ComputeTestSetup.cs
+++ b/Engine/Source/Programs/Horde/Plugins/Compute/HordeServer.Compute.Tests/ComputeTestSetup.cs
@@ -5,6 +5,7 @@ using System.Collections.Generic;
 using System.Linq;
 using System.Text;
 using System.Threading.Tasks;
+using HordeServer.Acls;
 using HordeServer.Compute;
 using HordeServer.Logs;
 using HordeServer.Storage;
@@ -28,6 +29,8 @@ namespace HordeServer.Tests
 		protected override void ConfigureServices(IServiceCollection services)
 		{
 			base.ConfigureServices(services);
+
+			services.AddSingleton<IDefaultAclModifier, ComputeAclModifier>();
 		}
 	}
 }
diff --git a/Engine/Source/Programs/Horde/Plugins/Compute/HordeServer.Compute/ComputeAclModifier.cs b/Engine/Source/Programs/Horde/Plugins/Compute/HordeServer.Compute/ComputeAclModifier.cs
new file mode 100644
--- /dev/null
+++ b/Engine/Source/Programs/Horde/Plugins/Compute/HordeServer.Compute/ComputeAclModifier.cs
@@ -0,0 +1,33 @@
+// Copyright Epic Games, Inc. All Rights Reserved.
+
+using System.Security.Claims;
+using HordeServer.Acls;
+using HordeServer.Agents;
+using HordeServer.Agents.Pools;
+using HordeServer.Agents.Sessions;
+using HordeServer.Agents.Software;
+using HordeServer.Logs;
+using HordeServer.Server;
+using HordeServer.Utilities;
+
+namespace HordeServer
+{
+	class ComputeAclModifier : IDefaultAclModifier
+	{
+		/// <inheritdoc/>
+		public void Apply(DefaultAclBuilder acl)
+		{
+			acl.AddCustomRole(new AclClaimConfig(ClaimTypes.Role, "internal:AgentRegistration"), new[] { AgentAclAction.CreateAgent, SessionAclAction.CreateSession });
+			acl.AddCustomRole(HordeClaims.AgentRegistrationClaim, new[] { AgentAclAction.CreateAgent, SessionAclAction.CreateSession, AgentAclAction.UpdateAgent, AgentSoftwareAclAction.DownloadSoftware, PoolAclAction.CreatePool, PoolAclAction.UpdatePool, PoolAclAction.ViewPool, PoolAclAction.DeletePool, PoolAclAction.ListPools, ServerAclAction.ViewCosts });
+			acl.AddCustomRole(HordeClaims.DownloadSoftwareClaim, new[] { AgentSoftwareAclAction.DownloadSoftware });
+			acl.AddCustomRole(HordeClaims.UploadToolsClaim, new[] { AgentSoftwareAclAction.UploadSoftware });
+
+			acl.AddDefaultReadAction(AgentAclAction.ListAgents);
+			acl.AddDefaultReadAction(AgentAclAction.ViewAgent);
+			acl.AddDefaultReadAction(LogAclAction.ViewEvent);
+			acl.AddDefaultReadAction(LogAclAction.ViewLog);
+			acl.AddDefaultReadAction(PoolAclAction.ListPools);
+			acl.AddDefaultReadAction(PoolAclAction.ViewPool);
+		}
+	}
+}
diff --git a/Engine/Source/Programs/Horde/Plugins/Compute/HordeServer.Compute/ComputePlugin.cs b/Engine/Source/Programs/Horde/Plugins/Compute/HordeServer.Compute/ComputePlugin.cs
--- a/Engine/Source/Programs/Horde/Plugins/Compute/HordeServer.Compute/ComputePlugin.cs
+++ b/Engine/Source/Programs/Horde/Plugins/Compute/HordeServer.Compute/ComputePlugin.cs
@@ -6,6 +6,7 @@ using Amazon.CloudWatch;
 using Amazon.EC2;
 using Amazon.Extensions.NETCore.Setup;
 using Amazon.SQS;
+using HordeServer.Acls;
 using HordeServer.Agents;
 using HordeServer.Agents.Enrollment;
 using HordeServer.Agents.Fleet;
@@ -46,6 +47,8 @@ namespace HordeServer
 		/// <inheritdoc/>
 		public void ConfigureServices(IServiceCollection services)
 		{
+			services.AddSingleton<IDefaultAclModifier, ComputeAclModifier>();
+
 			services.AddSingleton<IAgentCollection, AgentCollection>();
 			services.AddSingleton<ILeaseCollection, LeaseCollection>();
 			services.AddSingleton<ILogCollection, LogCollection>();
diff --git a/Engine/Source/Programs/Horde/Plugins/Ddc/HordeServer.Ddc.Tests/Ddc/Functional/ContentId/CompressedBlobTests.cs b/Engine/Source/Programs/Horde/Plugins/Ddc/HordeServer.Ddc.Tests/Ddc/Functional/ContentId/CompressedBlobTests.cs
--- a/Engine/Source/Programs/Horde/Plugins/Ddc/HordeServer.Ddc.Tests/Ddc/Functional/ContentId/CompressedBlobTests.cs
+++ b/Engine/Source/Programs/Horde/Plugins/Ddc/HordeServer.Ddc.Tests/Ddc/Functional/ContentId/CompressedBlobTests.cs
@@ -15,6 +15,7 @@ using Blake3;
 using EpicGames.AspNet;
 using EpicGames.Core;
 using EpicGames.Horde.Storage;
+using HordeServer.Acls;
 using HordeServer.Configuration;
 using HordeServer.Ddc;
 using HordeServer.Plugins;
@@ -47,7 +48,7 @@ namespace HordeServer.Tests.Ddc.FunctionalTests.CompressedBlobs
 
 			GlobalConfig globalConfig = ServiceProvider.GetRequiredService<IOptions<GlobalConfig>>().Value;
 			globalConfig.Plugins.AddStorageConfig(storageConfig);
-			globalConfig.PostLoad(serverSettings, new List<ILoadedPlugin>());
+			globalConfig.PostLoad(serverSettings, Array.Empty<ILoadedPlugin>(), Array.Empty<IDefaultAclModifier>());
 
 			ConfigService configService = ServiceProvider.GetRequiredService<ConfigService>();
 			configService.OverrideConfig(globalConfig);
diff --git a/Engine/Source/Programs/Horde/Plugins/Ddc/HordeServer.Ddc.Tests/Ddc/Functional/Objects/ReferencesTests.cs b/Engine/Source/Programs/Horde/Plugins/Ddc/HordeServer.Ddc.Tests/Ddc/Functional/Objects/ReferencesTests.cs
--- a/Engine/Source/Programs/Horde/Plugins/Ddc/HordeServer.Ddc.Tests/Ddc/Functional/Objects/ReferencesTests.cs
+++ b/Engine/Source/Programs/Horde/Plugins/Ddc/HordeServer.Ddc.Tests/Ddc/Functional/Objects/ReferencesTests.cs
@@ -18,6 +18,7 @@ using EpicGames.AspNet;
 using EpicGames.Core;
 using EpicGames.Horde.Storage;
 using EpicGames.Serialization;
+using HordeServer.Acls;
 using HordeServer.Configuration;
 using HordeServer.Ddc;
 using HordeServer.Plugins;
@@ -63,7 +64,7 @@ namespace HordeServer.Tests.Ddc.FunctionalTests.References
 
 			GlobalConfig globalConfig = ServiceProvider.GetRequiredService<IOptions<GlobalConfig>>().Value;
 			globalConfig.Plugins.AddStorageConfig(storageConfig);
-			globalConfig.PostLoad(serverSettings, new List<ILoadedPlugin>());
+			globalConfig.PostLoad(serverSettings, Array.Empty<ILoadedPlugin>(), Array.Empty<IDefaultAclModifier>());
 
 			ConfigService configService = ServiceProvider.GetRequiredService<ConfigService>();
 			configService.OverrideConfig(globalConfig);
diff --git a/Engine/Source/Programs/Horde/Plugins/Tools/HordeServer.Tools.Tests/Tools/ToolsControllerTests.cs b/Engine/Source/Programs/Horde/Plugins/Tools/HordeServer.Tools.Tests/Tools/ToolsControllerTests.cs
--- a/Engine/Source/Programs/Horde/Plugins/Tools/HordeServer.Tools.Tests/Tools/ToolsControllerTests.cs
+++ b/Engine/Source/Programs/Horde/Plugins/Tools/HordeServer.Tools.Tests/Tools/ToolsControllerTests.cs
@@ -51,7 +51,7 @@ public class ToolsControllerTests
 		globalConfig.Plugins.AddToolsConfig(toolsConfig);
 		
 		ServerSettings serverSettings = new() { AuthMethod = AuthMethod.Horde };
-		globalConfig.PostLoad(serverSettings, pluginCollection.LoadedPlugins);
+		globalConfig.PostLoad(serverSettings, pluginCollection.LoadedPlugins, Array.Empty<IDefaultAclModifier>());
 
 		Dictionary<string, string> settings = new() { { "Horde:AuthMethod", AuthMethod.Horde.ToString() } };
 		await using FakeHordeWebApp app = new(settings: settings);
diff --git a/Engine/Source/Programs/Horde/Plugins/Tools/HordeServer.Tools/ToolsAclModifier.cs b/Engine/Source/Programs/Horde/Plugins/Tools/HordeServer.Tools/ToolsAclModifier.cs
new file mode 100644
--- /dev/null
+++ b/Engine/Source/Programs/Horde/Plugins/Tools/HordeServer.Tools/ToolsAclModifier.cs
@@ -0,0 +1,17 @@
+// Copyright Epic Games, Inc. All Rights Reserved.
+
+using HordeServer.Acls;
+using HordeServer.Tools;
+using HordeServer.Utilities;
+
+namespace HordeServer
+{
+	class ToolsAclModifier : IDefaultAclModifier
+	{
+		/// <inheritdoc/>
+		public void Apply(DefaultAclBuilder acl)
+		{
+			acl.AddCustomRole(HordeClaims.UploadToolsClaim, new[] { ToolAclAction.UploadTool });
+		}
+	}
+}
diff --git a/Engine/Source/Programs/Horde/Plugins/Tools/HordeServer.Tools/ToolsPlugin.cs b/Engine/Source/Programs/Horde/Plugins/Tools/HordeServer.Tools/ToolsPlugin.cs
--- a/Engine/Source/Programs/Horde/Plugins/Tools/HordeServer.Tools/ToolsPlugin.cs
+++ b/Engine/Source/Programs/Horde/Plugins/Tools/HordeServer.Tools/ToolsPlugin.cs
@@ -2,6 +2,7 @@
 
 using System.Reflection;
 using EpicGames.Core;
+using HordeServer.Acls;
 using HordeServer.Plugins;
 using HordeServer.Tools;
 using Microsoft.Extensions.DependencyInjection;
@@ -17,6 +18,8 @@ namespace HordeServer
 		/// <inheritdoc/>
 		public void ConfigureServices(IServiceCollection services)
 		{
+			services.AddSingleton<IDefaultAclModifier, ToolsAclModifier>();
+
 			services.AddSingleton<IToolCollection, ToolCollection>();
 			services.AddCommandsFromAssembly(Assembly.GetExecutingAssembly());
 		}
