Horde: Add a callback for plugins to configure the application endpoints after services have been created.

#rnx

diff --git a/Engine/Source/Programs/Horde/HordeServer.Shared/Plugins/IPluginStartup.cs b/Engine/Source/Programs/Horde/HordeServer.Shared/Plugins/IPluginStartup.cs
--- a/Engine/Source/Programs/Horde/HordeServer.Shared/Plugins/IPluginStartup.cs
+++ b/Engine/Source/Programs/Horde/HordeServer.Shared/Plugins/IPluginStartup.cs
@@ -1,6 +1,8 @@
 // Copyright Epic Games, Inc. All Rights Reserved.
 
 using System;
+using Microsoft.AspNetCore.Builder;
+using Microsoft.AspNetCore.Hosting;
 using Microsoft.CodeAnalysis.CSharp.Syntax;
 using Microsoft.Extensions.Configuration;
 using Microsoft.Extensions.DependencyInjection;
@@ -12,6 +14,12 @@ namespace HordeServer.Plugins
 	/// </summary>
 	public interface IPluginStartup
 	{
+		/// <summary>
+		/// Configure the application
+		/// </summary>
+		/// <param name="app">Application builder instance</param>
+		void Configure(IApplicationBuilder app);
+
 		/// <summary>
 		/// Configure the services provided by the plugin
 		/// </summary>
diff --git a/Engine/Source/Programs/Horde/HordeServer/Plugins/PluginCollection.cs b/Engine/Source/Programs/Horde/HordeServer/Plugins/PluginCollection.cs
--- a/Engine/Source/Programs/Horde/HordeServer/Plugins/PluginCollection.cs
+++ b/Engine/Source/Programs/Horde/HordeServer/Plugins/PluginCollection.cs
@@ -54,6 +54,9 @@ namespace HordeServer.Plugins
 				serviceCollection.AddPluginConfig<TGlobalConfig>(Name);
 
 				TStartup startup = CreateStartup(config, serverInfo);
+				serviceCollection.AddSingleton<TStartup>(startup);
+				serviceCollection.AddSingleton<IPluginStartup>(startup);
+
 				startup.ConfigureServices(serviceCollection);
 			}
 
diff --git a/Engine/Source/Programs/Horde/HordeServer/Startup.cs b/Engine/Source/Programs/Horde/HordeServer/Startup.cs
--- a/Engine/Source/Programs/Horde/HordeServer/Startup.cs
+++ b/Engine/Source/Programs/Horde/HordeServer/Startup.cs
@@ -29,17 +29,10 @@ using Grpc.Core.Interceptors;
 using HordeCommon;
 using HordeServer.Accounts;
 using HordeServer.Acls;
-using HordeServer.Agents.Enrollment;
-using HordeServer.Agents.Relay;
 using HordeServer.Auditing;
 using HordeServer.Authentication;
 using HordeServer.Configuration;
 using HordeServer.Dashboard;
-using HordeServer.Devices;
-using HordeServer.Jobs;
-using HordeServer.Logs;
-using HordeServer.Notifications;
-using HordeServer.Perforce;
 using HordeServer.Plugins;
 using HordeServer.Server;
 using HordeServer.Server.Notices;
@@ -346,14 +339,11 @@ namespace HordeServer
 			services.AddSingleton<IAccountCollection, AccountCollection>();
 			services.AddSingleton<IServiceAccountCollection, ServiceAccountCollection>();
 			services.AddSingleton<IUserCollection, UserCollectionV2>();
-			services.AddSingleton<INotificationTriggerCollection, NotificationTriggerCollection>();
-			services.AddSingleton<IDeviceCollection, DeviceCollection>();
 			services.AddSingleton<INoticeCollection, NoticeCollection>();
 			services.AddSingleton<IDashboardPreviewCollection, DashboardPreviewCollection>();
 
 			services.AddSingleton<IConfigSource, InMemoryConfigSource>();
 			services.AddSingleton<IConfigSource, FileConfigSource>();
-			services.AddSingleton<IConfigSource, PerforceConfigSource>();
 
 			services.AddSingleton<ConfigService>();
 			services.AddSingleton<IConfigService>(sp => sp.GetRequiredService<ConfigService>());
@@ -957,11 +947,6 @@ namespace HordeServer
 			app.UseEndpoints(endpoints =>
 			{
 				endpoints.MapGrpcService<HealthService>();
-				endpoints.MapGrpcService<RpcService>();
-				endpoints.MapGrpcService<JobRpcService>();
-				endpoints.MapGrpcService<LogRpcService>();
-				endpoints.MapGrpcService<AgentRelayService>();
-				endpoints.MapGrpcService<EnrollmentRpc>();
 
 				endpoints.MapGrpcReflectionService();
 
@@ -975,6 +960,11 @@ namespace HordeServer
 				app.MapWhen(IsSpaRequest, builder => builder.UseSpa(spa => spa.Options.SourcePath = "DashboardApp"));
 			}
 
+			foreach (IPluginStartup pluginStartup in app.ApplicationServices.GetRequiredService<IEnumerable<IPluginStartup>>())
+			{
+				pluginStartup.Configure(app);
+			}
+
 			if (settings.Value.OpenBrowser && RuntimeInformation.IsOSPlatform(OSPlatform.Windows))
 			{
 				lifetime.ApplicationStarted.Register(() => LaunchBrowser(app));
diff --git a/Engine/Source/Programs/Horde/Plugins/Analytics/HordeServer.Analytics/AnalyticsPlugin.cs b/Engine/Source/Programs/Horde/Plugins/Analytics/HordeServer.Analytics/AnalyticsPlugin.cs
--- a/Engine/Source/Programs/Horde/Plugins/Analytics/HordeServer.Analytics/AnalyticsPlugin.cs
+++ b/Engine/Source/Programs/Horde/Plugins/Analytics/HordeServer.Analytics/AnalyticsPlugin.cs
@@ -9,6 +9,7 @@ using HordeServer.Plugins;
 using HordeServer.Telemetry;
 using HordeServer.Telemetry.Metrics;
 using HordeServer.Telemetry.Sinks;
+using Microsoft.AspNetCore.Builder;
 using Microsoft.Extensions.Configuration;
 using Microsoft.Extensions.DependencyInjection;
 
@@ -20,6 +21,10 @@ namespace HordeServer
 	[Plugin("Analytics", GlobalConfigType = typeof(AnalyticsConfig), ServerConfigType = typeof(AnalyticsServerConfig))]
 	public class AnalyticsPlugin : IPluginStartup
 	{
+		/// <inheritdoc/>
+		public void Configure(IApplicationBuilder app)
+		{ }
+
 		/// <inheritdoc/>
 		public void ConfigureServices(IServiceCollection services)
 		{
diff --git a/Engine/Source/Programs/Horde/Plugins/Build/HordeServer.Build/BuildPlugin.cs b/Engine/Source/Programs/Horde/Plugins/Build/HordeServer.Build/BuildPlugin.cs
--- a/Engine/Source/Programs/Horde/Plugins/Build/HordeServer.Build/BuildPlugin.cs
+++ b/Engine/Source/Programs/Horde/Plugins/Build/HordeServer.Build/BuildPlugin.cs
@@ -14,6 +14,7 @@ using HordeServer.Agents.Pools;
 using HordeServer.Agents.Utilization;
 using HordeServer.Artifacts;
 using HordeServer.Commits;
+using HordeServer.Configuration;
 using HordeServer.Devices;
 using HordeServer.Issues;
 using HordeServer.Issues.External;
@@ -34,6 +35,7 @@ using HordeServer.Streams;
 using HordeServer.Tasks;
 using HordeServer.Ugs;
 using HordeServer.Users;
+using Microsoft.AspNetCore.Builder;
 using Microsoft.Extensions.DependencyInjection;
 
 namespace HordeServer
@@ -53,6 +55,15 @@ namespace HordeServer
 			_staticConfig = staticConfig;
 		}
 
+		/// <inheritdoc/>
+		public void Configure(IApplicationBuilder app)
+		{
+			app.UseEndpoints(endpoints =>
+			{
+				endpoints.MapGrpcService<JobRpcService>();
+			});
+		}
+
 		/// <inheritdoc/>
 		public void ConfigureServices(IServiceCollection services)
 		{
@@ -92,6 +103,11 @@ namespace HordeServer
 
 			services.AddSingleton<ScheduleService>();
 
+			services.AddSingleton<INotificationTriggerCollection, NotificationTriggerCollection>();
+			services.AddSingleton<IDeviceCollection, DeviceCollection>();
+
+			services.AddSingleton<IConfigSource, PerforceConfigSource>();
+
 			// Notifications can be triggered from any instance, so always make sure we're ticking the background task.
 			services.AddHostedService(provider => (NotificationService)provider.GetRequiredService<INotificationService>());
 
diff --git a/Engine/Source/Programs/Horde/Plugins/Compute/HordeServer.Compute/ComputePlugin.cs b/Engine/Source/Programs/Horde/Plugins/Compute/HordeServer.Compute/ComputePlugin.cs
--- a/Engine/Source/Programs/Horde/Plugins/Compute/HordeServer.Compute/ComputePlugin.cs
+++ b/Engine/Source/Programs/Horde/Plugins/Compute/HordeServer.Compute/ComputePlugin.cs
@@ -21,6 +21,7 @@ using HordeServer.Logs;
 using HordeServer.Plugins;
 using HordeServer.Server;
 using HordeServer.Tasks;
+using Microsoft.AspNetCore.Builder;
 using Microsoft.Extensions.Configuration;
 using Microsoft.Extensions.DependencyInjection;
 
@@ -44,6 +45,18 @@ namespace HordeServer
 			_staticComputeConfig = staticComputeConfig;
 		}
 
+		/// <inheritdoc/>
+		public void Configure(IApplicationBuilder app)
+		{
+			app.UseEndpoints(endpoints =>
+			{
+				endpoints.MapGrpcService<RpcService>();
+				endpoints.MapGrpcService<LogRpcService>();
+				endpoints.MapGrpcService<AgentRelayService>();
+				endpoints.MapGrpcService<EnrollmentRpc>();
+			});
+		}
+
 		/// <inheritdoc/>
 		public void ConfigureServices(IServiceCollection services)
 		{
diff --git a/Engine/Source/Programs/Horde/Plugins/Ddc/HordeServer.Ddc/DdcPlugin.cs b/Engine/Source/Programs/Horde/Plugins/Ddc/HordeServer.Ddc/DdcPlugin.cs
--- a/Engine/Source/Programs/Horde/Plugins/Ddc/HordeServer.Ddc/DdcPlugin.cs
+++ b/Engine/Source/Programs/Horde/Plugins/Ddc/HordeServer.Ddc/DdcPlugin.cs
@@ -2,6 +2,7 @@
 
 using HordeServer.Ddc;
 using HordeServer.Plugins;
+using Microsoft.AspNetCore.Builder;
 using Microsoft.Extensions.DependencyInjection;
 
 namespace HordeServer
@@ -12,6 +13,10 @@ namespace HordeServer
 	[Plugin("Ddc")]
 	public class DdcPlugin : IPluginStartup
 	{
+		/// <inheritdoc/>
+		public void Configure(IApplicationBuilder app)
+		{ }
+
 		/// <inheritdoc/>
 		public void ConfigureServices(IServiceCollection services)
 		{
diff --git a/Engine/Source/Programs/Horde/Plugins/Secrets/HordeServer.Secrets/SecretsPlugin.cs b/Engine/Source/Programs/Horde/Plugins/Secrets/HordeServer.Secrets/SecretsPlugin.cs
--- a/Engine/Source/Programs/Horde/Plugins/Secrets/HordeServer.Secrets/SecretsPlugin.cs
+++ b/Engine/Source/Programs/Horde/Plugins/Secrets/HordeServer.Secrets/SecretsPlugin.cs
@@ -3,6 +3,7 @@
 using HordeServer.Plugins;
 using HordeServer.Secrets;
 using HordeServer.Secrets.Providers;
+using Microsoft.AspNetCore.Builder;
 using Microsoft.Extensions.DependencyInjection;
 
 namespace HordeServer
@@ -13,6 +14,10 @@ namespace HordeServer
 	[Plugin("Secrets", GlobalConfigType = typeof(SecretsConfig))]
 	public class SecretsPlugin : IPluginStartup
 	{
+		/// <inheritdoc/>
+		public void Configure(IApplicationBuilder app)
+		{ }
+
 		/// <inheritdoc/>
 		public void ConfigureServices(IServiceCollection services)
 		{
diff --git a/Engine/Source/Programs/Horde/Plugins/Storage/HordeServer.Storage/StoragePlugin.cs b/Engine/Source/Programs/Horde/Plugins/Storage/HordeServer.Storage/StoragePlugin.cs
--- a/Engine/Source/Programs/Horde/Plugins/Storage/HordeServer.Storage/StoragePlugin.cs
+++ b/Engine/Source/Programs/Horde/Plugins/Storage/HordeServer.Storage/StoragePlugin.cs
@@ -11,6 +11,7 @@ using EpicGames.Horde.Storage.Bundles;
 using EpicGames.Horde.Storage.ObjectStores;
 using HordeServer.Plugins;
 using HordeServer.Storage.ObjectStores;
+using Microsoft.AspNetCore.Builder;
 using Microsoft.Extensions.DependencyInjection;
 using Microsoft.Extensions.Logging;
 using Microsoft.Extensions.Options;
@@ -32,6 +33,10 @@ namespace HordeServer.Storage
 			_staticConfig = staticConfig;
 		}
 
+		/// <inheritdoc/>
+		public void Configure(IApplicationBuilder app)
+		{ }
+
 		/// <inheritdoc/>
 		public void ConfigureServices(IServiceCollection services)
 		{
diff --git a/Engine/Source/Programs/Horde/Plugins/Tools/HordeServer.Tools/ToolsPlugin.cs b/Engine/Source/Programs/Horde/Plugins/Tools/HordeServer.Tools/ToolsPlugin.cs
--- a/Engine/Source/Programs/Horde/Plugins/Tools/HordeServer.Tools/ToolsPlugin.cs
+++ b/Engine/Source/Programs/Horde/Plugins/Tools/HordeServer.Tools/ToolsPlugin.cs
@@ -5,6 +5,7 @@ using EpicGames.Core;
 using HordeServer.Acls;
 using HordeServer.Plugins;
 using HordeServer.Tools;
+using Microsoft.AspNetCore.Builder;
 using Microsoft.Extensions.DependencyInjection;
 
 namespace HordeServer
@@ -15,6 +16,10 @@ namespace HordeServer
 	[Plugin("Tools", GlobalConfigType = typeof(ToolsConfig), ServerConfigType = typeof(ToolsServerConfig))]
 	public class ToolsPlugin : IPluginStartup
 	{
+		/// <inheritdoc/>
+		public void Configure(IApplicationBuilder app)
+		{ }
+
 		/// <inheritdoc/>
 		public void ConfigureServices(IServiceCollection services)
 		{
