Horde: Remove hard dependencies between CppApi command and specific plugins.

#rnx

diff --git a/Engine/Source/Programs/Horde/HordeServer.Shared/Utilities/CppApiAttribute.cs b/Engine/Source/Programs/Horde/HordeServer.Shared/Utilities/CppApiAttribute.cs
new file mode 100644
--- /dev/null
+++ b/Engine/Source/Programs/Horde/HordeServer.Shared/Utilities/CppApiAttribute.cs
@@ -0,0 +1,12 @@
+﻿// Copyright Epic Games, Inc. All Rights Reserved.
+
+namespace HordeServer.Utilities
+{
+	/// <summary>
+	/// Indicates that requests and responses should be included in the exported C++ API
+	/// </summary>
+	[AttributeUsage(AttributeTargets.Class)]
+	public sealed class CppApiAttribute : Attribute
+	{
+	}
+}
diff --git a/Engine/Source/Programs/Horde/HordeServer/Commands/Generate/CppApiCommand.cs b/Engine/Source/Programs/Horde/HordeServer/Commands/Generate/CppApiCommand.cs
--- a/Engine/Source/Programs/Horde/HordeServer/Commands/Generate/CppApiCommand.cs
+++ b/Engine/Source/Programs/Horde/HordeServer/Commands/Generate/CppApiCommand.cs
@@ -4,14 +4,15 @@ using System;
 using System.Collections.Generic;
 using System.ComponentModel;
 using System.IO;
+using System.Linq;
 using System.Reflection;
 using System.Text.Json.Serialization;
 using System.Threading;
 using System.Threading.Tasks;
 using EpicGames.Core;
 using EpicGames.Horde;
-using HordeServer.Projects;
-using HordeServer.Server;
+using HordeServer.Plugins;
+using HordeServer.Utilities;
 using Microsoft.AspNetCore.Mvc;
 using Microsoft.Extensions.Logging;
 
@@ -24,12 +25,6 @@ namespace HordeServer.Commands.Generate
 		record class ApiTypeInfo(string Name, Type Type, string Description, List<ApiPropertyInfo> Properties);
 		record class ApiGroupInfo(string Name, Dictionary<Type, ApiTypeInfo> Types);
 
-		static readonly Type[] s_controllerTypes = new Type[]
-		{
-			typeof(ServerController),
-			typeof(ProjectsController)
-		};
-
 		static readonly string[] s_boilerplateLines = new[]
 		{
 			"// ---------------------------------------------------------------------------------------------------",
@@ -42,6 +37,13 @@ namespace HordeServer.Commands.Generate
 		[Description("Specifies the output directory for generated files.")]
 		public DirectoryReference? OutputDir { get; set; }
 
+		readonly IPluginCollection _pluginCollection;
+
+		public CppApiCommand(IPluginCollection pluginCollection)
+		{
+			_pluginCollection = pluginCollection;
+		}
+
 		public override Task<int> ExecuteAsync(ILogger logger)
 		{
 			// Try to find the output directory assuming we're being run from the default output directory in the engine source tree
@@ -58,11 +60,22 @@ namespace HordeServer.Commands.Generate
 				OutputDir = defaultHordeDir;
 			}
 
+			// Find all the assemblies containing potential controller types
+			HashSet<Assembly> inputAssemblies = new HashSet<Assembly>();
+			inputAssemblies.Add(Assembly.GetExecutingAssembly());
+			inputAssemblies.UnionWith(_pluginCollection.LoadedPlugins.Select(x => x.Assembly));
+
 			// Find all the referenced types from each controller
 			HashSet<Type> inputTypes = new HashSet<Type>();
-			foreach (Type controllerType in s_controllerTypes)
+			foreach (Assembly inputAssembly in inputAssemblies)
 			{
-				FindTypesFromController(controllerType, inputTypes);
+				foreach (Type exportedType in inputAssembly.GetExportedTypes())
+				{
+					if (exportedType.GetCustomAttribute<CppApiAttribute>() != null)
+					{
+						FindTypesFromController(exportedType, inputTypes);
+					}
+				}
 			}
 
 			// Create all the types 
diff --git a/Engine/Source/Programs/Horde/HordeServer/Server/ServerController.cs b/Engine/Source/Programs/Horde/HordeServer/Server/ServerController.cs
--- a/Engine/Source/Programs/Horde/HordeServer/Server/ServerController.cs
+++ b/Engine/Source/Programs/Horde/HordeServer/Server/ServerController.cs
@@ -20,6 +20,7 @@ namespace HordeServer.Server
 	/// </summary>
 	[ApiController]
 	[Authorize]
+	[CppApi]
 	[Route("[controller]")]
 	public class ServerController : HordeControllerBase
 	{
diff --git a/Engine/Source/Programs/Horde/Plugins/Build/HordeServer.Build/Projects/ProjectsController.cs b/Engine/Source/Programs/Horde/Plugins/Build/HordeServer.Build/Projects/ProjectsController.cs
--- a/Engine/Source/Programs/Horde/Plugins/Build/HordeServer.Build/Projects/ProjectsController.cs
+++ b/Engine/Source/Programs/Horde/Plugins/Build/HordeServer.Build/Projects/ProjectsController.cs
@@ -20,6 +20,7 @@ namespace HordeServer.Projects
 	/// </summary>
 	[ApiController]
 	[Authorize]
+	[CppApi]
 	[Route("[controller]")]
 	public class ProjectsController : HordeControllerBase
 	{
