Horde: Add a IPluginResponseFilter interface allowing plugins to filter the controller responses from other classes.

#rnx

diff --git a/Engine/Source/Programs/Horde/HordeServer.Shared/Plugins/IPluginResponseFilter.cs b/Engine/Source/Programs/Horde/HordeServer.Shared/Plugins/IPluginResponseFilter.cs
new file mode 100644
--- /dev/null
+++ b/Engine/Source/Programs/Horde/HordeServer.Shared/Plugins/IPluginResponseFilter.cs
@@ -0,0 +1,17 @@
+// Copyright Epic Games, Inc. All Rights Reserved.
+
+using Microsoft.AspNetCore.Http;
+
+namespace HordeServer.Plugins
+{
+	/// <summary>
+	/// Interface allowing plugins to modify API responses
+	/// </summary>
+	public interface IPluginResponseFilter
+	{
+		/// <summary>
+		/// Apply any changes to the given response
+		/// </summary>
+		public void Apply(HttpContext context, object response);
+	}
+}
diff --git a/Engine/Source/Programs/Horde/HordeServer/Dashboard/DashboardController.cs b/Engine/Source/Programs/Horde/HordeServer/Dashboard/DashboardController.cs
--- a/Engine/Source/Programs/Horde/HordeServer/Dashboard/DashboardController.cs
+++ b/Engine/Source/Programs/Horde/HordeServer/Dashboard/DashboardController.cs
@@ -10,6 +10,7 @@ using EpicGames.Horde.Accounts;
 using EpicGames.Horde.Dashboard;
 using EpicGames.Horde.Server;
 using HordeServer.Accounts;
+using HordeServer.Plugins;
 using HordeServer.Server;
 using HordeServer.Telemetry;
 using HordeServer.Utilities;
@@ -28,33 +29,23 @@ namespace HordeServer.Dashboard
 	[Route("[controller]")]
 	public class DashboardController : Controller
 	{
-		/// <summary>
-		/// Authentication scheme in use
-		/// </summary>
 		readonly string _authenticationScheme;
-
-		/// <summary>
-		/// Server settings
-		/// </summary>
 		readonly ServerSettings _settings;
-		readonly StaticBuildConfig _staticBuildConfig;
-
 		readonly IDashboardPreviewCollection _previewCollection;
-
 		readonly IAccountCollection _hordeAccounts;
-
+		readonly IEnumerable<IPluginResponseFilter> _responseFilters;
 		readonly IOptionsSnapshot<GlobalConfig> _globalConfig;
 
 		/// <summary>
 		/// Constructor
 		/// </summary>
-		public DashboardController(IDashboardPreviewCollection previewCollection, IAccountCollection hordeAccounts, IOptionsMonitor<ServerSettings> serverSettings, IOptions<StaticBuildConfig> staticBuildConfig, IOptionsSnapshot<GlobalConfig> globalConfig)
+		public DashboardController(IDashboardPreviewCollection previewCollection, IAccountCollection hordeAccounts, IEnumerable<IPluginResponseFilter> responseFilters, IOptionsMonitor<ServerSettings> serverSettings, IOptionsSnapshot<GlobalConfig> globalConfig)
 		{
 			_authenticationScheme = AccountController.GetAuthScheme(serverSettings.CurrentValue.AuthMethod);
 			_previewCollection = previewCollection;
 			_hordeAccounts = hordeAccounts;
+			_responseFilters = responseFilters;
 			_settings = serverSettings.CurrentValue;
-			_staticBuildConfig = staticBuildConfig.Value;
 			_globalConfig = globalConfig;
 		}
 
@@ -155,22 +146,9 @@ namespace HordeServer.Dashboard
 
 			dashboardConfigResponse.AuthMethod = _settings.AuthMethod;
 
-			if (_staticBuildConfig.JiraUrl != null)
-			{
-				dashboardConfigResponse.ExternalIssueServiceName = "Jira";
-				dashboardConfigResponse.ExternalIssueServiceUrl = _staticBuildConfig.JiraUrl.ToString().TrimEnd('/');
-			}
-
-			if (_staticBuildConfig.P4SwarmUrl != null)
-			{
-				dashboardConfigResponse.PerforceSwarmUrl = _staticBuildConfig.P4SwarmUrl.ToString().TrimEnd('/');
-			}
-
 			dashboardConfigResponse.HelpEmailAddress = _settings.HelpEmailAddress;
 			dashboardConfigResponse.HelpSlackChannel = _settings.HelpSlackChannel;
 
-			dashboardConfigResponse.DeviceProblemCooldownMinutes = _staticBuildConfig.DeviceProblemCooldownMinutes;
-
 			foreach (DashboardAgentCategoryConfig category in _globalConfig.Value.Dashboard.AgentCategories)
 			{
 				dashboardConfigResponse.AgentCategories.Add(new GetDashboardAgentCategoryResponse { Name = category.Name, Condition = category.Condition });
@@ -218,6 +196,11 @@ namespace HordeServer.Dashboard
 				}
 			}
 
+			foreach (IPluginResponseFilter responseFilter in _responseFilters)
+			{
+				responseFilter.Apply(HttpContext, dashboardConfigResponse);
+			}
+
 			return dashboardConfigResponse;
 		}
 
diff --git a/Engine/Source/Programs/Horde/HordeServer/Users/UserController.cs b/Engine/Source/Programs/Horde/HordeServer/Users/UserController.cs
--- a/Engine/Source/Programs/Horde/HordeServer/Users/UserController.cs
+++ b/Engine/Source/Programs/Horde/HordeServer/Users/UserController.cs
@@ -1,5 +1,7 @@
 // Copyright Epic Games, Inc. All Rights Reserved.
 
+using System.Collections;
+using System.Collections.Generic;
 using System.Linq;
 using System.Security.Claims;
 using System.Threading;
@@ -8,8 +10,7 @@ using EpicGames.Horde.Jobs;
 using EpicGames.Horde.Jobs.Bisect;
 using EpicGames.Horde.Users;
 using HordeServer.Accounts;
-using HordeServer.Agents;
-using HordeServer.Agents.Pools;
+using HordeServer.Plugins;
 using HordeServer.Server;
 using HordeServer.Server.Notices;
 using HordeServer.Utilities;
@@ -29,16 +30,18 @@ namespace HordeServer.Users
 	{
 		readonly IUserCollection _userCollection;
 		readonly IAvatarService? _avatarService;
+		readonly IEnumerable<IPluginResponseFilter> _responseFilters;
 		readonly IOptionsSnapshot<GlobalConfig> _globalConfig;
 		readonly IOptionsMonitor<ServerSettings> _settings;
 
 		/// <summary>
 		/// Constructor
 		/// </summary>
-		public UserController(IUserCollection userCollection, IAvatarService? avatarService, IOptionsSnapshot<GlobalConfig> globalConfig, IOptionsMonitor<ServerSettings> settings)
+		public UserController(IUserCollection userCollection, IAvatarService? avatarService, IEnumerable<IPluginResponseFilter> responseFilters, IOptionsSnapshot<GlobalConfig> globalConfig, IOptionsMonitor<ServerSettings> settings)
 		{
 			_userCollection = userCollection;
 			_avatarService = avatarService;
+			_responseFilters = responseFilters;
 			_globalConfig = globalConfig;
 			_settings = settings;
 		}
@@ -68,7 +71,7 @@ namespace HordeServer.Users
 			return PropertyFilter.Apply(response, filter);
 		}
 
-		static GetDashboardFeaturesResponse GetDashboardFeatures(GlobalConfig globalConfig, ServerSettings settings, ClaimsPrincipal principal)
+		GetDashboardFeaturesResponse GetDashboardFeatures(GlobalConfig globalConfig, ServerSettings settings, ClaimsPrincipal principal)
 		{
 			GetDashboardFeaturesResponse response = new GetDashboardFeaturesResponse();
 			response.ShowLandingPage = globalConfig.Dashboard.ShowLandingPage;
@@ -79,9 +82,13 @@ namespace HordeServer.Users
 			response.ShowDeviceManager = globalConfig.Dashboard.ShowDeviceManager;
 			response.ShowTests = globalConfig.Dashboard.ShowTests;
 			response.ShowNoticeEditor = globalConfig.Authorize(NoticeAclAction.CreateNotice, principal) || globalConfig.Authorize(NoticeAclAction.UpdateNotice, principal);
-			response.ShowPoolEditor = globalConfig.VersionEnum < ConfigVersion.PoolsInConfigFiles && (globalConfig.Authorize(PoolAclAction.CreatePool, principal) || globalConfig.Authorize(PoolAclAction.UpdatePool, principal));
-			response.ShowRemoteDesktop = globalConfig.Authorize(AgentAclAction.UpdateAgent, principal);
 			response.ShowAccounts = settings.AuthMethod == EpicGames.Horde.Server.AuthMethod.Horde && globalConfig.Authorize(AccountAclAction.UpdateAccount, principal);
+
+			foreach (IPluginResponseFilter responseFilter in _responseFilters)
+			{
+				responseFilter.Apply(HttpContext, response);
+			}
+
 			return response;
 		}
 
diff --git a/Engine/Source/Programs/Horde/Plugins/Build/HordeServer.Build/BuildPlugin.cs b/Engine/Source/Programs/Horde/Plugins/Build/HordeServer.Build/BuildPlugin.cs
--- a/Engine/Source/Programs/Horde/Plugins/Build/HordeServer.Build/BuildPlugin.cs
+++ b/Engine/Source/Programs/Horde/Plugins/Build/HordeServer.Build/BuildPlugin.cs
@@ -1,11 +1,6 @@
 // Copyright Epic Games, Inc. All Rights Reserved.
 
-using System;
-using System.Collections.Generic;
 using System.Diagnostics.CodeAnalysis;
-using System.Linq;
-using System.Text;
-using System.Threading.Tasks;
 using EpicGames.Core;
 using HordeServer.Acls;
 using HordeServer.Agents;
@@ -68,6 +63,7 @@ namespace HordeServer
 		public void ConfigureServices(IServiceCollection services)
 		{
 			services.AddSingleton<IDefaultAclModifier, BuildAclModifier>();
+			services.AddSingleton<IPluginResponseFilter, BuildResponseFilter>();
 
 			services.AddSingleton<IArtifactCollection, ArtifactCollection>();
 			services.AddSingleton<IGraphCollection, GraphCollection>();
diff --git a/Engine/Source/Programs/Horde/Plugins/Build/HordeServer.Build/BuildResponseFilter.cs b/Engine/Source/Programs/Horde/Plugins/Build/HordeServer.Build/BuildResponseFilter.cs
new file mode 100644
--- /dev/null
+++ b/Engine/Source/Programs/Horde/Plugins/Build/HordeServer.Build/BuildResponseFilter.cs
@@ -0,0 +1,39 @@
+// Copyright Epic Games, Inc. All Rights Reserved.
+
+using EpicGames.Horde.Dashboard;
+using HordeServer.Plugins;
+using Microsoft.AspNetCore.Http;
+using Microsoft.Extensions.Options;
+
+namespace HordeServer
+{
+	class BuildResponseFilter : IPluginResponseFilter
+	{
+		readonly StaticBuildConfig _staticBuildConfig;
+
+		public BuildResponseFilter(IOptions<StaticBuildConfig> staticBuildConfig)
+		{
+			_staticBuildConfig = staticBuildConfig.Value;
+		}
+
+		/// <inheritdoc/>
+		public void Apply(HttpContext context, object response)
+		{
+			if (response is GetDashboardConfigResponse dashboardConfigResponse)
+			{
+				if (_staticBuildConfig.JiraUrl != null)
+				{
+					dashboardConfigResponse.ExternalIssueServiceName = "Jira";
+					dashboardConfigResponse.ExternalIssueServiceUrl = _staticBuildConfig.JiraUrl.ToString().TrimEnd('/');
+				}
+
+				if (_staticBuildConfig.P4SwarmUrl != null)
+				{
+					dashboardConfigResponse.PerforceSwarmUrl = _staticBuildConfig.P4SwarmUrl.ToString().TrimEnd('/');
+				}
+
+				dashboardConfigResponse.DeviceProblemCooldownMinutes = _staticBuildConfig.DeviceProblemCooldownMinutes;
+			}
+		}
+	}
+}
diff --git a/Engine/Source/Programs/Horde/Plugins/Compute/HordeServer.Compute/ComputePlugin.cs b/Engine/Source/Programs/Horde/Plugins/Compute/HordeServer.Compute/ComputePlugin.cs
--- a/Engine/Source/Programs/Horde/Plugins/Compute/HordeServer.Compute/ComputePlugin.cs
+++ b/Engine/Source/Programs/Horde/Plugins/Compute/HordeServer.Compute/ComputePlugin.cs
@@ -61,6 +61,7 @@ namespace HordeServer
 		public void ConfigureServices(IServiceCollection services)
 		{
 			services.AddSingleton<IDefaultAclModifier, ComputeAclModifier>();
+			services.AddSingleton<IPluginResponseFilter, ComputeResponseFilter>();
 
 			services.AddSingleton<IAgentCollection, AgentCollection>();
 			services.AddSingleton<ILeaseCollection, LeaseCollection>();
diff --git a/Engine/Source/Programs/Horde/Plugins/Compute/HordeServer.Compute/ComputeResponseFilter.cs b/Engine/Source/Programs/Horde/Plugins/Compute/HordeServer.Compute/ComputeResponseFilter.cs
new file mode 100644
--- /dev/null
+++ b/Engine/Source/Programs/Horde/Plugins/Compute/HordeServer.Compute/ComputeResponseFilter.cs
@@ -0,0 +1,31 @@
+// Copyright Epic Games, Inc. All Rights Reserved.
+
+using EpicGames.Horde.Users;
+using HordeServer.Agents;
+using HordeServer.Agents.Pools;
+using HordeServer.Plugins;
+using Microsoft.AspNetCore.Http;
+using Microsoft.Extensions.Options;
+
+namespace HordeServer
+{
+	class ComputeResponseFilter : IPluginResponseFilter
+	{
+		readonly IOptionsMonitor<ComputeConfig> _computeConfig;
+
+		public ComputeResponseFilter(IOptionsMonitor<ComputeConfig> computeConfig)
+		{
+			_computeConfig = computeConfig;
+		}
+
+		public void Apply(HttpContext context, object response)
+		{
+			if (response is GetDashboardFeaturesResponse featuresResponse)
+			{
+				ComputeConfig computeConfig = _computeConfig.CurrentValue;
+				featuresResponse.ShowPoolEditor = computeConfig.VersionEnum < ConfigVersion.PoolsInConfigFiles && (computeConfig.Authorize(PoolAclAction.CreatePool, context.User) || computeConfig.Authorize(PoolAclAction.UpdatePool, context.User));
+				featuresResponse.ShowRemoteDesktop = computeConfig.Authorize(AgentAclAction.UpdateAgent, context.User);
+			}
+		}
+	}
+}
