Horde: Move fixup for pool workspaces into build config.

diff --git a/Engine/Source/Programs/Horde/HordeServer.Shared/Plugins/IPluginConfig.cs b/Engine/Source/Programs/Horde/HordeServer.Shared/Plugins/IPluginConfig.cs
--- a/Engine/Source/Programs/Horde/HordeServer.Shared/Plugins/IPluginConfig.cs
+++ b/Engine/Source/Programs/Horde/HordeServer.Shared/Plugins/IPluginConfig.cs
@@ -26,7 +26,7 @@ namespace HordeServer.Plugins
 	/// <summary>
 	/// Options passed to <see cref="IPluginConfig.PostLoad(AclConfig)"/>
 	/// </summary>
-	public record class PluginConfigOptions(ConfigVersion Version, AclConfig ParentAcl);
+	public record class PluginConfigOptions(ConfigVersion Version, IEnumerable<IPluginConfig> Plugins, AclConfig ParentAcl);
 
 	/// <summary>
 	/// Empty implementation of <see cref="IPluginConfig"/>
diff --git a/Engine/Source/Programs/Horde/HordeServer/Server/GlobalConfig.cs b/Engine/Source/Programs/Horde/HordeServer/Server/GlobalConfig.cs
--- a/Engine/Source/Programs/Horde/HordeServer/Server/GlobalConfig.cs
+++ b/Engine/Source/Programs/Horde/HordeServer/Server/GlobalConfig.cs
@@ -8,15 +8,10 @@ using System.Security.Claims;
 using System.Text.Json.Serialization;
 using EpicGames.Core;
 using EpicGames.Horde.Acls;
-using EpicGames.Horde.Agents.Pools;
 using HordeServer.Acls;
-using HordeServer.Agents;
-using HordeServer.Agents.Pools;
 using HordeServer.Configuration;
 using HordeServer.Dashboard;
-using HordeServer.Perforce;
 using HordeServer.Plugins;
-using HordeServer.Streams;
 using HordeServer.Utilities;
 
 namespace HordeServer.Server
@@ -123,14 +118,12 @@ namespace HordeServer.Server
 				}
 			}
 
-			PluginConfigOptions pluginConfigOptions = new PluginConfigOptions(VersionEnum, Acl);
+			PluginConfigOptions pluginConfigOptions = new PluginConfigOptions(VersionEnum, Plugins.Values, Acl);
 			foreach (IPluginConfig pluginConfig in Plugins.Values)
 			{
 				pluginConfig.PostLoad(pluginConfigOptions);
 			}
 
-			UpdateWorkspacesForPools();
-
 			_aclLookup.Clear();
 			BuildAclScopeLookup(Acl, _aclLookup);
 		}
@@ -173,80 +166,6 @@ namespace HordeServer.Server
 			}
 		}
 
-		void UpdateWorkspacesForPools()
-		{
-			// Try to get the compute config, and skip if it isn't configured
-			ComputeConfig? computeConfig;
-			if (!Plugins.TryGetValue(new PluginName("compute"), out computeConfig))
-			{
-				return;
-			}
-
-			BuildConfig? buildConfig;
-			if (!Plugins.TryGetValue(new PluginName("build"), out buildConfig))
-			{
-				return;
-			}
-
-			// Lookup table of pool id to workspaces
-			Dictionary<PoolId, AutoSdkConfig> poolToAutoSdkView = new Dictionary<PoolId, AutoSdkConfig>();
-			Dictionary<PoolId, List<AgentWorkspaceInfo>> poolToAgentWorkspaces = new Dictionary<PoolId, List<AgentWorkspaceInfo>>();
-
-			// Populate the workspace list from the current stream
-			foreach (StreamConfig streamConfig in buildConfig.Streams)
-			{
-				foreach (KeyValuePair<string, AgentConfig> agentTypePair in streamConfig.AgentTypes)
-				{
-					// Create the new agent workspace
-					if (streamConfig.TryGetAgentWorkspace(agentTypePair.Value, out AgentWorkspaceInfo? agentWorkspace, out AutoSdkConfig? autoSdkConfig))
-					{
-						AgentConfig agentType = agentTypePair.Value;
-
-						// Find or add a list of workspaces for this pool
-						List<AgentWorkspaceInfo>? agentWorkspaces;
-						if (!poolToAgentWorkspaces.TryGetValue(agentType.Pool, out agentWorkspaces))
-						{
-							agentWorkspaces = new List<AgentWorkspaceInfo>();
-							poolToAgentWorkspaces.Add(agentType.Pool, agentWorkspaces);
-						}
-
-						// Add it to the list
-						if (!agentWorkspaces.Contains(agentWorkspace))
-						{
-							agentWorkspaces.Add(agentWorkspace);
-						}
-						if (autoSdkConfig != null)
-						{
-							AutoSdkConfig? existingAutoSdkConfig;
-							poolToAutoSdkView.TryGetValue(agentType.Pool, out existingAutoSdkConfig);
-							poolToAutoSdkView[agentType.Pool] = AutoSdkConfig.Merge(autoSdkConfig, existingAutoSdkConfig);
-						}
-					}
-				}
-			}
-
-			// Update the list of workspaces for each pool
-			foreach (PoolConfig pool in computeConfig.Pools)
-			{
-				// Get the new list of workspaces for this pool
-				List<AgentWorkspaceInfo>? newWorkspaces;
-				if (!poolToAgentWorkspaces.TryGetValue(pool.Id, out newWorkspaces))
-				{
-					newWorkspaces = new List<AgentWorkspaceInfo>();
-				}
-
-				// Get the autosdk view
-				AutoSdkConfig? newAutoSdkConfig;
-				if (!poolToAutoSdkView.TryGetValue(pool.Id, out newAutoSdkConfig))
-				{
-					newAutoSdkConfig = AutoSdkConfig.None;
-				}
-
-				pool.Workspaces = newWorkspaces;
-				pool.AutoSdkConfig = newAutoSdkConfig;
-			}
-		}
-
 		/// <summary>
 		/// Authorizes a user to perform a given action
 		/// </summary>
diff --git a/Engine/Source/Programs/Horde/Plugins/Build/HordeServer.Build.Tests/Compute/ComputeServiceTest.cs b/Engine/Source/Programs/Horde/Plugins/Build/HordeServer.Build.Tests/Compute/ComputeServiceTest.cs
--- a/Engine/Source/Programs/Horde/Plugins/Build/HordeServer.Build.Tests/Compute/ComputeServiceTest.cs
+++ b/Engine/Source/Programs/Horde/Plugins/Build/HordeServer.Build.Tests/Compute/ComputeServiceTest.cs
@@ -309,7 +309,7 @@ namespace HordeServer.Tests.Compute
 			]
 			};
 
-			PluginConfigOptions configOptions = new PluginConfigOptions(ConfigVersion.Latest, new HordeServer.Acls.AclConfig());
+			PluginConfigOptions configOptions = new PluginConfigOptions(ConfigVersion.Latest, Enumerable.Empty<IPluginConfig>(), new HordeServer.Acls.AclConfig());
 			computeConfig.PostLoad(configOptions);
 
 			Assert.AreEqual(new ClusterId("compute1"), ComputeService.FindBestComputeClusterId(computeConfig, IPAddress.Parse("11.0.0.1")));
@@ -391,7 +391,7 @@ namespace HordeServer.Tests.Compute
 				NullLogger<ComputeService>.Instance);
 			await CreateAgentAsync(new PoolId(pool), properties: ["ComputeIp=11.0.0.1", "ComputePort=5000"]);
 			computeConfig.CurrentValue.Clusters = [ccc];
-			computeConfig.CurrentValue.PostLoad(new PluginConfigOptions(ConfigVersion.Latest, new AclConfig()));
+			computeConfig.CurrentValue.PostLoad(new PluginConfigOptions(ConfigVersion.Latest, Enumerable.Empty<IPluginConfig>(), new AclConfig()));
 			return cs;
 		}
 
diff --git a/Engine/Source/Programs/Horde/Plugins/Build/HordeServer.Build/BuildConfig.cs b/Engine/Source/Programs/Horde/Plugins/Build/HordeServer.Build/BuildConfig.cs
--- a/Engine/Source/Programs/Horde/Plugins/Build/HordeServer.Build/BuildConfig.cs
+++ b/Engine/Source/Programs/Horde/Plugins/Build/HordeServer.Build/BuildConfig.cs
@@ -109,6 +109,72 @@ namespace HordeServer
 			{
 				_artifactTypeLookup.Add(artifactType.Type, artifactType);
 			}
+
+			UpdateWorkspacesForPools(configOptions.Plugins);
+		}
+
+		void UpdateWorkspacesForPools(IEnumerable<IPluginConfig> plugins)
+		{
+			// Try to get the compute config, and skip if it isn't configured
+			ComputeConfig computeConfig = plugins.OfType<ComputeConfig>().First();
+
+			// Lookup table of pool id to workspaces
+			Dictionary<PoolId, AutoSdkConfig> poolToAutoSdkView = new Dictionary<PoolId, AutoSdkConfig>();
+			Dictionary<PoolId, List<AgentWorkspaceInfo>> poolToAgentWorkspaces = new Dictionary<PoolId, List<AgentWorkspaceInfo>>();
+
+			// Populate the workspace list from the current stream
+			foreach (StreamConfig streamConfig in Streams)
+			{
+				foreach (KeyValuePair<string, AgentConfig> agentTypePair in streamConfig.AgentTypes)
+				{
+					// Create the new agent workspace
+					if (streamConfig.TryGetAgentWorkspace(agentTypePair.Value, out AgentWorkspaceInfo? agentWorkspace, out AutoSdkConfig? autoSdkConfig))
+					{
+						AgentConfig agentType = agentTypePair.Value;
+
+						// Find or add a list of workspaces for this pool
+						List<AgentWorkspaceInfo>? agentWorkspaces;
+						if (!poolToAgentWorkspaces.TryGetValue(agentType.Pool, out agentWorkspaces))
+						{
+							agentWorkspaces = new List<AgentWorkspaceInfo>();
+							poolToAgentWorkspaces.Add(agentType.Pool, agentWorkspaces);
+						}
+
+						// Add it to the list
+						if (!agentWorkspaces.Contains(agentWorkspace))
+						{
+							agentWorkspaces.Add(agentWorkspace);
+						}
+						if (autoSdkConfig != null)
+						{
+							AutoSdkConfig? existingAutoSdkConfig;
+							poolToAutoSdkView.TryGetValue(agentType.Pool, out existingAutoSdkConfig);
+							poolToAutoSdkView[agentType.Pool] = AutoSdkConfig.Merge(autoSdkConfig, existingAutoSdkConfig);
+						}
+					}
+				}
+			}
+
+			// Update the list of workspaces for each pool
+			foreach (PoolConfig pool in computeConfig.Pools)
+			{
+				// Get the new list of workspaces for this pool
+				List<AgentWorkspaceInfo>? newWorkspaces;
+				if (!poolToAgentWorkspaces.TryGetValue(pool.Id, out newWorkspaces))
+				{
+					newWorkspaces = new List<AgentWorkspaceInfo>();
+				}
+
+				// Get the autosdk view
+				AutoSdkConfig? newAutoSdkConfig;
+				if (!poolToAutoSdkView.TryGetValue(pool.Id, out newAutoSdkConfig))
+				{
+					newAutoSdkConfig = AutoSdkConfig.None;
+				}
+
+				pool.Workspaces = newWorkspaces;
+				pool.AutoSdkConfig = newAutoSdkConfig;
+			}
 		}
 
 		/// <summary>
