Horde: Static analysis fix.

#rnx

diff --git a/Engine/Source/Programs/Horde/HordeServer/Users/UserController.cs b/Engine/Source/Programs/Horde/HordeServer/Users/UserController.cs
--- a/Engine/Source/Programs/Horde/HordeServer/Users/UserController.cs
+++ b/Engine/Source/Programs/Horde/HordeServer/Users/UserController.cs
@@ -1,6 +1,5 @@
 // Copyright Epic Games, Inc. All Rights Reserved.
 
-using System.Collections;
 using System.Collections.Generic;
 using System.Linq;
 using System.Security.Claims;
