Horde: Fix test failure.

#rnx

diff --git a/Engine/Source/Programs/Horde/HordeServer.Tests/Configuration/ConfigTests.cs b/Engine/Source/Programs/Horde/HordeServer.Tests/Configuration/ConfigTests.cs
--- a/Engine/Source/Programs/Horde/HordeServer.Tests/Configuration/ConfigTests.cs
+++ b/Engine/Source/Programs/Horde/HordeServer.Tests/Configuration/ConfigTests.cs
@@ -182,6 +182,7 @@ namespace HordeServer.Tests.Configuration
 			buildConfig.Projects.Add(new ProjectConfig { Streams = [new StreamConfig { WorkspaceTypes = inputWorkspaces }]});
 
 			GlobalConfig gc = new();
+			gc.Plugins.AddComputeConfig(new ComputeConfig());
 			gc.Plugins.AddBuildConfig(buildConfig);
 			gc.PostLoad(new ServerSettings(), Array.Empty<ILoadedPlugin>(), Array.Empty<IDefaultAclModifier>());
 
@@ -216,6 +217,7 @@ namespace HordeServer.Tests.Configuration
 			});
 
 			GlobalConfig gc = new();
+			gc.Plugins.AddComputeConfig(new ComputeConfig());
 			gc.Plugins.AddBuildConfig(buildConfig);
 			gc.PostLoad(new ServerSettings(), Array.Empty<ILoadedPlugin>(), Array.Empty<IDefaultAclModifier>());
 
diff --git a/Engine/Source/Programs/Horde/Plugins/Build/HordeServer.Build/BuildPlugin.cs b/Engine/Source/Programs/Horde/Plugins/Build/HordeServer.Build/BuildPlugin.cs
--- a/Engine/Source/Programs/Horde/Plugins/Build/HordeServer.Build/BuildPlugin.cs
+++ b/Engine/Source/Programs/Horde/Plugins/Build/HordeServer.Build/BuildPlugin.cs
@@ -222,19 +222,19 @@ namespace HordeServer
 	public static class BuildPluginExtensions
 	{
 		/// <summary>
-		/// Configures the storage plugin
+		/// Configures the build plugin
 		/// </summary>
 		public static void AddBuildConfig(this IDictionary<PluginName, IPluginConfig> dictionary, BuildConfig buildConfig)
 			=> dictionary[new PluginName("Build")] = buildConfig;
 
 		/// <summary>
-		/// Configures the storage plugin
+		/// Gets configuration for the build plugin
 		/// </summary>
 		public static BuildConfig GetBuildConfig(this IDictionary<PluginName, IPluginConfig> dictionary)
 			=> (BuildConfig)dictionary[new PluginName("Build")];
 
 		/// <summary>
-		/// Configures the storage plugin
+		/// Gets configuration for the build plugin
 		/// </summary>
 		public static bool TryGetBuildConfig(this IDictionary<PluginName, IPluginConfig> dictionary, [NotNullWhen(true)] out BuildConfig? buildConfig)
 		{
diff --git a/Engine/Source/Programs/Horde/Plugins/Compute/HordeServer.Compute/ComputePlugin.cs b/Engine/Source/Programs/Horde/Plugins/Compute/HordeServer.Compute/ComputePlugin.cs
--- a/Engine/Source/Programs/Horde/Plugins/Compute/HordeServer.Compute/ComputePlugin.cs
+++ b/Engine/Source/Programs/Horde/Plugins/Compute/HordeServer.Compute/ComputePlugin.cs
@@ -162,6 +162,12 @@ namespace HordeServer
 	/// </summary>
 	public static class ComputePluginExtensions
 	{
+		/// <summary>
+		/// Configures the compute plugin
+		/// </summary>
+		public static void AddComputeConfig(this IDictionary<PluginName, IPluginConfig> dictionary, ComputeConfig computeConfig)
+			=> dictionary[new PluginName("Compute")] = computeConfig;
+
 		/// <summary>
 		/// Get the compute plugin config
 		/// </summary>
