Horde: Update release notes.

#rnx

diff --git a/Engine/Source/Programs/Horde/Docs/ReleaseNotes.md b/Engine/Source/Programs/Horde/Docs/ReleaseNotes.md
--- a/Engine/Source/Programs/Horde/Docs/ReleaseNotes.md
+++ b/Engine/Source/Programs/Horde/Docs/ReleaseNotes.md
@@ -1,5 +1,30 @@
 # Release Notes
 
+## 2024-07-15
+
+* Split server functionality into plugins. Plugins are still currently statically configured, which will be changed in future. (34620916, many others)
+* Fix threading issue causing block cache to always attempt to access element zero, causing tests to get stuck in an infinite loop. (34759616)
+* Add a debug endpoint for writing memory mapped file cache stats to the log. (34757974)
+* Fix resource leaks for blob data reads. (34756957, 34757889)
+* Fix tracking of allocated size in MemoryMappedFileCache. (34747643)
+* Fix binding of plugin server config instances. (34745408)
+* Move build functionality into a plugin. (34744727)
+* Respect Forwarded-For header for client's IP during compute cluster ID resolving. Also return a better error message when no network range or cluster is found. (34743435)
+* Add support for automatic computer cluster assignment in Horde, based on internal and external IP. (34734153)
+* Add endpoint for resolving a suitable compute cluster ID (34706120)
+* Set agent property if Wine executable is configured. Needed for scheduling Wine-compatible compute tasks. (34675604)
+* Do not log cancellation exception as a warning when checking if blob exists. (34672928)
+* Remove legacy artifacts and log system. (34663531)
+* Resolve HTTP client IP for Datadog trace enricher (34644555)
+* Fix user ID/username not getting set in Datadog enricher for OpenTelemetry traces. Must be accessed once response has been sent. (34636046)
+* Support for artifact searches (34605400)
+* Log each configuration source during agent start for better visibility where the agent is reading config from. Also log the actual logs dir. (34603382)
+* Log where agent registration file is saved to/loaded from. When the file exists, the agent provides no hinting it reads this file which can be confusing when trying to reset the agent. (34600505)
+* Add documentation for setting up a secret store. (34569155)
+* Add support for compression of all responses using gzip, brotli, and zstd. (34565711)
+* Add creation time to artifact responses. (34559619)
+* Change config reader to operate on JsonNode instances rather than writing directly into target objects. (34546472)
+
 ## 2024-06-20
 
 * Compute a digest for each block stored in the block cache, and validate it before returning values. (34515388)
