Horde: Rename keys used for config updates, to allow rolling upgrades from previous config versions.

#rnx

diff --git a/Engine/Source/Programs/Horde/HordeServer/Configuration/ConfigService.cs b/Engine/Source/Programs/Horde/HordeServer/Configuration/ConfigService.cs
--- a/Engine/Source/Programs/Horde/HordeServer/Configuration/ConfigService.cs
+++ b/Engine/Source/Programs/Horde/HordeServer/Configuration/ConfigService.cs
@@ -118,7 +118,7 @@ namespace HordeServer.Configuration
 		readonly Dictionary<string, IConfigSource> _sources;
 		readonly IPluginCollection _pluginCollection;
 		readonly JsonSerializerOptions _jsonOptions;
-		readonly RedisKey _snapshotKey = "config";
+		readonly RedisKey _snapshotKey = "config-v2";
 		readonly IEnumerable<IDefaultAclModifier> _defaultAclModifiers;
 		readonly IHealthMonitor _health;
 		readonly ILogger _logger;
@@ -126,7 +126,7 @@ namespace HordeServer.Configuration
 		readonly ITicker _ticker;
 		readonly TimeSpan _tickInterval = TimeSpan.FromMinutes(1.0);
 
-		readonly RedisChannel _updateChannel = RedisChannel.Literal("config-update");
+		readonly RedisChannel _updateChannel = RedisChannel.Literal("config-update-v2");
 		readonly BackgroundTask _updateTask;
 
 		Task<ConfigState> _stateTask;
