Horde: Add artifact redirector to browse to job/step with modal open to artifact
#rnx
#skipci

diff --git a/Engine/Source/Programs/Horde/HordeDashboard/src/App.tsx b/Engine/Source/Programs/Horde/HordeDashboard/src/App.tsx
--- a/Engine/Source/Programs/Horde/HordeDashboard/src/App.tsx
+++ b/Engine/Source/Programs/Horde/HordeDashboard/src/App.tsx
@@ -42,6 +42,7 @@ import { StepIssueReportTest } from './components/test/IssueStepReport';
 import { preloadFonts } from './styles/Styles';
 import { darkTheme } from './styles/darkTheme';
 import { lightTheme } from './styles/lightTheme';
+import { ArtifactRedirector } from './components/artifacts/ArtifactsRedirector';
 
 let router: any;
 
@@ -123,6 +124,7 @@ const Main: React.FC = () => {
                { path: "stream/:streamId", element: <StreamView /> },
                { path: "agents", element: <AgentView /> },
                { path: "agents/registration", element: <AgentRequestsView /> },
+               { path: "artifact/:artifactId", element: <ArtifactRedirector /> },
                { path: "admin/token", element: <AdminToken /> },
                { path: "reports/utilization", element: <UtilizationReportView /> },
                { path: "preflight", element: <PreflightRedirector /> },
diff --git a/Engine/Source/Programs/Horde/HordeDashboard/src/components/artifacts/ArtifactsRedirector.tsx b/Engine/Source/Programs/Horde/HordeDashboard/src/components/artifacts/ArtifactsRedirector.tsx
new file mode 100644
--- /dev/null
+++ b/Engine/Source/Programs/Horde/HordeDashboard/src/components/artifacts/ArtifactsRedirector.tsx
@@ -0,0 +1,73 @@
+// Copyright Epic Games, Inc. All Rights Reserved.
+
+import React from 'react';
+import { Navigate, useNavigate, useParams } from 'react-router-dom';
+import ErrorHandler from '../ErrorHandler';
+import backend from '../../backend';
+
+function setError(artifactId: string, message: string) {
+
+   console.error(message);
+
+   message = `${message}\n\nReferring URL: ${encodeURI(window.location.href)}`
+
+   ErrorHandler.set({
+
+      title: `Error Getting Artifact ${artifactId}`,
+      message: message
+
+   }, true);
+
+}
+
+let redirecting = false;
+
+// redirect from external source, where horde stream id, etc are not known by that application
+export const ArtifactRedirector: React.FC = () => {
+
+   const { artifactId } = useParams<{ artifactId: string }>();
+   const navigate = useNavigate();
+
+   const getData = async () => {
+      try {
+         const data = await backend.getArtifactData(artifactId!);
+         let key = data.keys.find(k => k.startsWith("job:") && k.indexOf("/step:") !== -1);
+         let jobId = "";
+         let stepId = "";
+         if (key) {
+            jobId = key.slice(4, 28);
+            stepId = key.slice(-4)
+         } else {
+            key = data.keys.find(k => k.startsWith("job:"));
+            if (key) {
+               jobId = key.slice(4, 28);
+            }
+         }
+
+         if (!jobId) {
+            setError(artifactId!, "Unable to get job data");
+         } else {
+            let url = `/job/${jobId}`;
+            if (stepId) {
+               url += `?step=${stepId}&artifactId=${artifactId}`;
+            } 
+            navigate(url, { replace: true });
+         }
+
+      } catch (reason) {
+         setError(artifactId!, "Error getting artifact data");
+      }
+   }
+
+   if (!artifactId) {
+      setError("None", "No artifact id specified");
+   } else {
+      if (!redirecting) {
+         redirecting = true;
+         getData();
+      }
+   }
+
+   return null;
+
+}
\ No newline at end of file
diff --git a/Engine/Source/Programs/Horde/HordeDashboard/src/components/jobDetailsV2/JobOperationsBar.tsx b/Engine/Source/Programs/Horde/HordeDashboard/src/components/jobDetailsV2/JobOperationsBar.tsx
--- a/Engine/Source/Programs/Horde/HordeDashboard/src/components/jobDetailsV2/JobOperationsBar.tsx
+++ b/Engine/Source/Programs/Horde/HordeDashboard/src/components/jobDetailsV2/JobOperationsBar.tsx
@@ -5,7 +5,9 @@ import { observer } from 'mobx-react-lite';
 import React, { useState } from 'react';
 import { Link, useNavigate } from 'react-router-dom';
 import { ArtifactContextType, JobState, JobStepOutcome, JobStepState, StepData } from '../../backend/Api';
+import { getSiteConfig } from '../../backend/Config';
 import dashboard from '../../backend/Dashboard';
+import { getHordeStyling } from '../../styles/Styles';
 import { EditJobModal } from '../EditJobModal';
 import { useQuery } from '../JobDetailCommon';
 import { NewBuild } from '../NewBuild';
@@ -13,12 +15,10 @@ import { NotificationDropdown } from '../NotificationDropdown';
 import { PauseStepModal } from '../StepPauseModal';
 import { JobArtifactsModal } from '../artifacts/ArtifactsModal';
 import { BisectionCreateModal } from '../bisection/CreateModal';
+import { NewBuildV2 } from '../build/NewBuildV2';
 import { AbortJobModal } from './AbortJobModal';
 import { JobDetailsV2 } from './JobDetailsViewCommon';
 import { RetryStepsModal, StepRetryModal, StepRetryType } from './StepRetryModal';
-import { getSiteConfig } from '../../backend/Config';
-import { getHordeStyling } from '../../styles/Styles';
-import { NewBuildV2 } from '../build/NewBuildV2';
 
 enum ParameterState {
    Hidden,
@@ -233,10 +233,9 @@ const StepArtifactsOperations: React.FC<{ jobDetails: JobDetailsV2, stepId: stri
    const navigate = useNavigate();
 
    const query = useQuery();
-   const artifactContext = !!query.get("artifactContext") ? query.get("artifactContext")! as ArtifactContextType : undefined;
+   let artifactContext = !!query.get("artifactContext") ? query.get("artifactContext")! as ArtifactContextType : undefined;
    const artifactPath = !!query.get("artifactPath") ? query.get("artifactPath")! : undefined;
-
-   const jobData = jobDetails.jobData;
+   const artifactId = !!query.get("artifactId") ? query.get("artifactId")! : undefined;
 
    // subscribe
    if (dashboard.updated) { }
@@ -258,6 +257,11 @@ const StepArtifactsOperations: React.FC<{ jobDetails: JobDetailsV2, stepId: stri
       atypes.set(a.type, c);
    });
 
+   let artifact = stepArtifacts?.find(a => a.id === artifactId);
+   if (!artifactContext && artifact) {
+      artifactContext = artifact.type;
+   }
+   
    const opsList: IContextualMenuItem[] = [];
 
    const baseUrl = window.location.pathname + window.location.search;
