Horde: Allow specifying a secondary object store to read objects from, allowing a migration from one store in the background without causing downtime.

diff --git a/Engine/Source/Programs/Horde/Plugins/Storage/HordeServer.Storage/Storage/ObjectStores/ChainedObjectStore.cs b/Engine/Source/Programs/Horde/Plugins/Storage/HordeServer.Storage/Storage/ObjectStores/ChainedObjectStore.cs
new file mode 100644
--- /dev/null
+++ b/Engine/Source/Programs/Horde/Plugins/Storage/HordeServer.Storage/Storage/ObjectStores/ChainedObjectStore.cs
@@ -0,0 +1,111 @@
+// Copyright Epic Games, Inc. All Rights Reserved.
+
+using System;
+using System.Collections.Generic;
+using System.Linq;
+using System.Runtime.CompilerServices;
+using System.Text;
+using System.Threading.Tasks;
+using EpicGames.Core;
+using EpicGames.Horde.Storage;
+
+namespace HordeServer.Storage.Storage.ObjectStores
+{
+	/// <summary>
+	/// Object store which allows falling back to reading from a secondary object store if a key is not found in the primary store. Useful for migrations.
+	/// </summary>
+	class ChainedObjectStore : IObjectStore
+	{
+		readonly IObjectStore _primary;
+		readonly IObjectStore _secondary;
+
+		/// <inheritdoc/>
+		public bool SupportsRedirects => _primary.SupportsRedirects;
+
+		/// <summary>
+		/// Constructor
+		/// </summary>
+		/// <param name="primary">The primary object store to read from, and the store to use for all writes</param>
+		/// <param name="secondary">The secondary object store</param>
+		public ChainedObjectStore(IObjectStore primary, IObjectStore secondary)
+		{
+			_primary = primary;
+			_secondary = secondary;
+		}
+
+		/// <inheritdoc/>
+		public Task DeleteAsync(ObjectKey key, CancellationToken cancellationToken = default)
+			=> _primary.DeleteAsync(key, cancellationToken);
+
+		/// <inheritdoc/>
+		public async IAsyncEnumerable<ObjectKey> EnumerateAsync([EnumeratorCancellation] CancellationToken cancellationToken = default)
+		{
+			await foreach (ObjectKey key in _primary.EnumerateAsync(cancellationToken))
+			{
+				yield return key;
+			}
+			await foreach (ObjectKey key in _secondary.EnumerateAsync(cancellationToken))
+			{
+				yield return key;
+			}
+		}
+
+		/// <inheritdoc/>
+		public async Task<bool> ExistsAsync(ObjectKey key, CancellationToken cancellationToken = default)
+			=> await _primary.ExistsAsync(key, cancellationToken) || await _secondary.ExistsAsync(key, cancellationToken);
+
+		/// <inheritdoc/>
+		public void GetStats(StorageStats stats)
+		{
+			_primary.GetStats(stats);
+			_secondary.GetStats(stats);
+		}
+
+		/// <inheritdoc/>
+		public async Task<Stream> OpenAsync(ObjectKey key, int offset, int? length, CancellationToken cancellationToken = default)
+		{
+			if (await _primary.ExistsAsync(key, cancellationToken) || !await _secondary.ExistsAsync(key, cancellationToken))
+			{
+				return await _primary.OpenAsync(key, offset, length, cancellationToken);
+			}
+			else
+			{
+				return await _secondary.OpenAsync(key, offset, length, cancellationToken);
+			}
+		}
+
+		/// <inheritdoc/>
+		public async Task<IReadOnlyMemoryOwner<byte>> ReadAsync(ObjectKey key, int offset, int? length, CancellationToken cancellationToken = default)
+		{
+			if (await _primary.ExistsAsync(key, cancellationToken) || !await _secondary.ExistsAsync(key, cancellationToken))
+			{
+				return await _primary.ReadAsync(key, offset, length, cancellationToken);
+			}
+			else
+			{
+				return await _secondary.ReadAsync(key, offset, length, cancellationToken);
+			}
+		}
+
+		/// <inheritdoc/>
+		public async ValueTask<Uri?> TryGetReadRedirectAsync(ObjectKey key, CancellationToken cancellationToken = default)
+		{
+			if (await _primary.ExistsAsync(key, cancellationToken) || !await _secondary.ExistsAsync(key, cancellationToken))
+			{
+				return await _primary.TryGetReadRedirectAsync(key, cancellationToken);
+			}
+			else
+			{
+				return await _secondary.TryGetReadRedirectAsync(key, cancellationToken);
+			}
+		}
+
+		/// <inheritdoc/>
+		public ValueTask<Uri?> TryGetWriteRedirectAsync(ObjectKey key, CancellationToken cancellationToken = default)
+			=> _primary.TryGetWriteRedirectAsync(key, cancellationToken);
+
+		/// <inheritdoc/>
+		public Task WriteAsync(ObjectKey key, Stream stream, CancellationToken cancellationToken = default)
+			=> _primary.WriteAsync(key, stream, cancellationToken);
+	}
+}
diff --git a/Engine/Source/Programs/Horde/Plugins/Storage/HordeServer.Storage/Storage/StorageService.cs b/Engine/Source/Programs/Horde/Plugins/Storage/HordeServer.Storage/Storage/StorageService.cs
--- a/Engine/Source/Programs/Horde/Plugins/Storage/HordeServer.Storage/Storage/StorageService.cs
+++ b/Engine/Source/Programs/Horde/Plugins/Storage/HordeServer.Storage/Storage/StorageService.cs
@@ -17,6 +17,7 @@ using EpicGames.Horde.Storage.ObjectStores;
 using EpicGames.Redis;
 using EpicGames.Redis.Utility;
 using HordeServer.Server;
+using HordeServer.Storage.Storage.ObjectStores;
 using HordeServer.Utilities;
 using Microsoft.Extensions.Caching.Memory;
 using Microsoft.Extensions.Configuration;
@@ -538,16 +539,7 @@ namespace HordeServer.Storage
 						NamespaceId namespaceId = namespaceConfig.Id;
 						try
 						{
-							BackendId backendId = namespaceConfig.Backend;
-
-							BackendConfig? backendConfig = _staticStorageConfig.Value.Backends.FirstOrDefault(x => x.Id == namespaceConfig.Backend);
-							if (backendConfig == null && !storageConfig.TryGetBackend(namespaceConfig.Backend, out backendConfig))
-							{
-								throw new StorageException($"Missing or invalid backend identifier for namespace {namespaceConfig.Id}");
-							}
-
-							IObjectStore objectStore = _objectStoreFactory.CreateObjectStore(backendConfig);
-
+							IObjectStore objectStore = CreateObjectStore(storageConfig, namespaceConfig.Backend);
 							if (!String.IsNullOrEmpty(namespaceConfig.Prefix))
 							{
 								objectStore = new PrefixedObjectStore(namespaceConfig.Prefix, objectStore);
@@ -571,6 +563,24 @@ namespace HordeServer.Storage
 			}
 		}
 
+		IObjectStore CreateObjectStore(StorageConfig storageConfig, BackendId backendId)
+		{
+			BackendConfig? backendConfig = _staticStorageConfig.Value.Backends.FirstOrDefault(x => x.Id == backendId);
+			if (backendConfig == null && !storageConfig.TryGetBackend(backendId, out backendConfig))
+			{
+				throw new StorageException($"Missing or invalid backend identifier '{backendId}');");
+			}
+
+			IObjectStore objectStore = _objectStoreFactory.CreateObjectStore(backendConfig);
+			if (!backendConfig.Secondary.IsEmpty)
+			{
+				IObjectStore secondaryObjectStore = CreateObjectStore(storageConfig, backendConfig.Secondary);
+				objectStore = new ChainedObjectStore(objectStore, secondaryObjectStore);
+			}
+
+			return objectStore;
+		}
+
 		#endregion
 
 		#region Blobs
diff --git a/Engine/Source/Programs/Horde/Plugins/Storage/HordeServer.Storage/StorageConfig.cs b/Engine/Source/Programs/Horde/Plugins/Storage/HordeServer.Storage/StorageConfig.cs
--- a/Engine/Source/Programs/Horde/Plugins/Storage/HordeServer.Storage/StorageConfig.cs
+++ b/Engine/Source/Programs/Horde/Plugins/Storage/HordeServer.Storage/StorageConfig.cs
@@ -256,6 +256,11 @@ namespace HordeServer.Storage
 		/// </summary>
 		public BackendId Base { get; set; }
 
+		/// <summary>
+		/// Specifies another backend to read from if an object is not found in this one. Can be used when migrating data from one backend to another.
+		/// </summary>
+		public BackendId Secondary { get; set; }
+
 		/// <inheritdoc/>
 		public StorageBackendType? Type { get; set; }
 
