Horde: Updating schema documentation

#rnx

diff --git a/Engine/Source/Programs/Horde/Docs/Config/Schema/Globals.md b/Engine/Source/Programs/Horde/Docs/Config/Schema/Globals.md
--- a/Engine/Source/Programs/Horde/Docs/Config/Schema/Globals.md
+++ b/Engine/Source/Programs/Horde/Docs/Config/Schema/Globals.md
@@ -386,6 +386,7 @@ Name | Description
 ---- | -----------
 `id` | `string`<br>
 `base` | `string`<br>
+`secondary` | `string`<br>
 `type` | [`StorageBackendType`](#storagebackendtype-enum)<br>
 `baseDir` | `string`<br>
 `awsBucketName` | `string`<br>
