Horde: Allow setting access permissions for different artifact types on a per-type, per-project and per-stream basis.

diff --git a/Engine/Source/Programs/Horde/HordeServer.Shared/Acls/AclConfig.cs b/Engine/Source/Programs/Horde/HordeServer.Shared/Acls/AclConfig.cs
--- a/Engine/Source/Programs/Horde/HordeServer.Shared/Acls/AclConfig.cs
+++ b/Engine/Source/Programs/Horde/HordeServer.Shared/Acls/AclConfig.cs
@@ -90,7 +90,13 @@ namespace HordeServer.Acls
 			return false;
 		}
 
-		bool? AuthorizeSingleScope(AclAction action, ClaimsPrincipal user)
+		/// <summary>
+		/// Tests whether a user is authorized to perform the given actions in this specific scope
+		/// </summary>
+		/// <param name="action">Action that is being performed. This should be a single flag.</param>
+		/// <param name="user">The principal to authorize</param>
+		/// <returns>True/false if the action is allowed or denied, null if there is no specific setting for this user</returns>
+		public bool? AuthorizeSingleScope(AclAction action, ClaimsPrincipal user)
 		{
 			// Check if there's a specific entry for this action
 			if (Entries != null)
diff --git a/Engine/Source/Programs/Horde/Plugins/Build/HordeServer.Build/Artifacts/ArtifactsController.cs b/Engine/Source/Programs/Horde/Plugins/Build/HordeServer.Build/Artifacts/ArtifactsController.cs
--- a/Engine/Source/Programs/Horde/Plugins/Build/HordeServer.Build/Artifacts/ArtifactsController.cs
+++ b/Engine/Source/Programs/Horde/Plugins/Build/HordeServer.Build/Artifacts/ArtifactsController.cs
@@ -38,6 +38,8 @@ using Microsoft.Extensions.Logging;
 using Microsoft.Extensions.Options;
 using EpicGames.Horde.Acls;
 using HordeServer.Agents;
+using System.Security.Claims;
+using HordeServer.Projects;
 
 namespace HordeServer.Artifacts
 {
@@ -203,7 +205,7 @@ namespace HordeServer.Artifacts
 			FindArtifactsResponse response = new FindArtifactsResponse();
 			await foreach (IArtifact artifact in _artifactCollection.FindAsync(streamId, minChange, maxChange, name, type, keys, maxResults, HttpContext.RequestAborted))
 			{
-				if (_aclService.Authorize(artifact.AclScope, ArtifactAclAction.ReadArtifact, User))
+				if (_buildConfig.AuthorizeArtifact(artifact.Type, artifact.StreamId, ArtifactAclAction.ReadArtifact, User))
 				{
 					response.Artifacts.Add(new GetArtifactResponse(artifact.Id, artifact.Name, artifact.Type, artifact.Description, artifact.StreamId, artifact.Change, artifact.Keys, artifact.Metadata, artifact.CreatedAtUtc));
 				}
@@ -228,9 +230,9 @@ namespace HordeServer.Artifacts
 			{
 				return NotFound(id);
 			}
-			if (!_aclService.Authorize(artifact.AclScope, ArtifactAclAction.ReadArtifact, User))
+			if (!_buildConfig.AuthorizeArtifact(artifact.Type, artifact.StreamId, ArtifactAclAction.ReadArtifact, User))
 			{
-				return Forbid(ArtifactAclAction.ReadArtifact, artifact.AclScope);
+				return Forbid(ArtifactAclAction.ReadArtifact, artifact.StreamId);
 			}
 
 			return PropertyFilter.Apply(new GetArtifactResponse(artifact.Id, artifact.Name, artifact.Type, artifact.Description, artifact.StreamId, artifact.Change, artifact.Keys, artifact.Metadata, artifact.CreatedAtUtc), filter);
@@ -252,9 +254,9 @@ namespace HordeServer.Artifacts
 			{
 				return NotFound(id);
 			}
-			if (!_aclService.Authorize(artifact.AclScope, ArtifactAclAction.ReadArtifact, User))
+			if (!_buildConfig.AuthorizeArtifact(artifact.Type, artifact.StreamId, ArtifactAclAction.ReadArtifact, User))
 			{
-				return Forbid(ArtifactAclAction.ReadArtifact, artifact.AclScope);
+				return Forbid(ArtifactAclAction.ReadArtifact, artifact.StreamId);
 			}
 			if (!locator.WithinFolder(artifact.RefName.Text))
 			{
@@ -280,9 +282,9 @@ namespace HordeServer.Artifacts
 			{
 				return NotFound(id);
 			}
-			if (!_aclService.Authorize(artifact.AclScope, ArtifactAclAction.ReadArtifact, User))
+			if (!_buildConfig.AuthorizeArtifact(artifact.Type, artifact.StreamId, ArtifactAclAction.ReadArtifact, User))
 			{
-				return Forbid(ArtifactAclAction.ReadArtifact, artifact.AclScope);
+				return Forbid(ArtifactAclAction.ReadArtifact, artifact.StreamId);
 			}
 
 			return await StorageController.ReadRefInternalAsync(_storageService, artifact.NamespaceId, artifact.RefName, Request.Headers, cancellationToken);
@@ -307,9 +309,9 @@ namespace HordeServer.Artifacts
 			{
 				return NotFound(id);
 			}
-			if (!_aclService.Authorize(artifact.AclScope, ArtifactAclAction.ReadArtifact, User))
+			if (!_buildConfig.AuthorizeArtifact(artifact.Type, artifact.StreamId, ArtifactAclAction.ReadArtifact, User))
 			{
-				return Forbid(ArtifactAclAction.ReadArtifact, artifact.AclScope);
+				return Forbid(ArtifactAclAction.ReadArtifact, artifact.StreamId);
 			}
 
 			using IStorageClient storageClient = _storageService.CreateClient(artifact.NamespaceId);
@@ -474,9 +476,9 @@ namespace HordeServer.Artifacts
 			{
 				return NotFound(id);
 			}
-			if (!_aclService.Authorize(artifact.AclScope, ArtifactAclAction.ReadArtifact, User))
+			if (!_buildConfig.AuthorizeArtifact(artifact.Type, artifact.StreamId, ArtifactAclAction.ReadArtifact, User))
 			{
-				return Forbid(ArtifactAclAction.ReadArtifact, artifact.AclScope);
+				return Forbid(ArtifactAclAction.ReadArtifact, artifact.StreamId);
 			}
 
 			using IStorageClient storageClient = _storageService.CreateClient(artifact.NamespaceId);
@@ -549,9 +551,9 @@ namespace HordeServer.Artifacts
 			{
 				return NotFound(id);
 			}
-			if (!_aclService.Authorize(artifact.AclScope, ArtifactAclAction.ReadArtifact, User))
+			if (!_buildConfig.AuthorizeArtifact(artifact.Type, artifact.StreamId, ArtifactAclAction.ReadArtifact, User))
 			{
-				return Forbid(ArtifactAclAction.ReadArtifact, artifact.AclScope);
+				return Forbid(ArtifactAclAction.ReadArtifact, artifact.StreamId);
 			}
 
 			switch (format ?? DownloadArtifactFormat.Zip)
@@ -641,9 +643,9 @@ namespace HordeServer.Artifacts
 			{
 				return NotFound(id);
 			}
-			if (!_aclService.Authorize(artifact.AclScope, ArtifactAclAction.ReadArtifact, User))
+			if (!_buildConfig.AuthorizeArtifact(artifact.Type, artifact.StreamId, ArtifactAclAction.ReadArtifact, User))
 			{
-				return Forbid(ArtifactAclAction.ReadArtifact, artifact.AclScope);
+				return Forbid(ArtifactAclAction.ReadArtifact, artifact.StreamId);
 			}
 
 			ReadOnlyMemory<byte> manifestData = await _unsyncCache.GetManifestDataAsync(artifact, cancellationToken);
@@ -674,9 +676,9 @@ namespace HordeServer.Artifacts
 			{
 				return NotFound(id);
 			}
-			if (!_aclService.Authorize(artifact.AclScope, ArtifactAclAction.ReadArtifact, User))
+			if (!_buildConfig.AuthorizeArtifact(artifact.Type, artifact.StreamId, ArtifactAclAction.ReadArtifact, User))
 			{
-				return Forbid(ArtifactAclAction.ReadArtifact, artifact.AclScope);
+				return Forbid(ArtifactAclAction.ReadArtifact, artifact.StreamId);
 			}
 
 			if (!String.Equals(request.HashStrong, "Blake3.160", StringComparison.OrdinalIgnoreCase))
diff --git a/Engine/Source/Programs/Horde/Plugins/Build/HordeServer.Build/Artifacts/IArtifact.cs b/Engine/Source/Programs/Horde/Plugins/Build/HordeServer.Build/Artifacts/IArtifact.cs
--- a/Engine/Source/Programs/Horde/Plugins/Build/HordeServer.Build/Artifacts/IArtifact.cs
+++ b/Engine/Source/Programs/Horde/Plugins/Build/HordeServer.Build/Artifacts/IArtifact.cs
@@ -68,10 +68,5 @@ namespace HordeServer.Artifacts
 		/// Time at which the artifact was created
 		/// </summary>
 		public DateTime CreatedAtUtc { get; }
-
-		/// <summary>
-		/// Permissions scope for this object
-		/// </summary>
-		public AclScopeName AclScope { get; }
 	}
 }
diff --git a/Engine/Source/Programs/Horde/Plugins/Build/HordeServer.Build/BuildConfig.cs b/Engine/Source/Programs/Horde/Plugins/Build/HordeServer.Build/BuildConfig.cs
--- a/Engine/Source/Programs/Horde/Plugins/Build/HordeServer.Build/BuildConfig.cs
+++ b/Engine/Source/Programs/Horde/Plugins/Build/HordeServer.Build/BuildConfig.cs
@@ -26,6 +26,7 @@ using HordeServer.Perforce;
 using HordeServer.Plugins;
 using HordeServer.Projects;
 using HordeServer.Streams;
+using HordeServer.Utilities;
 
 #pragma warning disable CA2227 // Change x to be read-only by removing the property setter
 
@@ -83,7 +84,6 @@ namespace HordeServer
 
 		private readonly Dictionary<ProjectId, ProjectConfig> _projectLookup = new Dictionary<ProjectId, ProjectConfig>();
 		private readonly Dictionary<StreamId, StreamConfig> _streamLookup = new Dictionary<StreamId, StreamConfig>();
-		private readonly Dictionary<ArtifactType, ArtifactTypeConfig> _artifactTypeLookup = new Dictionary<ArtifactType, ArtifactTypeConfig>();
 
 		/// <inheritdoc/>
 		public void PostLoad(PluginConfigOptions configOptions)
@@ -104,12 +104,6 @@ namespace HordeServer
 				}
 			}
 
-			_artifactTypeLookup.Clear();
-			foreach (ArtifactTypeConfig artifactType in ArtifactTypes)
-			{
-				_artifactTypeLookup.Add(artifactType.Type, artifactType);
-			}
-
 			UpdateWorkspacesForPools(configOptions.Plugins);
 		}
 
@@ -178,12 +172,36 @@ namespace HordeServer
 		}
 
 		/// <summary>
-		/// Attempts to get configuration for a specific artifact type
+		/// Authorize an action to perform on an artifact
 		/// </summary>
-		/// <param name="type">The artifact type name</param>
-		/// <param name="config">Configuration for the stream</param>
-		/// <returns>True if the stream configuration was found</returns>
-		public bool TryGetArtifactType(ArtifactType type, [NotNullWhen(true)] out ArtifactTypeConfig? config) => _artifactTypeLookup.TryGetValue(type, out config);
+		public bool AuthorizeArtifact(ArtifactType type, StreamId streamId, AclAction action, ClaimsPrincipal user)
+		{
+			if (user.HasAdminClaim())
+			{
+				return true;
+			}
+
+			bool? auth = null;
+			if (TryGetStream(streamId, out StreamConfig? streamConfig))
+			{
+				auth = AuthorizeArtifactType(streamConfig.ArtifactTypes, type, action, user);
+				auth ??= streamConfig.Acl.AuthorizeSingleScope(action, user);
+
+				ProjectConfig projectConfig = streamConfig.ProjectConfig;
+
+				auth ??= AuthorizeArtifactType(projectConfig.ArtifactTypes, type, action, user);
+				auth ??= projectConfig.Acl.AuthorizeSingleScope(action, user);
+			}
+
+			auth ??= AuthorizeArtifactType(ArtifactTypes, type, action, user);
+			return auth ?? Acl.Authorize(action, user);
+		}
+
+		static bool? AuthorizeArtifactType(IReadOnlyList<ArtifactTypeAclConfig> artifactTypes, ArtifactType artifactType, AclAction action, ClaimsPrincipal user)
+		{
+			ArtifactTypeAclConfig? config = artifactTypes.FirstOrDefault(x => x.Type == artifactType);
+			return config?.Acl?.AuthorizeSingleScope(action, user);
+		}
 
 		/// <summary>
 		/// Attempts to get configuration for a project from this object
@@ -285,15 +303,26 @@ namespace HordeServer
 	}
 
 	/// <summary>
-	/// Configuration for an artifact
+	/// ACL configuration for an artifact type
 	/// </summary>
-	public class ArtifactTypeConfig
+	public class ArtifactTypeAclConfig
 	{
 		/// <summary>
 		/// Name of the artifact type
 		/// </summary>
 		public ArtifactType Type { get; set; }
 
+		/// <summary>
+		/// Acl for the artifact type
+		/// </summary>
+		public AclConfig? Acl { get; set; }
+	}
+
+	/// <summary>
+	/// Configuration for an artifact
+	/// </summary>
+	public class ArtifactTypeConfig : ArtifactTypeAclConfig
+	{
 		/// <summary>
 		/// Legacy 'Name' property
 		/// </summary>
diff --git a/Engine/Source/Programs/Horde/Plugins/Build/HordeServer.Build/Projects/ProjectConfig.cs b/Engine/Source/Programs/Horde/Plugins/Build/HordeServer.Build/Projects/ProjectConfig.cs
--- a/Engine/Source/Programs/Horde/Plugins/Build/HordeServer.Build/Projects/ProjectConfig.cs
+++ b/Engine/Source/Programs/Horde/Plugins/Build/HordeServer.Build/Projects/ProjectConfig.cs
@@ -105,6 +105,11 @@ namespace HordeServer.Projects
 		/// </summary>
 		public List<StreamConfig> Streams { get; set; } = new List<StreamConfig>();
 
+		/// <summary>
+		/// Permissions for artifact types
+		/// </summary>
+		public List<ArtifactTypeAclConfig> ArtifactTypes { get; set; } = new List<ArtifactTypeAclConfig>();
+
 		/// <summary>
 		/// Acl entries
 		/// </summary>
diff --git a/Engine/Source/Programs/Horde/Plugins/Build/HordeServer.Build/Streams/StreamConfig.cs b/Engine/Source/Programs/Horde/Plugins/Build/HordeServer.Build/Streams/StreamConfig.cs
--- a/Engine/Source/Programs/Horde/Plugins/Build/HordeServer.Build/Streams/StreamConfig.cs
+++ b/Engine/Source/Programs/Horde/Plugins/Build/HordeServer.Build/Streams/StreamConfig.cs
@@ -245,6 +245,11 @@ namespace HordeServer.Streams
 		/// </summary>
 		public List<TokenConfig> Tokens { get; set; } = new List<TokenConfig>();
 
+		/// <summary>
+		/// Permissions for artifact types
+		/// </summary>
+		public List<ArtifactTypeAclConfig> ArtifactTypes { get; set; } = new List<ArtifactTypeAclConfig>();
+
 		/// <inheritdoc cref="AclConfig.Authorize(AclAction, ClaimsPrincipal)"/>
 		public bool Authorize(AclAction action, ClaimsPrincipal user)
 			=> Acl.Authorize(action, user);
