Horde: Support for disabling conform tasks through the build config file.

diff --git a/Engine/Source/Programs/Horde/Plugins/Build/HordeServer.Build/BuildConfig.cs b/Engine/Source/Programs/Horde/Plugins/Build/HordeServer.Build/BuildConfig.cs
--- a/Engine/Source/Programs/Horde/Plugins/Build/HordeServer.Build/BuildConfig.cs
+++ b/Engine/Source/Programs/Horde/Plugins/Build/HordeServer.Build/BuildConfig.cs
@@ -77,6 +77,11 @@ namespace HordeServer
 		[JsonIgnore]
 		public IReadOnlyList<StreamConfig> Streams { get; private set; } = null!;
 
+		/// <summary>
+		/// Whether to allow conform tasks to run
+		/// </summary>
+		public bool EnableConformTasks { get; set; } = true;
+
 		/// <summary>
 		/// Commit tag to use for marking issues as fixed
 		/// </summary>
diff --git a/Engine/Source/Programs/Horde/Plugins/Build/HordeServer.Build/Tasks/ConformTaskSource.cs b/Engine/Source/Programs/Horde/Plugins/Build/HordeServer.Build/Tasks/ConformTaskSource.cs
--- a/Engine/Source/Programs/Horde/Plugins/Build/HordeServer.Build/Tasks/ConformTaskSource.cs
+++ b/Engine/Source/Programs/Horde/Plugins/Build/HordeServer.Build/Tasks/ConformTaskSource.cs
@@ -157,7 +157,7 @@ namespace HordeServer.Tasks
 		/// <inheritdoc/>
 		public override async Task<Task<AgentLease?>> AssignLeaseAsync(IAgent agent, CancellationToken cancellationToken)
 		{
-			if (!_settings.CurrentValue.EnableConformTasks)
+			if (!_settings.CurrentValue.EnableConformTasks || !_buildConfig.CurrentValue.EnableConformTasks)
 			{
 				return SkipAsync(cancellationToken);
 			}
