Fix Horde installer issue introduced in previous submit that nested Server files in a win-x64 folder

#rnx

diff --git a/Engine/Source/Programs/Horde/BuildHorde.xml b/Engine/Source/Programs/Horde/BuildHorde.xml
--- a/Engine/Source/Programs/Horde/BuildHorde.xml
+++ b/Engine/Source/Programs/Horde/BuildHorde.xml
@@ -292,10 +292,10 @@
 			<!-- Create the server file list -->
 			<Delete Files="$(StagingDir)/Server-Bulk/..."/>
 			<Tag Files="$(StagingDir)/Server/win-x64/..." Except="HordeServer.exe" With="#ServerFiles"/>
-			<Copy Files="#ServerFiles" From="$(StagingDir)" To="$(StagingDir)/Server-Bulk"/>
+			<Copy Files="#ServerFiles" From="$(StagingDir)/Server/win-x64" To="$(StagingDir)/Server-Bulk"/>
 
-			<Spawn Exe="$(WixDir)/heat.exe" Arguments="dir &quot;$(StagingDir)/Server-Bulk&quot; -cg HordeServerFiles -scom -sreg -gg -dr InstallDir -srd -var var.SourceDir -out &quot;$(InstallerDir)/HordeServerFiles.wxs&quot;"/>
-			<Spawn Exe="$(WixDir)/candle.exe" Arguments="$(CommonArgs) HordeServerFiles.wxs -dSourceDir=&quot;$(StagingDir)&quot; -out $(ObjDir)/HordeServerFiles.wixobj" WorkingDir="$(InstallerDir)"/>
+			<Spawn Exe="$(WixDir)/heat.exe" Arguments="dir &quot;$(StagingDir)/Server-Bulk&quot; -cg HordeServerFiles -scom -sreg -gg -dr ServerInstallDir -srd -var var.SourceDir -out &quot;$(InstallerDir)/HordeServerFiles.wxs&quot;"/>
+			<Spawn Exe="$(WixDir)/candle.exe" Arguments="$(CommonArgs) HordeServerFiles.wxs -dSourceDir=&quot;$(StagingDir)/Server/win-x64&quot; -out $(ObjDir)/HordeServerFiles.wixobj" WorkingDir="$(InstallerDir)"/>
 
 			<!-- Main component -->
 			<Spawn Exe="$(WixDir)/candle.exe" Arguments="$(CommonArgs) ServerInstaller.wxs -dStagingDir=&quot;$(StagingDir)&quot; -out $(ObjDir)/ServerInstaller.wixobj -ext WixUtilExtension" WorkingDir="$(InstallerDir)"/>
