Horde: Update release notes.

#rnx

diff --git a/Engine/Source/Programs/Horde/Docs/ReleaseNotes.md b/Engine/Source/Programs/Horde/Docs/ReleaseNotes.md
--- a/Engine/Source/Programs/Horde/Docs/ReleaseNotes.md
+++ b/Engine/Source/Programs/Horde/Docs/ReleaseNotes.md
@@ -1,5 +1,11 @@
 # Release Notes
 
+## 2024-07-16
+
+* Support for disabling conform tasks through the build config file. (34836886)
+* Allow setting access permissions for different artifact types on a per-type, per-project and per-stream basis. (34836758)
+* Allow specifying a secondary object store to read objects from, allowing a migration from one store in the background without causing downtime. (34828731)
+
 ## 2024-07-15
 
 * Split server functionality into plugins. Plugins are still currently statically configured, which will be changed in future. (34620916, many others)
