Horde: Adding plugin info to server info response
#rnx

diff --git a/Engine/Source/Programs/Horde/HordeServer/Server/ServerController.cs b/Engine/Source/Programs/Horde/HordeServer/Server/ServerController.cs
--- a/Engine/Source/Programs/Horde/HordeServer/Server/ServerController.cs
+++ b/Engine/Source/Programs/Horde/HordeServer/Server/ServerController.cs
@@ -6,12 +6,15 @@ using System.Diagnostics;
 using System.Linq;
 using System.Reflection;
 using System.Threading.Tasks;
+using EpicGames.Core;
 using EpicGames.Horde;
 using EpicGames.Horde.Server;
+using HordeServer.Plugins;
 using HordeServer.Utilities;
 using Microsoft.AspNetCore.Authorization;
 using Microsoft.AspNetCore.Mvc;
 using Microsoft.Extensions.Options;
+using MongoDB.Driver;
 
 namespace HordeServer.Server
 {
@@ -25,14 +28,16 @@ namespace HordeServer.Server
 	public class ServerController : HordeControllerBase
 	{
 		readonly IAgentVersionProvider? _agentVersionProvider;
+		readonly IPluginCollection _pluginCollection;
 		readonly IOptionsSnapshot<GlobalConfig> _globalConfig;
 
 		/// <summary>
 		/// Constructor
 		/// </summary>
-		public ServerController(IEnumerable<IAgentVersionProvider> agentVersionProviders, IOptionsSnapshot<GlobalConfig> globalConfig)
+		public ServerController(IEnumerable<IAgentVersionProvider> agentVersionProviders, IPluginCollection pluginCollection, IOptionsSnapshot<GlobalConfig> globalConfig)
 		{
 			_agentVersionProvider = agentVersionProviders.FirstOrDefault();
+			_pluginCollection = pluginCollection;
 			_globalConfig = globalConfig;
 		}
 
@@ -68,6 +73,11 @@ namespace HordeServer.Server
 				response.AgentVersion = await _agentVersionProvider.GetAsync(HttpContext.RequestAborted);
 			}
 
+			response.Plugins = _pluginCollection.LoadedPlugins.ConvertAll(plugin => {
+				FileVersionInfo pluginVersion = FileVersionInfo.GetVersionInfo(plugin.Assembly.Location);
+				return new ServerPluginInfoResponse(plugin.Metadata.Name.ToString(), plugin.Metadata.Description, true, pluginVersion.ProductVersion ?? String.Empty);
+			}).ToArray();
+
 			return response;
 		}
 
diff --git a/Engine/Source/Programs/Shared/EpicGames.Horde/Server/ServerMessages.cs b/Engine/Source/Programs/Shared/EpicGames.Horde/Server/ServerMessages.cs
--- a/Engine/Source/Programs/Shared/EpicGames.Horde/Server/ServerMessages.cs
+++ b/Engine/Source/Programs/Shared/EpicGames.Horde/Server/ServerMessages.cs
@@ -27,6 +27,12 @@ namespace EpicGames.Horde.Server
 		/// The current agent version string
 		/// </summary>
 		public string? AgentVersion { get; set; }
+
+		/// <summary>
+		/// List of plugins
+		/// </summary>
+		public ServerPluginInfoResponse[] Plugins { get; set; } = Array.Empty<ServerPluginInfoResponse>();
+
 	}
 
 	/// <summary>
@@ -231,5 +237,15 @@ namespace EpicGames.Horde.Server
 		/// </summary>
 		public ServerStatusSubsystem[] Statuses { get; set; } = Array.Empty<ServerStatusSubsystem>();
 	}
+
+	/// <summary>
+	/// Information about a server plugin
+	/// </summary>
+	/// <param name="Name">Name of the plugin</param>
+	/// <param name="Description">Optional description of the plugin</param>
+	/// <param name="Loaded">Whether the plugin is loaded</param>
+	/// <param name="Version">The version of the plugin assembly</param>
+	public record ServerPluginInfoResponse(string Name, string? Description, bool Loaded, string? Version = null);
+
 }
 
