Horde: Updating schema documentation

#rnx

diff --git a/Engine/Source/Programs/Horde/Docs/Config/Schema/Globals.md b/Engine/Source/Programs/Horde/Docs/Config/Schema/Globals.md
--- a/Engine/Source/Programs/Horde/Docs/Config/Schema/Globals.md
+++ b/Engine/Source/Programs/Horde/Docs/Config/Schema/Globals.md
@@ -64,6 +64,7 @@ Name | Description
 `agentShutdownIfDisabledGracePeriod` | `string`<br>
 `artifactTypes` | [`ArtifactTypeConfig`](#artifacttypeconfig)`[]`<br>
 `projects` | [`ProjectConfig`](Projects.md)`[]`<br>
+`enableConformTasks` | `boolean`<br>
 `issueFixedTag` | `string`<br>
 
 ## PerforceCluster
@@ -145,22 +146,11 @@ Name | Description
 
 Name | Description
 ---- | -----------
-`type` | `string`<br>
 `name` | `string`<br>
 `keepCount` | `integer`<br>
 `keepDays` | `integer`<br>
-
-## ComputeConfig
-
-Name | Description
----- | -----------
+`type` | `string`<br>
 `acl` | [`AclConfig`](#aclconfig)<br>
-`versionEnum` | [`ConfigVersion`](#configversion-enum)<br>
-`rates` | [`AgentRateConfig`](#agentrateconfig)`[]`<br>
-`clusters` | [`ComputeClusterConfig`](#computeclusterconfig)`[]`<br>
-`pools` | [`PoolConfig`](#poolconfig)`[]`<br>
-`software` | [`AgentSoftwareConfig`](#agentsoftwareconfig)`[]`<br>
-`networks` | [`NetworkConfig`](#networkconfig)`[]`<br>
 
 ## AclConfig
 
@@ -195,6 +185,18 @@ Name | Description
 `excludeActions` | `string[]`<br>
 `extends` | `string[]`<br>
 
+## ComputeConfig
+
+Name | Description
+---- | -----------
+`acl` | [`AclConfig`](#aclconfig)<br>
+`versionEnum` | [`ConfigVersion`](#configversion-enum)<br>
+`rates` | [`AgentRateConfig`](#agentrateconfig)`[]`<br>
+`clusters` | [`ComputeClusterConfig`](#computeclusterconfig)`[]`<br>
+`pools` | [`PoolConfig`](#poolconfig)`[]`<br>
+`software` | [`AgentSoftwareConfig`](#agentsoftwareconfig)`[]`<br>
+`networks` | [`NetworkConfig`](#networkconfig)`[]`<br>
+
 ## ConfigVersion (Enum)
 
 Name | Description
diff --git a/Engine/Source/Programs/Horde/Docs/Config/Schema/Projects.md b/Engine/Source/Programs/Horde/Docs/Config/Schema/Projects.md
--- a/Engine/Source/Programs/Horde/Docs/Config/Schema/Projects.md
+++ b/Engine/Source/Programs/Horde/Docs/Config/Schema/Projects.md
@@ -18,6 +18,7 @@ Name | Description
 `workspaceTypes` | `string` `->` [`WorkspaceConfig`](#workspaceconfig)<br>
 `telemetryStoreId` | `string`<br>
 `streams` | [`StreamConfig`](Streams.md)`[]`<br>
+`artifactTypes` | [`ArtifactTypeAclConfig`](#artifacttypeaclconfig)`[]`<br>
 `acl` | [`AclConfig`](#aclconfig)<br>
 
 ## ConfigInclude
@@ -185,6 +186,13 @@ Name | Description
 `minScratchSpace` | `integer`<br>
 `conformDiskFreeSpace` | `integer`<br>
 
+## ArtifactTypeAclConfig
+
+Name | Description
+---- | -----------
+`type` | `string`<br>
+`acl` | [`AclConfig`](#aclconfig)<br>
+
 ## AclConfig
 
 Name | Description
diff --git a/Engine/Source/Programs/Horde/Docs/Config/Schema/Streams.md b/Engine/Source/Programs/Horde/Docs/Config/Schema/Streams.md
--- a/Engine/Source/Programs/Horde/Docs/Config/Schema/Streams.md
+++ b/Engine/Source/Programs/Horde/Docs/Config/Schema/Streams.md
@@ -32,6 +32,7 @@ Name | Description
 `replicators` | [`ReplicatorConfig`](#replicatorconfig)`[]`<br>
 `workflows` | [`WorkflowConfig`](#workflowconfig)`[]`<br>
 `tokens` | [`TokenConfig`](#tokenconfig)`[]`<br>
+`artifactTypes` | [`ArtifactTypeAclConfig`](#artifacttypeaclconfig)`[]`<br>
 
 ## ConfigInclude
 
@@ -426,3 +427,10 @@ Name | Description
 `projectKey` | `string`<br>
 `defaultComponentId` | `string`<br>
 `defaultIssueTypeId` | `string`<br>
+
+## ArtifactTypeAclConfig
+
+Name | Description
+---- | -----------
+`type` | `string`<br>
+`acl` | [`AclConfig`](#aclconfig)<br>
