Horde: Rename ILogEvent to ILogAnchor. "Log event" is an overloaded term, and "anchor" more accurately describes its use within a log document as an indexed point of interest.

#rnx

diff --git a/Engine/Source/Programs/Horde/Plugins/Build/HordeServer.Build.Tests/Issues/IssueServiceTests.cs b/Engine/Source/Programs/Horde/Plugins/Build/HordeServer.Build.Tests/Issues/IssueServiceTests.cs
--- a/Engine/Source/Programs/Horde/Plugins/Build/HordeServer.Build.Tests/Issues/IssueServiceTests.cs
+++ b/Engine/Source/Programs/Horde/Plugins/Build/HordeServer.Build.Tests/Issues/IssueServiceTests.cs
@@ -906,9 +906,9 @@ namespace HordeServer.Tests.Issues
 
 				ILog? log = await LogCollection.GetAsync(job.Batches[0].Steps[0].LogId!.Value, CancellationToken.None);
 				Assert.IsNotNull(log);
-				List<ILogEvent> events = await log.GetEventsAsync();
-				Assert.AreEqual(1, events.Count);
-				Assert.AreEqual(2, events[0].LineCount);
+				List<ILogAnchor> anchors = await log.GetAnchorsAsync();
+				Assert.AreEqual(1, anchors.Count);
+				Assert.AreEqual(2, anchors[0].LineCount);
 
 				IReadOnlyList<IIssue> issues = await IssueCollection.FindIssuesAsync();
 				Assert.AreEqual(1, issues.Count);
@@ -950,9 +950,9 @@ namespace HordeServer.Tests.Issues
 				ILog? log = await LogCollection.GetAsync(job.Batches[0].Steps[0].LogId!.Value, CancellationToken.None);
 				Assert.IsNotNull(log);
 
-				List<ILogEvent> events = await log.GetEventsAsync();
-				Assert.AreEqual(1, events.Count);
-				Assert.AreEqual(1, events[0].LineCount);
+				List<ILogAnchor> anchors = await log.GetAnchorsAsync();
+				Assert.AreEqual(1, anchors.Count);
+				Assert.AreEqual(1, anchors[0].LineCount);
 
 				IReadOnlyList<IIssue> issues = await IssueCollection.FindIssuesAsync();
 				Assert.AreEqual(1, issues.Count);
diff --git a/Engine/Source/Programs/Horde/Plugins/Build/HordeServer.Build/Issues/IssueService.cs b/Engine/Source/Programs/Horde/Plugins/Build/HordeServer.Build/Issues/IssueService.cs
--- a/Engine/Source/Programs/Horde/Plugins/Build/HordeServer.Build/Issues/IssueService.cs
+++ b/Engine/Source/Programs/Horde/Plugins/Build/HordeServer.Build/Issues/IssueService.cs
@@ -148,13 +148,13 @@ namespace HordeServer.Issues
 	{
 		class IssueEventInternal : IssueEvent
 		{
-			public ILogEvent Event { get; }
+			public ILogAnchor Anchor { get; }
 			public ILogEventData EventData { get; }
 
-			public IssueEventInternal(ILogEvent logEvent, ILogEventData logEventData)
-				: base(logEvent.LineIndex, GetLogLevelFromSeverity(logEventData.Severity), logEventData.EventId, logEventData.Lines)
+			public IssueEventInternal(ILogAnchor anchor, ILogEventData logEventData)
+				: base(anchor.LineIndex, GetLogLevelFromSeverity(logEventData.Severity), logEventData.EventId, logEventData.Lines)
 			{
-				Event = logEvent;
+				Anchor = anchor;
 				EventData = logEventData;
 			}
 
@@ -702,12 +702,12 @@ namespace HordeServer.Issues
 			handlers.SortBy(x => -x.Priority);
 
 			// Create all the issue definitions by passing each log event to the handlers in order until one attaches it to an issue
-			List <ILogEvent> stepEvents = await log.GetEventsAsync(cancellationToken: cancellationToken);
-			foreach (ILogEvent stepEvent in stepEvents)
+			List <ILogAnchor> stepAnchors = await log.GetAnchorsAsync(cancellationToken: cancellationToken);
+			foreach (ILogAnchor stepAnchor in stepAnchors)
 			{
-				ILogEventData stepEventData = await stepEvent.GetDataAsync(cancellationToken);
+				ILogEventData stepEventData = await stepAnchor.GetDataAsync(cancellationToken);
 
-				IssueEventInternal issueEvent = new IssueEventInternal(stepEvent, stepEventData);
+				IssueEventInternal issueEvent = new IssueEventInternal(stepAnchor, stepEventData);
 				foreach (IssueHandler handler in handlers)
 				{
 					if (handler.HandleEvent(issueEvent))
@@ -749,7 +749,7 @@ namespace HordeServer.Issues
 			}
 
 			// Print the list of new events
-			_logger.LogInformation("UpdateCompleteStep({JobId}, {BatchId}, {StepId}): {NumEvents} events, {NumFingerprints} unique fingerprints", job.Id, batch.Id, step.Id, stepEvents.Count, eventGroups.Count);
+			_logger.LogInformation("UpdateCompleteStep({JobId}, {BatchId}, {StepId}): {NumEvents} events, {NumFingerprints} unique fingerprints", job.Id, batch.Id, step.Id, stepAnchors.Count, eventGroups.Count);
 			foreach (IssueEventGroupInternal eventGroup in eventGroups)
 			{
 				_logger.LogInformation("Group {Digest}: Type '{FingerprintType}', keys '{FingerprintKeys}', {NumEvents} events", eventGroup.Id.ToString(), eventGroup.Fingerprint.Type, String.Join(", ", eventGroup.Fingerprint.Keys), eventGroup.Events.Count);
@@ -809,7 +809,7 @@ namespace HordeServer.Issues
 						}
 
 						// Assign all the events to the span
-						await _logCollection.AddSpanToEventsAsync(matchEventGroups.SelectMany(x => x.Events.Select(x => x.Event)), newSpan.Id, cancellationToken);
+						await _logCollection.AddSpanToEventsAsync(matchEventGroups.SelectMany(x => x.Events.Select(x => x.Anchor)), newSpan.Id, cancellationToken);
 
 						// Remove the matches from the set of events
 						newEventGroups.ExceptWith(matchEventGroups);
@@ -907,7 +907,7 @@ namespace HordeServer.Issues
 
 				// Update the log events
 				_logger.LogDebug("Created new span {SpanId} from event group {Group}", newSpan.Id, eventGroup.Id.ToString());
-				await _logCollection.AddSpanToEventsAsync(eventGroup.Events.Select(x => x.Event), newSpan.Id, cancellationToken);
+				await _logCollection.AddSpanToEventsAsync(eventGroup.Events.Select(x => x.Anchor), newSpan.Id, cancellationToken);
 
 				// Remove the events from the remaining list of events to match
 				newEventGroups.ExceptWith(sourceEventGroups);
diff --git a/Engine/Source/Programs/Horde/Plugins/Build/HordeServer.Build/Issues/IssuesController.cs b/Engine/Source/Programs/Horde/Plugins/Build/HordeServer.Build/Issues/IssuesController.cs
--- a/Engine/Source/Programs/Horde/Plugins/Build/HordeServer.Build/Issues/IssuesController.cs
+++ b/Engine/Source/Programs/Horde/Plugins/Build/HordeServer.Build/Issues/IssuesController.cs
@@ -737,24 +737,24 @@ namespace HordeServer.Issues
 			}
 
 			IReadOnlyList<IIssueSpan> spans = await _issueCollection.FindSpansAsync(issueId, cancellationToken);
-			IReadOnlyList<ILogEvent> events = await _logCollection.FindEventsForSpansAsync(spans.Select(x => x.Id), logIdValues.ToArray(), index, count, cancellationToken);
+			IReadOnlyList<ILogAnchor> anchors = await _logCollection.FindAnchorsForSpansAsync(spans.Select(x => x.Id), logIdValues.ToArray(), index, count, cancellationToken);
 
 			JobPermissionsCache permissionsCache = new JobPermissionsCache();
 			Dictionary<LogId, ILog?> logs = new Dictionary<LogId, ILog?>();
 
 			List<object> responses = new List<object>();
-			foreach (ILogEvent logEvent in events)
+			foreach (ILogAnchor anchor in anchors)
 			{
 				ILog? log;
-				if (!logs.TryGetValue(logEvent.LogId, out log))
+				if (!logs.TryGetValue(anchor.LogId, out log))
 				{
-					log = await _logCollection.GetAsync(logEvent.LogId, cancellationToken);
-					logs[logEvent.LogId] = log;
+					log = await _logCollection.GetAsync(anchor.LogId, cancellationToken);
+					logs[anchor.LogId] = log;
 				}
 				if (log != null && await _jobService.AuthorizeAsync(log.JobId, LogAclAction.ViewLog, User, _buildConfig.Value, cancellationToken))
 				{
-					ILogEventData data = await logEvent.GetDataAsync(cancellationToken);
-					GetLogEventResponse response = LogsController.CreateGetLogEventResponse(logEvent, data, issueId);
+					ILogEventData data = await anchor.GetDataAsync(cancellationToken);
+					GetLogEventResponse response = LogsController.CreateGetLogEventResponse(anchor, data, issueId);
 					responses.Add(PropertyFilter.Apply(response, filter));
 				}
 			}
diff --git a/Engine/Source/Programs/Horde/Plugins/Build/HordeServer.Build/Notifications/NotificationService.cs b/Engine/Source/Programs/Horde/Plugins/Build/HordeServer.Build/Notifications/NotificationService.cs
--- a/Engine/Source/Programs/Horde/Plugins/Build/HordeServer.Build/Notifications/NotificationService.cs
+++ b/Engine/Source/Programs/Horde/Plugins/Build/HordeServer.Build/Notifications/NotificationService.cs
@@ -685,9 +685,9 @@ namespace HordeServer.Notifications
 				return;
 			}
 
-			List<ILogEvent> jobStepEvents = await log.GetEventsAsync(cancellationToken: cancellationToken);
+			List<ILogAnchor> jobStepEvents = await log.GetAnchorsAsync(cancellationToken: cancellationToken);
 			List<ILogEventData> jobStepEventData = new List<ILogEventData>();
-			foreach (ILogEvent logEvent in jobStepEvents)
+			foreach (ILogAnchor logEvent in jobStepEvents)
 			{
 				ILogEventData eventData = await logEvent.GetDataAsync(cancellationToken);
 				jobStepEventData.Add(eventData);
diff --git a/Engine/Source/Programs/Horde/Plugins/Build/HordeServer.Build/Notifications/Sinks/SlackNotificationSink.cs b/Engine/Source/Programs/Horde/Plugins/Build/HordeServer.Build/Notifications/Sinks/SlackNotificationSink.cs
--- a/Engine/Source/Programs/Horde/Plugins/Build/HordeServer.Build/Notifications/Sinks/SlackNotificationSink.cs
+++ b/Engine/Source/Programs/Horde/Plugins/Build/HordeServer.Build/Notifications/Sinks/SlackNotificationSink.cs
@@ -1012,7 +1012,7 @@ namespace HordeServer.Notifications.Sinks
 			if (isNew)
 			{
 				// Create the summary text
-				List<ILogEvent> events = new List<ILogEvent>();
+				List<ILogAnchor> anchors = new List<ILogAnchor>();
 				List<ILogEventData> eventDataItems = new List<ILogEventData>();
 
 				if (span.FirstFailure.LogId != null)
@@ -1021,16 +1021,16 @@ namespace HordeServer.Notifications.Sinks
 					ILog? log = await _logCollection.GetAsync(logId, cancellationToken);
 					if (log != null)
 					{
-						events = await log.GetEventsAsync(span.Id, 0, 50, cancellationToken);
-						if (events.Any(x => x.Severity == LogEventSeverity.Error))
+						anchors = await log.GetAnchorsAsync(span.Id, 0, 50, cancellationToken);
+						if (anchors.Any(x => x.Severity == LogEventSeverity.Error))
 						{
-							events.RemoveAll(x => x.Severity == LogEventSeverity.Warning);
+							anchors.RemoveAll(x => x.Severity == LogEventSeverity.Warning);
 						}
 
 						List<string> eventStrings = new List<string>();
-						for (int idx = 0; idx < Math.Min(events.Count, 3); idx++)
+						for (int idx = 0; idx < Math.Min(anchors.Count, 3); idx++)
 						{
-							ILogEventData data = await events[idx].GetDataAsync(cancellationToken);
+							ILogEventData data = await anchors[idx].GetDataAsync(cancellationToken);
 							eventDataItems.Add(data);
 						}
 					}
@@ -1043,7 +1043,7 @@ namespace HordeServer.Notifications.Sinks
 				{
 					message.Blocks.Add(new SectionBlock(new TextObject(QuoteText(eventDataItem.Message, MaxLineLength))));
 				}
-				if (events.Count > eventDataItems.Count)
+				if (anchors.Count > eventDataItems.Count)
 				{
 					message.Blocks.Add(new SectionBlock(new TextObject("```...```")));
 				}
@@ -1457,7 +1457,7 @@ namespace HordeServer.Notifications.Sinks
 				ILog? log = await _logCollection.GetAsync(logId, cancellationToken);
 				if (log != null)
 				{
-					List<ILogEvent> events = await log.GetEventsAsync(lastSpan.Id, 0, 20, cancellationToken);
+					List<ILogAnchor> events = await log.GetAnchorsAsync(lastSpan.Id, 0, 20, cancellationToken);
 					if (events.Any(x => x.Severity == LogEventSeverity.Error))
 					{
 						events.RemoveAll(x => x.Severity == LogEventSeverity.Warning);
diff --git a/Engine/Source/Programs/Horde/Plugins/Build/HordeServer.Build/Ugs/UgsController.cs b/Engine/Source/Programs/Horde/Plugins/Build/HordeServer.Build/Ugs/UgsController.cs
--- a/Engine/Source/Programs/Horde/Plugins/Build/HordeServer.Build/Ugs/UgsController.cs
+++ b/Engine/Source/Programs/Horde/Plugins/Build/HordeServer.Build/Ugs/UgsController.cs
@@ -250,9 +250,9 @@ namespace HordeServer.Ugs
 			List<GetUgsIssueDiagnosticResponse> diagnostics = new List<GetUgsIssueDiagnosticResponse>();
 
 			IReadOnlyList<IIssueSpan> spans = await _issueService.Collection.FindSpansAsync(issueId, cancellationToken);
-			IReadOnlyList<ILogEvent> events = await _logCollection.FindEventsForSpansAsync(spans.Select(x => x.Id), null, 0, count: 10, cancellationToken);
+			IReadOnlyList<ILogAnchor> events = await _logCollection.FindAnchorsForSpansAsync(spans.Select(x => x.Id), null, 0, count: 10, cancellationToken);
 
-			foreach (ILogEvent logEvent in events)
+			foreach (ILogAnchor logEvent in events)
 			{
 				ILogEventData eventData = await logEvent.GetDataAsync(cancellationToken);
 				long buildId = logEvent.LogId.GetHashCode();
diff --git a/Engine/Source/Programs/Horde/Plugins/Compute/HordeServer.Compute/Logs/ILog.cs b/Engine/Source/Programs/Horde/Plugins/Compute/HordeServer.Compute/Logs/ILog.cs
--- a/Engine/Source/Programs/Horde/Plugins/Compute/HordeServer.Compute/Logs/ILog.cs
+++ b/Engine/Source/Programs/Horde/Plugins/Compute/HordeServer.Compute/Logs/ILog.cs
@@ -148,15 +148,15 @@ namespace HordeServer.Logs
 		/// <param name="count">Maximum number of results to return</param>
 		/// <param name="cancellationToken">Cancellation token for the operation</param>
 		/// <returns>List of events matching the query</returns>
-		Task<List<ILogEvent>> GetEventsAsync(ObjectId? spanId = null, int? index = null, int? count = null, CancellationToken cancellationToken = default);
+		Task<List<ILogAnchor>> GetAnchorsAsync(ObjectId? spanId = null, int? index = null, int? count = null, CancellationToken cancellationToken = default);
 
 		#endregion
 	}
 
 	/// <summary>
-	/// Represents a node in the graph
+	/// Anchor within a log files
 	/// </summary>
-	public interface ILogEvent
+	public interface ILogAnchor
 	{
 		/// <summary>
 		/// Unique id of the log containing this event
diff --git a/Engine/Source/Programs/Horde/Plugins/Compute/HordeServer.Compute/Logs/ILogCollection.cs b/Engine/Source/Programs/Horde/Plugins/Compute/HordeServer.Compute/Logs/ILogCollection.cs
--- a/Engine/Source/Programs/Horde/Plugins/Compute/HordeServer.Compute/Logs/ILogCollection.cs
+++ b/Engine/Source/Programs/Horde/Plugins/Compute/HordeServer.Compute/Logs/ILogCollection.cs
@@ -49,7 +49,7 @@ namespace HordeServer.Logs
 		/// <param name="count">Number of results to return</param>
 		/// <param name="cancellationToken">Cancellation token for the call</param>
 		/// <returns>List of events for this issue</returns>
-		Task<IReadOnlyList<ILogEvent>> FindEventsForSpansAsync(IEnumerable<ObjectId> spanIds, LogId[]? logIds = null, int index = 0, int count = 10, CancellationToken cancellationToken = default);
+		Task<IReadOnlyList<ILogAnchor>> FindAnchorsForSpansAsync(IEnumerable<ObjectId> spanIds, LogId[]? logIds = null, int index = 0, int count = 10, CancellationToken cancellationToken = default);
 
 		/// <summary>
 		/// Update the span for an event
@@ -58,7 +58,7 @@ namespace HordeServer.Logs
 		/// <param name="spanId">New span id</param>
 		/// <param name="cancellationToken">Cancellation token for the call</param>
 		/// <returns>Async task</returns>
-		Task AddSpanToEventsAsync(IEnumerable<ILogEvent> events, ObjectId spanId, CancellationToken cancellationToken = default);
+		Task AddSpanToEventsAsync(IEnumerable<ILogAnchor> events, ObjectId spanId, CancellationToken cancellationToken = default);
 
 		#endregion
 	}
diff --git a/Engine/Source/Programs/Horde/Plugins/Compute/HordeServer.Compute/Logs/LogCollection.cs b/Engine/Source/Programs/Horde/Plugins/Compute/HordeServer.Compute/Logs/LogCollection.cs
--- a/Engine/Source/Programs/Horde/Plugins/Compute/HordeServer.Compute/Logs/LogCollection.cs
+++ b/Engine/Source/Programs/Horde/Plugins/Compute/HordeServer.Compute/Logs/LogCollection.cs
@@ -84,8 +84,8 @@ namespace HordeServer.Logs
 			public Task AddEventsAsync(List<NewLogEventData> newEvents, CancellationToken cancellationToken = default)
 				=> _collection.AddEventsAsync(_document.Id, newEvents, cancellationToken);
 
-			public Task<List<ILogEvent>> GetEventsAsync(ObjectId? spanId = null, int? index = null, int? count = null, CancellationToken cancellationToken = default)
-				=> _collection.GetEventsAsync(this, spanId, index, count, cancellationToken);
+			public Task<List<ILogAnchor>> GetAnchorsAsync(ObjectId? spanId = null, int? index = null, int? count = null, CancellationToken cancellationToken = default)
+				=> _collection.GetAnchorsAsync(this, spanId, index, count, cancellationToken);
 		}
 
 		class LogDocument
@@ -173,21 +173,21 @@ namespace HordeServer.Logs
 			}
 		}
 
-		class LogEvent : ILogEvent
+		class LogAnchor : ILogAnchor
 		{
 			readonly ILog _log;
 			readonly LogCollection _collection;
-			readonly LogEventDocument _document;
+			readonly LogAnchorDocument _document;
 
-			public LogEventDocument Document => _document;
+			public LogAnchorDocument Document => _document;
 
-			LogId ILogEvent.LogId => _document.Id.LogId;
-			LogEventSeverity ILogEvent.Severity => _document.IsWarning ? LogEventSeverity.Warning : LogEventSeverity.Error;
-			int ILogEvent.LineIndex => _document.Id.LineIndex;
-			int ILogEvent.LineCount => _document.LineCount ?? 1;
-			ObjectId? ILogEvent.SpanId => _document.SpanId;
+			LogId ILogAnchor.LogId => _document.Id.LogId;
+			LogEventSeverity ILogAnchor.Severity => _document.IsWarning ? LogEventSeverity.Warning : LogEventSeverity.Error;
+			int ILogAnchor.LineIndex => _document.Id.LineIndex;
+			int ILogAnchor.LineCount => _document.LineCount ?? 1;
+			ObjectId? ILogAnchor.SpanId => _document.SpanId;
 
-			public LogEvent(ILog log, LogCollection collection, LogEventDocument document)
+			public LogAnchor(ILog log, LogCollection collection, LogAnchorDocument document)
 			{
 				_log = log;
 				_collection = collection;
@@ -198,7 +198,7 @@ namespace HordeServer.Logs
 				=> await _collection.GetEventDataAsync(_log, _document.Id.LineIndex, _document.LineCount ?? 1, cancellationToken);
 		}
 
-		class LogEventId
+		class LogAnchorId
 		{
 			[BsonElement("l")]
 			public LogId LogId { get; set; }
@@ -207,10 +207,10 @@ namespace HordeServer.Logs
 			public int LineIndex { get; set; }
 		}
 
-		class LogEventDocument
+		class LogAnchorDocument
 		{
 			[BsonId]
-			public LogEventId Id { get; set; }
+			public LogAnchorId Id { get; set; }
 
 			[BsonElement("w"), BsonIgnoreIfDefault, BsonDefaultValue(false)]
 			public bool IsWarning { get; set; }
@@ -222,20 +222,20 @@ namespace HordeServer.Logs
 			public ObjectId? SpanId { get; set; }
 
 			[BsonConstructor]
-			public LogEventDocument()
+			public LogAnchorDocument()
 			{
-				Id = new LogEventId();
+				Id = new LogAnchorId();
 			}
 
-			public LogEventDocument(LogId logId, LogEventSeverity severity, int lineIndex, int lineCount, ObjectId? spanId)
+			public LogAnchorDocument(LogId logId, LogEventSeverity severity, int lineIndex, int lineCount, ObjectId? spanId)
 			{
-				Id = new LogEventId { LogId = logId, LineIndex = lineIndex };
+				Id = new LogAnchorId { LogId = logId, LineIndex = lineIndex };
 				IsWarning = severity == LogEventSeverity.Warning;
 				LineCount = (lineCount > 1) ? (int?)lineCount : null;
 				SpanId = spanId;
 			}
 
-			public LogEventDocument(LogId logId, NewLogEventData data)
+			public LogAnchorDocument(LogId logId, NewLogEventData data)
 				: this(logId, data.Severity, data.LineIndex, data.LineCount, data.SpanId)
 			{
 			}
@@ -261,7 +261,7 @@ namespace HordeServer.Logs
 		}
 
 		readonly IMongoCollection<LogDocument> _logCollection;
-		readonly IMongoCollection<LogEventDocument> _logEvents;
+		readonly IMongoCollection<LogAnchorDocument> _logEvents;
 		readonly IMongoCollection<LegacyLogEventDocument> _legacyEvents;
 		readonly LogTailService _logTailService;
 		readonly StorageService _storageService;
@@ -281,10 +281,10 @@ namespace HordeServer.Logs
 			_logger = logger;
 			_logCache = new MemoryCache(new MemoryCacheOptions());
 
-			List<MongoIndex<LogEventDocument>> logEventIndexes = new List<MongoIndex<LogEventDocument>>();
+			List<MongoIndex<LogAnchorDocument>> logEventIndexes = new List<MongoIndex<LogAnchorDocument>>();
 			logEventIndexes.Add(keys => keys.Ascending(x => x.Id.LogId));
 			logEventIndexes.Add(keys => keys.Ascending(x => x.SpanId).Ascending(x => x.Id));
-			_logEvents = mongoService.GetCollection<LogEventDocument>("LogEvents", logEventIndexes);
+			_logEvents = mongoService.GetCollection<LogAnchorDocument>("LogEvents", logEventIndexes);
 
 			_legacyEvents = mongoService.GetCollection<LegacyLogEventDocument>("Events", keys => keys.Ascending(x => x.LogId));
 		}
@@ -778,23 +778,23 @@ namespace HordeServer.Logs
 		/// <inheritdoc/>
 		Task AddEventsAsync(LogId logId, List<NewLogEventData> newEvents, CancellationToken cancellationToken)
 		{
-			return _logEvents.InsertManyAsync(newEvents.ConvertAll(x => new LogEventDocument(logId, x)), cancellationToken: cancellationToken);
+			return _logEvents.InsertManyAsync(newEvents.ConvertAll(x => new LogAnchorDocument(logId, x)), cancellationToken: cancellationToken);
 		}
 
 		/// <inheritdoc/>
-		async Task<List<ILogEvent>> GetEventsAsync(ILog log, ObjectId? spanId = null, int? index = null, int? count = null, CancellationToken cancellationToken = default)
+		async Task<List<ILogAnchor>> GetAnchorsAsync(ILog log, ObjectId? spanId = null, int? index = null, int? count = null, CancellationToken cancellationToken = default)
 		{
 			_logger.LogInformation("Querying for log events for log {LogId}", log.Id);
 
-			FilterDefinitionBuilder<LogEventDocument> builder = Builders<LogEventDocument>.Filter;
+			FilterDefinitionBuilder<LogAnchorDocument> builder = Builders<LogAnchorDocument>.Filter;
 
-			FilterDefinition<LogEventDocument> filter = builder.Eq(x => x.Id.LogId, log.Id);
+			FilterDefinition<LogAnchorDocument> filter = builder.Eq(x => x.Id.LogId, log.Id);
 			if (spanId != null)
 			{
 				filter &= builder.Eq(x => x.SpanId, spanId.Value);
 			}
 
-			IFindFluent<LogEventDocument, LogEventDocument> results = _logEvents.Find(filter).SortBy(x => x.Id);
+			IFindFluent<LogAnchorDocument, LogAnchorDocument> results = _logEvents.Find(filter).SortBy(x => x.Id);
 			if (index != null)
 			{
 				results = results.Skip(index.Value);
@@ -804,28 +804,28 @@ namespace HordeServer.Logs
 				results = results.Limit(count.Value);
 			}
 
-			List<LogEventDocument> logEventDocuments = await results.ToListAsync(cancellationToken);
-			return logEventDocuments.ConvertAll<ILogEvent>(x => new LogEvent(log, this, x));
+			List<LogAnchorDocument> logEventDocuments = await results.ToListAsync(cancellationToken);
+			return logEventDocuments.ConvertAll<ILogAnchor>(x => new LogAnchor(log, this, x));
 		}
 
 		/// <inheritdoc/>
-		public async Task<IReadOnlyList<ILogEvent>> FindEventsForSpansAsync(IEnumerable<ObjectId> spanIds, LogId[]? logIds, int index, int count, CancellationToken cancellationToken)
+		public async Task<IReadOnlyList<ILogAnchor>> FindAnchorsForSpansAsync(IEnumerable<ObjectId> spanIds, LogId[]? logIds, int index, int count, CancellationToken cancellationToken)
 		{
-			FilterDefinition<LogEventDocument> filter = Builders<LogEventDocument>.Filter.In(x => x.SpanId, spanIds.Select<ObjectId, ObjectId?>(x => x));
+			FilterDefinition<LogAnchorDocument> filter = Builders<LogAnchorDocument>.Filter.In(x => x.SpanId, spanIds.Select<ObjectId, ObjectId?>(x => x));
 			if (logIds != null && logIds.Length > 0)
 			{
-				filter &= Builders<LogEventDocument>.Filter.In(x => x.Id.LogId, logIds);
+				filter &= Builders<LogAnchorDocument>.Filter.In(x => x.Id.LogId, logIds);
 			}
 
-			List<LogEvent> logEvents = new List<LogEvent>();
+			List<LogAnchor> logEvents = new List<LogAnchor>();
 
-			List<LogEventDocument> logEventDocuments = await _logEvents.Find(filter).Skip(index).Limit(count).ToListAsync(cancellationToken);
-			foreach (IGrouping<LogId, LogEventDocument> logEventGroup in logEventDocuments.GroupBy(x => x.Id.LogId))
+			List<LogAnchorDocument> logEventDocuments = await _logEvents.Find(filter).Skip(index).Limit(count).ToListAsync(cancellationToken);
+			foreach (IGrouping<LogId, LogAnchorDocument> logEventGroup in logEventDocuments.GroupBy(x => x.Id.LogId))
 			{
 				ILog? log = await GetAsync(logEventGroup.Key, cancellationToken);
 				if (log != null)
 				{
-					logEvents.AddRange(logEventGroup.Select(x => new LogEvent(log, this, x)));
+					logEvents.AddRange(logEventGroup.Select(x => new LogAnchor(log, this, x)));
 				}
 			}
 
@@ -833,10 +833,10 @@ namespace HordeServer.Logs
 		}
 
 		/// <inheritdoc/>
-		public async Task AddSpanToEventsAsync(IEnumerable<ILogEvent> events, ObjectId spanId, CancellationToken cancellationToken)
+		public async Task AddSpanToEventsAsync(IEnumerable<ILogAnchor> events, ObjectId spanId, CancellationToken cancellationToken)
 		{
-			FilterDefinition<LogEventDocument> eventFilter = Builders<LogEventDocument>.Filter.In(x => x.Id, events.Select(x => ((LogEvent)x).Document.Id));
-			UpdateDefinition<LogEventDocument> eventUpdate = Builders<LogEventDocument>.Update.Set(x => x.SpanId, spanId);
+			FilterDefinition<LogAnchorDocument> eventFilter = Builders<LogAnchorDocument>.Filter.In(x => x.Id, events.Select(x => ((LogAnchor)x).Document.Id));
+			UpdateDefinition<LogAnchorDocument> eventUpdate = Builders<LogAnchorDocument>.Update.Set(x => x.SpanId, spanId);
 			await _logEvents.UpdateManyAsync(eventFilter, eventUpdate, cancellationToken: cancellationToken);
 		}
 
diff --git a/Engine/Source/Programs/Horde/Plugins/Compute/HordeServer.Compute/Logs/LogsController.cs b/Engine/Source/Programs/Horde/Plugins/Compute/HordeServer.Compute/Logs/LogsController.cs
--- a/Engine/Source/Programs/Horde/Plugins/Compute/HordeServer.Compute/Logs/LogsController.cs
+++ b/Engine/Source/Programs/Horde/Plugins/Compute/HordeServer.Compute/Logs/LogsController.cs
@@ -324,26 +324,26 @@ namespace HordeServer.Logs
 				return Forbid();
 			}
 
-			List<ILogEvent> logEvents = await log.GetEventsAsync(null, index, count, cancellationToken);
+			List<ILogAnchor> anchors = await log.GetAnchorsAsync(null, index, count, cancellationToken);
 
 			Dictionary<ObjectId, int?> spanIdToIssueId = new Dictionary<ObjectId, int?>();
 
 			List<GetLogEventResponse> responses = new List<GetLogEventResponse>();
-			foreach (ILogEvent logEvent in logEvents)
+			foreach (ILogAnchor anchor in anchors)
 			{
-				ILogEventData logEventData = await logEvent.GetDataAsync(cancellationToken);
+				ILogEventData logEventData = await anchor.GetDataAsync(cancellationToken);
 
 				int? issueId = null;
-				if (logEvent.SpanId != null && !spanIdToIssueId.TryGetValue(logEvent.SpanId.Value, out issueId))
+				if (anchor.SpanId != null && !spanIdToIssueId.TryGetValue(anchor.SpanId.Value, out issueId))
 				{
 					foreach (ILogExtIssueProvider issueProvider in _issueProviders)
 					{
-						issueId ??= await issueProvider.GetIssueIdAsync(logEvent.SpanId.Value, cancellationToken);
+						issueId ??= await issueProvider.GetIssueIdAsync(anchor.SpanId.Value, cancellationToken);
 					}
-					spanIdToIssueId[logEvent.SpanId.Value] = issueId;
+					spanIdToIssueId[anchor.SpanId.Value] = issueId;
 				}
 
-				responses.Add(CreateGetLogEventResponse(logEvent, logEventData, issueId));
+				responses.Add(CreateGetLogEventResponse(anchor, logEventData, issueId));
 			}
 			return responses;
 		}
@@ -351,16 +351,16 @@ namespace HordeServer.Logs
 		/// <summary>
 		/// Create a log event response message
 		/// </summary>
-		/// <param name="logEvent">The event to construct from</param>
+		/// <param name="anchor">The event to construct from</param>
 		/// <param name="eventData">The event data</param>
 		/// <param name="issueId">The issue for this event</param>
-		public static GetLogEventResponse CreateGetLogEventResponse(ILogEvent logEvent, ILogEventData eventData, int? issueId)
+		public static GetLogEventResponse CreateGetLogEventResponse(ILogAnchor anchor, ILogEventData eventData, int? issueId)
 		{
 			GetLogEventResponse response = new GetLogEventResponse();
-			response.Severity = logEvent.Severity;
-			response.LogId = logEvent.LogId;
-			response.LineIndex = logEvent.LineIndex;
-			response.LineCount = logEvent.LineCount;
+			response.Severity = anchor.Severity;
+			response.LogId = anchor.LogId;
+			response.LineIndex = anchor.LineIndex;
+			response.LineCount = anchor.LineCount;
 			response.IssueId = issueId;
 			response.Lines.AddRange(eventData.Lines.Select(x => JsonDocument.Parse(x.Data).RootElement));
 			return response;
