Horde: Fix issue with missing agent preventing lease log from loading
#rnx
#skipci

diff --git a/Engine/Source/Programs/Horde/HordeDashboard/src/components/LogSource.tsx b/Engine/Source/Programs/Horde/HordeDashboard/src/components/LogSource.tsx
--- a/Engine/Source/Programs/Horde/HordeDashboard/src/components/LogSource.tsx
+++ b/Engine/Source/Programs/Horde/HordeDashboard/src/components/LogSource.tsx
@@ -652,7 +652,15 @@ class LeaseLogSource extends LogSource {
 
                this.lease = values[0];
 
-               this.agent = await backend.getAgent(this.lease.agentId!);
+               try {
+                  this.agent = await backend.getAgent(this.lease.agentId!);
+               } catch (reason) {
+                  this.agent = {
+                     id: "missing-agent",
+                     name: "Missing Agent"
+                  } as any;
+               }
+               
 
                this.leaseUpdated();
                this.initComplete();
