Horde: Fix issue with semicolon in render
#rnx
#skipci

diff --git a/Engine/Source/Programs/Horde/HordeDashboard/src/components/tools/ToolManagement.tsx b/Engine/Source/Programs/Horde/HordeDashboard/src/components/tools/ToolManagement.tsx
--- a/Engine/Source/Programs/Horde/HordeDashboard/src/components/tools/ToolManagement.tsx
+++ b/Engine/Source/Programs/Horde/HordeDashboard/src/components/tools/ToolManagement.tsx
@@ -297,7 +297,7 @@ export const ToolManagementModal: React.FC<{ toolId: string; toolName?: string,
             <Stack>
                <Text variant="mediumPlus">{headerText}</Text>
             </Stack>
-            <Stack grow />;
+            <Stack grow />
             <Stack>
                <IconButton
                   iconProps={{ iconName: 'Cancel' }}
