Horde: Refactor current plugins to legacy plugins
#rnx

diff --git a/Engine/Source/Programs/Horde/HordeDashboard/src/App.tsx b/Engine/Source/Programs/Horde/HordeDashboard/src/App.tsx
--- a/Engine/Source/Programs/Horde/HordeDashboard/src/App.tsx
+++ b/Engine/Source/Programs/Horde/HordeDashboard/src/App.tsx
@@ -3,7 +3,7 @@
 import { Image, Spinner, SpinnerSize, Stack, Text, ThemeProvider } from '@fluentui/react';
 import React, { useState } from 'react';
 import { Navigate, Outlet, RouteObject, RouterProvider, createBrowserRouter } from 'react-router-dom';
-import hordePlugins from './Plugins';
+import hordePlugins from './legacyPlugins';
 import backend from './backend';
 import { getSiteConfig } from './backend/Config';
 import dashboard from './backend/Dashboard';
diff --git a/Engine/Source/Programs/Horde/HordeDashboard/src/Plugins/AutomatedTestSession/components/AutomatedTestSession.tsx b/Engine/Source/Programs/Horde/HordeDashboard/src/Plugins/AutomatedTestSession/components/AutomatedTestSession.tsx
deleted file mode 100644
diff --git a/Engine/Source/Programs/Horde/HordeDashboard/src/Plugins/AutomatedTestSession/index.tsx b/Engine/Source/Programs/Horde/HordeDashboard/src/Plugins/AutomatedTestSession/index.tsx
deleted file mode 100644
diff --git a/Engine/Source/Programs/Horde/HordeDashboard/src/Plugins/AutomatedTestSession/models/AutomatedTestSession.ts b/Engine/Source/Programs/Horde/HordeDashboard/src/Plugins/AutomatedTestSession/models/AutomatedTestSession.ts
deleted file mode 100644
diff --git a/Engine/Source/Programs/Horde/HordeDashboard/src/Plugins/SimpleTestReport/components/SimpleReport.tsx b/Engine/Source/Programs/Horde/HordeDashboard/src/Plugins/SimpleTestReport/components/SimpleReport.tsx
deleted file mode 100644
diff --git a/Engine/Source/Programs/Horde/HordeDashboard/src/Plugins/SimpleTestReport/index.tsx b/Engine/Source/Programs/Horde/HordeDashboard/src/Plugins/SimpleTestReport/index.tsx
deleted file mode 100644
diff --git a/Engine/Source/Programs/Horde/HordeDashboard/src/Plugins/SimpleTestReport/models/SimpleReport.ts b/Engine/Source/Programs/Horde/HordeDashboard/src/Plugins/SimpleTestReport/models/SimpleReport.ts
deleted file mode 100644
diff --git a/Engine/Source/Programs/Horde/HordeDashboard/src/Plugins/UnrealAutomatedTests/components/UnrealAutomatedTests.tsx b/Engine/Source/Programs/Horde/HordeDashboard/src/Plugins/UnrealAutomatedTests/components/UnrealAutomatedTests.tsx
deleted file mode 100644
diff --git a/Engine/Source/Programs/Horde/HordeDashboard/src/Plugins/UnrealAutomatedTests/index.tsx b/Engine/Source/Programs/Horde/HordeDashboard/src/Plugins/UnrealAutomatedTests/index.tsx
deleted file mode 100644
diff --git a/Engine/Source/Programs/Horde/HordeDashboard/src/Plugins/UnrealAutomatedTests/models/UnrealAutomatedTests.ts b/Engine/Source/Programs/Horde/HordeDashboard/src/Plugins/UnrealAutomatedTests/models/UnrealAutomatedTests.ts
deleted file mode 100644
diff --git a/Engine/Source/Programs/Horde/HordeDashboard/src/Plugins/example/components/JobDetailPluginExample.tsx b/Engine/Source/Programs/Horde/HordeDashboard/src/Plugins/example/components/JobDetailPluginExample.tsx
deleted file mode 100644
diff --git a/Engine/Source/Programs/Horde/HordeDashboard/src/Plugins/example/index.tsx b/Engine/Source/Programs/Horde/HordeDashboard/src/Plugins/example/index.tsx
deleted file mode 100644
diff --git a/Engine/Source/Programs/Horde/HordeDashboard/src/Plugins/example/routes/ExampleBar.tsx b/Engine/Source/Programs/Horde/HordeDashboard/src/Plugins/example/routes/ExampleBar.tsx
deleted file mode 100644
diff --git a/Engine/Source/Programs/Horde/HordeDashboard/src/Plugins/index.tsx b/Engine/Source/Programs/Horde/HordeDashboard/src/Plugins/index.tsx
deleted file mode 100644
diff --git a/Engine/Source/Programs/Horde/HordeDashboard/src/Plugins/pluginTypes.ts b/Engine/Source/Programs/Horde/HordeDashboard/src/Plugins/pluginTypes.ts
deleted file mode 100644
diff --git a/Engine/Source/Programs/Horde/HordeDashboard/src/components/TestReportPanel.tsx b/Engine/Source/Programs/Horde/HordeDashboard/src/components/TestReportPanel.tsx
--- a/Engine/Source/Programs/Horde/HordeDashboard/src/components/TestReportPanel.tsx
+++ b/Engine/Source/Programs/Horde/HordeDashboard/src/components/TestReportPanel.tsx
@@ -5,8 +5,8 @@ import { mergeStyleSets } from '@fluentui/react/lib/Styling';
 import { Stack, Text } from '@fluentui/react';
 import { Link } from 'react-router-dom';
 import { TestData } from '../backend/Api';
-import { PluginMount } from  '../Plugins/pluginTypes';
-import hordePlugins from  '../Plugins';
+import { PluginMount } from  '../legacyPlugins/pluginTypes';
+import hordePlugins from  '../legacyPlugins';
 import { ComponentMount } from './TestReportView';
 import { getHordeStyling } from '../styles/Styles';
 import { getHordeTheme } from '../styles/theme';
diff --git a/Engine/Source/Programs/Horde/HordeDashboard/src/components/TestReportView.tsx b/Engine/Source/Programs/Horde/HordeDashboard/src/components/TestReportView.tsx
--- a/Engine/Source/Programs/Horde/HordeDashboard/src/components/TestReportView.tsx
+++ b/Engine/Source/Programs/Horde/HordeDashboard/src/components/TestReportView.tsx
@@ -4,8 +4,8 @@ import React, { useEffect, useRef, useState } from 'react';
 import { TopNav } from './TopNav';
 import { Breadcrumbs, BreadcrumbItem } from './Breadcrumbs';
 import { Stack, Spinner, SpinnerSize, Text } from '@fluentui/react';
-import { PluginMount } from  '../Plugins/pluginTypes';
-import hordePlugins from  '../Plugins';
+import { PluginMount } from  '../legacyPlugins/pluginTypes';
+import hordePlugins from  '../legacyPlugins';
 import { projectStore } from '../backend/ProjectStore';
 import { useParams } from 'react-router-dom';
 import { TestDataCollection, TestDataWrapper } from '../backend/TestDataHandler'
diff --git a/Engine/Source/Programs/Horde/HordeDashboard/src/components/jobDetailsV2/StepDetailTestPanel.tsx b/Engine/Source/Programs/Horde/HordeDashboard/src/components/jobDetailsV2/StepDetailTestPanel.tsx
--- a/Engine/Source/Programs/Horde/HordeDashboard/src/components/jobDetailsV2/StepDetailTestPanel.tsx
+++ b/Engine/Source/Programs/Horde/HordeDashboard/src/components/jobDetailsV2/StepDetailTestPanel.tsx
@@ -8,8 +8,8 @@ import { Link } from 'react-router-dom';
 import backend from '../../backend';
 import { TestData } from '../../backend/Api';
 import { ISideRailLink } from '../../base/components/SideRail';
-import { PluginMount } from  '../../Plugins/pluginTypes';
-import hordePlugins from  '../../Plugins';
+import { PluginMount } from  '../../legacyPlugins/pluginTypes';
+import hordePlugins from  '../../legacyPlugins';
 import { ComponentMount } from '../TestReportView';
 import { JobDataView, JobDetailsV2 } from './JobDetailsViewCommon';
 import { getHordeStyling } from '../../styles/Styles';
diff --git a/Engine/Source/Programs/Horde/HordeDashboard/src/legacyPlugins/AutomatedTestSession/components/AutomatedTestSession.tsx b/Engine/Source/Programs/Horde/HordeDashboard/src/legacyPlugins/AutomatedTestSession/components/AutomatedTestSession.tsx
new file mode 100644
--- /dev/null
+++ b/Engine/Source/Programs/Horde/HordeDashboard/src/legacyPlugins/AutomatedTestSession/components/AutomatedTestSession.tsx
@@ -0,0 +1,1440 @@
+// Copyright Epic Games, Inc. All Rights Reserved.
+
+import { ActionButton, Checkbox, DefaultButton, DirectionalHint, Dropdown, HoverCard, HoverCardType, IContextualMenuProps, IDropdownOption, IPlainCardProps, ITextFieldProps, IconButton, Image, Modal, Pivot, PivotItem, ProgressIndicator, Spinner, SpinnerSize, Stack, Text, TextField } from '@fluentui/react';
+import { FontIcon } from '@fluentui/react/lib/Icon';
+import { mergeStyleSets, mergeStyles } from '@fluentui/react/lib/Styling';
+import React, { useEffect, useMemo, useReducer, useRef, useState } from 'react';
+import { generatePath, useNavigate, useParams } from 'react-router';
+import { Link } from 'react-router-dom';
+import dashboard, { StatusColor } from "../../../backend/Dashboard";
+import { msecToElapsed } from '../../../base/utilities/timeUtils';
+import { BreadcrumbItem, Breadcrumbs } from '../../../components/Breadcrumbs';
+import { useQuery } from '../../../components/JobDetailCommon';
+import { getProjectName, getStreamData, testDataHandler } from '../../../components/TestReportView';
+import { TopNav } from '../../../components/TopNav';
+import { getHordeStyling } from '../../../styles/Styles';
+import { getHordeTheme } from '../../../styles/theme';
+import { ArtifactType, EventType, FilterType, Loader, MetaFilterTools, MetaWrapper, Metadata, MetadataFilter, Section, SectionCollection, SectionFilter, TestCase, TestDevice, TestEvent, TestResult, TestSessionCollection, TestSessionWrapper, TestState, TestStats } from '../models/AutomatedTestSession';
+
+const gutterClass = mergeStyles({
+   borderLeftWidth: 6,
+   padding: 0,
+   margin: 0,
+   paddingTop: 0,
+   paddingBottom: 0,
+   paddingRight: 8,
+   marginTop: 0,
+   marginBottom: 0,
+   height: 20
+});
+
+let _styles: any;
+let _stateColorStyles: any;
+
+const getStyling = () => {
+
+   const theme = getHordeTheme();
+   const colors = dashboard.getStatusColors();
+
+   const styles = _styles ?? mergeStyleSets({
+      container: {
+         overflow: 'auto'
+      },
+      event: [
+         {
+            fontSize: "11px",
+            fontFamily: "Horde Cousine Regular"
+         }
+      ],
+      gutter: [
+         {
+            padding: 0,
+            margin: 0,
+            paddingTop: 0,
+            paddingBottom: 0,
+            paddingRight: 14,
+            marginTop: 0,
+            marginBottom: 0,
+            height: 20
+         }
+      ],
+      gutterError: [
+         {
+            background: dashboard.darktheme ? "#330606" : "#FEF6F6",
+            borderLeftStyle: 'solid',
+            borderLeftColor: colors.get(StatusColor.Failure)!
+         }, gutterClass
+      ],
+      gutterWarning: [
+         {
+            background: dashboard.darktheme ? "#1E1817" : "#FEF8E7",            
+            borderLeftStyle: 'solid',
+            borderLeftColor: colors.get(StatusColor.Warnings)!
+         }, gutterClass
+      ],
+      gutterSuccess: [
+         {
+            borderLeftStyle: 'solid',
+            borderLeftColor: colors.get(StatusColor.Success)!
+         }, gutterClass
+      ],
+      eventWarning: [
+         {
+            background: dashboard.darktheme ? "#1E1817" : "#FEF8E7",
+         }
+      ],
+      eventError: [
+         {
+            background: dashboard.darktheme ? "#330606" : "#FEF6F6",
+         }
+      ],
+      itemHover: {
+         cursor: "pointer",
+         selectors: {
+            ':hover': {
+               background: theme.palette.neutralLight
+            }
+         }
+      },
+      icon: {
+         userSelect: 'none',
+         fontSize: '16px' //'1.2vw'
+      },
+      sectionItem: {
+         zIndex: 1,
+         borderWidth: "1px",
+         borderColor: "#888",
+         borderStyle: "solid",
+         background: dashboard.darktheme ? "#111" : "#FFF",
+         padding: "6px",
+         overflow: "auto"
+      },
+      plainCard: {
+         borderWidth: "1px",
+         borderColor: "#888",
+         borderStyle: "solid",
+         background: dashboard.darktheme ? "#111" : "#FFF",
+         padding: "6px",
+         display: 'flex',
+      },
+      filterCard: {
+         borderWidth: "1px",
+         borderColor: "#888",
+         borderStyle: "solid",
+         background: dashboard.darktheme ? "#111" : "#FFF",
+         padding: "6px",
+         display: 'flex',
+      },
+      bottomLoadMore: {
+         borderWidth: "1px",
+         borderColor: "#CCC",
+         borderStyle: "solid",
+         padding: 3,
+         background: dashboard.darktheme ? "#111" : "#FFF",
+         fontSize: 11
+      },
+      compactItem: {
+         paddingLeft: 6
+      },
+      compactItemRight: {
+         paddingRight: 6
+      },
+      reportLeftPanel: {
+         borderWidth: "1px",
+         borderColor: "#CCC",
+         borderStyle: "solid",
+         borderRadius: "5px",
+         width: "220px",
+         padding: 4,
+         background: dashboard.darktheme ? "#111" : "#FFF",
+      },
+      reportPivotItem: {
+         paddingTop: 8
+      },
+      bottomArrow: {
+         fontSize: 0,
+         position: "relative",
+         selectors: {
+            ':before': {
+               content: "''",
+               display: 'inline-block',
+               width: 0,
+               height: 0,
+               border: '8px solid transparent',
+               verticalAlign: 'middle',
+               borderBottomColor: '#666',
+               position: "absolute",
+               transform: 'translateY(9px)',
+            }
+         }
+      },
+      rotateText: {
+         transform: 'translateX(15px) translateY(15px) rotate(90deg)',
+         transformOrigin: '0 0',
+         position: 'absolute'
+      },
+      borderSolid: {
+         border: '1px solid transparent',
+      },
+      sideFilterButton: {
+         padding: 4,
+         minWidth: 0,
+         marginLeft: '4px !important'
+      },
+      stripes: {
+         backgroundImage: 'repeating-linear-gradient(-45deg, rgba(255, 255, 255, .2) 25%, transparent 25%, transparent 50%, rgba(255, 255, 255, .2) 50%, rgba(255, 255, 255, .2) 75%, transparent 75%, transparent)',
+      }
+   });
+
+   const stateColorStyles = _stateColorStyles ?? new Map<string, string>([
+      [TestState.Success, colors.get(StatusColor.Success)!],
+      [TestState.InProcess, colors.get(StatusColor.Running)!],
+      [TestState.NotRun, colors.get(StatusColor.Waiting)!],
+      [TestState.SuccessWithWarnings, colors.get(StatusColor.Warnings)!],
+      [TestState.Failed, colors.get(StatusColor.Failure)!],
+      [TestState.Unknown, "#000000"],
+   ]);
+
+   _styles = styles;
+   _stateColorStyles = stateColorStyles;
+
+   return {
+      theme: theme,
+      styles: _styles,
+      stateColorStyles: _stateColorStyles
+   }
+}
+
+
+const getIconTestStateStyles = (state: string): string => {
+
+   const { styles, stateColorStyles } = getStyling();
+
+   let color: string = "#000000";
+   if (!stateColorStyles.has(state)) {
+      color = stateColorStyles.get(TestState.Unknown)!;
+   } else {
+      color = stateColorStyles.get(state)!;
+   }
+
+   return mergeStyles({ color: color, cursor: "pointer" }, styles.icon);
+}
+
+const stateIconName = new Map<string, string>([
+   [TestState.Success, "BoxCheckmarkSolid"],
+   [TestState.InProcess, "BoxPlaySolid"],
+   [TestState.NotRun, "CheckboxFill"],
+   [TestState.SuccessWithWarnings, "BoxCheckmarkSolid"],
+   [TestState.Failed, "BoxMultiplySolid"],
+   [TestState.Skipped, 'BorderDash'],
+   [TestState.Unknown, "SquareShape"],
+]);
+const getStateIconName = (state: string): string => {
+   if (!stateIconName.has(state)) {
+      return stateIconName.get(TestState.Unknown)!;
+   }
+   return stateIconName.get(state)!;
+}
+
+const stateLabels = new Map<string, string>([
+   [TestState.Success, 'Passed'],
+   [TestState.InProcess, 'Incomplete'],
+   [TestState.NotRun, 'Unexecuted'],
+   [TestState.SuccessWithWarnings, 'Passed'],
+   [TestState.Failed, 'Failed'],
+   [TestState.Skipped, 'Skipped'],
+   [TestState.Unknown, 'Unknown'],
+]);
+const getStateLabel = (state: string): string | undefined => { return stateLabels.has(state) ? stateLabels.get(state) : state }
+
+const copyToClipboard = (value: string | undefined) => {
+
+   if (!value) {
+      return;
+   }
+
+   const el = document.createElement('textarea');
+   el.value = value;
+   document.body.appendChild(el);
+   el.select();
+   document.execCommand('copy');
+   document.body.removeChild(el);
+}
+
+const missingImage = "/images/missing-image.png";
+const MissingImageLabel = (): JSX.Element => { return <span style={{ fontWeight: 'bold' }}> [missing image]</span> }
+type imageData = { link?: string, ref?: string }
+type ImageLinks = { [key: string]: imageData; }
+
+const Stats = (stats: TestStats): JSX.Element => {
+   const { stateColorStyles } = getStyling();
+   return (
+      <Stack horizontal verticalAlign="center">
+         <Text styles={{ root: { fontWeight: 'bold' } }}>
+            <span>[ </span>
+            {stats.Failed > 0 && <span>{stats.Failed} <span style={{ fontSize: 12 }} className={mergeStyles({ color: stateColorStyles.get(TestState.Failed) })}>{stateLabels.get(TestState.Failed)?.toLowerCase()}</span>. </span>}
+            {stats.Incomplete > 0 && <span>{stats.Incomplete} <span style={{ fontSize: 12 }} className={mergeStyles({ color: stateColorStyles.get(TestState.InProcess) })}>{stateLabels.get(TestState.InProcess)?.toLowerCase()}</span>. </span>}
+            {stats.Unexecuted > 0 && <span>{stats.Unexecuted} <span style={{ fontSize: 12 }} className={mergeStyles({ color: stateColorStyles.get(TestState.NotRun) })}>{stateLabels.get(TestState.NotRun)?.toLowerCase()}</span>. </span>}
+            {stats.Passed > 0 && <span>{stats.Passed} <span style={{ fontSize: 12 }} className={mergeStyles({ color: stateColorStyles.get(TestState.Success) })}>{stateLabels.get(TestState.Success)?.toLowerCase()}</span>. </span>}
+            {stats.Skipped > 0 && <span>{stats.Skipped} <span style={{ fontSize: 12 }} className={mergeStyles({ color: stateColorStyles.get(TestState.Skipped) })}>{stateLabels.get(TestState.Skipped)?.toLowerCase()}</span>. </span>}
+            <span> ]</span>
+         </Text>
+      </Stack>
+   );
+}
+
+const ItemMeta = (item: TestResult, metaMask: string[]): JSX.Element => {
+   const { styles } = getStyling();
+   const meta = item.Session.MetaHandler;
+   const metaInfo: string[] | undefined = meta?.maskByKeys(metaMask);
+   return (
+      <Link to={onClickTestResult(item.TestUID, item.Session.MetaHandler?.hash)} key={item.Session?.MetaHandler?.hash}>
+         <Stack styles={{ root: { fontSize: 0 } }} disableShrink horizontal>
+            <Stack className={styles.compactItem}><FontIcon title={getStateLabel(item.State)} className={getIconTestStateStyles(item.State)} iconName={getStateIconName(item.State)} /></Stack>
+            {metaInfo && metaInfo.length > 0 && <Text nowrap>{metaInfo.map((info) => <span className={styles.compactItem} key={info}>{info}</span>)}</Text>}
+            <Text><span className={styles.compactItem}>#{item.Session.Testdata.change}</span></Text>
+         </Stack>
+      </Link>
+   );
+}
+
+const OnIconCard = (item: TestCase): JSX.Element => {
+   const { styles } = getStyling();
+   const metas = item.getResults(true);
+   const metaMask = metas.length > 1 ? GetCommonMetaKeys(item.getResults().map((item) => item.Session.Metadata)) : [];
+   return (
+      <Stack className={styles.plainCard}>
+         <Text>{item.Name}</Text>
+         {metas.map((meta) => ItemMeta(meta, metaMask))}
+      </Stack>
+   );
+};
+
+const OnTestCard = (item: TestResult): JSX.Element => {
+   const { styles } = getStyling();
+   const meta = item.Session.MetaHandler;
+   return <Stack className={styles.plainCard}>
+      {TestInfoItemRow('Changelist', item.Session.Testdata.change)}
+      {meta?.map((key) => TestInfoItemRow(key, meta.get(key)))}
+   </Stack>;
+};
+
+const OnDeviceCard = (device: TestDevice): JSX.Element => {
+   const { styles } = getStyling();
+   const meta = device.Metadata && new MetaWrapper(device.Metadata);
+   return <Stack className={styles.plainCard}>
+      {meta && meta.map((key, value) => TestInfoItemRow(key, value))}
+      {!meta && <Text>no metadata</Text>}
+   </Stack>;
+};
+
+const LoadingSpinner = (text: string, overlay: boolean = false): JSX.Element => {
+   return (
+      <Stack>
+         <Stack.Item align="center" styles={{ root: [{ padding: 10 }, overlay ? { position: 'fixed', bottom: '20px', left: '20px' } : undefined] }}>
+            <Text variant="mediumPlus">{text}</Text>
+            <Spinner styles={{ root: { padding: 10 } }} size={SpinnerSize.large}></Spinner>
+         </Stack.Item>
+      </Stack>
+   );
+};
+
+const setQuery = (query: URLSearchParams, type: FilterType, value?: string) => {
+   if (query.has(type)) {
+      if (value) {
+         query.set(type, value);
+      } else {
+         query.delete(type);
+      }
+   } else if (value) {
+      query.append(type, value);
+   }
+}
+
+const urlFromFilterType = (type: FilterType, target?: string) => {
+   const searchQueries = new URLSearchParams(window.location.search);
+   setQuery(searchQueries, type, target);
+   return `${window.location.pathname}?${searchQueries.toString()}`;
+}
+
+const urlFromFilter = (filter: Metadata) => {
+   const searchQueries = new URLSearchParams(window.location.search);
+   const keys = Object.keys(filter);
+   keys.forEach((key) => {
+      setQuery(searchQueries, key as FilterType, filter[key]);
+   });
+   return `${window.location.pathname}?${searchQueries.toString()}`;
+}
+
+const urlFromMetaFilterType = (key: string, value?: string[]) => {
+   const searchQueries = new URLSearchParams(window.location.search);
+   if (searchQueries.has(FilterType.meta)) {
+      let filters = searchQueries.getAll(FilterType.meta);
+      searchQueries.delete(FilterType.meta);
+      filters = filters.filter((meta) => !MetaFilterTools.Match(key, meta));
+      filters.forEach((meta) => searchQueries.append(FilterType.meta, meta));
+   }
+   if (value && value.length > 0) {
+      searchQueries.append(FilterType.meta, MetaFilterTools.Encode(key, value));
+   }
+   return `${window.location.pathname}?${searchQueries.toString()}`;
+}
+
+const onClickTestResult = (testuid?: string, metauid?: string, change?: number, view?: string) => {
+   return urlFromFilter({ testuid: testuid, metauid: metauid, testchange: change?.toString(), view: view });
+}
+
+const IconItem = (item: TestCase): JSX.Element => {
+   const plainCardProps: IPlainCardProps = {
+      renderData: item,
+      onRenderPlainCard: OnIconCard,
+      gapSpace: 8,
+   };
+   return (
+      <Stack.Item styles={{ root: { fontSize: 0 } }} key={item.TestUID}>
+         <HoverCard plainCardProps={plainCardProps} type={HoverCardType.plain} cardOpenDelay={200}>
+            <Link to={onClickTestResult(item.TestUID, item.StateMetaUID)}><FontIcon title={getStateLabel(item.State)} className={getIconTestStateStyles(item.State)} iconName={getStateIconName(item.State)} /></Link>
+         </HoverCard>
+      </Stack.Item>
+   );
+}
+
+const CompactItem = (item: TestResult, metaMask: string[], highlighted: boolean = false): JSX.Element => {
+
+   const { styles, stateColorStyles } = getStyling();
+
+   const plainCardProps: IPlainCardProps = {
+      renderData: item,
+      onRenderPlainCard: OnTestCard,
+      directionalHint: DirectionalHint.rightTopEdge,
+      gapSpace: 8,
+   };
+   const meta = item.Session.MetaHandler;
+   const metaInfo: string[] | undefined = meta?.maskByKeys(metaMask);
+
+   return (
+      <Stack.Item styles={{ root: { fontSize: 0 } }} key={item.TestUID + item.Session?.MetaHandler?.hash} disableShrink style={{ background: highlighted ? 'rgba(0,0,0,0.1)' : undefined }}>
+         <HoverCard plainCardProps={plainCardProps} type={HoverCardType.plain} cardOpenDelay={200}>
+            <Link to={onClickTestResult(item.TestUID, item.Session.MetaHandler?.hash)}>
+               <Stack horizontal>
+                  <Stack className={styles.compactItem}><FontIcon title={getStateLabel(item.State)} className={getIconTestStateStyles(item.State)} iconName={getStateIconName(item.State)} /></Stack>
+                  {metaInfo && metaInfo.length > 0 && <Text nowrap>{metaInfo.map((info) => <span className={styles.compactItem} key={info}>{info}</span>)}</Text>}
+                  {item.ErrorCount > 0 && <Text className={styles.compactItem} nowrap>({item.ErrorCount} <span style={{ color: stateColorStyles.get(TestState.Failed), fontWeight: 'bold' }}>errors</span>)</Text>}
+                  {item.WarningCount > 0 && <Text className={styles.compactItem} nowrap>({item.WarningCount} <span style={{ color: stateColorStyles.get(TestState.SuccessWithWarnings), fontWeight: 'bold' }}>warnings</span>)</Text>}
+               </Stack>
+            </Link>
+         </HoverCard>
+      </Stack.Item>
+   );
+}
+
+const TestDetailItemRow = (label: string, element?: JSX.Element): JSX.Element => {
+   return (
+      <Stack horizontal key={label}>
+         <Stack.Item styles={{ root: { width: '90px' } }}><Text style={{ fontWeight: 'bold' }} nowrap>{label}:</Text></Stack.Item>
+         {element && <Stack.Item>{element}</Stack.Item>}
+      </Stack>
+   );
+}
+
+const TestInfoItemRow = (label: string, value?: any, color?: string): JSX.Element => {
+   return TestDetailItemRow(label, value && <Text styles={{ root: { color: color } }}>{value.toString()}</Text>);
+}
+
+const TestResultModal: React.FC<{ testuid: string, collection: TestSessionCollection, onLoadMore?: () => void }> = (props) => {
+   const { testuid, collection, onLoadMore } = props;
+   const [loading, setLoading] = useState(true);
+   const [loadingTest, setLoadingTest] = useState(true);
+   const testRef = useRef<TestResult | undefined>(undefined);
+   const testcaseRef = useRef<TestCase | undefined>(undefined);
+   const navigate = useNavigate();
+   const onClickCancel = () => {
+      navigate(onClickTestResult());
+   }
+   const onClickPivot = (item?: PivotItem) => {
+      navigate(urlFromFilterType(FilterType.view, item?.props.itemKey), { replace: true });
+   }
+
+   const query = useQuery();
+   const metauid = query.get(FilterType.metauid) ?? undefined;
+   const change = parseInt(query.get(FilterType.testchange) ?? "");
+   const view = query.get(FilterType.view) ?? undefined;
+
+   useEffect(() => {
+      const cancel = { current: false };
+      const assignTestFromCache = async () => {
+         if (!testcaseRef.current) {
+            setLoadingTest(false);
+            return;
+         }
+         setLoadingTest(true);
+         try {
+            const testcase = testcaseRef.current;
+            if (testcase && metauid) {
+               const testFromQuery = await collection.getTestResultByQuery(testuid, metauid, change);
+               if (testFromQuery && !cancel.current) {
+                  testRef.current = testFromQuery;
+               }
+            }
+            if (!cancel.current && !testRef.current) {
+               const first_metauid = testcase.getResults(true)[0].Session.MetaHandler!.hash;
+               testRef.current = await collection.getTestResultByQuery(testuid, first_metauid);
+            }
+         } catch (reason) {
+            console.error(reason);
+         } finally { !cancel.current && setLoadingTest(false) }
+      }
+      const getTestResults = async () => {
+         setLoading(true);
+         try {
+            const testcase = await collection.getTestCaseByUID(testuid);
+            if (testcase && !cancel.current) {
+               testcaseRef.current = testcase;
+            }
+         } catch (reason) {
+            console.error(reason);
+         } finally { if (!cancel.current) { setLoading(false); assignTestFromCache() } }
+      }
+
+      const test = testRef.current;
+      if (!test || testuid !== test.TestUID) {
+
+         testRef.current = undefined;
+         testcaseRef.current = undefined;
+         getTestResults();
+
+      } else if (testcaseRef.current
+         && (!test || metauid !== test.Session.MetaHandler?.hash || change !== test.Session.Testdata.change)) {
+
+         testRef.current = undefined;
+         assignTestFromCache();
+      }
+      return function cleanup() { cancel.current = true }
+   }, [testuid, metauid, change, collection])
+
+   const test = testRef.current;
+   const meta = test?.Session.MetaHandler;
+   const alternates = testcaseRef.current?.getResults(true);
+   const metaMask: string[] = alternates && !testDataHandler.cursor ? GetCommonMetaKeys(alternates.map((item) => item.Session.Metadata)) : [];
+   const selectedChange = change ? change : (test?.Session.Testdata.change ?? 0);
+   let testName: string | undefined = undefined;
+   if (testcaseRef.current) {
+      testName = testcaseRef.current.Name;
+   }
+
+   const { hordeClasses, modeColors } = getHordeStyling();
+   const { styles } = getStyling();
+
+   return (
+      <Modal isOpen styles={{ main: { padding: 8, width: '100vw', height: 'calc(100vh - 24px)', background: modeColors.background} }} className={hordeClasses.modal} onDismiss={onClickCancel}>
+         <Stack>
+            <Stack horizontal styles={{ root: { padding: 8 } }}>
+               <Stack verticalAlign="center" styles={{ root: { overflow: 'hidden' } }}>
+                  <Text styles={{ root: { fontWeight: "unset", fontFamily: "Horde Open Sans SemiBold", fontSize: "16px", color: "#087BC4", textOverflow: 'ellipsis' } }} nowrap block>{testName ?? (loading ? "Loading..." : "Unknown test")}</Text>
+               </Stack>
+               <Stack verticalAlign="center" grow>
+                  <div><FontIcon iconName="AddToShoppingList" title="Copy Name to Clipboard" onClick={(ev) => { copyToClipboard(testName); }} style={{ padding: 6, fontSize: 18 }} className={styles.itemHover} /></div>
+               </Stack>
+               <Stack.Item align="end" disableShrink>
+                  <IconButton iconProps={{ iconName: 'Cancel' }} ariaLabel="Close test result" onClick={onClickCancel} />
+               </Stack.Item>
+            </Stack>
+            {!loading && alternates &&
+               <Stack horizontal styles={{ root: { paddingLeft: 8, paddingRight: 8 } }} tokens={{ childrenGap: 16 }} verticalFill disableShrink>
+                  <Stack.Item className={styles.reportLeftPanel} grow>
+                     <Pivot selectedKey={view} onLinkClick={onClickPivot}>
+                        <PivotItem headerText="General" className={styles.reportPivotItem}>
+                           {!loadingTest && test?.Details !== undefined && TestDetails(test, alternates, metaMask)}
+                           {loadingTest && LoadingSpinner('Loading Details')}
+                        </PivotItem>
+                        <PivotItem headerText="History" className={styles.reportPivotItem} itemKey="history">
+                           {testcaseRef.current && <TestHistory collection={collection} testcase={testcaseRef.current} metaMask={metaMask} metauid={metauid ?? meta?.hash} change={selectedChange} onLoadMore={onLoadMore} />}
+                        </PivotItem>
+                        <PivotItem headerText="Artifacts" className={styles.reportPivotItem} itemKey="artifacts">
+                           Available soon
+                        </PivotItem>
+                     </Pivot>
+                  </Stack.Item>
+                  <Stack.Item grow style={{ width: '100%' }} className={styles.container}>
+                     {!loadingTest && test?.Details !== undefined && <TestEvents test={test} />}
+                     {loadingTest && LoadingSpinner('Loading Details')}
+                  </Stack.Item>
+               </Stack>
+            }
+            {!loading && !loadingTest && !test && <Stack styles={{ root: { paddingLeft: 8, paddingRight: 8 } }}><Text>Could not find Test {testuid} for meta {metauid}</Text></Stack>}
+            {loading && LoadingSpinner('Loading Test Result')}
+         </Stack>
+      </Modal>
+   );
+}
+
+const TestDetails = (test: TestResult, alternates: TestResult[], metaMask: string[]): JSX.Element => {
+   const { stateColorStyles } = getStyling();
+   const meta = test.Session.MetaHandler;
+   const devices = test.Session.getDevices(test.DeviceAppInstanceName);
+
+   return (
+      <Stack>
+         {TestDetailItemRow('State',
+            <Stack horizontal wrap tokens={{ childrenGap: 6 }}>
+               <FontIcon className={getIconTestStateStyles(test.State)} iconName={getStateIconName(test.State)} />
+               <Text style={{ color: stateColorStyles.get(test.State) }}>{getStateLabel(test.State)}</Text>
+            </Stack>)
+         }
+         {TestInfoItemRow('Errors', test.ErrorCount ?? 0, stateColorStyles.get(TestState.Failed))}
+         {TestInfoItemRow('Warnings', test.WarningCount ?? 0, stateColorStyles.get(TestState.SuccessWithWarnings))}
+         <Stack>
+            {TestInfoItemRow('Conditions')}
+            <Stack.Item styles={{ root: { marginLeft: 30 } }}>
+               {TestInfoItemRow('Changelist', test.Session.Testdata.change)}
+               {meta?.map((key, value) => TestInfoItemRow(key, value))}
+            </Stack.Item>
+         </Stack>
+         {test.DateTime && TestInfoItemRow('Run at', test.DateTime)}
+         {test.TimeElapseSec && TestInfoItemRow('For', msecToElapsed(test.TimeElapseSec * 1000))}
+         {devices && devices.length > 0 &&
+            TestDetailItemRow('Devices', <Stack>{devices.map((device) => TestDeviceItem(device))}</Stack>)
+         }
+         {TestDetailItemRow('Link to', <Link to={test.Session.Testdata.getJobStepLink()}>Job Step</Link>)}
+         {alternates.length > 0 &&
+            <Stack>
+               {TestInfoItemRow(testDataHandler.cursor ? `Result on ${alternates[0].Session.Testdata.change}` : 'Latest Alternates')}
+               <Stack>
+                  {alternates.map((item) => CompactItem(item, metaMask, item.Session.MetaHandler!.hash === meta!.hash))}
+               </Stack>
+            </Stack>
+         }
+      </Stack>
+   );
+}
+
+const TestDeviceItem = (device: TestDevice): JSX.Element => {
+   const deviceCardProps: IPlainCardProps = {
+      renderData: device,
+      onRenderPlainCard: OnDeviceCard,
+      directionalHint: DirectionalHint.rightTopEdge,
+      gapSpace: 8,
+   };
+
+   return (
+      <HoverCard plainCardProps={deviceCardProps} type={HoverCardType.plain} cardOpenDelay={20} key={device.AppInstanceName}>
+         <Text style={{ cursor: 'pointer' }}>{device.Name}</Text>
+      </HoverCard>
+   );
+
+}
+
+const TestHistory: React.FC<{ collection: TestSessionCollection, testcase: TestCase, metaMask: string[], metauid?: string, change: number, onLoadMore?: () => void }> = (props) => {
+   const { collection, testcase, metaMask, metauid, change, onLoadMore } = props;
+   const [loading, setLoading] = useState(true);
+   const [history, updateHistory] = useState(testcase.history);
+   const [needFetching, forceFetch] = useReducer(x => x + 1, 1);
+   const testuid = testcase.TestUID;
+
+   useEffect(() => {
+      const cancel = { current: false };
+      const historyChangeDisposer = testcase.reactOnHistoryChange((version) => updateHistory(version));
+      const FetchTestResultHistory = async () => {
+         setLoading(true);
+         try {
+            await collection.fetchSessionHistory(testcase.SessionName, testcase.TestUID);
+         } catch (reason) {
+            console.error(reason);
+         } finally { !cancel.current && setLoading(false) }
+      }
+      needFetching && FetchTestResultHistory();
+      return function cleanup() { cancel.current = true; historyChangeDisposer() }
+   }, [testcase, collection, needFetching])
+
+   const { styles } = getStyling();
+
+   let metasRef: MetaWrapper[] | undefined = undefined;
+   let historyMetaMask: string[] = metaMask;
+   const changesByMeta = collection.getLoadedTestHistoryByMeta(testcase.TestUID);
+   if (changesByMeta && changesByMeta.size > 0) {
+      // get the sorted list of meta uids
+      const metauids = testDataHandler.cursor ?
+         Array.from(changesByMeta.keys()).sort()
+         : testcase.getResults(true).map((test) => test.Session.MetaHandler!.hash);
+      if (testDataHandler.cursor) {
+         historyMetaMask = GetCommonMetaKeys(metauids.map((key) => changesByMeta?.get(key)![0].Session.Metadata));
+      }
+      metasRef = metauids.map((key) => changesByMeta?.get(key)![0].Session.MetaHandler!);
+   }
+
+   return (
+      <Stack style={{ height: 'calc(100vh - 160px)' }} className={styles.container}>
+         {loading && <LoaderProgressBar loader={testcase} target={testcase.historyTarget} />}
+         {!loading && onLoadMore &&
+            <Stack className={styles.bottomLoadMore} style={{ fontSize: 11 }} horizontal horizontalAlign="center">
+               <Stack style={{ color: "rgb(0, 120, 212)", cursor: 'pointer' }} onClick={() => { onLoadMore(); forceFetch() }}>Load More History...</Stack>
+            </Stack>
+         }
+         {history > -1 && metasRef &&
+            <Stack tokens={{ childrenGap: 6 }} className={styles.container} verticalFill>
+               {metasRef.map((meta) => TestHistoryMeta(
+                  meta, historyMetaMask, changesByMeta?.get(meta.hash) ?? [], testuid, meta.hash === metauid, change)
+               )
+               }
+            </Stack>
+         }
+      </Stack>
+   );
+}
+
+const TestHistoryMeta = (meta: MetaWrapper, metaMask: string[], tests: TestResult[], testuid: string, selected: boolean, change: number): JSX.Element => {
+
+   const { styles } = getStyling();
+
+   let metaInfo: string[] = meta.maskByKeys(metaMask);
+   if (metaInfo.length === 0) {
+      metaInfo = meta.map((key, value) => value);
+   }
+   return (
+      <Stack key={meta.hash}>
+         <Text variant="smallPlus">{metaInfo.map((info) => <span className={styles.compactItemRight} key={info}>{info}</span>)}</Text>
+         <Stack horizontal style={{ paddingLeft: 6 }}>
+            {tests.map((test) => <TestHistoryItem key={test.Session.Testdata.change} test={test} selected={selected && test.Session.Testdata.change === change} />)}
+         </Stack>
+      </Stack>
+   );
+}
+
+const TestHistoryItem: React.FC<{ test: TestResult, selected: boolean }> = (props) => {
+   const { test, selected } = props;
+   const navigate = useNavigate();
+   const { styles } = getStyling();
+
+   const hoverCardProps: IPlainCardProps = {
+      renderData: test,
+      onRenderPlainCard: TestHistoryItemCard,
+      directionalHint: DirectionalHint.topCenter,
+      gapSpace: 8,
+      calloutProps: { isBeakVisible: true }
+   };
+
+   const onClickIcon = () => navigate(
+      urlFromFilter({
+         metauid: test.Session.MetaHandler?.hash,
+         testchange: test.Session.Testdata.change?.toString(),
+         view: 'history'
+      })
+   );
+
+   const scrollToSelected = (ref: any) => {
+      if (ref && selected) {
+         ref.scrollIntoView({ behavior: 'smooth', block: 'start' });
+      }
+   }
+
+   return (
+      <HoverCard plainCardProps={hoverCardProps} type={HoverCardType.plain} cardOpenDelay={20}>
+         <div className={selected ? styles.bottomArrow : undefined} ref={selected ? scrollToSelected : undefined}>
+            {test && <FontIcon title={getStateLabel(test.State)} className={getIconTestStateStyles(test.State)} iconName={getStateIconName(test.State)} onClick={onClickIcon} />}
+         </div>
+      </HoverCard>
+   );
+}
+
+const TestHistoryItemCard = (test: TestResult): JSX.Element => {
+   const { styles } = getStyling();
+   return <Stack className={styles.plainCard}>
+      {TestInfoItemRow('Changelist', test.Session.Testdata.change)}
+   </Stack>;
+};
+
+const TestEvents: React.FC<{ test: TestResult }> = (props) => {
+   const { test } = props;
+   const [filterError, setFilterError] = useState(true);
+   const [filterWarning, setFilterWarning] = useState(true);
+   const [filterInfo, setFilterInfo] = useState(test.State !== TestState.Failed);
+   const { styles } = getStyling();
+
+   function onChangeFilterError(ev?: React.FormEvent<HTMLElement | HTMLInputElement>, checked?: boolean) {
+      checked !== undefined && setFilterError(checked);
+   }
+   function onChangeFilterWarning(ev?: React.FormEvent<HTMLElement | HTMLInputElement>, checked?: boolean) {
+      checked !== undefined && setFilterWarning(checked);
+   }
+   function onChangeFilterInfo(ev?: React.FormEvent<HTMLElement | HTMLInputElement>, checked?: boolean) {
+      checked !== undefined && setFilterInfo(checked);
+   }
+
+   function isEntryNeedDisplay(entry: TestEvent, index?: number, array?: TestEvent[]): boolean {
+      const entryEventType = entry.Type;
+      if (filterError && entryEventType === EventType.Error) {
+         return true;
+      }
+      if (filterWarning && entryEventType === EventType.Warning) {
+         return true;
+      }
+      if (filterInfo && entryEventType === EventType.Info) {
+         return true;
+      }
+
+      return false;
+   }
+
+   return (
+      <Stack styles={{ root: { width: '100%' } }} className={styles.container}>
+         <Stack horizontal tokens={{ childrenGap: 10 }} styles={{ root: { padding: 5 } }}>
+            <FontIcon iconName="PageListFilter" title="Filters" className={styles.icon} />
+            <Checkbox label="Error" defaultChecked={filterError} onChange={(onChangeFilterError)} />
+            <Checkbox label="Warning" defaultChecked={filterWarning} onChange={onChangeFilterWarning} />
+            <Checkbox label="Info" defaultChecked={filterInfo} onChange={onChangeFilterInfo} />
+         </Stack>
+         <Stack styles={{ root: { height: 'calc(100vh - 125px)', width: '100%', overflow: 'auto' } }}>
+            {test.Details.Events.filter(isEntryNeedDisplay).map((value, index) => <EventPane key={index} entry={value} test={test} />)}
+            {test.Details.Events.length === 0 && <Text styles={{ root: { padding: 8, fontWeight: 'bold' } }}>No event for this test.</Text>}
+         </Stack>
+      </Stack>
+   );
+}
+
+const EventPane: React.FC<{ entry: TestEvent, test: TestResult }> = (props) => {
+
+   const { styles } = getStyling();
+
+   const { entry, test } = props;
+   const eventType = entry.Type;
+   const is_error = eventType === EventType.Error;
+   const is_warning = eventType === EventType.Warning;
+   const style = is_error ? styles.eventError : (is_warning ? styles.eventWarning : styles.event);
+   const gutterStyle = is_error ? styles.gutterError : (is_warning ? styles.gutterWarning : styles.gutter);
+
+   const [imageInfo, setImageInfo] = useState<ImageLinks>({});
+   const need_image_comparison = entry.Artifacts && entry.Artifacts.length > 0 && entry.Tag === ArtifactType.ImageCompare;
+   const need_image_comparisonRef = useRef(need_image_comparison);
+
+
+   useEffect(() => {
+      const cancel = { current: false };
+      if (need_image_comparisonRef.current) {
+         const findLinks = async () => {
+            const imageInfo: ImageLinks = {};
+            for (const key in entry.Artifacts) {
+
+               const item = entry.Artifacts[key];
+               const url = await test.Session.Testdata?.getArtifactImageLink(item.ReferencePath);
+               imageInfo[item.Tag] = { link: url, ref: item.ReferencePath };
+            }
+            !cancel.current && setImageInfo(imageInfo);
+         }
+         findLinks();
+      }
+      return function cleanup() { cancel.current = true }
+   }, [entry.Artifacts, test]);
+
+   return (
+      <Stack className={styles.event} disableShrink={true} styles={{ root: { paddingLeft: 8, width: "100%" } }}>
+         <Stack horizontal>
+            <Stack className={gutterStyle}></Stack>
+            <Stack.Item className={style} align="center" styles={{ root: { paddingLeft: 8, width: "100%" } }}><pre style={{ margin: 0, whiteSpace: "pre-wrap" }}>[{entry.DateTime}] {eventType}: {entry.Message}</pre></Stack.Item>
+         </Stack>
+         {need_image_comparison &&
+            <Stack styles={{ root: { paddingLeft: 16, width: '100%' } }}>
+               <Stack><Text variant="medium"><span style={{ fontWeight: "bold" }}>Image comparison: </span>{entry.Context}</Text></Stack>
+               <Stack horizontal>
+                  {(is_error || is_warning) &&
+                     <Stack styles={{ root: { padding: 5 } }}>
+                        <a href={imageInfo[ArtifactType.Approved]?.link}><Image width="100%" style={{ minWidth: 200, maxWidth: 400, minHeight: 120 }} src={imageInfo[ArtifactType.Approved]?.link || missingImage} alt={imageInfo[ArtifactType.Approved]?.ref} /></a>
+                        <Stack.Item align="center">Reference{!imageInfo[ArtifactType.Approved] && MissingImageLabel()}</Stack.Item>
+                     </Stack>
+                  }
+                  {(is_error || is_warning) &&
+                     <Stack styles={{ root: { padding: 5 } }}>
+                        <a href={imageInfo[ArtifactType.Difference]?.link}><Image width="100%" style={{ minWidth: 200, maxWidth: 400, minHeight: 120 }} src={imageInfo[ArtifactType.Difference]?.link || missingImage} alt={imageInfo[ArtifactType.Difference]?.ref} /></a>
+                        <Stack.Item align="center">Difference{!imageInfo[ArtifactType.Difference] && MissingImageLabel()}</Stack.Item>
+                     </Stack>
+                  }
+                  <Stack styles={{ root: { padding: 5 } }}>
+                     <a href={imageInfo[ArtifactType.Unapproved]?.link}><Image width="100%" style={{ minWidth: 200, maxWidth: 400, minHeight: 120 }} src={imageInfo[ArtifactType.Unapproved]?.link || missingImage} alt={imageInfo[ArtifactType.Unapproved]?.ref} /></a>
+                     <Stack.Item align="center">Produced{!imageInfo[ArtifactType.Unapproved] && MissingImageLabel()}</Stack.Item>
+                  </Stack>
+               </Stack>
+            </Stack>
+         }
+      </Stack>
+   );
+}
+
+type ChartStack = {
+   value: number,
+   title?: string,
+   color?: string,
+   onClick?: () => void,
+   stripes?: boolean,
+
+}
+
+const Chart = (stack: ChartStack[], width: number, height: number, basecolor?: string, style?: any): JSX.Element => {
+
+   const { styles } = getStyling();
+
+   const mainTitle = stack.map((item) => `${item.value}% ${item.title}`).join(' ');
+
+   return (
+      <div className={mergeStyles({ backgroundColor: basecolor, width: width, height: height, verticalAlign: 'middle' }, style)} title={mainTitle}>
+         {stack.map((item) => <span key={item.title!}
+            onClick={item.onClick}
+            className={item.stripes ? styles.stripes : undefined}
+            style={{
+               width: `${item.value}%`, height: '100%',
+               backgroundColor: item.color,
+               display: 'inline-block',
+               cursor: item.onClick ? 'pointer' : 'inherit',
+               backgroundSize: `${height * 2}px ${height * 2}px`
+            }} />)}
+      </div>
+   );
+}
+
+type BreadCrumbFill = {
+   title: string,
+   linkTo: any,
+};
+
+const SectionBreadCrumb = (filling: BreadCrumbFill[], stack?: ChartStack[], total?: number): JSX.Element => {
+
+   const { stateColorStyles } = getStyling();
+
+   return (
+      <Stack grow>
+         <Stack horizontal wrap>
+            {filling.slice(0, -1).map((crumb, index) => <Stack styles={{ root: { padding: 3 } }} key={index}><Link to={crumb.linkTo}><Text>{crumb.title} &gt;</Text></Link></Stack>)}
+         </Stack>
+         <Stack horizontal verticalAlign="center" styles={{ root: { padding: 6 } }}>
+            <Text variant="mediumPlus">{filling[filling.length - 1].title}</Text>
+            {total && stack &&
+               Chart(stack, 200, 20, stateColorStyles.get(TestState.Success), { margin: '3px !important' })
+            }
+            {total !== undefined && <Text style={{ fontWeight: 'bold' }}>[ {total} results ]</Text>}
+         </Stack>
+      </Stack>
+   );
+}
+
+const sectionCrumbLink = (type: string, target: string) => {
+   const searchQueries = new URLSearchParams(window.location.search);
+   searchQueries.set(type, target);
+   if (type === FilterType.suite) {
+      searchQueries.delete(FilterType.name);
+   }
+
+   return `${window.location.pathname}?${searchQueries.toString()}`;
+}
+
+const noSectionCrumbLink = () => {
+   const searchQueries = new URLSearchParams(window.location.search);
+   if (searchQueries.has(FilterType.suite)) {
+      searchQueries.delete(FilterType.suite);
+   }
+   if (searchQueries.has(FilterType.name)) {
+      searchQueries.delete(FilterType.name);
+   }
+
+   return `${window.location.pathname}?${searchQueries.toString()}`;
+}
+
+const ResultsFolder: React.FC<{ items: TestCase[], name: string, expanded?: boolean }> = (props) => {
+   const { items, name, expanded } = props;
+   const [isExpanded, setExpanded] = useState(expanded ?? false);
+   const { styles } = getStyling();
+   return (
+      <Stack styles={{ root: { padding: 5, userSelect: "none" } }}>
+         <Stack verticalAlign="center" horizontal styles={{ root: { cursor: 'pointer' } }} onClick={() => setExpanded(!isExpanded)}>
+            <FontIcon className={styles.icon} iconName={isExpanded ? "CaretSolidDown" : "CaretSolidRight"} style={{ fontSize: 11 }} />
+            <Text variant="smallPlus" styles={{ root: { fontFamily: "Horde Open Sans SemiBold", paddingLeft: 3 } }}>{name} ({items.length})</Text>
+         </Stack>
+         {isExpanded &&
+            <Stack horizontal wrap tokens={{ childrenGap: 2 }} styles={{ root: { paddingTop: 4 } }}>
+               {items.map((testcase) => IconItem(testcase))}
+            </Stack>
+         }
+      </Stack>
+   );
+}
+
+const ResultsSelector: React.FC<{ items: TestCase[], stats: TestStats }> = (props) => {
+   const { items, stats } = props;
+   const [isFailureExpanded, setFailureExpanded] = useState(false);
+   const [isOtherExpanded, setOtherExpanded] = useState(false);
+
+
+   const failures = useMemo(() => {
+      return items.filter((testcase) => testcase.State !== TestState.Success && testcase.State !== TestState.Skipped)
+   }, [items]);
+
+   const others = useMemo(() => {
+      const othersFilter: TestCase[] = [];
+      items.forEach((testcase) => {
+         const filteredCase = testcase.getFilteredByState(TestState.SuccessWithWarnings);
+         if (filteredCase) {
+            othersFilter.push(filteredCase);
+         }
+      });
+      SectionCollection.sortTestCases(othersFilter);
+      return othersFilter;
+   }, [items]);
+
+   const failedFactor = Math.ceil((stats.Failed + stats.Incomplete) / (stats.TotalRun || 1) * 20) / 20;
+   const { stateColorStyles } = getStyling();
+
+   const stack: ChartStack[] = [
+      {
+         value: failedFactor * 100,
+         title: "Failure",
+         color: stateColorStyles.get(TestState.Failed),
+         onClick: () => { setFailureExpanded(!isFailureExpanded); setOtherExpanded(false) },
+         stripes: true
+      },
+      {
+         value: (1 - failedFactor) * 100,
+         title: "Passed",
+         color: 'transparent',
+         onClick: () => { setFailureExpanded(false); setOtherExpanded(!isOtherExpanded) }
+      }
+   ]
+
+   return (
+      <Stack styles={{ root: { padding: 5, userSelect: "none" } }}>
+         <Stack horizontal verticalAlign="center">
+            {Chart(stack, 500, 30, stateColorStyles.get(TestState.Success), { margin: '3px !important' })}
+         </Stack>
+         {isFailureExpanded &&
+            <Stack horizontal wrap tokens={{ childrenGap: 2 }} styles={{ root: { paddingTop: 4 } }}>
+               {failures.map((testcase) => IconItem(testcase))}
+            </Stack>
+         }
+         {isOtherExpanded &&
+            <Stack horizontal wrap tokens={{ childrenGap: 2 }} styles={{ root: { paddingTop: 4 } }}>
+               {others.map((testcase) => IconItem(testcase))}
+            </Stack>
+         }
+      </Stack>
+   );
+}
+
+const GetCommonMetaKeys = (items: Metadata[]) => {
+   const first = items[0];
+   const commonKeys = new Set(Object.keys(first));
+   if (items.length > 1) {
+      items.slice(1).forEach((item) => {
+         Array.from(commonKeys.values()).forEach((key) => {
+            if (first[key] !== item[key]) {
+               commonKeys.delete(key);
+            }
+         });
+      });
+   }
+   return Array.from(commonKeys);
+}
+
+const SectionItem = (item: Section): JSX.Element => {
+   const { styles } = getStyling();
+   const key = item.Type + item.FullName;
+
+   if (item.isOneTest) {
+      const testcase = item.TestCases[0];
+
+      return (
+         <Stack className={styles.sectionItem} styles={{ root: { background: item.Stats.Failed > 0 ? (dashboard.darktheme ? "#111" : "#FFF") : undefined } }} grow key={key}>
+            <Stack verticalAlign="center" horizontal styles={{ root: { paddingBottom: 4 } }}>
+               {IconItem(testcase)}
+               <Link to={onClickTestResult(testcase.TestUID, testcase.StateMetaUID)}>
+                  <Text variant="smallPlus" styles={{ root: { fontFamily: "Horde Open Sans SemiBold", paddingLeft: 3 } }}>{item.Name}</Text>
+               </Link>
+            </Stack>
+            {testcase.Metas.size > 1 && Stats(testcase.Stats)}
+         </Stack>
+      );
+
+   } else {
+      const onClickSection = () => {
+         const searchQueries = new URLSearchParams(window.location.search);
+         searchQueries.set(item.Type, item.FullName);
+         return `${window.location.pathname}?${searchQueries.toString()}`;
+      }
+
+      return (
+         <Stack className={styles.sectionItem} styles={{ root: { background: item.Stats.Failed > 0 ? (dashboard.darktheme ? "#111" : "#FFF") : undefined } }} grow key={key}>
+            <Link to={onClickSection()} style={{ fontSize: 0, color: 'inherit' }}>
+               <Stack horizontal styles={{ root: { paddingBottom: 4 } }}>
+                  <FontIcon className={styles.icon} iconName="FolderList" />
+                  <Text variant="smallPlus" styles={{ root: { fontFamily: "Horde Open Sans SemiBold", paddingLeft: 3 } }}>{item.Name}</Text>
+               </Stack>
+            </Link>
+            {item.Stats.TotalRun > 0 && <ResultsSelector items={item.TestCases} stats={item.Stats} />}
+            {item.Stats.TotalRun === 0 && <ResultsFolder items={item.TestCases} name="Unexecuted" />}
+            {Stats(item.Stats)}
+         </Stack>
+      );
+
+   }
+}
+
+const TestMetaFilter: React.FC<{ metaOptions: Map<string, Set<string>>, filter?: MetadataFilter }> = (props) => {
+   const { filter, metaOptions } = props;
+   const navigate = useNavigate();
+   const newFilter = useRef<string | undefined>(undefined);
+   const [metaFilters, setMetaFilters] = useState<{ key: string, values: string[] }[]>(
+      (filter ? Array.from(Object.entries(filter)) : []).map((meta) => {
+         return { key: meta[0], values: meta[1] }
+      }).filter((item) => metaOptions.has(item.key))
+   );
+   const selectedMeta = new Set(metaFilters.map((item) => item.key));
+   const metaCanAddKeys: IDropdownOption[] = Array.from(metaOptions.entries()).map((data) => {
+      return { key: data[0], text: data[0] }
+   }).filter((item) => !selectedMeta.has(item.key));
+
+   const { styles } = getStyling();
+
+   return (
+      <Stack tokens={{ childrenGap: 4 }}>
+         {metaFilters.length > 0 &&
+            metaFilters.map((entry) => {
+               const possibilities = metaOptions.get(entry.key)!;
+               const selected = entry.values.filter((item) => possibilities.has(item));
+               const options: IDropdownOption[] = Array.from(possibilities.values()).map((item) => { return { key: item, text: item } });
+               return (
+                  <Stack key={entry.key} horizontal horizontalAlign="end" verticalAlign="center">
+                     <Text style={{ fontWeight: 'bold', paddingRight: 4 }}>{entry.key}</Text>
+                     <Dropdown
+                        placeholder="select options"
+                        styles={{ dropdown: { width: 150 } }}
+                        onChange={(ev, option) => {
+                           const isSelected = option!.selected;
+                           const key = option!.key as string;
+                           const selectedIndex = selected.findIndex((item) => item === key);
+                           if (isSelected) {
+                              selectedIndex === -1 && selected.push(key);
+                           } else {
+                              selectedIndex > -1 && selected.splice(selectedIndex, 1);
+                           }
+                           navigate(urlFromMetaFilterType(entry.key, selected));
+                           setMetaFilters(metaFilters.map((item) => item.key !== entry.key ? item : { key: entry.key, values: selected }));
+                        }}
+                        multiSelect
+                        options={options}
+                        defaultSelectedKeys={selected} />
+                     <DefaultButton
+                        onClick={() => {
+                           setMetaFilters(metaFilters.filter((item) => item.key !== entry.key));
+                           selected.length > 0 && navigate(urlFromMetaFilterType(entry.key, undefined));
+                        }}
+                        title="Update filtered results"
+                        className={styles.sideFilterButton}>
+                        <FontIcon className={styles.icon} iconName="CalculatorMultiply" style={{ fontSize: 11 }} />
+                     </DefaultButton>
+                  </Stack>
+               );
+            })
+         }
+         {metaCanAddKeys.length > 0 &&
+            <Stack horizontal horizontalAlign="end" verticalAlign="center">
+               <Dropdown
+                  placeholder="select filter to Add"
+                  key={`add ${metaCanAddKeys.length}`}
+                  styles={{ dropdown: { width: 150 } }}
+                  options={metaCanAddKeys}
+                  onChange={(ev, item) => newFilter.current = item?.key as string} />
+               <DefaultButton
+                  onClick={() => {
+                     newFilter.current && setMetaFilters([...metaFilters, { key: newFilter.current, values: [] }]);
+                     newFilter.current = undefined
+                  }}
+                  title="Add filter"
+                  className={styles.sideFilterButton}>
+                  <FontIcon className={styles.icon} iconName="CalculatorAddition" style={{ fontSize: 11 }} />
+               </DefaultButton>
+            </Stack>
+         }
+      </Stack>
+   );
+}
+
+const TestNameFilter: React.FC<{ filter: SectionFilter, onChange?: (value: string) => void, onValidate?: (value?: string) => void, dismissMenu?: (ev?: any, dismissAll?: boolean) => void }> = (props) => {
+   const { filter, onChange, onValidate, dismissMenu } = props;
+   const [name, setName] = useState<string>(filter.namecontains ?? '');
+
+   const { styles } = getStyling();
+
+   return (
+      <Stack horizontal horizontalAlign="end">
+         <TextField
+            deferredValidationTime={750}
+            validateOnLoad={false}
+            spellCheck={false}
+            placeholder="search test name"
+            value={name}
+            onChange={(ev, value) => setName(value ?? '')}
+            onGetErrorMessage={onChange && ((value) => { onChange(value); return undefined })}
+            onKeyDown={onValidate && ((ev) => {
+               switch (ev.key) {
+                  case "Enter":
+                     onValidate((ev.target as ITextFieldProps).value);
+                     dismissMenu && dismissMenu();
+                     break;
+                  case "Escape":
+                     onValidate(filter.namecontains);
+                     break;
+               }
+            })}
+            onBlur={onValidate && ((ev) => onValidate(ev.target.value))}
+         />
+         <DefaultButton
+            onClick={() => { onValidate && onValidate(undefined); setName(''); !filter.meta && dismissMenu && dismissMenu(); }}
+            title="Clear filter"
+            className={styles.sideFilterButton}>
+            <FontIcon className={styles.icon} iconName="CalculatorMultiply" style={{ fontSize: 11 }} />
+         </DefaultButton>
+      </Stack>
+   );
+}
+
+const TestFilterCard = (filter: SectionFilter, metaOptions: Map<string, Set<string>>, onChange?: (value: string) => void, onValidate?: (value?: string) => void, dismissMenu?: (ev?: any, dismissAll?: boolean) => void): JSX.Element => {
+   return (
+      <Stack style={{ padding: 8, paddingTop: 16, paddingBottom: 12 }} tokens={{ childrenGap: 4 }}>
+         <TestNameFilter filter={filter} onChange={onChange} onValidate={onValidate} dismissMenu={dismissMenu} />
+         <TestMetaFilter filter={filter.meta} metaOptions={metaOptions} />
+      </Stack>
+   );
+}
+
+const TestCollection: React.FC<{ collection: TestSessionCollection, onLoadMore?: () => void }> = (props) => {
+
+   const { modeColors } = getHordeStyling();
+   const { styles, stateColorStyles } = getStyling();
+   const colors = dashboard.getStatusColors();
+
+   const { collection, onLoadMore } = props;
+   const sectionsHandler = useRef<SectionCollection | undefined>(undefined);
+   const navigate = useNavigate();
+   if (!sectionsHandler.current) {
+      sectionsHandler.current = new SectionCollection();
+   }
+   const [onTheFlyTestNameFilter, setOnTheFlyTestNameFilter] = useState<string | undefined>(undefined);
+   const query = useQuery();
+
+   const filter: SectionFilter = {
+      suite: query.get(FilterType.suite) ?? undefined,
+      name: query.get(FilterType.name) ?? undefined,
+      project: query.get(FilterType.project) ?? undefined,
+      namecontains: query.get(FilterType.namecontains) ?? undefined,
+      meta: MetaFilterTools.Decode(query.getAll(FilterType.meta)),
+   };
+   const testuid = query.get(FilterType.testuid) ?? undefined;
+
+   const data = useMemo(() => testDataHandler.cursor ? collection.iterSessions().next().value : undefined, [collection]); // for single view mode
+
+   const sections = !onTheFlyTestNameFilter ?
+      sectionsHandler.current.extractSections(filter, collection) :
+      sectionsHandler.current.reExtractSectionOnTheFly(filter, onTheFlyTestNameFilter, collection);
+
+   const sectionBreadCrumbs: BreadCrumbFill[] = [{ title: 'All', linkTo: noSectionCrumbLink }];
+   const totalTests = sections.reduce((a, b) => a + b.Stats.TotalRun, 0);
+   const stackFailure: ChartStack = {
+      value: Math.ceil(sections.reduce((a, b) => a + b.Stats.Failed + b.Stats.Incomplete, 0) / (totalTests || 1) * 20) * 5,
+      title: "Failure",
+      color: stateColorStyles.get(TestState.Failed),
+      stripes: true
+   };
+   if (filter.suite) {
+      const item: BreadCrumbFill = { title: `${filter.suite} [${FilterType.suite}]`, linkTo: sectionCrumbLink(FilterType.suite, filter.suite!) };
+      sectionBreadCrumbs.push(item);
+   }
+   if (filter.name) {
+      const firstSection = sections.length > 0 ? sections[0] : undefined;
+      const commonSection = firstSection ? firstSection.FullName.slice(0, firstSection.FullName.length - firstSection.Name.length - 1) : filter.name;
+      commonSection.split('.').forEach((part, i, allParts) => {
+         const item: BreadCrumbFill = { title: part, linkTo: sectionCrumbLink(FilterType.name, allParts.slice(0, i + 1).join('.')) };
+         sectionBreadCrumbs.push(item);
+      });
+   }
+
+   const metaOptions = data ? new Map() : collection.getMetaValueByKey();
+
+   const addFilters: IContextualMenuProps = {
+      items: [
+         {
+            key: 'search_name',
+            onRender: (item, dismissMenu) => TestFilterCard(
+               filter,
+               metaOptions,
+               (value) => setOnTheFlyTestNameFilter(value),
+               (value) => {
+                  if (onTheFlyTestNameFilter && onTheFlyTestNameFilter === value) {
+                     sectionsHandler.current?.promoteOnTheFlyCache();
+                  }
+                  setOnTheFlyTestNameFilter(undefined);
+                  if (filter.namecontains !== value) {
+                     navigate(urlFromFilterType(FilterType.namecontains, value?.toLowerCase()));
+                  }
+               },
+               dismissMenu
+            )
+         }
+      ]
+   };
+
+   return <Stack className={styles.container} styles={{ root: { paddingTop: 8 } }} grow>
+      {!!data &&
+         <Stack styles={{ root: { paddingTop: 4, paddingLeft: 20, paddingBottom: 8, paddingRight: 8 } }}>
+            <Text>This test pass ran on <span style={{ fontWeight: 'bold' }}>{data.TestSessionInfo.DateTime}</span> for a duration of <span style={{ fontWeight: 'bold' }}>{msecToElapsed(data.TestSessionInfo.TimeElapseSec * 1000)}</span> for <span style={{ fontWeight: 'bold' }}>{data.MetaHandler?.map((key, value) => <span key={key}>{value} </span>)}</span></Text>
+         </Stack>
+      }
+
+      <Stack styles={{ root: { background: modeColors.crumbs, padding: 4, paddingLeft: 10 } }} horizontal>
+         {SectionBreadCrumb(sectionBreadCrumbs, [stackFailure], totalTests)}
+         <Stack horizontalAlign="end" disableShrink>
+            <ActionButton menuProps={addFilters}>
+               <FontIcon
+                  iconName={filter.namecontains || filter.meta ? "Warning" : "CirclePlus"}
+                  style={{ padding: 3, color: (filter.namecontains || filter.meta) ? colors.get(StatusColor.Warnings) : undefined }} /> Filters
+            </ActionButton>
+         </Stack>
+      </Stack>
+
+      <Stack className={styles.container} grow>
+         <Stack tokens={{ childrenGap: 12 }} styles={{ root: { paddingLeft: 20, paddingRight: 10, paddingTop: 6, paddingBottom: 6 } }}>
+            {sections.length > 0 && sections.map((item: Section) => SectionItem(item))}
+            {sections.length === 0 && <Stack><Stack.Item align="center" style={{ fontSize: 16 }}>No test found</Stack.Item></Stack>}
+         </Stack>
+      </Stack>
+
+      {testuid && <TestResultModal testuid={testuid} collection={collection} onLoadMore={onLoadMore} />}
+   </Stack>
+}
+
+export const TestSessionView: React.FC<any> = (props) => {
+   const sessions = useRef<TestSessionCollection | undefined>(undefined);
+   const item = testDataHandler.cursor;
+   let sessionHandler = item?.getDataHandler() as TestSessionWrapper;
+   if (item && !sessionHandler) {
+      sessionHandler = new TestSessionWrapper(item);
+      item.setDataHandler(sessionHandler);
+   }
+   if (!sessions.current) {
+      sessions.current = new TestSessionCollection(testDataHandler);
+   }
+   const session_uid = sessionHandler.getUID();
+   if (!sessions.current.has(session_uid)) {
+      sessions.current.clear() //Keep only one session when in single session view
+      sessions.current.set(session_uid, sessionHandler);
+   }
+   return <TestCollection collection={sessions.current} />
+}
+
+const getStreamBreadcrumbProps = (streamId: string, key: string) => {
+   const stream = getStreamData(streamId);
+   const projectName = getProjectName(streamId);
+   const streamName = stream?.name ?? "Unknown Stream";
+
+   const crumbItems: BreadcrumbItem[] = [
+      {
+         text: projectName,
+         link: `/project/${stream?.project?.id}`
+      },
+      {
+         text: streamName,
+         link: `/stream/${stream?.id}`
+      },
+      {
+         text: key
+      }
+   ];
+
+   return {
+      items: crumbItems,
+      title: `Horde - ${key}: ${projectName}/${streamName}`
+   }
+}
+
+export const TestSessionViewAll: React.FC = () => {
+   const [loading, setLoading] = useState(false);
+   const [maxCount, setMaxCount] = useState(500);
+   const [version, updateVersion] = useState(0);
+   const [crumbProps, setCrumbProps] = useState<{ items: BreadcrumbItem[], title: string }>({ items: [], title: "Test Report" });
+   const sessions = useRef<TestSessionCollection | undefined>(undefined);
+
+   const { streamId } = useParams<{ streamId: string }>();
+
+   const { hordeClasses } = getHordeStyling();
+
+   const count = testDataHandler.items?.size ?? 0;
+   const changeRange = !loading && count >= maxCount ? testDataHandler.getChanges() : undefined;
+
+   const loadMore = () => setMaxCount(maxCount + 500);
+
+   useEffect(() => {
+      if (!sessions.current || testDataHandler.streamId !== streamId) {
+         sessions.current = new TestSessionCollection(testDataHandler);
+      }
+      const versionChangeDisposer = sessions.current.reactOnVersionChange((version) => updateVersion(version))
+      const fetchData = async () => {
+         setCrumbProps(getStreamBreadcrumbProps(streamId!, "Automated Test Session"));
+         setLoading(true);
+         try {
+            await sessions.current?.getSessionsForStream(streamId!, maxCount);
+            if (sessions.current?.version === 0 && testDataHandler.activated) {
+               // No data was fetched, but we still want to mount the TestCollection view
+               sessions.current.updateVersion();
+            }
+         } finally { testDataHandler.activated && setLoading(false); }
+      }
+      fetchData();
+
+      return function cleanup() { versionChangeDisposer(); testDataHandler.desactivate(); }
+   }, [maxCount, streamId]);
+
+   const { styles } = getStyling();
+
+   return (
+      <Stack className={hordeClasses.horde} style={{ height: '100vh' }}>
+         <TopNav />
+         <Breadcrumbs items={crumbProps.items.concat()} title={crumbProps.title} />
+         <Stack className={styles.container} grow>
+            {loading && sessions.current && <LoaderProgressBar loader={sessions.current} target={maxCount} />}
+            {changeRange &&
+               <Stack className={styles.bottomLoadMore} tokens={{ childrenGap: 8 }} horizontal horizontalAlign="center">
+                  <Stack>Changelist range [#{changeRange[0]} - #{changeRange[changeRange.length - 1]}] -</Stack>
+                  <Stack style={{ color: "rgb(0, 120, 212)", cursor: 'pointer' }} onClick={loadMore}>Load more data...</Stack>
+               </Stack>
+            }
+            {version > 0 && sessions.current && <TestCollection collection={sessions.current} onLoadMore={loadMore} />}
+         </Stack>
+      </Stack>
+   );
+}
+
+const LoaderProgressBar: React.FC<{ loader: Loader, target: number }> = ({ loader, target }) => {
+   const [progress, setProgress] = useState(loader.loadingProgress);
+   const whenUpdate = useRef<Promise<void> & { cancel: () => void } | undefined>(undefined);
+
+   whenUpdate.current = loader.whenLoadUpdate();
+   whenUpdate.current.then(
+      () => testDataHandler.activated && setProgress(loader.loadingProgress)
+   ).catch((cancel) => null /*ignore cancel event*/);
+
+   useEffect(() => {
+      return function cleanup() { whenUpdate.current?.cancel(); }
+   }, []);
+
+   return (
+      <Stack>
+         <ProgressIndicator percentComplete={target > 1 && target !== progress ? progress / target : undefined} barHeight={7} />
+      </Stack>
+   );
+};
+
+export const routePath: string = "/automatedtestsession/:streamId";
+
+export const BuildHealthTestSessionView: React.FC<{ streamId: string }> = (props) => {
+   const { streamId } = props;
+
+   const theme = getHordeTheme();
+
+   return (
+      <DefaultButton href={generatePath(routePath, { streamId: streamId })} text="Test Results" style={{ color: theme.semanticColors.bodyText }} />
+   )
+}
diff --git a/Engine/Source/Programs/Horde/HordeDashboard/src/legacyPlugins/AutomatedTestSession/index.tsx b/Engine/Source/Programs/Horde/HordeDashboard/src/legacyPlugins/AutomatedTestSession/index.tsx
new file mode 100644
--- /dev/null
+++ b/Engine/Source/Programs/Horde/HordeDashboard/src/legacyPlugins/AutomatedTestSession/index.tsx
@@ -0,0 +1,30 @@
+// Copyright Epic Games, Inc. All Rights Reserved.
+
+import { Plugin, PluginMount } from "../pluginTypes";
+import { TestSessionView, TestSessionViewAll, BuildHealthTestSessionView, routePath } from "./components/AutomatedTestSession";
+
+class AutomatedTestSessionPlugin implements Plugin {
+
+    name = "Automated Test Session";
+    components = [
+        {
+            mount: PluginMount.TestReportPanel,
+            component: TestSessionView,
+            id: "Automated Test Session"
+        },
+        {
+            mount: PluginMount.BuildHealthSummary,
+            component: BuildHealthTestSessionView,
+        }
+    ];
+    routes = [
+        {
+            path: routePath,
+            component: TestSessionViewAll
+        }
+    ];
+}
+
+const automatedTestSessionPlugin = new AutomatedTestSessionPlugin();
+export default automatedTestSessionPlugin;
+
diff --git a/Engine/Source/Programs/Horde/HordeDashboard/src/legacyPlugins/AutomatedTestSession/models/AutomatedTestSession.ts b/Engine/Source/Programs/Horde/HordeDashboard/src/legacyPlugins/AutomatedTestSession/models/AutomatedTestSession.ts
new file mode 100644
--- /dev/null
+++ b/Engine/Source/Programs/Horde/HordeDashboard/src/legacyPlugins/AutomatedTestSession/models/AutomatedTestSession.ts
@@ -0,0 +1,1280 @@
+// Copyright Epic Games, Inc. All Rights Reserved.
+
+import { observable, action, when, reaction, makeObservable } from "mobx";
+import { DataWrapper, TestDataWrapper, TestDataCollection } from '../../../backend/TestDataHandler'
+
+export enum TestState {
+    Failed = "Fail",
+    Success = "Success",
+    SuccessWithWarnings = "SuccessWithWarnings",
+    NotRun = "NotRun",
+    InProcess = "InProcess",
+    Skipped = "Skipped",
+    Unknown = "Unknown",
+}
+
+const TestStateRank: { [Key in string]: number } = {
+    "Fail": 0,
+    "InProcess": 1,
+    "NotRun": 2,
+    "SuccessWithWarnings": 3,
+    "Success": 4,
+    "Skipped": 5,
+    "Unknown": 6,
+};
+
+export enum EventType {
+    Info = "Info",
+    Error = "Error",
+    Warning = "Warning",
+}
+
+export enum ArtifactType {
+    ImageCompare = "image comparison",
+    Approved = "approved",
+    Unapproved = "unapproved",
+    Difference = "difference",
+}
+
+export type TestArtifact = {
+    Tag: string;
+    ReferencePath: string;
+}
+
+export type TestEvent = {
+    Message: string;
+    Context: string;
+    Type: string;
+    Tag: string;
+    Hash: string;
+    DateTime: string;
+    Artifacts: TestArtifact[];
+}
+
+export type TestResult = {
+    Name: string;
+    TestUID: string;
+    Suite: string;
+    State: string;
+    DeviceAppInstanceName: string | string[];
+    ErrorCount: number;
+    WarningCount: number;
+    ErrorHashAggregate: string;
+    DateTime: string;
+    TimeElapseSec: number;
+    // references that need to be initiated post fetch
+    Details: TestResultDetails;
+    Session: TestSessionWrapper;
+}
+
+export type TestResultDetails = {
+    Events: TestEvent[];
+}
+
+export type TestDevice = MetadataHolder & {
+    Name: string;
+    AppInstanceName: string;
+    AppInstanceLog: string;
+}
+
+export type TestSessionInfo = {
+    DateTime: string;
+    TimeElapseSec: number;
+    Tests: TestResultsMap;
+    TestResultsTestDataUID: string;
+}
+
+export type IndexedError = {
+    Message: string;
+    Tag: string;
+    TestUIDs: string[];
+}
+
+export type Metadata = { [Key in string]: string | undefined }
+
+type MetadataHolder = {
+    Metadata?: Metadata
+}
+
+export type Loader = {
+    loadingProgress: number;
+    whenLoadUpdate: () => Promise<void> & { cancel: () => void };
+}
+
+const generatetHashFromString = (aggregate: string) => {
+    let hash = 0;
+    for (let i = 0; i < aggregate.length; i++) {
+        const char = aggregate.charCodeAt(i);
+        hash = ((hash << 5) - hash) + char;
+        hash = hash & hash;
+    }
+    return hash.toString(16);
+}
+
+export const generateHashFromMetadata = (obj: Metadata) => {
+    const keys = Object.keys(obj).sort();
+    const aggregate = keys.map((key) => key + (obj[key] ?? '')).join();
+    return generatetHashFromString(aggregate);
+}
+
+export class MetaWrapper {
+    private collection?: Metadata
+    private _hash: string
+
+    constructor(collection?: Metadata) {
+        this.collection = collection;
+        this._hash = "";
+    }
+
+    get(key: string): string | undefined {
+        if (this.collection) {
+            return this.collection[key];
+        }
+
+        return undefined;
+    }
+
+    map(mapFunc: (key: string, value?: string) => any): any[] {
+        if (this.collection) {
+            return Object.keys(this.collection).map((key) => mapFunc(key, this.collection![key]));
+        }
+        return [];
+    }
+
+    forEach(predicate: (key: string, value?: string) => void) {
+        if (this.collection) {
+            Object.keys(this.collection).forEach((key) => predicate(key, this.collection![key]));
+        }
+    }
+
+    maskByKeys(keyMask: string[]) {
+        return this.map((key, value) => keyMask.includes(key) ? undefined : value).filter((info) => info)
+    }
+
+    private generateHashCode(): string {
+        if (this.collection) {
+            return generateHashFromMetadata(this.collection);
+        }
+        return "";
+    }
+
+    get hash(): string {
+        if (!this._hash) {
+            this._hash = this.generateHashCode();
+        }
+        return this._hash;
+    }
+}
+
+export type IndexedErrorMap = { [Key in string]: IndexedError }
+
+export type TestResultDetailsMap = { [Key in string]: TestResultDetails }
+
+export type TestResultsMap = { [Key in string]: TestResult }
+
+export type TestSession = MetadataHolder & {
+    Name: string;
+    Type: string;
+    PreFlightChange: string;
+    TestSessionInfo: TestSessionInfo;
+    IndexedErrors: IndexedErrorMap;
+    Devices: TestDevice[];
+}
+
+type CheckList = { [Key in string]: boolean }
+
+export enum SessionFields {
+    initiate = "data.Metadata,data.Name,data.PreFlightChange",
+    remaining = "data.TestSessionInfo,data.IndexedErrors,data.Devices,data.Type",
+}
+
+export class TestSessionWrapper implements TestSession, DataWrapper {
+    get Name(): string { return this.getOwnProperty('Name'); }
+    get Type(): string { return this.getOwnProperty('Type'); }
+    get PreFlightChange(): string { return this.getOwnProperty('PreFlightChange'); }
+    get TestSessionInfo(): TestSessionInfo { return this.getOwnProperty('TestSessionInfo'); }
+    get IndexedErrors(): IndexedErrorMap { return this.getOwnProperty('IndexedErrors'); }
+    get Devices(): TestDevice[] { return this.getOwnProperty('Devices'); }
+    get Metadata(): Metadata { return this.getOwnProperty('Metadata'); }
+
+    Testdata: TestDataWrapper;
+    MetaHandler?: MetaWrapper;
+
+    @observable
+    propertiesLoaded: CheckList = {
+        TestSessionInfo: false
+    };
+    @action
+    setPropertiesLoaded(...properties: string[]) {
+        properties.forEach((key) => {
+            this.propertiesLoaded[key] = true;
+        })
+    }
+    isPropertyLoaded(key: string) {
+        return this.propertiesLoaded[key];
+    }
+    isAllPropertiesLoaded() {
+        return this.propertiesLoaded['TestSessionInfo']; // check TestSessionInfo only for now
+    }
+    getMissingProperties() {
+        return SessionFields.remaining.split(',').map((item) => item.slice(5));
+    }
+    whenPropertyLoaded(key: string) {
+        return when(() => this.propertiesLoaded[key]);
+    }
+
+    private _testsession: TestSession;
+
+    constructor(testdata: TestDataWrapper) {
+        makeObservable(this);
+        this._testsession = testdata.data;
+        this.Testdata = testdata;
+        // first initialization, we should call the post Properties Update
+        this.postPropertiesUpdate();
+    }
+
+    getOwnProperty<T>(key: string): T {
+        return Object.getOwnPropertyDescriptor(this._testsession, key)?.value as T;
+    }
+
+    getUID(): string {
+        return this.Name + (this.MetaHandler?.hash ?? '');
+    }
+
+    postPropertiesUpdate() {
+        this.capitalizeKeys()
+        const data = this._testsession;
+        if (!this.MetaHandler && data.Metadata) {
+            this.MetaHandler = new MetaWrapper(data.Metadata);
+        }
+        if (data.TestSessionInfo && !this.isPropertyLoaded('TestSessionInfo')) {
+            for (const key in data.TestSessionInfo.Tests) {
+                const test = data.TestSessionInfo.Tests[key];
+                test.Session = this;
+            }
+            this.setPropertiesLoaded('TestSessionInfo');
+        }
+    }
+
+    getTestResultByUID(uid: string) {
+        if (this.TestSessionInfo) {
+            return this.TestSessionInfo.Tests[uid];
+        }
+
+        return undefined
+    }
+
+    getDevices(AppInstanceName: string | string[]) {
+        if (this.Devices) {
+            if (typeof (AppInstanceName) === 'string') {
+                const asString = AppInstanceName as string;
+                const device = this.Devices.find((item) => item.AppInstanceName === asString);
+                if (device) {
+                    return [device];
+                }
+            }
+            else {
+                const asArray = AppInstanceName as string[];
+                return this.Devices.filter((item) => asArray.includes(item.AppInstanceName))
+            }
+        }
+
+        return undefined
+    }
+
+    private capitalizeKeys() {
+        // Fix inconsistances in object keys fetched from Horde backend.
+        const data = this._testsession as any;
+        for (var key in data) {
+            const first = key.charAt(0);
+            const firstUpper = first.toUpperCase();
+            if (first !== firstUpper) {
+                const value = data[key];
+                delete data[key];
+                data[firstUpper + key.slice(1)] = value;
+            }
+        }
+    }
+}
+
+export class TestStats {
+    Failed: number = 0;
+    Passed: number = 0;
+    Skipped: number = 0;
+    Unexecuted: number = 0;
+    Incomplete: number = 0;
+    TotalRun: number = 0;
+
+    update(state: string, increment: boolean = true) {
+        switch (state) {
+            case TestState.Failed:
+                increment ? this.Failed++ : this.Failed--;
+                increment ? this.TotalRun++ : this.TotalRun--
+                break;
+            case TestState.NotRun:
+                increment ? this.Unexecuted++ : this.Unexecuted--;
+                break;
+            case TestState.Success:
+            case TestState.SuccessWithWarnings:
+                increment ? this.Passed++ : this.Passed--;
+                increment ? this.TotalRun++ : this.TotalRun--
+                break;
+            case TestState.InProcess:
+                increment ? this.Incomplete++ : this.Incomplete--;
+                increment ? this.TotalRun++ : this.TotalRun--
+                break;
+            case TestState.Skipped:
+                increment ? this.Skipped++ : this.Skipped--;
+                break;
+        }
+    }
+
+    updateFromStats(stats: TestStats) {
+        this.Failed += stats.Failed;
+        this.Passed += stats.Passed;
+        this.Skipped += stats.Skipped;
+        this.Unexecuted += stats.Unexecuted;
+        this.Incomplete += stats.Incomplete;
+        this.TotalRun += stats.TotalRun;
+    }
+}
+
+export class TestCase implements Loader {
+    Name: string;
+    DisplayName: string;
+    TestUID: string;
+    SessionName: string;
+    Suite: string;
+    State: string;
+    StateMetaUID: string | undefined;
+    Metas: Map<string, TestResult>;
+    Stats: TestStats;
+
+    @observable
+    version: number = 0;
+    @action
+    updateVersion() {
+        this.version++;
+    }
+    whenVersionChange() {
+        const start = this.version;
+        return when(() => this.version !== start);
+    }
+
+    @observable
+    history: number = 0;
+    @action
+    updateHistory() {
+        this.history++;
+    }
+    reactOnHistoryChange(callback: (version: number) => void) {
+        return reaction(
+            () => this.history,
+            (version) => callback(version)
+        );
+    }
+
+    @observable
+    historyTarget: number = 0;
+    @action
+    incrementHistoryTarget() {
+        this.historyTarget++;
+    }
+    get loadingProgress() { return this.history }
+    whenLoadUpdate() {
+        const previousVersion = this.history;
+        return when(() => this.version !== previousVersion);
+    }
+
+    constructor(name: string, testuid: string, suite: string, sessionName: string, displayName?: string) {
+        makeObservable(this);
+        this.Name = name;
+        this.DisplayName = displayName ?? this.toDislayName(name);
+        this.TestUID = testuid;
+        this.State = TestState.Unknown;
+        this.StateMetaUID = undefined;
+        this.Suite = suite;
+        this.SessionName = sessionName;
+        this.Metas = new Map();
+        this.Stats = new TestStats();
+    }
+
+    removeMeta(meta: MetaWrapper) {
+        const metauid = meta.hash;
+        const test = this.Metas.get(metauid);
+        if (!test) {
+            return false;
+        }
+        this.Metas.delete(metauid);
+        this.Stats.update(test.State, false);
+        this.updateVersion();
+        return true;
+    }
+
+    isEmpty() {
+        return this.Metas.size === 0;
+    }
+
+    addMeta(item: TestResult) {
+        const metauid = item.Session.MetaHandler!.hash;
+        if (this.Metas.has(metauid)) {
+            return false;
+        }
+        this.Metas.set(metauid, item);
+        let state = item.State;
+        if (state === TestState.Success && item.WarningCount > 0) {
+            state = TestState.SuccessWithWarnings;
+            item.State = state;
+        }
+        if (this.setState(state)) {
+            this.StateMetaUID = metauid;
+        }
+        this.Stats.update(state);
+        this.updateVersion();
+        return true;
+    }
+
+    forEachMeta(predicate: (test: TestResult) => void) {
+        this.Metas.forEach(predicate);
+    }
+
+    getResults(sorted: boolean = false) {
+        const results = Array.from(this.Metas.values());
+        if (sorted) {
+            results.sort(this.sortResults);
+        }
+        return results;
+    }
+
+    getFilteredByMeta(metafilter: MetadataFilter) {
+        const metas = Array.from(this.Metas.values()).filter(
+            (test) => Object.keys(metafilter).every((key) => {
+                const sessionKey = test.Session.Metadata[key];
+                return sessionKey && metafilter[key].includes(sessionKey);
+            })
+        );
+        if (metas.length === 0) {
+            return undefined;
+        }
+        const testcase = new TestCase(this.Name, this.TestUID, this.Suite, this.SessionName, this.DisplayName);
+        metas.forEach((test) => testcase.addMeta(test));
+        return testcase;
+    }
+
+    getFilteredByState(state: TestState) {
+        const targetRank = TestStateRank[state];
+        const metas = Array.from(this.Metas.values()).filter(
+            (test) => targetRank <= TestStateRank[test.State]
+        );
+        if (metas.length === 0) {
+            return undefined;
+        }
+        const testcase = new TestCase(this.Name, this.TestUID, this.Suite, this.SessionName, this.DisplayName);
+        metas.forEach((test) => testcase.addMeta(test));
+        return testcase;
+    }
+
+    private toDislayName(name: string): string {
+        // It is necessary to sanitize the test names that contains / instead of . syntax
+        if (name.indexOf('/') > -1) {
+            return name.replaceAll(/\.*\//g, ".");
+        }
+        return name;
+    }
+
+    private setState(state: string) {
+        if (this.State !== state) {
+            const currentRank = TestStateRank[this.State];
+            const newRank = TestStateRank[state];
+            if (newRank < currentRank) {
+                this.State = state;
+                return true;
+            }
+        }
+        return false;
+    }
+
+    private sortResults(a: TestResult, b: TestResult) {
+        const aRank = TestStateRank[a.State];
+        const bRank = TestStateRank[b.State];
+        if (aRank > bRank) {
+            return 1;
+        } else if (aRank < bRank) {
+            return -1;
+        }
+        return 0;
+    }
+}
+
+export class Section {
+    Name: string;
+    FullName: string;
+    Type: FilterType;
+    Stats: TestStats;
+    TestCases: TestCase[];
+    TestNames: Set<string>;
+
+    constructor(name: string, fullName: string, type: FilterType) {
+        this.Name = name;
+        this.FullName = fullName;
+        this.Type = type;
+        this.Stats = new TestStats();
+        this.TestCases = [];
+        this.TestNames = new Set();
+    }
+
+    addTestCase(testcase: TestCase, testName?: string) {
+        this.TestCases.push(testcase);
+        this.TestNames.add(testName ?? testcase.DisplayName);
+        this.Stats.updateFromStats(testcase.Stats);
+    }
+
+    get isOneTest(): boolean {
+        return this.TestNames.size === 1 && this.TestNames.has(this.FullName)
+    }
+}
+
+export enum FilterType {
+    suite = 'suite',
+    name = 'name',
+    project = 'project',
+    testuid = 'testuid',
+    metauid = 'metauid',
+    testchange = 'testchange',
+    view = 'view',
+    namecontains = 'namecontains',
+    meta = 'meta',
+}
+
+export type MetadataFilter = { [Key in string]: string[] }
+export class MetaFilterTools {
+    static Match = (key: string, target: string): boolean => target.startsWith(`${key}.`)
+    static Encode = (key: string, items: string[]): string => `${key}.${items.join('~')}`
+    static Decode = (items: string[]): MetadataFilter | undefined => {
+        return items.length === 0 ? undefined :
+            Object.fromEntries(items.map((item) => {
+                const [key, value] = item.split('.', 2);
+                return [key, value.split('~')]
+            }))
+    }
+}
+
+export type SectionFilter = {
+    suite?: string;
+    name?: string;
+    project?: string;
+    namecontains?: string;
+    meta?: MetadataFilter;
+}
+
+export type ReportFilter = {
+    testuid?: string;
+    metauid?: string;
+    testchange?: number;
+    view?: string;
+}
+
+export const generateHashFromSectionFilter = (obj: SectionFilter) => {
+    const baseFilter: Metadata = Object.fromEntries(
+        Object.keys(obj).filter((key) => key !== 'meta').map((key) => [
+            key, Object.getOwnPropertyDescriptor(obj, key)?.value
+        ])
+    );
+    let aggregate = Object.keys(baseFilter).sort().map((key) => key + (baseFilter[key] ?? '')).join();
+    const subObj = obj.meta;
+    if (subObj !== undefined) {
+        aggregate += 'meta' + Object.keys(subObj).sort().map((key) => key + (subObj[key]?.join('~') ?? '')).join();
+    }
+    return generatetHashFromString(aggregate);
+}
+
+type GetKeyFunction = (testcase: TestCase) => [string, string];
+const regexNameFilter = /"([^"]+)"|([^\s"]+)/g;
+
+export class SectionCollection {
+
+    private filterCache: Map<string, { hash: string, sections: Section[] }>;
+    private onTheFlyCache: { filter: string, hash: string, sections: Section[] } | undefined;
+
+    constructor() {
+        this.filterCache = new Map();
+    }
+
+    extractSections(filter: SectionFilter, sessions: TestSessionCollection) {
+        // Get the sections, first look for in the cache, otherwise generate it
+        const filterCacheKey = generateHashFromSectionFilter(filter);
+        const cache = this.filterCache.get(filterCacheKey);
+        const isSameHash = cache && cache.hash === sessions.hash;
+        const sections = cache && isSameHash ? cache.sections : this.filterSections(filter, sessions.getTestCases());
+        if (!this.filterCache.has(filterCacheKey) || !isSameHash) {
+            this.filterCache.set(filterCacheKey, { hash: sessions.hash, sections: sections })
+        }
+        return sections;
+    }
+
+    reExtractSectionOnTheFly(filter: SectionFilter, testNameFilter: string, sessions: TestSessionCollection) {
+        // to use only when iterating through cached data
+        let sections = this.extractSections({ ...filter, namecontains: undefined }, sessions);
+        const onTheFlyFilter = { ...filter, namecontains: testNameFilter.toLowerCase() }
+        sections = this.filterSections(onTheFlyFilter, sections.flatMap((section) => section.TestCases));
+        this.onTheFlyCache = { filter: generateHashFromSectionFilter(onTheFlyFilter), hash: sessions.hash, sections };
+        return sections;
+    }
+
+    promoteOnTheFlyCache() {
+        if (this.onTheFlyCache) {
+            this.filterCache.set(this.onTheFlyCache.filter, this.onTheFlyCache);
+            this.onTheFlyCache = undefined;
+        }
+    }
+
+    private filterSections(filter: SectionFilter, tests: TestCase[]) {
+        const filterFunc = (test: TestCase) => {
+            if (filter.suite) {
+                if (test.Suite !== filter.suite) {
+                    return false;
+                }
+            }
+            if (filter.namecontains) {
+                const substrings = Array.from(filter.namecontains.matchAll(regexNameFilter), match => match[1] || match[2]);
+                const testName = test.Name.toLowerCase();
+                if (!substrings.some((value) => testName.indexOf(value) >= 0)) {
+                    return false;
+                }
+            }
+            if (filter.name) {
+                if (test.DisplayName.indexOf(filter.name + '.') !== 0) {
+                    return false;
+                }
+            }
+            return true;
+        }
+        const filteredTests = tests.reduce((filtered: TestCase[], item: TestCase) => {
+            if (filterFunc(item)) {
+                if (filter.meta === undefined) {
+                    filtered.push(item);
+                } else {
+                    const filteredCase = item.getFilteredByMeta(filter.meta);
+                    if (filteredCase) {
+                        filtered.push(filteredCase);
+                    }
+                }
+            }
+            return filtered
+        }, []);
+
+        const commonSection = filter.name ? this.getCommonSection(filteredTests.map((item) => item.DisplayName)) : "";
+        const sectionKeyFunc = (key: FilterType): GetKeyFunction | undefined => {
+            switch (key) {
+                case FilterType.suite:
+                    return (testcase: TestCase) => {
+                        return [testcase.Suite, testcase.Suite];
+                    }
+                case FilterType.name:
+                    return (testcase: TestCase) => {
+                        const suite = testcase.DisplayName.slice(commonSection.length).split('.', 1)[0];
+                        return [suite, commonSection + suite];
+                    }
+            }
+        };
+        const sectionFocus = filter.suite || filter.name ? FilterType.name : FilterType.suite;
+        const sections = this.buildSections(filteredTests, sectionFocus, sectionKeyFunc(sectionFocus)!);
+        // sorting
+        sections.forEach((section) => this.sortSectionResults(section));
+        this.sortSections(sections);
+        return sections;
+    }
+
+    private buildSections(testcases: TestCase[], type: FilterType, keyFunc: GetKeyFunction): Section[] {
+        const sections: Section[] = [];
+        const sectionsMap = new Map<string, Section>();
+
+        testcases.forEach((test) => {
+            const [key, fullKey] = keyFunc(test);
+            let targetSection = sectionsMap.get(key);
+            if (!targetSection) {
+                targetSection = new Section(key, fullKey, type);
+                sections.push(targetSection);
+                sectionsMap.set(key, targetSection);
+            }
+            targetSection.addTestCase(test, test.DisplayName);
+        });
+
+        return sections;
+    }
+
+    private getCommonString(list: string[]): string {
+        if (list.length === 0) {
+            return "";
+        }
+        let shortest: string = list[0];
+        let longest: string = "";
+        list.forEach((item) => {
+            if (item > longest) {
+                longest = item;
+                return;
+            }
+            if (item < shortest) {
+                shortest = item;
+            }
+        });
+        while (longest.indexOf(shortest) !== 0) {
+            shortest = shortest.slice(0, -1)
+        }
+
+        return shortest;
+    }
+
+    private getCommonSection(list: string[]): string {
+        let commonKey = this.getCommonString(list);
+        if (commonKey[commonKey.length - 1] !== '.' && commonKey.indexOf('.') >= 0) {
+            commonKey = commonKey.slice(0, commonKey.lastIndexOf('.') + 1);
+        }
+
+        return commonKey;
+    }
+
+    // sort comparing the number of Failed, Incomplete, Unexecuted and Passed
+    private sortSections(sections: Section[]) {
+        sections.sort((a, b) => {
+            const aStats = a.isOneTest ? a.TestCases[0].Stats : a.Stats;
+            const bStats = b.isOneTest ? b.TestCases[0].Stats : b.Stats;
+            if (aStats.Failed > bStats.Failed) {
+                return -1;
+            } else if (aStats.Failed < bStats.Failed) {
+                return 1;
+            } else if (aStats.Incomplete > bStats.Incomplete) {
+                return -1;
+            } else if (aStats.Incomplete < bStats.Incomplete) {
+                return 1;
+            } else if (aStats.Unexecuted > bStats.Unexecuted) {
+                return -1;
+            } else if (aStats.Unexecuted < bStats.Unexecuted) {
+                return 1;
+            } else if (aStats.Passed > bStats.Passed) {
+                return -1;
+            } else if (aStats.Passed < bStats.Passed) {
+                return 1;
+            } else if (a.isOneTest && !b.isOneTest) {
+                return 1;
+            } else if (b.isOneTest && !a.isOneTest) {
+                return -1;
+            }
+            return 0;
+        });
+    }
+
+    // sort in that order Failed, Incomplete, Unexecuted, Passed, Skipped
+    private sortSectionResults(section: Section) {
+        SectionCollection.sortTestCases(section.TestCases);
+    }
+
+    static sortTestCases(testcases: TestCase[]) {
+        testcases.sort((a, b) => {
+            const aRank = TestStateRank[a.State];
+            const bRank = TestStateRank[b.State];
+            if (aRank > bRank) {
+                return 1;
+            } else if (aRank < bRank) {
+                return -1;
+            }
+
+            if (a.Name > b.Name) {
+                return 1;
+            } else if (a.Name < b.Name) {
+                return -1;
+            }
+
+            return 0;
+        })
+    }
+}
+
+type SessionReferences = {
+    SessionId: string;
+    ResultsId: string;
+}
+
+export class TestSessionCollection implements Loader {
+
+    @observable
+    version: number = 0;
+    @action
+    updateVersion() {
+        this.version++;
+    }
+    reactOnVersionChange(callback: (version: number) => void) {
+        return reaction(
+            () => this.version,
+            (version) => callback(version)
+        );
+    }
+
+    @observable
+    loadingCount: number = 0;
+    @action
+    incrementLoading() {
+        this.loadingCount++;
+    }
+    @action
+    decrementLoading() {
+        this.loadingCount--;
+    }
+    whenLoadUpdate() {
+        const previousCount = this.loadingCount;
+        return when(() => this.loadingCount !== previousCount)
+    }
+    get loadingProgress(): number { return this._datahandler.items?.size ?? 0 }
+
+    get size(): number { return this._collection.size; }
+    get hash(): string {
+        let aggregate = "";
+        this.forEachSession( // only include fully loaded session
+            (session) => aggregate += session.isAllPropertiesLoaded() ? session.getUID() : ""
+        );
+        return generatetHashFromString(aggregate);
+    }
+
+    private _collection: Map<string, TestSessionWrapper>;
+    private _testcases: Map<string, TestCase>;
+    private _datahandler: TestDataCollection;
+    private _includePreflight: boolean = false;
+    private _cachedMetaValueByKey: { hash: string, value: Map<string, Set<string>> } | undefined;
+
+    constructor(handler: TestDataCollection, includePreflight: boolean = false) {
+        makeObservable(this);
+        this._collection = new Map();
+        this._testcases = new Map();
+        this._datahandler = handler;
+        this._includePreflight = includePreflight;
+    }
+
+    set(uid: string, item: TestSessionWrapper) {
+        this._collection.set(uid, item);
+        if (item.isAllPropertiesLoaded()) {
+            this.indexTestsFromSession(item);
+            this.updateVersion();
+        }
+    }
+
+    get(uid: string) {
+        return this._collection.get(uid);
+    }
+
+    has(uid: string) {
+        return this._collection.has(uid);
+    }
+
+    delete(uid: string) {
+        const session = this._collection.get(uid);
+        this._collection.delete(uid);
+        // clean up testuids map
+        if (session && session.TestSessionInfo) {
+            Object.keys(session.TestSessionInfo.Tests).forEach((key) => {
+                const testcase = this._testcases.get(key)!;
+                if (testcase.removeMeta(session.MetaHandler!)) {
+                    if (testcase.isEmpty()) {
+                        this._testcases.delete(key)
+                    }
+                }
+            });
+        }
+        this.updateVersion();
+    }
+
+    clear() {
+        this._collection.clear()
+        this._testcases.clear()
+        this.updateVersion();
+    }
+
+    forEachSession(func: (item: TestSessionWrapper) => void) {
+        this._collection.forEach(func);
+    }
+
+    getMetaValueByKey() {
+        if (this._cachedMetaValueByKey !== undefined) {
+            if (this._cachedMetaValueByKey.hash === this.hash) {
+                return this._cachedMetaValueByKey.value;
+            }
+        }
+        const keys: Map<string, Set<string>> = new Map();
+        this.forEachSession((item) => {
+            const meta = item.MetaHandler;
+            if (meta) {
+                meta.forEach((key, value) => {
+                    if (value) {
+                        let valueSet = keys.get(key);
+                        if (!valueSet) {
+                            valueSet = new Set();
+                            keys.set(key, valueSet);
+                        }
+                        valueSet.add(value);
+                    }
+                });
+            }
+        });
+        this._cachedMetaValueByKey = { hash: this.hash, value: keys };
+        return keys;
+    }
+
+    *iterSessions() {
+        const iterator = this._collection.values();
+        let result = iterator.next();
+        while (!result.done) {
+            yield result.value;
+            result = iterator.next();
+        }
+    }
+
+    *iterTests() {
+        const iterSession = this._collection.values();
+        let resultSession = iterSession.next();
+        while (!resultSession.done) {
+            const sessionInfo = resultSession.value.TestSessionInfo;
+            if (sessionInfo) {
+                for (const key in sessionInfo.Tests) {
+                    yield sessionInfo.Tests[key];
+                }
+            }
+            resultSession = iterSession.next();
+        }
+    }
+
+    async getSessionsForStream(streamId: string, maxCount: number = 1000) {
+        this.incrementLoading();
+        const onFetch = (items: (TestDataWrapper | undefined)[]) => this.onInitialFetch(items);
+        try {
+            if (streamId === this._datahandler.streamId
+                && this._datahandler.items
+                && maxCount > this._datahandler.items.size) {
+                // Max count resize
+                this._datahandler.setActive(true);
+                await this._datahandler.queueFetchItems(onFetch, SessionFields.initiate, maxCount);
+            } else {
+                // Initial fetch or reload
+                await this._datahandler.setFromStream(
+                    streamId, "Automated Test Session", onFetch, SessionFields.initiate, maxCount
+                );
+            }
+        } finally { this.decrementLoading() }
+        await when(() => this.loadingCount === 0);
+    }
+
+    private onInitialFetch(items: (TestDataWrapper | undefined)[]) {
+        const sessionsToFetch: string[] = [];
+        items.forEach((item) => {
+            if (item) {
+                let data = item.getDataHandler() as TestSessionWrapper;
+                if (!data) {
+                    data = new TestSessionWrapper(item);
+                    item.setDataHandler(data);
+                }
+                if (!this._includePreflight && data.PreFlightChange) {
+                    return;
+                }
+                const session_uid = data.getUID();
+                if (!this.has(session_uid)) {
+                    this.set(session_uid, data);
+                    if (!data.isAllPropertiesLoaded()) {
+                        sessionsToFetch.push(item.id);
+                    }
+                }
+            }
+        });
+        if (sessionsToFetch.length > 0) {
+            this.incrementLoading();
+            this._datahandler.fetchUpdateItems(
+                sessionsToFetch, (items) => this.onRemainingFetch(items), SessionFields.remaining
+            ).finally(() => this.decrementLoading());
+        }
+    }
+
+    private onRemainingFetch(items: (TestDataWrapper | undefined)[]) {
+        let hasUpdated = false;
+        items.forEach((item) => {
+            if (item) {
+                const data = item.getDataHandler() as TestSessionWrapper;
+                if (data && data.isAllPropertiesLoaded()) {
+                    this.indexTestsFromSession(data);
+                    hasUpdated = true;
+                }
+            }
+        });
+        if (hasUpdated) { this.updateVersion() }
+    }
+
+    private indexTestsFromSession(session: TestSessionWrapper) {
+        const tests = Object.values(session.TestSessionInfo!.Tests);
+        if (tests.length === 0) {
+            return;
+        }
+        const first = this._testcases.get(tests[0].TestUID);
+        if (first && first.SessionName !== session.Name) {
+            return; // Skip conflicting session name, most likely session was renamed.
+        }
+        tests.forEach((item) => {
+            const testuid = item.TestUID;
+            let testcase = this._testcases.get(testuid);
+            if (!testcase) {
+                testcase = new TestCase(item.Name, testuid, item.Suite, session.Name);
+                this._testcases.set(testuid, testcase);
+            }
+            else if (testcase.SessionName !== session.Name) {
+                console.warn(`'${testcase.DisplayName}' is in another session. Appears in '${testcase.SessionName}' and '${session.Name}'.`);
+            }
+            testcase.addMeta(item);
+        });
+    }
+
+    async getTestCaseByUID(uid: string): Promise<TestCase | undefined> {
+        let testcase: TestCase | undefined = undefined;
+        const toFetch: TestSessionWrapper[] = [];
+        const testcases = this._testcases;
+
+        function getTestCase() {
+            testcase = testcases.get(uid);
+            if (testcase) {
+                testcase.forEachMeta((test) => {
+                    if (test.Details === undefined) {
+                        toFetch.push(test.Session);
+                    }
+                });
+            }
+        }
+        getTestCase();
+
+        while (!testcase && toFetch.length === 0 && this.loadingCount !== 0) {
+            // Give a chance to the collection to load a bit more if nothing was found
+            // It is possible the test is run infrequently and is deep in the result collection
+            const previousVersion = this.version;
+            await when(() => this.version !== previousVersion || this.loadingCount === 0);
+            // then retry
+            getTestCase();
+        }
+
+        if (toFetch.length > 0) {
+            // Fetch missing data from session test results
+            this.fetchSessionResults(
+                toFetch.map((item) => {
+                    return { SessionId: item.Testdata.id, ResultsId: item.TestSessionInfo.TestResultsTestDataUID }
+                }), 'data.' + uid
+            );
+        }
+
+        return testcase;
+    }
+
+    async getTestResultByQuery(uid: string, metauid: string, change?: number): Promise<TestResult | undefined> {
+        const testcase = await this.getTestCaseByUID(uid);
+        let test: TestResult | undefined = undefined;
+
+        if (testcase) {
+            test = testcase.Metas.get(metauid);
+            while (!test && this.loadingCount !== 0) {
+                const previousVersion = testcase.version;
+                await when(() => testcase.version !== previousVersion || this.loadingCount === 0);
+                test = testcase.Metas.get(metauid);
+            }
+            if (change && test?.Session.Testdata.change !== change) {
+                // Need to get the target change from the history
+                test = this.getLoadedTestHistoryByMeta(uid)?.get(metauid)?.find(
+                    (item) => item.Session.Testdata.change === change
+                );
+                if (!test) {
+                    // trigger fetch history in background
+                    this.fetchSessionHistory(testcase.SessionName, uid);
+                }
+                while (!test && this.loadingCount !== 0) {
+                    const previousVersion = this.version;
+                    await when(() => this.version !== previousVersion || this.loadingCount === 0);
+                    test = this.getLoadedTestHistoryByMeta(uid)?.get(metauid)?.find(
+                        (item) => item.Session.Testdata.change === change
+                    );
+                }
+
+            }
+        }
+
+        if (test && test.Details === undefined) {
+            // Fetch missing data from the session test result
+            await this.fetchSessionResults(
+                [{ SessionId: test.Session.Testdata.id, ResultsId: test.Session.TestSessionInfo.TestResultsTestDataUID }]
+                , 'data.' + uid
+            );
+            if (test.Details === undefined) {
+                // if nothing was fetched or failed to fetch, return nothing
+                return undefined;
+            }
+        }
+
+        return test;
+    }
+
+    private async fetchSessionResults(items: SessionReferences[], filter?: string, parallelCount: number = 30) {
+        let count = 0;
+        for (let i = 0; i < items.length; i += parallelCount) {
+            if (!this._datahandler.activated) {
+                break;
+            }
+            const chunk = items.slice(i, i + parallelCount);
+            count += (
+                await Promise.all(chunk.map((item) => this.fetchResultItemFromSession(item, filter)))
+            ).reduce((a, b) => a + b);
+        }
+        return count;
+    }
+
+    private async fetchResultItemFromSession(item: SessionReferences, filter?: string): Promise<number> {
+        const update = await this._datahandler.fetchItemsFromAlternateKey(`${this._datahandler.key} Result Details::${item.ResultsId}`, undefined, 1, 0, filter);
+        // We only expect one returned item
+        if (update.length > 0) {
+            const results = update[0].data as TestResultDetailsMap;
+            const session = this._datahandler.items?.get(item.SessionId)?.getDataHandler() as TestSessionWrapper;
+            if (session) {
+                for (const uid in results) {
+                    const test = session.getTestResultByUID(uid);
+                    if (test) {
+                        test.Details = results[uid];
+                    }
+                }
+                return 1;
+            }
+        }
+
+        return 0;
+    }
+
+    getTests() {
+        const tests: TestResult[] = [];
+        const iterator = this._collection.values();
+        let result = iterator.next();
+        while (!result.done) {
+            const sessionInfo = result.value.TestSessionInfo;
+            if (sessionInfo) {
+                tests.push(...Object.values(result.value.TestSessionInfo.Tests) as TestResult[]);
+            }
+            result = iterator.next();
+        }
+        return tests;
+    }
+
+    getTestCases() {
+        return Array.from(this._testcases.values());
+    }
+
+    getLoadedTestHistoryByMeta(uid: string) {
+        const testcase = this._testcases.get(uid);
+        if (!testcase) {
+            return undefined;
+        }
+
+        const testresultsByMeta: Map<string, TestResult[]> = new Map();
+        this.forEachSessionHistoryByChange(testcase.SessionName,
+            (session) => {
+                if (session.isAllPropertiesLoaded()) {
+                    const test = session.getTestResultByUID(uid);
+                    if (test) {
+                        const metauid = session.MetaHandler!.hash;
+                        const meta = testresultsByMeta.get(metauid);
+                        if (meta) {
+                            const change = session.Testdata.change;
+                            if (!meta.some((item) => item.Session.Testdata.change === change)) {
+                                // Keep only one result per change
+                                meta.push(test);
+                            }
+                        } else {
+                            testresultsByMeta.set(metauid, [test]);
+                        }
+                    }
+                }
+            }
+        );
+
+        return testresultsByMeta;
+    }
+
+    async fetchSessionHistory(sessionName: string, testuid: string, maxCount: number = 1000) {
+        if (this._datahandler.cursor) {
+            // for single view mode we need to fetch the whole pool for context
+            const onFetch = (items: (TestDataWrapper | undefined)[]) => this.onInitialHistoryFetch(items, sessionName, testuid)
+            await this._datahandler.fetchCursorHistory(onFetch, SessionFields.initiate, maxCount);
+            return;
+        }
+
+        await this.fetchIncompleteSessionsForHistory(sessionName, testuid);
+        while (this.loadingCount !== 0) {
+            // wait for the collection to be fully loaded, in the meanwhile load the missing session history
+            const previousVersion = this.version;
+            await when(() => this.version !== previousVersion || this.loadingCount === 0);
+            await this.fetchIncompleteSessionsForHistory(sessionName, testuid);
+        }
+    }
+
+    private async fetchIncompleteSessionsForHistory(sessionName: string, testuid: string) {
+        const missingSessions: string[] = [];
+        this.forEachSessionHistoryByChange(sessionName,
+            (session) => {
+                if (!session.isAllPropertiesLoaded()) {
+                    missingSessions.push(session.Testdata.id);
+                    this._testcases.get(testuid)?.incrementHistoryTarget();
+                }
+            }
+        );
+
+        if (missingSessions.length > 0) {
+            await this._datahandler.fetchUpdateItems(
+                missingSessions, (item) => this.onFetchTestCaseHistory(item, testuid), SessionFields.remaining
+            )
+        }
+        return missingSessions.length;
+    }
+
+    private forEachSessionHistoryByChange(sessionName: string, callback: (session: TestSessionWrapper) => void) {
+        const iterator = this._datahandler.iterItemsByChange();
+        let change = iterator.next();
+        while (!change.done) {
+            (change.value as TestDataWrapper[]).forEach((item) => {
+                const session = item.getDataHandler() as TestSessionWrapper;
+                if (session.Name === sessionName) {
+                    if (!this._includePreflight && session.PreFlightChange) {
+                        return;
+                    }
+                    callback(session);
+                }
+            });
+            change = iterator.next();
+        }
+    }
+
+    private onInitialHistoryFetch(items: (TestDataWrapper | undefined)[], sessionName: string, testuid: string) {
+        const toFetch: TestDataWrapper[] = [];
+        items.forEach((item) => {
+            if (item) {
+                let data = item.getDataHandler() as TestSessionWrapper;
+                if (!data) {
+                    data = new TestSessionWrapper(item);
+                    item.setDataHandler(data);
+                    if (data.Name === sessionName) {
+                        toFetch.push(item);
+                        this._testcases.get(testuid)?.incrementHistoryTarget();
+                    }
+                }
+            }
+        });
+        if (toFetch.length > 0) {
+            this._datahandler.fetchUpdateItems(
+                toFetch.map((item) => item.id), (items) => this.onFetchTestCaseHistory(items, testuid), SessionFields.remaining
+            )
+        }
+    }
+
+    private onFetchTestCaseHistory(items: (TestDataWrapper | undefined)[], testuid: string) {
+        items.forEach((item) => {
+            if (item) {
+                this._testcases.get(testuid)?.updateHistory();
+            }
+        });
+    }
+
+}
diff --git a/Engine/Source/Programs/Horde/HordeDashboard/src/legacyPlugins/SimpleTestReport/components/SimpleReport.tsx b/Engine/Source/Programs/Horde/HordeDashboard/src/legacyPlugins/SimpleTestReport/components/SimpleReport.tsx
new file mode 100644
--- /dev/null
+++ b/Engine/Source/Programs/Horde/HordeDashboard/src/legacyPlugins/SimpleTestReport/components/SimpleReport.tsx
@@ -0,0 +1,484 @@
+// Copyright Epic Games, Inc. All Rights Reserved.
+
+import { Icon, Spinner, SpinnerSize, Stack, Text } from '@fluentui/react';
+import { mergeStyleSets, mergeStyles } from '@fluentui/react/lib/Styling';
+import React, { useEffect, useRef, useState } from 'react';
+import { Link } from 'react-router-dom';
+
+import backend from '../../../backend';
+import { TestData } from '../../../backend/Api';
+import dashboard from '../../../backend/Dashboard';
+import { msecToElapsed } from '../../../base/utilities/timeUtils';
+import { testDataHandler } from '../../../components/TestReportView';
+import { getHordeStyling } from '../../../styles/Styles';
+import { getHordeTheme } from '../../../styles/theme';
+import { HistoryItem, LogLevel, SimpleReport } from '../models/SimpleReport';
+
+let _styles: any;
+let _stateStyles: any;
+
+const getStyling = () => {
+
+   const theme = getHordeTheme();
+
+   const styles = _styles ?? mergeStyleSets({
+      container: {
+         overflow: 'auto',
+         height: 'calc(100vh - 165px)',
+         marginTop: 8
+      },
+      item: [
+         {
+            fontSize: "11px",
+            fontFamily: "Horde Cousine Regular"
+         }
+      ],
+      gutter: [
+         {
+            padding: 0,
+            margin: 0,
+            paddingTop: 0,
+            paddingBottom: 0,
+            paddingRight: 14,
+            marginTop: 0,
+            marginBottom: 0,
+            height: 20
+         }
+      ],
+      gutterError: [
+         {
+            background: dashboard.darktheme ? "#330606" : "#FEF6F6",
+            borderLeftStyle: 'solid',
+            borderLeftColor: "#EC4C47",
+            borderLeftWidth: 6,
+            padding: 0,
+            margin: 0,
+            paddingTop: 0,
+            paddingBottom: 0,
+            paddingRight: 8,
+            marginTop: 0,
+            marginBottom: 0,
+            height: 20
+         }
+      ],
+      gutterWarning: [
+         {
+            background: dashboard.darktheme ? "#1E1817" : "#FEF8E7",
+            borderLeftStyle: 'solid',
+            borderLeftColor: "#F7D154",
+            borderLeftWidth: 6,
+            padding: 0,
+            margin: 0,
+            paddingTop: 0,
+            paddingBottom: 0,
+            paddingRight: 8,
+            marginTop: 0,
+            marginBottom: 0,
+            height: 20
+         }
+      ],
+      gutterSuccess: [
+         {
+            borderLeftStyle: 'solid',
+            borderLeftColor: theme.palette.green,
+            borderLeftWidth: 6,
+            padding: 0,
+            margin: 0,
+            paddingTop: 0,
+            paddingBottom: 0,
+            paddingRight: 8,
+            marginTop: 0,
+            marginBottom: 0,
+            height: 20
+         }
+      ],
+      itemWarning: [
+         {
+            background: dashboard.darktheme ? "#1E1817" :"#FEF8E7"
+         }
+      ],
+      itemError: [
+         {
+            background: dashboard.darktheme ? "#330606" :"#FEF6F6"
+         }
+      ],
+      itemHover: {
+         selectors: {
+            ':hover': {
+               background: theme.palette.neutralLight
+            }
+         }
+      },
+      itemHighlighted: {
+         background: theme.palette.neutralLighter
+      },
+      historyList: {
+         zIndex: 1,
+         borderWidth: "1px",
+         borderColor: "#888",
+         borderStyle: "solid",
+         backgroundColor: dashboard.darktheme ? "#111" : "#FFF",
+         padding: "3px",
+         overflow: "auto"
+      },
+      jobDetailsItem: {
+         padding: 8,
+         borderBottom: '1px solid ' + theme.palette.neutralLighter,
+         selectors: {
+            ':hover': { background: theme.palette.neutralLight }
+         }
+      }
+   });
+
+
+   const stateStyles = _stateStyles ?? mergeStyleSets({
+      success: [{ color: theme.palette.green, userSelect: "none" }, iconClass],
+      warnings: [{ color: "#F7D154", userSelect: "none" }, iconClass],
+      failure: [{ color: "#EC4C47", userSelect: "none" }, iconClass],
+   });
+
+   _stateStyles = stateStyles;
+   _styles = styles;
+
+   return {
+      theme: theme,
+      styles: styles,
+      stateStyles: stateStyles
+   };
+}
+
+const iconClass = mergeStyles({
+   fontSize: 13
+});
+
+const wait = (ms: number) => new Promise(resolve => setTimeout(resolve, ms));
+const secToElapsed = (sec: number) => msecToElapsed(sec * 1000);
+
+const markdownURL = /^\[([^\]]+)\]\((.+)\)$/;
+function getUrlFromLink(link: string): { text: string, url: string } {
+   let url = link;
+   let text = link;
+   const match = markdownURL.exec(link);
+   if (match !== null) {
+      text = match[1];
+      url = match[2];
+   }
+   return { text: text, url: url };
+}
+
+const URLLink: React.FC<{ link: string }> = (props) => {
+
+   const { styles } = getStyling();
+
+   const { link } = props;
+
+   const linkInfo = getUrlFromLink(link);
+
+   return (<Stack disableShrink={true}><Text className={styles.itemHover}><a href={linkInfo.url}>{linkInfo.text}</a></Text></Stack>);
+}
+
+const LogItem: React.FC<{ line: string, level: string }> = (props) => {
+
+   const { styles } = getStyling();
+
+   const { line, level } = props;
+   const is_error = level === LogLevel.Error;
+   const is_warning = level === LogLevel.Warning;
+   const style = is_error ? styles.itemError : (is_warning ? styles.itemWarning : styles.item);
+   const gutterStyle = is_error ? styles.gutterError : (is_warning ? styles.gutterWarning : styles.gutter);
+
+   return (
+      <Stack className={styles.item} disableShrink={true} styles={{ root: { paddingLeft: 8, width: "100%" } }}>
+         <Stack horizontal>
+            <Stack className={gutterStyle}></Stack>
+            <Stack.Item className={style} align="center" styles={{ root: { paddingLeft: 8, width: "100%" } }}><pre style={{ margin: 0, whiteSpace: "pre-wrap" }}>{line}</pre></Stack.Item>
+         </Stack>
+      </Stack>
+   );
+}
+
+const LogLink: React.FC<{ to: string }> = (props) => {
+
+   const { styles } = getStyling();
+
+   const { to } = props;
+   //const [link, setLink] = useState<string | undefined>(undefined)
+
+   // this view is deprecated and these links were nonfunctional in tests I checked,  would also need to be updated to artifacts v2
+   /*
+    useEffect(() => {
+       const getLink = async () => {
+          
+           
+            const artifact = await testDataHandler.cursor?.findArtifactData(to);
+            if (artifact !== undefined) {
+                return `${backend.serverUrl}/api/v1/artifacts/${artifact.id}/download?Code=${artifact.code}`
+            }
+           
+            return undefined;
+        }
+        getLink().then((item) => setLink(item));
+    });
+    */
+
+   // <a href={link}>{to}</a>
+
+   return (
+      <Stack disableShrink={true} styles={{ root: { paddingLeft: 8 } }}><Text className={styles.itemHover}>{to}</Text></Stack>
+   );
+}
+
+const HistoryItemPane: React.FC<{ item: HistoryItem, selected: boolean }> = (props) => {
+   const { styles } = getStyling();
+   const { item, selected } = props;
+   const gutterStyle = item.HasSucceeded ? styles.gutterSuccess : styles.gutterError;
+   const state = item.HasSucceeded ? "Success" : "Failed";
+   const errorCount = item.HasSucceeded ? "" : ` (${item.ErrorCount} error${item.ErrorCount > 1 ? "s" : ""})`;
+
+   return (
+      <Stack horizontal className={`${styles.itemHover} ${selected ? styles.itemHighlighted : ""}`}>
+         <Stack className={gutterStyle}></Stack>
+         <Stack horizontal horizontalAlign="space-between" grow tokens={{ childrenGap: 11 }}>
+            <Stack.Item styles={{ root: { whiteSpace: "nowrap" } }}>
+               {item.Url.indexOf('http') < 0 ?
+                  <Link to={item.Url}>
+                     <Text variant="smallPlus">{state} on {item.Change}</Text>
+                     {selected && <Text styles={{ root: { fontWeight: 'bold' } }}> &lt;</Text>}
+                  </Link> :
+                  <a href={item.Url}>
+                     <Text variant="smallPlus">{state} on {item.Change}</Text>
+                     {selected && <Text styles={{ root: { fontWeight: 'bold' } }}> &lt;</Text>}
+                  </a>
+               }
+            </Stack.Item>
+            <Stack.Item styles={{ root: { width: '22ch' } }}><Text>{item.Date}</Text></Stack.Item>
+            <Stack.Item styles={{ root: { width: '8ch' } }}>
+               <Text styles={{ root: { color: "#0089DD" } }}>{secToElapsed(item.TotalDurationSeconds)}</Text>
+            </Stack.Item>
+            <Stack.Item styles={{ root: { width: '12ch', textAlign: 'right' } }}><Text>{errorCount}</Text></Stack.Item>
+         </Stack>
+      </Stack>
+   );
+}
+
+const HistoryPanel: React.FC = () => {
+   const [historyLoaded, setHistoryLoaded] = useState(false);
+   const historyItemsRef = useRef<HistoryItem[]>([]);
+   const { hordeClasses } = getHordeStyling();
+   const historyItems = historyItemsRef.current;
+
+   useEffect(() => {
+      testDataHandler.getCursorHistory(undefined, undefined, 100).then(
+         (items) => {
+            if (items.length > 0) {
+               items.forEach((item) => {
+                  const testdata = item.data as SimpleReport;
+                  const testHistoryItem: HistoryItem = {
+                     Change: item.change,
+                     TestdataId: item.id,
+                     Url: `/testreport/${item.id}`,
+                     HasSucceeded: testdata.HasSucceeded,
+                     Date: testdata.ReportCreatedOn,
+                     TotalDurationSeconds: testdata.TotalDurationSeconds,
+                     ErrorCount: testdata.Errors.length,
+                  };
+                  historyItemsRef.current.push(testHistoryItem);
+               });
+            }
+         }
+      ).catch(
+         (reason) => console.error(reason)
+      ).finally(
+         () => { setHistoryLoaded(true) }
+      );
+   }, []);
+
+   return (
+      <Stack styles={{ root: { paddingTop: 18, paddingRight: 0 } }}>
+         <Stack className={hordeClasses.raised}>
+            <Stack tokens={{ childrenGap: 9 }} grow>
+               <Stack horizontal tokens={{ childrenGap: 12 }}>
+                  <Text variant="mediumPlus" styles={{ root: { fontFamily: "Horde Open Sans SemiBold" } }}>History</Text>
+               </Stack>
+               {!historyLoaded && <Stack><Spinner styles={{ root: { padding: 10 } }} size={SpinnerSize.medium}></Spinner></Stack>}
+               {(historyLoaded || historyItems.length > 0) &&
+                  <Stack styles={{ root: { width: 500 } }}>
+                     {historyItems.map((item) => <HistoryItemPane key={item.TestdataId} item={item} selected={testDataHandler.cursor?.id === item.TestdataId} />)}
+                  </Stack>
+               }
+               {(historyLoaded && historyItems.length === 0) && <Stack><Text> No history </Text></Stack>}
+            </Stack>
+         </Stack>
+      </Stack>
+   );
+}
+
+export const TestReportView: React.FC<{ data: SimpleReport, query: URLSearchParams }> = (props) => {
+
+   const { hordeClasses, modeColors } = getHordeStyling();
+   const { styles, stateStyles } = getStyling();
+
+   const { data } = props;
+
+   const style = data.HasSucceeded ? stateStyles.success : stateStyles.failure;
+
+   return (
+      <Stack className={styles.container} styles={{ root: { backgroundColor: modeColors.backend, paddingLeft: 24, paddingTop: 12, paddingRight: 12 } }}>
+
+         <Stack styles={{ root: { paddingTop: 18, paddingRight: 0 } }}>
+            <Stack className={hordeClasses.raised}>
+               <Stack tokens={{ childrenGap: 12 }} grow>
+                  <Stack>
+                     <Text variant="mediumPlus" styles={{ root: { fontFamily: "Horde Open Sans SemiBold" } }}>Summary</Text>
+                  </Stack>
+                  <Stack styles={{ root: { paddingLeft: 8 } }}>
+                     <Text>This test pass run on <span style={{ fontWeight: 'bold' }}>{data.ReportCreatedOn}</span> for a duration of <span style={{ fontWeight: 'bold' }}>{secToElapsed(data.TotalDurationSeconds)}</span> on <span style={{ fontWeight: 'bold' }}>{data.Description}</span></Text>
+                     <Stack horizontal>
+                        <Stack.Item><Text variant="medium" styles={{ root: { fontFamily: "Horde Open Sans SemiBold" } }}>Status:&nbsp;</Text></Stack.Item>
+                        <Stack.Item><Icon className={style} iconName="Square" /></Stack.Item>
+                        <Stack.Item align="center" styles={{ root: { paddingLeft: 8, width: "100%", fontWeight: "bold" } }}>{data.Status}</Stack.Item>
+                     </Stack>
+                     <Text>
+                        {data.Errors.length > 0 && <span style={{ fontWeight: 'bold' }}>{data.Errors.length} error messages. </span>}
+                        {data.Warnings.length > 0 && <span>{data.Warnings.length} warning messages. </span>}
+                     </Text>
+                     {data.URLLink &&
+                        <Stack horizontal><Text variant="medium" styles={{ root: { fontFamily: "Horde Open Sans SemiBold" } }}>Link to:&nbsp;</Text><URLLink link={data.URLLink} /></Stack>
+                     }
+                  </Stack>
+               </Stack>
+            </Stack>
+         </Stack>
+
+         <Stack styles={{ root: { paddingTop: 18, paddingRight: 0 } }}>
+            <Stack className={hordeClasses.raised}>
+               <Stack tokens={{ childrenGap: 9 }} grow>
+                  <Stack>
+                     <Text variant="mediumPlus" styles={{ root: { fontFamily: "Horde Open Sans SemiBold" } }}>Messages</Text>
+                  </Stack>
+                  {(data.Errors.length > 0 || data.Warnings.length > 0) &&
+                     <Stack>
+                        {data.Errors.map((line, index) => <LogItem key={`E${index}`} line={line} level={LogLevel.Error} />)}
+                        {data.Warnings.map((line, index) => <LogItem key={`W${index}`} line={line} level={LogLevel.Warning} />)}
+                     </Stack>
+                  }
+                  {(data.Errors.length === 0 && data.Warnings.length === 0) && <Stack><Text> No log event reported </Text></Stack>}
+                  {data.Logs.length > 0 && <Stack><Text styles={{ root: { fontWeight: "bold" } }}>Logs: </Text>{data.Logs.map((item) => <LogLink key={item} to={item} />)}</Stack>}
+               </Stack>
+            </Stack>
+         </Stack>
+
+         <HistoryPanel />
+
+      </Stack>
+   );
+}
+
+export const ReportLinkItemPane: React.FC<{ id: string, fullName: string }> = (props) => {
+   const { id, fullName } = props;
+   const testdata = useRef<TestData>();
+   const result = useRef<SimpleReport>();
+   const history = useRef<HistoryItem[]>();
+   const [testdataLoading, setTestdataLoading] = useState(false);
+   const [historyLoaded, setHistoryLoaded] = useState(false);
+   const [historyVisible, setHistoryVisible] = useState(false);
+
+   const { styles, stateStyles } = getStyling();
+
+   async function getTestData() {
+      if (testdata.current !== undefined || testdataLoading) {
+         return;
+      }
+      setTestdataLoading(true);
+      try {
+         testdata.current = await backend.getTestData(id);
+         result.current = testdata.current.data as SimpleReport;
+      } catch (reason) {
+         console.error(`Failed to load test data id: ${id}`);
+         throw reason;
+      } finally {
+         setTestdataLoading(false);
+      }
+   }
+
+   async function getTestHistory() {
+      if (history.current !== undefined || historyLoaded) {
+         return;
+      }
+      try {
+         if (testdataLoading) {
+            // we wait while the test data is loading.
+            do { await wait(3 * 1000); } while (testdataLoading);
+         } else {
+            // Otherwise we fetch it
+            await getTestData();
+         }
+         if (testdata.current !== undefined && history.current === undefined) {
+            const rawHistoryItems = await backend.getTestDataHistory(testdata.current.streamId, fullName, testdata.current.change);
+            const historyItems: HistoryItem[] = [];
+            rawHistoryItems.forEach((item) => {
+               if (historyItems.length > 0 && item.change === historyItems[historyItems.length - 1].Change) return; // keep only item per changelist
+               const testdata = item.data as SimpleReport;
+               const testHistoryItem: HistoryItem = {
+                  Change: item.change,
+                  TestdataId: item.id,
+                  Url: !testdata.URLLink ? `/testreport/${item.id}` : getUrlFromLink(testdata.URLLink).url,
+                  HasSucceeded: testdata.HasSucceeded,
+                  Date: testdata.ReportCreatedOn,
+                  TotalDurationSeconds: testdata.TotalDurationSeconds,
+                  ErrorCount: testdata.Errors.length,
+               };
+               historyItems.push(testHistoryItem);
+            });
+            history.current = historyItems;
+         }
+      } catch (reason) {
+         console.error(reason);
+      } finally {
+         setHistoryLoaded(true);
+      }
+   }
+
+   function onClickHistory() {
+      const needHistoryVisible = !historyVisible;
+      setHistoryVisible(needHistoryVisible);
+      if (needHistoryVisible) {
+         getTestHistory();
+      }
+   }
+   function onHistoryBlur(ev: any) {
+      if (!ev.currentTarget.contains(ev.relatedTarget as Node)) setHistoryVisible(false);
+   }
+
+   if (testdata.current === undefined) {
+      getTestData();
+   }
+
+   return (
+      <Stack horizontal wrap tokens={{ childrenGap: 30 }}>
+         {result.current !== undefined &&
+            <Stack horizontal>
+               <Stack.Item><Icon className={result.current.HasSucceeded ? stateStyles.success : stateStyles.failure} iconName="Square" /></Stack.Item>
+               <Stack.Item align="center" styles={{ root: { paddingLeft: 8, width: "100%", fontWeight: "bold" } }}>{result.current.Status}</Stack.Item>
+            </Stack>
+         }
+         {result.current !== undefined && result.current.URLLink &&
+            <a href={getUrlFromLink(result.current.URLLink).url}>Report Link</a>
+         }
+         <Stack onClick={(ev) => { onClickHistory() }} tabIndex={0} onBlur={onHistoryBlur} className={styles.itemHover} style={{ position: "relative", cursor: "pointer" }}>
+            <Text styles={{ root: { fontWeight: 'bold' } }}>History</Text>
+            {historyVisible &&
+               <div style={{ position: "absolute", width: 500, top: "100%", maxHeight: 300, cursor: 'auto' }} className={styles.historyList} onClick={(ev) => { ev.stopPropagation() }}>
+                  {!historyLoaded && <Spinner styles={{ root: { padding: 10 } }} size={SpinnerSize.small}></Spinner>}
+                  {historyLoaded && history.current !== undefined && history.current.length > 0 &&
+                     history.current.map((historyItem) => <HistoryItemPane key={historyItem.Change} item={historyItem} selected={id === historyItem.TestdataId} />)
+                  }
+                  {historyLoaded && history.current !== undefined && history.current.length === 0 &&
+                     <div><Text variant="smallPlus">no history found</Text></div>
+                  }
+               </div>
+            }
+         </Stack>
+      </Stack>
+   );
+}
diff --git a/Engine/Source/Programs/Horde/HordeDashboard/src/legacyPlugins/SimpleTestReport/index.tsx b/Engine/Source/Programs/Horde/HordeDashboard/src/legacyPlugins/SimpleTestReport/index.tsx
new file mode 100644
--- /dev/null
+++ b/Engine/Source/Programs/Horde/HordeDashboard/src/legacyPlugins/SimpleTestReport/index.tsx
@@ -0,0 +1,24 @@
+// Copyright Epic Games, Inc. All Rights Reserved.
+
+import { Plugin, PluginMount } from "../pluginTypes";
+import { TestReportView, ReportLinkItemPane } from "./components/SimpleReport";
+
+class SimpleTestReportPlugin implements Plugin {
+
+    name = "Simple Report";
+    components = [
+        {
+            mount: PluginMount.TestReportPanel,
+            component: TestReportView,
+            id: "Simple Report"
+        },
+        {
+            mount: PluginMount.TestReportLink,
+            component: ReportLinkItemPane,
+            id: "Simple Report"
+        }
+    ];
+}
+
+const simpleTestReportPlugin = new SimpleTestReportPlugin();
+export default simpleTestReportPlugin; 
diff --git a/Engine/Source/Programs/Horde/HordeDashboard/src/legacyPlugins/SimpleTestReport/models/SimpleReport.ts b/Engine/Source/Programs/Horde/HordeDashboard/src/legacyPlugins/SimpleTestReport/models/SimpleReport.ts
new file mode 100644
--- /dev/null
+++ b/Engine/Source/Programs/Horde/HordeDashboard/src/legacyPlugins/SimpleTestReport/models/SimpleReport.ts
@@ -0,0 +1,29 @@
+// Copyright Epic Games, Inc. All Rights Reserved.
+
+export enum LogLevel {
+    Info = "Info",
+    Error = "Error",
+    Warning = "Warning",
+}
+
+export type SimpleReport = {
+    Description: string;
+    HasSucceeded: boolean;
+    Status: string;
+    ReportCreatedOn: string;
+    TotalDurationSeconds: number;
+    URLLink: string;
+    Logs: string[];
+    Errors: string[];
+    Warnings: string[];
+}
+
+export type HistoryItem = {
+    TestdataId: string;
+    Url: string;
+    Change: number;
+    HasSucceeded: boolean;
+    Date: string;
+    TotalDurationSeconds: number;
+    ErrorCount: number;
+}
\ No newline at end of file
diff --git a/Engine/Source/Programs/Horde/HordeDashboard/src/legacyPlugins/UnrealAutomatedTests/components/UnrealAutomatedTests.tsx b/Engine/Source/Programs/Horde/HordeDashboard/src/legacyPlugins/UnrealAutomatedTests/components/UnrealAutomatedTests.tsx
new file mode 100644
--- /dev/null
+++ b/Engine/Source/Programs/Horde/HordeDashboard/src/legacyPlugins/UnrealAutomatedTests/components/UnrealAutomatedTests.tsx
@@ -0,0 +1,601 @@
+// Copyright Epic Games, Inc. All Rights Reserved.
+
+import { Checkbox, Icon, Image, Spinner, SpinnerSize, Stack, Text } from '@fluentui/react';
+import { mergeStyleSets, mergeStyles } from '@fluentui/react/lib/Styling';
+import React, { useEffect, useRef, useState } from 'react';
+import { useNavigate } from 'react-router';
+import { Link } from 'react-router-dom';
+import dashboard from "../../../backend/Dashboard";
+import { projectStore } from "../../../backend/ProjectStore";
+import { TestDataWrapper } from '../../../backend/TestDataHandler';
+import { msecToElapsed } from '../../../base/utilities/timeUtils';
+import { testDataHandler } from '../../../components/TestReportView';
+import { getHordeStyling } from '../../../styles/Styles';
+import { getHordeTheme } from '../../../styles/theme';
+import { EventType, Metadata, TestDetails, TestEntry, TestEntryArtifact, TestPassSummary, TestResult, TestState, TestStateHistoryItem } from '../models/UnrealAutomatedTests';
+
+let _styles: any;
+let _gutterHystoryStyles: any;
+let _stateStyles: any;
+
+const getStyling = () => {
+   const theme = getHordeTheme();
+
+   const styles = _styles ?? mergeStyleSets({
+      container: {
+         overflow: 'auto',
+         height: 'calc(100vh - 165px)',
+         marginTop: 8
+      },
+      item: [
+         {
+            fontSize: "11px",
+            fontFamily: "Horde Cousine Regular"
+         }
+      ],
+      gutter: [
+         {
+            padding: 0,
+            margin: 0,
+            paddingTop: 0,
+            paddingBottom: 0,
+            paddingRight: 14,
+            marginTop: 0,
+            marginBottom: 0,
+            height: 20
+         }
+      ],
+      gutterError: [
+         {
+            background: dashboard.darktheme ? "#330606" : "#FEF6F6",
+            borderLeftStyle: 'solid',
+            borderLeftColor: "#EC4C47"
+         }, gutterClass
+      ],
+      gutterWarning: [
+         {
+            background: dashboard.darktheme ? "#1E1817" : "#FEF8E7",
+            borderLeftStyle: 'solid',
+            borderLeftColor: "#F7D154"
+         }, gutterClass
+      ],
+      gutterSuccess: [
+         {
+            borderLeftStyle: 'solid',
+            borderLeftColor: theme.palette.green
+         }, gutterClass
+      ],
+      itemWarning: [
+         {
+            background: dashboard.darktheme ? "#1E1817" : "#FEF8E7",
+         }
+      ],
+      itemError: [
+         {
+            background: dashboard.darktheme ? "#330606" : "#FEF6F6",
+         }
+      ],
+      itemHover: {
+         cursor: "pointer",
+         selectors: {
+            ':hover': {
+               background: theme.palette.neutralLight
+            }
+         }
+      },
+      itemHighlighted: {
+         borderWidth: "1px",
+         borderColor: "#888",
+         borderStyle: "solid"
+      },
+      historyList: {
+         zIndex: 1,
+         borderWidth: "1px",
+         borderColor: "#888",
+         borderStyle: "solid",
+         backgroundColor: theme.palette.white,
+         padding: "3px",
+         overflow: "auto"
+      }
+   });
+
+   const gutterHystoryStyles = _gutterHystoryStyles ?? new Map<string, string>([
+      [TestState.Success, styles.gutterSuccess],
+      [TestState.InProcess, mergeStyles({
+         background: dashboard.darktheme ? "#141115" : "#E4F1F5",
+         borderLeftStyle: 'solid',
+         borderLeftColor: "#01BCF2"
+      }, gutterClass)],
+      [TestState.NotRun, mergeStyles({
+         borderLeftStyle: 'solid',
+         borderLeftColor: "#A19F9D"
+      }, gutterClass)],
+      [TestState.SuccessWithWarnings, styles.gutterWarning],
+      [TestState.Failed, styles.gutterError],
+   ]);
+
+   const stateStyles = _stateStyles ?? new Map<string, string>([
+      [TestState.Success, mergeStyles({ color: "#9BF95C", userSelect: "none" }, iconClass)],
+      [TestState.InProcess, mergeStyles({ color: "#01BCF2", userSelect: "none" }, iconClass)],
+      [TestState.NotRun, mergeStyles({ color: "#A19F9D", userSelect: "none" }, iconClass)],
+      [TestState.SuccessWithWarnings, mergeStyles({ color: "#F7D154", userSelect: "none" }, iconClass)],
+      [TestState.Failed, mergeStyles({ color: "#EC4C47", userSelect: "none" }, iconClass)],
+      [TestState.Unknown, mergeStyles({ color: dashboard.darktheme ? "#666666" : "#000000", userSelect: "none" }, iconClass)],
+   ]);
+
+   _styles = styles;
+   _gutterHystoryStyles = gutterHystoryStyles;
+   _stateStyles = stateStyles;
+
+   return {
+      theme: theme,
+      styles: styles,
+      stateStyles: stateStyles,
+      gutterHystoryStyles: gutterHystoryStyles
+   }
+}
+
+const gutterClass = mergeStyles({
+   borderLeftWidth: 6,
+   padding: 0,
+   margin: 0,
+   paddingTop: 0,
+   paddingBottom: 0,
+   paddingRight: 8,
+   marginTop: 0,
+   marginBottom: 0,
+   height: 20
+});
+
+const iconClass = mergeStyles({
+   fontSize: 12
+});
+
+const getTestStateStyles = (test: TestResult): string | undefined => {
+   const { stateStyles } = getStyling();
+   if (!stateStyles.has(test.State)) {
+      return stateStyles.get(TestState.Unknown);
+   }
+   return stateStyles.get(test.Warnings > 0 && test.State === TestState.Success ? TestState.SuccessWithWarnings : test.State)
+}
+
+const shadowBoxSlyle = {
+   boxShadow: '0 4px 8px 0 rgba(0, 0, 0, 0.2), 0 6px 20px 0 rgba(0, 0, 0, 0.19)',
+   margin: 5,
+   paddingRight: 5
+};
+
+const stateLabels = new Map<string, string>([
+   [TestState.Success, 'Passed'],
+   [TestState.InProcess, 'Incomplete'],
+   [TestState.NotRun, 'Unexecuted'],
+   [TestState.SuccessWithWarnings, 'Passed'],
+   [TestState.Failed, 'Failed'],
+   [TestState.Skipped, 'Skipped'],
+   [TestState.Unknown, 'Unknown'],
+]);
+const getStateLabel = (state: string): string | undefined => { return stateLabels.has(state) ? stateLabels.get(state) : state }
+
+function getMetadata(collection: Metadata, key: string): string | undefined {
+   if (collection) {
+      return collection[key];
+   }
+
+   return undefined;
+}
+
+const copyToClipboard = (value: string | undefined) => {
+
+   if (!value) {
+      return;
+   }
+
+   const el = document.createElement('textarea');
+   el.value = value;
+   document.body.appendChild(el);
+   el.select();
+   document.execCommand('copy');
+   document.body.removeChild(el);
+}
+
+const missingImage = "/images/missing-image.png";
+const MissingImageLabel = (): JSX.Element => { return <span style={{ fontWeight: 'bold' }}> [missing image]</span> }
+type ImageLinks = { approved?: string, unapproved?: string, difference?: string }
+
+
+const EntryPane: React.FC<{ entry: TestEntry, testArtifacts: TestEntryArtifact[] }> = (props) => {
+   const { entry, testArtifacts } = props;
+   const [imageLinks, setImageLinks] = useState<ImageLinks>({});
+
+
+   const { styles } = getStyling();
+
+   const eventType = entry.Event.Type;
+   const is_error = eventType === EventType.Error;
+   const is_warning = eventType === EventType.Warning;
+   const style = is_error ? styles.itemError : (is_warning ? styles.itemWarning : styles.item);
+   const gutterStyle = is_error ? styles.gutterError : (is_warning ? styles.gutterWarning : styles.gutter);
+
+   const artifact = testArtifacts.find((value) => value.Id === entry.Event.Artifact);
+   const artifactRef = useRef(artifact);
+
+   const need_image_comparison = artifact !== undefined && artifact.Type === "Comparison";
+   useEffect(() => {
+      const artifact = artifactRef.current;
+      if (artifact !== undefined) {
+         const findLinks = async () => {
+            const imageLinks: ImageLinks = {};
+            // Approved            
+            imageLinks.approved = artifact.Files.Approved && await testDataHandler.cursor?.getArtifactImageLink(artifact.Files.Approved);
+            // Unapproved            
+            imageLinks.unapproved = artifact.Files.Unapproved && await testDataHandler.cursor?.getArtifactImageLink(artifact.Files.Unapproved);
+            // Difference
+            imageLinks.difference = artifact.Files.Difference && await testDataHandler.cursor?.getArtifactImageLink(artifact.Files.Difference);
+
+            setImageLinks(imageLinks);
+         }
+         findLinks();
+      }
+   }, []);
+
+   return (
+      <Stack className={styles.item} disableShrink={true} styles={{ root: { paddingLeft: 8, width: "100%" } }}>
+         <Stack horizontal>
+            <Stack className={gutterStyle}></Stack>
+            <Stack.Item className={style} align="center" styles={{ root: { paddingLeft: 8, width: "100%" } }}><pre style={{ margin: 0, whiteSpace: "pre-wrap" }}>[{entry.Timestamp}] {eventType}: {entry.Event.Message}</pre></Stack.Item>
+         </Stack>
+         {need_image_comparison &&
+            <Stack styles={{ root: { paddingLeft: 16 } }}>
+               <Stack><Text variant="medium" styles={{ root: { fontWeight: "bold" } }}>Image comparison: {artifact?.Name}</Text></Stack>
+               <Stack horizontal>
+                  {(is_error || is_warning) &&
+                     <Stack styles={{ root: { padding: 5 } }}>
+                        <a href={imageLinks.approved}><Image width={400} src={imageLinks.approved || missingImage} alt={artifact?.Files.Approved} /></a>
+                        <Stack.Item align="center">Approved{!imageLinks.approved && MissingImageLabel()}</Stack.Item>
+                     </Stack>
+                  }
+                  {(is_error || is_warning) &&
+                     <Stack styles={{ root: { padding: 5 } }}>
+                        <a href={imageLinks.difference}><Image width={400} src={imageLinks.difference || missingImage} alt={artifact?.Files.Difference} /></a>
+                        <Stack.Item align="center">Difference{!imageLinks.difference && MissingImageLabel()}</Stack.Item>
+                     </Stack>
+                  }
+                  <Stack styles={{ root: { padding: 5 } }}>
+                     <a href={imageLinks.unapproved}><Image width={400} src={imageLinks.unapproved || missingImage} alt={artifact?.Files.Unapproved} /></a>
+                     <Stack.Item align="center">Unapproved{!imageLinks.unapproved && MissingImageLabel()}</Stack.Item>
+                  </Stack>
+               </Stack>
+            </Stack>
+         }
+      </Stack>
+   );
+}
+
+const HistoryItem: React.FC<{ item: TestStateHistoryItem, testName: string, selected: boolean }> = (props) => {
+   const { item, testName, selected } = props;
+
+   const {theme, gutterHystoryStyles, styles} = getStyling();
+
+   const gutterStyle = gutterHystoryStyles.get(item.State);
+
+   return (
+      <Stack horizontal className={`${selected ? styles.itemHighlighted : ""}`}>
+         <Stack className={gutterStyle}></Stack>
+         <Stack.Item>
+            <Link to={`/testreport/${item.TestdataId}?test=${testName}`} className={`${styles.itemHover}`}>
+               <Text variant="smallPlus">{getStateLabel(item.State)} on {item.Change}</Text>
+            </Link>
+            {item.RangeUrl &&
+               <a href={item.RangeUrl} target="blank" className={`${styles.itemHover}`}>
+                  <Text variant="smallPlus" styles={{ root: { color: theme.palette.neutralDark, paddingLeft: 4, paddingRight: 4 } }}>[ Swarm Range ]</Text>
+               </a>}
+         </Stack.Item>
+      </Stack>
+   );
+}
+
+const IconItem: React.FC<{ item: TestResult }> = (props) => {
+   const { item } = props;
+   const navigate = useNavigate();
+   const style = getTestStateStyles(item);
+   function onClickTest() {
+      navigate(`${window.location.pathname}?test=${item.FullTestPath}`, { replace: true });
+   }
+   return (
+      <Icon styles={{ root: { cursor: "pointer" } }} className={style} iconName="Square" title={item.FullTestPath} onClick={(ev) => { onClickTest() }} />
+   );
+}
+
+const TestResultPane: React.FC<{ test: TestResult, selected: boolean }> = (props) => {
+   const { test, selected } = props;
+   const [testDetails, setTestDetails] = useState<TestDetails | undefined>(undefined);
+   const [loading, setLoading] = useState(false);
+   const [visible, setVisible] = useState(selected);
+   const [filterError, setFilterError] = useState(true);
+   const [filterWarning, setFilterWarning] = useState(true);
+   const [filterInfo, setFilterInfo] = useState(test.State !== TestState.Failed);
+   const [historyLoaded, setHistoryLoaded] = useState(false);
+   const [historyVisible, setHistoryVisible] = useState(false);
+   const historyRef = useRef<TestStateHistoryItem[]>([]);
+
+   const navigate = useNavigate();
+
+   const { styles } = getStyling();
+
+   function onClickTest() {
+      if (selected) {
+         if (visible) {
+            // remove test name selection
+            navigate(window.location.pathname, { replace: true });
+            return;
+         }
+      }
+      // add test name selection
+      navigate(`${window.location.pathname}?test=${test.FullTestPath}`, { replace: true });
+   }
+
+   function onClickTestName() {
+      if (testDetails === undefined) {
+         getTestDetails();
+         setVisible(true);
+         return;
+      }
+
+      setVisible(!visible);
+   }
+
+   async function getTestDetails() {
+      if (testDetails !== undefined) {
+         return;
+      }
+
+      if (loading) {
+         return;
+      }
+      testDataHandler.cursor?.findArtifactData(test.ArtifactName).then((value) => {
+         setTestDetails(value as TestDetails);
+      }).catch((reason) => {
+         console.error(reason);
+      }).finally(() => { setLoading(false) });
+
+      setLoading(true);
+   }
+
+   function onChangeFilterError(ev?: React.FormEvent<HTMLElement | HTMLInputElement>, checked?: boolean) {
+      checked !== undefined && setFilterError(checked);
+   }
+   function onChangeFilterWarning(ev?: React.FormEvent<HTMLElement | HTMLInputElement>, checked?: boolean) {
+      checked !== undefined && setFilterWarning(checked);
+   }
+   function onChangeFilterInfo(ev?: React.FormEvent<HTMLElement | HTMLInputElement>, checked?: boolean) {
+      checked !== undefined && setFilterInfo(checked);
+   }
+
+   function isEntryNeedDisplay(entry: TestEntry, index?: number, array?: TestEntry[]): boolean {
+      const entryEventType = entry.Event.Type;
+      if (filterError && entryEventType === EventType.Error) {
+         return true;
+      }
+      if (filterWarning && entryEventType === EventType.Warning) {
+         return true;
+      }
+      if (filterInfo && entryEventType === EventType.Info) {
+         return true;
+      }
+
+      return false;
+   }
+
+   function getTestHistory(testFullName: string, testdataItems: TestDataWrapper[]): TestStateHistoryItem[] {
+      const testResults: TestStateHistoryItem[] = [];
+
+      let previousChange: number = 0;
+      const stream = projectStore.streamById(testDataHandler.streamId)!;
+      const project = projectStore.byId(stream?.projectId)!;
+      const name = project?.name === "Engine" ? "UE4" : project.name;
+
+      testdataItems.reverse().forEach((item) => {
+         const testdata = item.data as TestPassSummary;
+         const foundTest = testdata.Tests.find(testItem => testItem.FullTestPath === testFullName);
+
+         if (foundTest !== undefined) {
+            const testHistoryItem: TestStateHistoryItem = {
+               Change: item.change,
+               RangeUrl: previousChange > 0 ? `${dashboard.swarmUrl}/files/${name}/${stream.name}?range=@${previousChange},@${item.change}#commits` : undefined,
+               TestdataId: item.id,
+               State: foundTest.Warnings > 0 && foundTest.State === TestState.Success ? TestState.SuccessWithWarnings : foundTest.State,
+            };
+            testResults.push(testHistoryItem);
+         }
+         previousChange = item.change;
+      });
+
+      return testResults.reverse();
+   }
+
+   function onClickHistory() {
+      const needHistoryVisible = !historyVisible;
+      setHistoryVisible(needHistoryVisible);
+      if (needHistoryVisible && !historyLoaded) {
+         testDataHandler.getCursorHistory(undefined, undefined, 100).then(
+            (items) => historyRef.current = getTestHistory(test.FullTestPath, items)
+         ).catch(
+            (reason) => console.error(reason)
+         ).finally(
+            () => { setHistoryLoaded(true) }
+         );
+      }
+   }
+   function onHistoryBlur(ev: any) {
+      // onBlur is used on the History widget in order to make the widget hide when user click outside of it (it loses the focus).
+      // However, item links clicked in the widget make the parent widget lose its focus, so we check if the focus is lost in favor
+      // of one of its children; in that case the widget is kept visible.
+      if (!ev.currentTarget.contains(ev.relatedTarget as Node)) setHistoryVisible(false);
+   }
+
+   const style = getTestStateStyles(test);
+   const history = historyRef.current;
+
+   if (selected) {
+      getTestDetails();
+      if (!visible) {
+         setVisible(true); // Set the visibility when using the top icon for focus
+      }
+   }
+
+   const scrollToSelected = (ref: any) => {
+      if (ref && selected && visible) {
+         ref.scrollIntoView({ behavior: 'smooth', block: 'start' });
+      }
+   }
+
+   return (
+      <Stack styles={{ root: { paddingBottom: 6 } }} horizontalAlign="start" onClick={(ev) => { onClickTest() }}>
+         <div onClick={(ev) => { onClickTestName() }} className={styles.itemHover} style={{ width: "100%" }} ref={scrollToSelected}>
+            <Stack horizontal>
+               <Icon styles={{ root: { paddingTop: 7, paddingRight: 8 } }} className={style} iconName="Square" />
+               <Stack.Item disableShrink><Text variant="mediumPlus" >{test.TestDisplayName}</Text></Stack.Item>
+            </Stack>
+         </div>
+         {loading && <Spinner styles={{ root: { padding: 10 } }} size={SpinnerSize.medium}></Spinner>}
+         {testDetails && visible &&
+            <div style={shadowBoxSlyle} onClick={(ev) => { if (selected) ev.stopPropagation() }}>
+               <Stack horizontal tokens={{ childrenGap: 10 }} styles={{ root: { padding: 5 } }}>
+                  <Text styles={{ root: { fontWeight: 'bold' } }}>Filters:</Text>
+                  <Checkbox label="Error" defaultChecked={filterError} onChange={onChangeFilterError} />
+                  <Checkbox label="Warning" defaultChecked={filterWarning} onChange={onChangeFilterWarning} />
+                  <Checkbox label="Info" defaultChecked={filterInfo} onChange={onChangeFilterInfo} />
+                  <Stack onClick={(ev) => { onClickHistory() }} tabIndex={0} onBlur={onHistoryBlur} className={styles.itemHover} style={{ position: "relative" }}>
+                     <Text styles={{ root: { fontWeight: 'bold', paddingLeft: 4, paddingRight: 4 } }}>History</Text>
+                     {historyVisible &&
+                        <div style={{ position: "absolute", width: 250, top: "100%", maxHeight: 300 }} className={styles.historyList} onClick={(ev) => { ev.stopPropagation() }}>
+                           {!historyLoaded && <Spinner styles={{ root: { padding: 10 } }} size={SpinnerSize.small}></Spinner>}
+                           {historyLoaded && history.length > 0 &&
+                              history.map((historyItem) => <HistoryItem key={historyItem.Change} item={historyItem} testName={test.FullTestPath} selected={historyItem.Change === testDataHandler.cursor?.change} />)
+                           }
+                           {historyLoaded && history.length === 0 && <div><Text variant="smallPlus">no history found</Text></div>}
+                        </div>
+                     }
+                  </Stack>
+                  <Text onClick={(ev) => { copyToClipboard(test.FullTestPath); }} styles={{ root: { fontWeight: 'bold', paddingLeft: 4, paddingRight: 4 } }} className={styles.itemHover}>To Clipboard</Text>
+               </Stack>
+               {!!testDetails.Entries && testDetails.Entries.filter(isEntryNeedDisplay).map((value, index) => <EntryPane key={index} entry={value} testArtifacts={testDetails.Artifacts ?? []} />)}
+               {(!testDetails.Entries || testDetails.Entries.length === 0) && <Text styles={{ root: { padding: 8, fontWeight: 'bold' } }}>No event for this test.</Text>}
+            </div>
+         }
+      </Stack>
+   );
+}
+
+const TestResultPanel: React.FC<{ tests: TestResult[], title?: string, selected?: string }> = (props) => {
+
+   const { hordeClasses } = getHordeStyling();
+   const { theme } = getStyling();
+
+   const { tests, title, selected } = props;
+
+   const suites = new Map<string, TestResult[]>();
+   tests.forEach((test) => {
+      const testSuite = test.FullTestPath.replace('.' + test.TestDisplayName, '');
+      if (!suites.has(testSuite)) {
+         suites.set(testSuite, [test]);
+      } else {
+         suites.get(testSuite)?.push(test);
+      }
+   });
+   const suiteKeys = Array.from(suites.keys());
+
+   return (
+      <Stack styles={{ root: { paddingTop: 18, paddingRight: 0 } }}>
+         <Stack className={hordeClasses.raised}>
+            <Stack tokens={{ childrenGap: 9 }} grow>
+               <Stack>
+                  <Text variant="mediumPlus" styles={{ root: { fontFamily: "Horde Open Sans SemiBold" } }}>{title} tests ({tests.length})</Text>
+               </Stack>
+               {suiteKeys.map((suite) => {
+                  return (<Stack styles={{ root: { paddingLeft: 8 } }} tokens={{ childrenGap: 5 }} grow key={`${title}-${suite}`}>
+                     <Text variant="mediumPlus" styles={{ root: { color: theme.palette.neutralDark } }}>{suite}</Text>
+                     <Stack styles={{ root: { paddingLeft: 8 } }}>
+                        {suites.get(suite)?.map((test: TestResult) => <TestResultPane test={test} selected={selected === test.FullTestPath} key={test.ArtifactName} />)}
+                     </Stack>
+                  </Stack>)
+               })}
+            </Stack>
+         </Stack>
+      </Stack>
+   );
+}
+
+export const TestPassSummaryView: React.FC<{ data: TestPassSummary, query: URLSearchParams }> = (props) => {
+
+   const { hordeClasses, modeColors } = getHordeStyling();   
+   const { stateStyles, styles } = getStyling();
+
+   const { data, query } = props;
+   const failedTests: TestResult[] = [];
+   const notrunTests: TestResult[] = [];
+   const inprocessTests: TestResult[] = [];
+   const passedTests: TestResult[] = [];
+   const skippedTests: TestResult[] = [];
+   const TestsMap = new Map<string, TestResult[]>([
+      [TestState.Failed, failedTests],
+      [TestState.NotRun, notrunTests],
+      [TestState.InProcess, inprocessTests],
+      [TestState.Success, passedTests],
+      [TestState.Skipped, skippedTests],
+   ]);
+   data.Tests.forEach((test) => TestsMap.get(test.State)?.push(test));
+
+   const selectedTest = query.get('test') ? query.get('test')! : undefined;
+
+   return (
+      <Stack className={styles.container} styles={{ root: { backgroundColor: modeColors.background, paddingLeft: 24, paddingTop: 12, paddingRight: 12 } }}>
+
+         <Stack styles={{ root: { paddingTop: 18, paddingRight: 0 } }}>
+            <Stack className={hordeClasses.raised}>
+               <Stack tokens={{ childrenGap: 12 }} grow>
+                  <Stack>
+                     <Text variant="mediumPlus" styles={{ root: { fontFamily: "Horde Open Sans SemiBold" } }}>Summary</Text>
+                  </Stack>
+                  <Stack styles={{ root: { paddingLeft: 8 } }}>
+                     <Text>This test pass ran on <span style={{ fontWeight: 'bold' }}>{data.ReportCreatedOn}</span> for a duration of <span style={{ fontWeight: 'bold' }}>{msecToElapsed(data.TotalDurationSeconds * 1000)}</span> on <span style={{ fontWeight: 'bold' }}>{getMetadata(data.Metadata, 'Platform') ?? "Unknown"}</span></Text>
+                     <Text>
+                        {data.FailedCount > 0 && <span><span style={{ fontWeight: 'bold' }}>{data.FailedCount}</span> tests <span style={{ fontWeight: 'bold' }} className={stateStyles.get('Fail')}>{stateLabels.get('Fail')?.toLowerCase()}</span>. </span>}
+                        {data.InProcessCount > 0 && <span><span style={{ fontWeight: 'bold' }}>{data.InProcessCount}</span> tests <span style={{ fontWeight: 'bold' }} className={stateStyles.get('InProcess')}>{stateLabels.get('InProcess')?.toLowerCase()}</span>. </span>}
+                        {data.NotRunCount > 0 && <span><span style={{ fontWeight: 'bold' }}>{data.NotRunCount}</span> tests <span style={{ fontWeight: 'bold' }} className={stateStyles.get('NotRun')}>{stateLabels.get('NotRun')?.toLowerCase()}</span>. </span>}
+                        {data.SucceededCount > 0 &&
+                           <span>
+                              <span style={{ fontWeight: 'bold' }}>{(data.SucceededCount + data.SucceededWithWarningsCount)}</span> tests <span style={{ fontWeight: 'bold' }} className={stateStyles.get('Success')}>{stateLabels.get('Success')?.toLowerCase()}</span>
+                              {data.SucceededWithWarningsCount > 0 &&
+                                 <span>, including <span style={{ fontWeight: 'bold' }}>{data.SucceededWithWarningsCount}</span> with <span style={{ fontWeight: 'bold' }} className={stateStyles.get('SuccessWithWarnings')}>warnings</span></span>
+                              }.
+                           </span>
+                        }
+                     </Text>
+                     <Stack horizontal wrap>
+                        {data.Tests.map((item: TestResult) => <IconItem item={item} key={item.ArtifactName} />)}
+                     </Stack>
+                     {data.ReportURL && <Text>Link to: <a href={data.ReportURL}>External report</a></Text>}
+                  </Stack>
+               </Stack>
+            </Stack>
+         </Stack>
+
+         {failedTests.length > 0 &&
+            <TestResultPanel tests={failedTests} title={stateLabels.get('Fail')} selected={selectedTest} />
+         }
+         {inprocessTests.length > 0 &&
+            <TestResultPanel tests={inprocessTests} title={stateLabels.get('InProcess')} selected={selectedTest} />
+         }
+         {notrunTests.length > 0 &&
+            <TestResultPanel tests={notrunTests} title={stateLabels.get('NotRun')} selected={selectedTest} />
+         }
+         {skippedTests.length > 0 &&
+            <TestResultPanel tests={skippedTests} title={stateLabels.get('Skipped')} selected={selectedTest} />
+         }
+         {passedTests.length > 0 &&
+            <TestResultPanel tests={passedTests} title={stateLabels.get('Success')} selected={selectedTest} />
+         }
+      </Stack>
+   );
+}
diff --git a/Engine/Source/Programs/Horde/HordeDashboard/src/legacyPlugins/UnrealAutomatedTests/index.tsx b/Engine/Source/Programs/Horde/HordeDashboard/src/legacyPlugins/UnrealAutomatedTests/index.tsx
new file mode 100644
--- /dev/null
+++ b/Engine/Source/Programs/Horde/HordeDashboard/src/legacyPlugins/UnrealAutomatedTests/index.tsx
@@ -0,0 +1,19 @@
+// Copyright Epic Games, Inc. All Rights Reserved.
+
+import { Plugin, PluginMount } from "../pluginTypes";
+import { TestPassSummaryView } from "./components/UnrealAutomatedTests";
+
+class UnrealAutomatedTestsPlugin implements Plugin {
+
+    name = "Unreal Automated Tests";
+    components = [{
+        mount: PluginMount.TestReportPanel,
+        component: TestPassSummaryView,
+        id: "Unreal Automated Tests"
+    }];
+
+}
+
+const unrealAutomatedTestsPlugin = new UnrealAutomatedTestsPlugin();
+export default unrealAutomatedTestsPlugin;
+
diff --git a/Engine/Source/Programs/Horde/HordeDashboard/src/legacyPlugins/UnrealAutomatedTests/models/UnrealAutomatedTests.ts b/Engine/Source/Programs/Horde/HordeDashboard/src/legacyPlugins/UnrealAutomatedTests/models/UnrealAutomatedTests.ts
new file mode 100644
--- /dev/null
+++ b/Engine/Source/Programs/Horde/HordeDashboard/src/legacyPlugins/UnrealAutomatedTests/models/UnrealAutomatedTests.ts
@@ -0,0 +1,88 @@
+// Copyright Epic Games, Inc. All Rights Reserved.
+
+export enum TestState {
+    Failed = "Fail",
+    Success = "Success",
+    SuccessWithWarnings = "SuccessWithWarnings",
+    NotRun = "NotRun",
+    InProcess = "InProcess",
+    Skipped = "Skipped",
+    Unknown = "Unknown",
+}
+
+export enum EventType {
+    Info = "Info",
+    Error = "Error",
+    Warning = "Warning",
+}
+
+export type TestEntryArtifact = {
+    Id: string;
+    Name: string;
+    Type: string;
+    Files: { Difference: string, Approved: string, Unapproved: string };
+}
+
+export type TestEntry = {
+    Filename: string;
+    LineNumber: number;
+    Timestamp: string;
+    Event: { Type: string, Message: string, Context: string, Artifact: string };
+}
+
+export type TestDetails = {
+    TestDisplayName: string;
+    FullTestPath: string;
+    State: string;
+    DeviceInstance: string;
+    Errors: number;
+    Warnings: number;
+    Entries?: TestEntry[];
+    Artifacts?: TestEntryArtifact[];
+}
+
+export type TestResult = {
+    TestDisplayName: string;
+    FullTestPath: string;
+    State: string;
+    DeviceInstance: string;
+    ArtifactName: string;
+    Errors: number;
+    Warnings: number;
+}
+
+export type TestDevice = {
+    DeviceName: string;
+    Instance: string;
+    Platform: string;
+    OSVersion: string;
+    Model: string;
+    GPU: string;
+    CPUModel: string;
+    RAMInGB: number;
+    RenderMode: string;
+    RHI: string;
+}
+
+export type Metadata = {[Key in string]: string}
+
+export type TestPassSummary = {
+    ReportURL: string;
+    FailedCount: number;
+    NotRunCount: number;
+    InProcessCount: number;
+    ReportCreatedOn: string;
+    SucceededCount: number;
+    SucceededWithWarningsCount: number;
+    TotalDurationSeconds: number;
+    Tests: TestResult[];
+    Devices: TestDevice[];
+    Metadata: Metadata;
+}
+
+export type TestStateHistoryItem = {
+    TestdataId: string;
+    Change: number;
+    RangeUrl?: string;
+    State: string;
+}
diff --git a/Engine/Source/Programs/Horde/HordeDashboard/src/legacyPlugins/example/components/JobDetailPluginExample.tsx b/Engine/Source/Programs/Horde/HordeDashboard/src/legacyPlugins/example/components/JobDetailPluginExample.tsx
new file mode 100644
--- /dev/null
+++ b/Engine/Source/Programs/Horde/HordeDashboard/src/legacyPlugins/example/components/JobDetailPluginExample.tsx
@@ -0,0 +1,26 @@
+// Copyright Epic Games, Inc. All Rights Reserved.
+
+import { Stack, Text } from '@fluentui/react';
+import React from 'react';
+import { GetBatchResponse,  GetStepResponse } from '../../../backend/Api';
+import { JobDetails } from '../../../backend/JobDetails';
+import { getHordeStyling } from '../../../styles/Styles';
+
+export const JobDetailPluginExample: React.FC<{ jobDetails: JobDetails; step?: GetStepResponse; batch?: GetBatchResponse}> = ({ jobDetails, step, batch }) => {
+
+   const { hordeClasses } = getHordeStyling();
+
+    // do not render on batch
+    if (batch || !step) {
+        return null;
+    }
+
+    return <Stack styles={{ root: { paddingTop: 18, paddingRight: 12 } }}>
+        <Stack className={hordeClasses.raised}>
+            <Stack tokens={{ childrenGap: 12 }}>
+                <Text variant="mediumPlus" styles={{ root: { fontFamily: "Horde Open Sans SemiBold" } }}>Job Detail Plugin Example</Text>
+                <Text>{`This is an example job detail plugin example, job id is ${jobDetails.id}, step id is ${step.id} `}</Text>
+            </Stack>
+        </Stack>
+    </Stack>;
+};
\ No newline at end of file
diff --git a/Engine/Source/Programs/Horde/HordeDashboard/src/legacyPlugins/example/index.tsx b/Engine/Source/Programs/Horde/HordeDashboard/src/legacyPlugins/example/index.tsx
new file mode 100644
--- /dev/null
+++ b/Engine/Source/Programs/Horde/HordeDashboard/src/legacyPlugins/example/index.tsx
@@ -0,0 +1,21 @@
+// Copyright Epic Games, Inc. All Rights Reserved.
+import { Plugin, PluginMount } from "../pluginTypes";
+import { ExampleBar } from "./routes/ExampleBar";
+import { JobDetailPluginExample } from "./components/JobDetailPluginExample";
+
+class ExamplePlugin implements Plugin {
+
+    name = "ExamplePlugin";
+    routes = [{
+        path: "/example/bar",
+        component: ExampleBar
+    }];
+    components = [{
+        mount: PluginMount.JobDetailPanel,
+        component: JobDetailPluginExample
+    }]
+
+}
+
+const examplePlugin = new ExamplePlugin();
+export default examplePlugin;
\ No newline at end of file
diff --git a/Engine/Source/Programs/Horde/HordeDashboard/src/legacyPlugins/example/routes/ExampleBar.tsx b/Engine/Source/Programs/Horde/HordeDashboard/src/legacyPlugins/example/routes/ExampleBar.tsx
new file mode 100644
--- /dev/null
+++ b/Engine/Source/Programs/Horde/HordeDashboard/src/legacyPlugins/example/routes/ExampleBar.tsx
@@ -0,0 +1,10 @@
+// Copyright Epic Games, Inc. All Rights Reserved.
+
+import React from 'react';
+import { MessageBar, MessageBarType } from '@fluentui/react';
+
+export const ExampleBar: React.FC = () => {
+    return (
+        <MessageBar messageBarType={MessageBarType.info} isMultiline={false} > Hello from Example Plugin </MessageBar>
+    );
+};
\ No newline at end of file
diff --git a/Engine/Source/Programs/Horde/HordeDashboard/src/legacyPlugins/index.tsx b/Engine/Source/Programs/Horde/HordeDashboard/src/legacyPlugins/index.tsx
new file mode 100644
--- /dev/null
+++ b/Engine/Source/Programs/Horde/HordeDashboard/src/legacyPlugins/index.tsx
@@ -0,0 +1,53 @@
+// Copyright Epic Games, Inc. All Rights Reserved.
+
+import { Plugin, PluginMount, PluginRoute, PluginComponent } from "./pluginTypes";
+import SimpleTestReportPlugin from "./SimpleTestReport"
+import AutomatedTestSessionPlugin from "./AutomatedTestSession"
+import UnrealAutomatedTestsPlugin from "./UnrealAutomatedTests"
+
+export class Plugins {
+
+   get routes(): PluginRoute[] {
+      return this.plugins.map(p => p.routes ?? []).flat();
+   }
+
+   getComponents(mount: PluginMount, id?: string): PluginComponent[] {
+      return this.plugins.map(p => {
+         return p.components?.filter(c => c.mount === mount && c.id === id) ?? [];
+      }).flat();
+   }
+
+   loadPlugins(pluginList: string[] | undefined): void {
+
+      console.log("loading plugins...");
+
+      this.plugins.push(...[AutomatedTestSessionPlugin, SimpleTestReportPlugin, UnrealAutomatedTestsPlugin]);
+
+      /*
+      return new Promise<void>((resolve) => {
+
+         if (!pluginList || !pluginList.length || this.plugins.length) {
+            return resolve();
+         }
+
+         Promise.all(pluginList.map(async (plugin: any) => {
+            await import(`./${plugin}/index.js`).then((m: any) => {
+               this.plugins.push(m.default);
+               console.log(`loaded plugin: ${plugin}`);
+            }).catch((reason: any) => {
+               console.error(`unable to load plugin: ${plugin}, ${reason}`);
+            });
+         })).finally(() => {
+            resolve();
+         });
+      });
+      */
+   }
+
+   private plugins: Plugin[] = [];
+
+}
+
+const plugins = new Plugins();
+
+export default plugins;
\ No newline at end of file
diff --git a/Engine/Source/Programs/Horde/HordeDashboard/src/legacyPlugins/pluginTypes.ts b/Engine/Source/Programs/Horde/HordeDashboard/src/legacyPlugins/pluginTypes.ts
new file mode 100644
--- /dev/null
+++ b/Engine/Source/Programs/Horde/HordeDashboard/src/legacyPlugins/pluginTypes.ts
@@ -0,0 +1,28 @@
+
+// render components with top level routes
+export type PluginRoute = {
+    path: string;
+    component: React.FC;
+ }
+ 
+ // render components with specific mount points
+ export enum PluginMount {
+    JobDetailPanel,
+    TestReportPanel,
+    TestReportLink,
+    BuildHealthSummary,
+ }
+ 
+ export type PluginComponent = {
+    mount: PluginMount;
+    component: React.FC<any>;
+    // optional unique identifier for the component
+    id?: string;
+ }
+ 
+ export interface Plugin {
+    name: string;
+    routes?: PluginRoute[];
+    components?: PluginComponent[];
+ }
+ 
\ No newline at end of file
