Horde: Add support for tagging log events with issue fingerprints.

The $issue property on a log event can contain a serialized IssueFingerprint, which will be parsed by the ExternalIssueHandler handler inside Horde. This property can be added from C# code using the ILogger.BeginIssueScope() extension method.

diff --git a/Engine/Source/Programs/Horde/Plugins/Build/HordeServer.Build.Tests/Issues/IssueServiceTests.cs b/Engine/Source/Programs/Horde/Plugins/Build/HordeServer.Build.Tests/Issues/IssueServiceTests.cs
--- a/Engine/Source/Programs/Horde/Plugins/Build/HordeServer.Build.Tests/Issues/IssueServiceTests.cs
+++ b/Engine/Source/Programs/Horde/Plugins/Build/HordeServer.Build.Tests/Issues/IssueServiceTests.cs
@@ -49,6 +49,7 @@ namespace HordeServer.Tests.Issues
 			readonly LogBuilder _builder;
 			readonly List<(LogLevel, ReadOnlyMemory<byte>)> _events = new List<(LogLevel, ReadOnlyMemory<byte>)>();
 			readonly IStorageClient _storageClient;
+			readonly LoggerScopeCollection _scopeCollection = new LoggerScopeCollection();
 
 			int _lineIndex;
 
@@ -78,14 +79,15 @@ namespace HordeServer.Tests.Issues
 				_storageClient.Dispose();
 			}
 
-			public IDisposable? BeginScope<TState>(TState state) where TState : notnull => null!;
+			public IDisposable? BeginScope<TState>(TState state) where TState : notnull => _scopeCollection.BeginScope(state);
 
 			public bool IsEnabled(LogLevel logLevel) => true;
 
 			public void Log<TState>(LogLevel logLevel, EventId eventId, TState state, Exception? exception, Func<TState, Exception?, string> formatter)
 			{
-				JsonLogEvent logEvent = JsonLogEvent.FromLoggerState(logLevel, eventId, state, exception, formatter);
-				_events.Add((logLevel, logEvent.Data));
+				LogEvent logEvent = LogEvent.FromState(logLevel, eventId, state, exception, formatter);
+				logEvent.AddProperties(_scopeCollection.GetProperties());
+				_events.Add((logLevel, logEvent.ToJsonBytes()));
 			}
 
 			private async Task WriteAsync(LogLevel level, byte[] line)
@@ -377,24 +379,34 @@ namespace HordeServer.Tests.Issues
 		}
 
 		private async Task ParseEventsAsync(IJob job, int batchIdx, int stepIdx, string[] lines)
+		{
+			void WriteLinesToLogger(ILogger logger)
+			{
+				using (LogParser parser = new LogParser(logger, new List<string>()))
+				{
+					for (int idx = 0; idx < lines.Length; idx++)
+					{
+						parser.WriteLine(lines[idx]);
+					}
+				}
+			}
+
+			await WriteEventsAsync(job, batchIdx, stepIdx, WriteLinesToLogger);
+		}
+
+		private async Task WriteEventsAsync(IJob job, int batchIdx, int stepIdx, Action<ILogger> writeEvents)
 		{
 			LogId logId = job.Batches[batchIdx].Steps[stepIdx].LogId!.Value;
 			ILog log = (await LogCollection.GetAsync(logId))!;
 
-			IStorageClient storageClient = StorageService.CreateClient(Namespace.Logs);
+			using IStorageClient storageClient = StorageService.CreateClient(Namespace.Logs);
 			await using (TestJsonLogger logger = new TestJsonLogger(log, storageClient))
 			{
 				PerforceMetadataLogger perforceLogger = new PerforceMetadataLogger(logger);
 				perforceLogger.AddClientView(_autoSdkDir, "//depot/CarefullyRedist/...", 12345);
 				perforceLogger.AddClientView(_workspaceDir, "//UE4/Main/...", 12345);
 
-				using (LogParser parser = new LogParser(perforceLogger, new List<string>()))
-				{
-					for (int idx = 0; idx < lines.Length; idx++)
-					{
-						parser.WriteLine(lines[idx]);
-					}
-				}
+				writeEvents(perforceLogger);
 			}
 		}
 
@@ -1554,6 +1566,37 @@ namespace HordeServer.Tests.Issues
 			Assert.AreEqual("Warnings in Phsyarum_BP2.uasset", issues2[1].Summary);
 		}
 
+		[TestMethod]
+		public async Task ExternalIssueTestAsync()
+		{
+			IJob job = CreateJob(_mainStreamId, 120, "Compile Test", _graph);
+
+			await WriteEventsAsync(job, 0, 0, logger => {
+				logger.LogInformation("Foo");
+
+				IssueFingerprint fingerprint = new IssueFingerprint("NewIssueType", "This is an issue with severity: {Severity}");
+				fingerprint.Keys.Add(IssueKey.FromFile("foo.cpp"));
+
+				using (IDisposable? scope = logger.BeginIssueScope(fingerprint))
+				{
+					logger.LogWarning("This is a warning!");
+					logger.LogError("This is an error!");
+				}
+
+				logger.LogInformation("Bar");
+			});
+			await UpdateCompleteStepAsync(job, 0, 0, JobStepOutcome.Failure);
+
+			IReadOnlyList<IIssue> issues = await IssueCollection.FindIssuesAsync();
+			Assert.AreEqual(1, issues.Count);
+
+			IIssue issue = issues[0];
+			Assert.AreEqual("This is an issue with severity: errors", issue.Summary);
+			Assert.AreEqual(1, issue.Fingerprints.Count);
+			Assert.AreEqual(1, issue.Fingerprints[0].Keys.Count);
+			Assert.AreEqual(IssueKey.FromFile("foo.cpp"), issue.Fingerprints[0].Keys.First());
+		}
+
 		[TestMethod]
 		public async Task HashedIssueTestAsync()
 		{
diff --git a/Engine/Source/Programs/Horde/Plugins/Build/HordeServer.Build/Issues/IIssueFingerprint.cs b/Engine/Source/Programs/Horde/Plugins/Build/HordeServer.Build/Issues/IIssueFingerprint.cs
deleted file mode 100644
diff --git a/Engine/Source/Programs/Horde/Plugins/Build/HordeServer.Build/Issues/IIssueSpan.cs b/Engine/Source/Programs/Horde/Plugins/Build/HordeServer.Build/Issues/IIssueSpan.cs
--- a/Engine/Source/Programs/Horde/Plugins/Build/HordeServer.Build/Issues/IIssueSpan.cs
+++ b/Engine/Source/Programs/Horde/Plugins/Build/HordeServer.Build/Issues/IIssueSpan.cs
@@ -1,6 +1,6 @@
 // Copyright Epic Games, Inc. All Rights Reserved.
 
-using System.Collections.Generic;
+using EpicGames.Horde.Issues;
 using EpicGames.Horde.Jobs.Templates;
 using EpicGames.Horde.Streams;
 using EpicGames.Horde.Users;
diff --git a/Engine/Source/Programs/Horde/Plugins/Build/HordeServer.Build/Issues/IssueCollection.cs b/Engine/Source/Programs/Horde/Plugins/Build/HordeServer.Build/Issues/IssueCollection.cs
--- a/Engine/Source/Programs/Horde/Plugins/Build/HordeServer.Build/Issues/IssueCollection.cs
+++ b/Engine/Source/Programs/Horde/Plugins/Build/HordeServer.Build/Issues/IssueCollection.cs
@@ -46,7 +46,7 @@ namespace HordeServer.Issues
 		}
 
 		[DebuggerDisplay("{Id}: {Summary}")]
-		class Issue : IIssue
+		class IssueDocument : IIssue
 		{
 			[BsonId]
 			public int Id { get; set; }
@@ -60,8 +60,8 @@ namespace HordeServer.Issues
 			public string? Description { get; set; }
 
 			[BsonIgnoreIfNull]
-			public IssueFingerprint? Fingerprint { get; set; }
-			public List<IssueFingerprint>? Fingerprints { get; set; }
+			public IssueFingerprintDocument? Fingerprint { get; set; }
+			public List<IssueFingerprintDocument>? Fingerprints { get; set; }
 
 			public IssueSeverity Severity { get; set; }
 
@@ -105,20 +105,20 @@ namespace HordeServer.Issues
 			[BsonIgnoreIfNull]
 			public int? FixChange { get; set; }
 
-			public List<IssueStream> Streams { get; set; } = new List<IssueStream>();
+			public List<IssueStreamDocument> Streams { get; set; } = new List<IssueStreamDocument>();
 
 			public int MinSuspectChange { get; set; }
 			public int MaxSuspectChange { get; set; }
 
 			[BsonElement("Suspects"), BsonIgnoreIfNull]
-			public List<IssueSuspect>? SuspectsDeprecated { get; set; }
+			public List<IssueSuspectDocument>? SuspectsDeprecated { get; set; }
 
 			[BsonIgnoreIfNull]
 			public List<ObjectId>? ExcludeSpans { get; set; }
 
 			public int UpdateIndex { get; set; }
 
-			IReadOnlyList<IIssueFingerprint> IIssue.Fingerprints => Fingerprints ?? ((Fingerprint == null) ? new List<IssueFingerprint>() : new List<IssueFingerprint> { Fingerprint });
+			IReadOnlyList<IIssueFingerprint> IIssue.Fingerprints => Fingerprints ?? ((Fingerprint == null) ? new List<IssueFingerprintDocument>() : new List<IssueFingerprintDocument> { Fingerprint });
 			UserId? IIssue.OwnerId => OwnerId ?? DefaultOwnerId ?? GetDefaultOwnerId();
 			IReadOnlyList<IIssueStream> IIssue.Streams => Streams;
 			DateTime IIssue.LastSeenAt => (LastSeenAt == default) ? DateTime.UtcNow : LastSeenAt;
@@ -139,13 +139,13 @@ namespace HordeServer.Issues
 			public Uri? WorkflowThreadUrl { get; set; }
 
 			[BsonConstructor]
-			private Issue()
+			private IssueDocument()
 			{
 				Summary = String.Empty;
 				Fingerprint = null!;
 			}
 
-			public Issue(int id, string summary)
+			public IssueDocument(int id, string summary)
 			{
 				Id = id;
 				Summary = summary;
@@ -183,18 +183,18 @@ namespace HordeServer.Issues
 			}
 		}
 
-		class IssueStream : IIssueStream
+		class IssueStreamDocument : IIssueStream
 		{
 			public StreamId StreamId { get; set; }
 			public bool? MergeOrigin { get; set; }
 			public bool? ContainsFix { get; set; }
 			public bool? FixFailed { get; set; }
 
-			public IssueStream()
+			public IssueStreamDocument()
 			{
 			}
 
-			public IssueStream(IIssueStream other)
+			public IssueStreamDocument(IIssueStream other)
 			{
 				StreamId = other.StreamId;
 				MergeOrigin = other.MergeOrigin;
@@ -203,7 +203,7 @@ namespace HordeServer.Issues
 			}
 		}
 
-		class IssueSuspect : IIssueSuspect
+		class IssueSuspectDocument : IIssueSuspect
 		{
 			public ObjectId Id { get; set; }
 			public int IssueId { get; set; }
@@ -212,11 +212,11 @@ namespace HordeServer.Issues
 			public DateTime? DeclinedAt { get; set; }
 			public DateTime? ResolvedAt { get; set; } // Degenerate
 
-			private IssueSuspect()
+			private IssueSuspectDocument()
 			{
 			}
 
-			public IssueSuspect(int issueId, NewIssueSuspectData newSuspect, DateTime? resolvedAt)
+			public IssueSuspectDocument(int issueId, NewIssueSuspectData newSuspect, DateTime? resolvedAt)
 			{
 				Id = ObjectId.GenerateNewId();
 				IssueId = issueId;
@@ -225,12 +225,12 @@ namespace HordeServer.Issues
 				ResolvedAt = resolvedAt;
 			}
 
-			public IssueSuspect(int issueId, IIssueSpanSuspect suspect)
+			public IssueSuspectDocument(int issueId, IIssueSpanSuspect suspect)
 				: this(issueId, suspect.AuthorId, suspect.OriginatingChange ?? suspect.Change, null, null)
 			{
 			}
 
-			public IssueSuspect(int issueId, UserId authorId, int change, DateTime? declinedAt, DateTime? resolvedAt)
+			public IssueSuspectDocument(int issueId, UserId authorId, int change, DateTime? declinedAt, DateTime? resolvedAt)
 			{
 				Id = ObjectId.GenerateNewId();
 				IssueId = issueId;
@@ -241,7 +241,7 @@ namespace HordeServer.Issues
 			}
 		}
 
-		class IssueFingerprint : IIssueFingerprint
+		class IssueFingerprintDocument : IIssueFingerprint
 		{
 			public string Type { get; set; }
 
@@ -280,12 +280,12 @@ namespace HordeServer.Issues
 			string IIssueFingerprint.ChangeFilter => ChangeFilter ?? String.Join(";", GetLegacyHandlerInfo(Type).ChangeFilter);
 
 			[BsonConstructor]
-			private IssueFingerprint()
+			private IssueFingerprintDocument()
 			{
 				Type = String.Empty;
 			}
 
-			public IssueFingerprint(IIssueFingerprint fingerprint)
+			public IssueFingerprintDocument(IIssueFingerprint fingerprint)
 			{
 				Type = fingerprint.Type;
 				SummaryTemplateValue = fingerprint.SummaryTemplate;
@@ -406,7 +406,7 @@ namespace HordeServer.Issues
 			#endregion
 		}
 
-		class IssueSpan : IIssueSpan
+		class IssueSpanDocument : IIssueSpan
 		{
 			public ObjectId Id { get; set; }
 
@@ -424,20 +424,20 @@ namespace HordeServer.Issues
 			public DateTime? ResolvedAt { get; set; } // Propagated from the owning issue
 
 			[BsonRequired]
-			public IssueFingerprint Fingerprint { get; set; }
+			public IssueFingerprintDocument Fingerprint { get; set; }
 
 			public int MinChange { get; set; }
 			public int MaxChange { get; set; } = Int32.MaxValue;
 
-			public IssueStep? LastSuccess { get; set; }
+			public IssueStepDocument? LastSuccess { get; set; }
 
 			[BsonRequired]
-			public IssueStep FirstFailure { get; set; }
+			public IssueStepDocument FirstFailure { get; set; }
 
 			[BsonRequired]
-			public IssueStep LastFailure { get; set; }
+			public IssueStepDocument LastFailure { get; set; }
 
-			public IssueStep? NextSuccess { get; set; }
+			public IssueStepDocument? NextSuccess { get; set; }
 
 			public bool? PromoteByDefault { get; set; }
 
@@ -446,7 +446,7 @@ namespace HordeServer.Issues
 
 			bool IIssueSpan.PromoteByDefault => PromoteByDefault ?? NotifySuspectsDeprecated;
 
-			public List<IssueSpanSuspect> Suspects { get; set; }
+			public List<IssueSpanSuspectDocument> Suspects { get; set; }
 			public int IssueId { get; set; }
 			public int UpdateIndex { get; set; }
 
@@ -457,54 +457,54 @@ namespace HordeServer.Issues
 			IReadOnlyList<IIssueSpanSuspect> IIssueSpan.Suspects => Suspects;
 			IIssueFingerprint IIssueSpan.Fingerprint => Fingerprint;
 
-			private IssueSpan()
+			private IssueSpanDocument()
 			{
 				StreamName = null!;
 				NodeName = null!;
 				Fingerprint = null!;
 				FirstFailure = null!;
 				LastFailure = null!;
-				Suspects = new List<IssueSpanSuspect>();
+				Suspects = new List<IssueSpanSuspectDocument>();
 			}
 
-			public IssueSpan(int issueId, NewIssueSpanData newSpan)
+			public IssueSpanDocument(int issueId, NewIssueSpanData newSpan)
 			{
 				Id = ObjectId.GenerateNewId();
 				StreamId = newSpan.StreamId;
 				StreamName = newSpan.StreamName;
 				TemplateRefId = newSpan.TemplateRefId;
 				NodeName = newSpan.NodeName;
-				Fingerprint = new IssueFingerprint(newSpan.Fingerprint);
+				Fingerprint = new IssueFingerprintDocument(newSpan.Fingerprint);
 				if (newSpan.LastSuccess != null)
 				{
 					MinChange = newSpan.LastSuccess.Change;
-					LastSuccess = new IssueStep(Id, newSpan.LastSuccess);
+					LastSuccess = new IssueStepDocument(Id, newSpan.LastSuccess);
 				}
-				FirstFailure = new IssueStep(Id, newSpan.FirstFailure);
-				LastFailure = new IssueStep(Id, newSpan.FirstFailure);
+				FirstFailure = new IssueStepDocument(Id, newSpan.FirstFailure);
+				LastFailure = new IssueStepDocument(Id, newSpan.FirstFailure);
 				if (newSpan.NextSuccess != null)
 				{
 					MaxChange = newSpan.NextSuccess.Change;
-					NextSuccess = new IssueStep(Id, newSpan.NextSuccess);
+					NextSuccess = new IssueStepDocument(Id, newSpan.NextSuccess);
 				}
 				PromoteByDefault = newSpan.FirstFailure.PromoteByDefault;
-				Suspects = newSpan.Suspects.ConvertAll(x => new IssueSpanSuspect(x));
+				Suspects = newSpan.Suspects.ConvertAll(x => new IssueSpanSuspectDocument(x));
 				IssueId = issueId;
 			}
 		}
 
-		class IssueSpanSuspect : IIssueSpanSuspect
+		class IssueSpanSuspectDocument : IIssueSpanSuspect
 		{
 			public int Change { get; set; }
 			public UserId AuthorId { get; set; }
 			public int? OriginatingChange { get; set; }
 
 			[BsonConstructor]
-			private IssueSpanSuspect()
+			private IssueSpanSuspectDocument()
 			{
 			}
 
-			public IssueSpanSuspect(NewIssueSpanSuspectData newSuspectData)
+			public IssueSpanSuspectDocument(NewIssueSpanSuspectData newSuspectData)
 			{
 				Change = newSuspectData.Change;
 				AuthorId = newSuspectData.AuthorId;
@@ -512,7 +512,7 @@ namespace HordeServer.Issues
 			}
 		}
 
-		class IssueStep : IIssueStep
+		class IssueStepDocument : IIssueStep
 		{
 			public ObjectId Id { get; set; }
 			public ObjectId SpanId { get; set; }
@@ -548,12 +548,12 @@ namespace HordeServer.Issues
 			bool IIssueStep.PromoteByDefault => PromoteByDefault ?? NotifySuspectsDeprecated;
 
 			[BsonConstructor]
-			private IssueStep()
+			private IssueStepDocument()
 			{
 				JobName = null!;
 			}
 
-			public IssueStep(ObjectId spanId, NewIssueStepData stepData)
+			public IssueStepDocument(ObjectId spanId, NewIssueStepData stepData)
 			{
 				Id = ObjectId.GenerateNewId();
 				SpanId = spanId;
@@ -600,10 +600,10 @@ namespace HordeServer.Issues
 		readonly IRedisService _redisService;
 		readonly IUserCollection _userCollection;
 		readonly ISingletonDocument<IssueLedger> _ledgerSingleton;
-		readonly IMongoCollection<Issue> _issues;
-		readonly IMongoCollection<IssueSpan> _issueSpans;
-		readonly IMongoCollection<IssueStep> _issueSteps;
-		readonly IMongoCollection<IssueSuspect> _issueSuspects;
+		readonly IMongoCollection<IssueDocument> _issues;
+		readonly IMongoCollection<IssueSpanDocument> _issueSpans;
+		readonly IMongoCollection<IssueStepDocument> _issueSteps;
+		readonly IMongoCollection<IssueSuspectDocument> _issueSuspects;
 		readonly IAuditLog<int> _auditLog;
 		readonly ITelemetryWriter _telemetryWriter;
 		readonly IOptionsMonitor<BuildConfig> _buildConfig;
@@ -615,6 +615,7 @@ namespace HordeServer.Issues
 			BsonClassMap.RegisterClassMap<IssueKey>(cm =>
 			{
 				cm.MapConstructor(() => new IssueKey("", IssueKeyType.None), nameof(IssueKey.Name), nameof(IssueKey.Type));
+				cm.MapConstructor(() => new IssueKey("", IssueKeyType.None, null), nameof(IssueKey.Name), nameof(IssueKey.Type), nameof(IssueKey.Scope));
 				cm.MapProperty(x => x.Name).SetElementName("n");
 				cm.MapProperty(x => x.Type).SetElementName("t");
 				cm.MapProperty(x => x.Scope).SetElementName("s").SetIgnoreIfNull(true);
@@ -639,27 +640,27 @@ namespace HordeServer.Issues
 
 			_ledgerSingleton = new SingletonDocument<IssueLedger>(mongoService);
 
-			List<MongoIndex<Issue>> issueIndexes = new List<MongoIndex<Issue>>();
+			List<MongoIndex<IssueDocument>> issueIndexes = new List<MongoIndex<IssueDocument>>();
 			issueIndexes.Add(keys => keys.Ascending(x => x.ResolvedAt));
 			issueIndexes.Add(keys => keys.Ascending(x => x.VerifiedAt));
-			_issues = mongoService.GetCollection<Issue>("IssuesV2");
+			_issues = mongoService.GetCollection<IssueDocument>("IssuesV2");
 
-			List<MongoIndex<IssueSpan>> issueSpanIndexes = new List<MongoIndex<IssueSpan>>();
+			List<MongoIndex<IssueSpanDocument>> issueSpanIndexes = new List<MongoIndex<IssueSpanDocument>>();
 			issueSpanIndexes.Add(keys => keys.Ascending(x => x.IssueId));
 			issueSpanIndexes.Add(keys => keys.Ascending(x => x.StreamId).Ascending(x => x.MinChange).Ascending(x => x.MaxChange));
 			issueSpanIndexes.Add("StreamChanges", keys => keys.Ascending(x => x.StreamId).Ascending(x => x.TemplateRefId).Ascending(x => x.NodeName).Ascending(x => x.MinChange).Ascending(x => x.MaxChange));
-			_issueSpans = mongoService.GetCollection<IssueSpan>("IssuesV2.Spans", issueSpanIndexes);
+			_issueSpans = mongoService.GetCollection<IssueSpanDocument>("IssuesV2.Spans", issueSpanIndexes);
 
-			List<MongoIndex<IssueStep>> issueStepIndexes = new List<MongoIndex<IssueStep>>();
+			List<MongoIndex<IssueStepDocument>> issueStepIndexes = new List<MongoIndex<IssueStepDocument>>();
 			issueStepIndexes.Add(keys => keys.Ascending(x => x.SpanId));
 			issueStepIndexes.Add(keys => keys.Ascending(x => x.JobId).Ascending(x => x.BatchId).Ascending(x => x.StepId));
-			_issueSteps = mongoService.GetCollection<IssueStep>("IssuesV2.Steps", issueStepIndexes);
+			_issueSteps = mongoService.GetCollection<IssueStepDocument>("IssuesV2.Steps", issueStepIndexes);
 
-			List<MongoIndex<IssueSuspect>> issueSuspectIndexes = new List<MongoIndex<IssueSuspect>>();
+			List<MongoIndex<IssueSuspectDocument>> issueSuspectIndexes = new List<MongoIndex<IssueSuspectDocument>>();
 			issueSuspectIndexes.Add(keys => keys.Ascending(x => x.Change));
 			issueSuspectIndexes.Add(keys => keys.Ascending(x => x.AuthorId).Ascending(x => x.ResolvedAt));
 			issueSuspectIndexes.Add(keys => keys.Ascending(x => x.IssueId).Ascending(x => x.Change), unique: true);
-			_issueSuspects = mongoService.GetCollection<IssueSuspect>("IssuesV2.Suspects", issueSuspectIndexes);
+			_issueSuspects = mongoService.GetCollection<IssueSuspectDocument>("IssuesV2.Suspects", issueSuspectIndexes);
 
 			_auditLog = auditLogFactory.Create("IssuesV2.History", "IssueId");
 		}
@@ -686,16 +687,16 @@ namespace HordeServer.Issues
 			return new IssueLock(issueLock, _tracer, _logger);
 		}
 
-		async Task<Issue?> TryUpdateIssueAsync(IIssue issue, UpdateDefinition<Issue> update, CancellationToken cancellationToken)
+		async Task<IssueDocument?> TryUpdateIssueAsync(IIssue issue, UpdateDefinition<IssueDocument> update, CancellationToken cancellationToken)
 		{
-			Issue issueDocument = (Issue)issue;
+			IssueDocument issueDocument = (IssueDocument)issue;
 
 			int prevUpdateIndex = issueDocument.UpdateIndex;
 			update = update.Set(x => x.UpdateIndex, prevUpdateIndex + 1);
 
-			FindOneAndUpdateOptions<Issue, Issue> options = new FindOneAndUpdateOptions<Issue, Issue> { ReturnDocument = ReturnDocument.After };
+			FindOneAndUpdateOptions<IssueDocument, IssueDocument> options = new FindOneAndUpdateOptions<IssueDocument, IssueDocument> { ReturnDocument = ReturnDocument.After };
 
-			Issue? newIssue = await _issues.FindOneAndUpdateAsync<Issue>(x => x.Id == issueDocument.Id && x.UpdateIndex == prevUpdateIndex, update, options, cancellationToken);
+			IssueDocument? newIssue = await _issues.FindOneAndUpdateAsync<IssueDocument>(x => x.Id == issueDocument.Id && x.UpdateIndex == prevUpdateIndex, update, options, cancellationToken);
 			if (newIssue != null)
 			{
 				SendTelemetry(newIssue);
@@ -703,16 +704,16 @@ namespace HordeServer.Issues
 			return newIssue;
 		}
 
-		async Task<IssueSpan?> TryUpdateSpanAsync(IIssueSpan issueSpan, UpdateDefinition<IssueSpan> update, CancellationToken cancellationToken)
+		async Task<IssueSpanDocument?> TryUpdateSpanAsync(IIssueSpan issueSpan, UpdateDefinition<IssueSpanDocument> update, CancellationToken cancellationToken)
 		{
-			IssueSpan issueSpanDocument = (IssueSpan)issueSpan;
+			IssueSpanDocument issueSpanDocument = (IssueSpanDocument)issueSpan;
 
 			int prevUpdateIndex = issueSpanDocument.UpdateIndex;
 			update = update.Set(x => x.UpdateIndex, prevUpdateIndex + 1);
 
-			FindOneAndUpdateOptions<IssueSpan, IssueSpan> options = new FindOneAndUpdateOptions<IssueSpan, IssueSpan> { ReturnDocument = ReturnDocument.After };
+			FindOneAndUpdateOptions<IssueSpanDocument, IssueSpanDocument> options = new FindOneAndUpdateOptions<IssueSpanDocument, IssueSpanDocument> { ReturnDocument = ReturnDocument.After };
 
-			IssueSpan? newIssueSpan = await _issueSpans.FindOneAndUpdateAsync<IssueSpan>(x => x.Id == issueSpanDocument.Id && x.UpdateIndex == prevUpdateIndex, update, options, cancellationToken);
+			IssueSpanDocument? newIssueSpan = await _issueSpans.FindOneAndUpdateAsync<IssueSpanDocument>(x => x.Id == issueSpanDocument.Id && x.UpdateIndex == prevUpdateIndex, update, options, cancellationToken);
 			if (newIssueSpan != null)
 			{
 				SendTelemetry(newIssueSpan);
@@ -727,7 +728,7 @@ namespace HordeServer.Issues
 		{
 			IssueLedger ledger = await _ledgerSingleton.UpdateAsync(x => x.NextId++, cancellationToken);
 
-			Issue newIssue = new Issue(ledger.NextId, summary);
+			IssueDocument newIssue = new IssueDocument(ledger.NextId, summary);
 			await _issues.InsertOneAsync(newIssue, null, cancellationToken);
 			SendTelemetry(newIssue);
 
@@ -794,14 +795,14 @@ namespace HordeServer.Issues
 			return user.Name;
 		}
 
-		async Task LogIssueChangesAsync(UserId? initiatedByUserId, Issue oldIssue, Issue newIssue, CancellationToken cancellationToken)
+		async Task LogIssueChangesAsync(UserId? initiatedByUserId, IssueDocument oldIssue, IssueDocument newIssue, CancellationToken cancellationToken)
 		{
 			ILogger issueLogger = GetLogger(oldIssue.Id);
 			using IDisposable? scope = issueLogger.BeginScope("User {UserName} ({UserId})", await GetUserNameAsync(initiatedByUserId, cancellationToken), initiatedByUserId ?? UserId.Empty);
 			await LogIssueChangesImplAsync(issueLogger, oldIssue, newIssue, cancellationToken);
 		}
 
-		async Task LogIssueChangesImplAsync(ILogger issueLogger, Issue oldIssue, Issue newIssue, CancellationToken cancellationToken)
+		async Task LogIssueChangesImplAsync(ILogger issueLogger, IssueDocument oldIssue, IssueDocument newIssue, CancellationToken cancellationToken)
 		{
 			if (newIssue.Severity != oldIssue.Severity)
 			{
@@ -943,7 +944,7 @@ namespace HordeServer.Issues
 			}
 		}
 
-		async Task LogIssueSuspectChangesAsync(ILogger issueLogger, IReadOnlyList<IssueSuspect> oldIssueSuspects, List<IssueSuspect> newIssueSuspects, CancellationToken cancellationToken)
+		async Task LogIssueSuspectChangesAsync(ILogger issueLogger, IReadOnlyList<IssueSuspectDocument> oldIssueSuspects, List<IssueSuspectDocument> newIssueSuspects, CancellationToken cancellationToken)
 		{
 			HashSet<(UserId, int)> oldSuspects = new HashSet<(UserId, int)>(oldIssueSuspects.Select(x => (x.AuthorId, x.Change)));
 			HashSet<(UserId, int)> newSuspects = new HashSet<(UserId, int)>(newIssueSuspects.Select(x => (x.AuthorId, x.Change)));
@@ -971,7 +972,7 @@ namespace HordeServer.Issues
 		/// <inheritdoc/>
 		public async Task<IIssue?> GetIssueAsync(int issueId, CancellationToken cancellationToken)
 		{
-			Issue issue = await _issues.Find(x => x.Id == issueId).FirstOrDefaultAsync(cancellationToken);
+			IssueDocument issue = await _issues.Find(x => x.Id == issueId).FirstOrDefaultAsync(cancellationToken);
 			return issue;
 		}
 
@@ -1012,45 +1013,45 @@ namespace HordeServer.Issues
 			}
 			else
 			{
-				FilterDefinition<IssueSpan> filter = Builders<IssueSpan>.Filter.Eq(x => x.StreamId, streamId.Value);
+				FilterDefinition<IssueSpanDocument> filter = Builders<IssueSpanDocument>.Filter.Eq(x => x.StreamId, streamId.Value);
 				if (ids != null)
 				{
-					filter &= Builders<IssueSpan>.Filter.In(x => x.IssueId, ids.Select<int, int?>(x => x));
+					filter &= Builders<IssueSpanDocument>.Filter.In(x => x.IssueId, ids.Select<int, int?>(x => x));
 				}
 				else
 				{
-					filter &= Builders<IssueSpan>.Filter.Exists(x => x.IssueId);
+					filter &= Builders<IssueSpanDocument>.Filter.Exists(x => x.IssueId);
 				}
 
 				if (minChange != null)
 				{
-					filter &= Builders<IssueSpan>.Filter.Not(Builders<IssueSpan>.Filter.Lt(x => x.MaxChange, minChange.Value));
+					filter &= Builders<IssueSpanDocument>.Filter.Not(Builders<IssueSpanDocument>.Filter.Lt(x => x.MaxChange, minChange.Value));
 				}
 				if (maxChange != null)
 				{
-					filter &= Builders<IssueSpan>.Filter.Not(Builders<IssueSpan>.Filter.Gt(x => x.MinChange, maxChange.Value));
+					filter &= Builders<IssueSpanDocument>.Filter.Not(Builders<IssueSpanDocument>.Filter.Gt(x => x.MinChange, maxChange.Value));
 				}
 
 				if (resolved != null)
 				{
 					if (resolved.Value)
 					{
-						filter &= Builders<IssueSpan>.Filter.Ne(x => x.ResolvedAt, null);
+						filter &= Builders<IssueSpanDocument>.Filter.Ne(x => x.ResolvedAt, null);
 					}
 					else
 					{
-						filter &= Builders<IssueSpan>.Filter.Eq(x => x.ResolvedAt, null);
+						filter &= Builders<IssueSpanDocument>.Filter.Eq(x => x.ResolvedAt, null);
 					}
 				}
 
 				using (IAsyncCursor<ProjectedIssueId> cursor = await _issueSpans.Aggregate().Match(filter).Group(x => x.IssueId, x => new ProjectedIssueId { _id = x.Key }).SortByDescending(x => x._id).ToCursorAsync(cancellationToken))
 				{
-					List<Issue> results = await PaginatedJoinAsync(cursor, (nextIds, nextIndex, nextCount) => FilterIssuesByOtherFieldsAsync(nextIds, null, null, null, promoted, nextIndex, nextCount, cancellationToken), index, count, cancellationToken);
+					List<IssueDocument> results = await PaginatedJoinAsync(cursor, (nextIds, nextIndex, nextCount) => FilterIssuesByOtherFieldsAsync(nextIds, null, null, null, promoted, nextIndex, nextCount, cancellationToken), index, count, cancellationToken);
 					if (resolved != null)
 					{
 						for (int idx = results.Count - 1; idx >= 0; idx--)
 						{
-							Issue issue = results[idx];
+							IssueDocument issue = results[idx];
 							if ((issue.ResolvedAt != null) != resolved.Value && issue.ResolvedById != IIssue.ResolvedByTimeoutId)
 							{
 								_logger.LogWarning("Issue {IssueId} has resolved state out of sync with spans", issue.Id);
@@ -1063,41 +1064,41 @@ namespace HordeServer.Issues
 			}
 		}
 
-		async Task<List<Issue>> FilterIssuesByOtherFieldsAsync(IEnumerable<int>? ids, int? minChange, int? maxChange, bool? resolved, bool? promoted, int index, int? count, CancellationToken cancellationToken)
+		async Task<List<IssueDocument>> FilterIssuesByOtherFieldsAsync(IEnumerable<int>? ids, int? minChange, int? maxChange, bool? resolved, bool? promoted, int index, int? count, CancellationToken cancellationToken)
 		{
-			FilterDefinition<Issue> filter = FilterDefinition<Issue>.Empty;
+			FilterDefinition<IssueDocument> filter = FilterDefinition<IssueDocument>.Empty;
 			if (ids != null)
 			{
-				filter &= Builders<Issue>.Filter.In(x => x.Id, ids);
+				filter &= Builders<IssueDocument>.Filter.In(x => x.Id, ids);
 			}
 			if (resolved != null)
 			{
 				if (resolved.Value)
 				{
-					filter &= Builders<Issue>.Filter.Ne(x => x.ResolvedAt, null);
+					filter &= Builders<IssueDocument>.Filter.Ne(x => x.ResolvedAt, null);
 				}
 				else
 				{
-					filter &= Builders<Issue>.Filter.Eq(x => x.ResolvedAt, null);
+					filter &= Builders<IssueDocument>.Filter.Eq(x => x.ResolvedAt, null);
 				}
 			}
 			if (minChange != null)
 			{
-				filter &= Builders<Issue>.Filter.Not(Builders<Issue>.Filter.Lt(x => x.MaxSuspectChange, minChange.Value));
+				filter &= Builders<IssueDocument>.Filter.Not(Builders<IssueDocument>.Filter.Lt(x => x.MaxSuspectChange, minChange.Value));
 			}
 			if (maxChange != null)
 			{
-				filter &= Builders<Issue>.Filter.Not(Builders<Issue>.Filter.Gt(x => x.MinSuspectChange, maxChange.Value));
+				filter &= Builders<IssueDocument>.Filter.Not(Builders<IssueDocument>.Filter.Gt(x => x.MinSuspectChange, maxChange.Value));
 			}
 			if (promoted != null)
 			{
 				if (promoted.Value)
 				{
-					filter &= Builders<Issue>.Filter.Eq(x => x.Promoted, true);
+					filter &= Builders<IssueDocument>.Filter.Eq(x => x.Promoted, true);
 				}
 				else
 				{
-					filter &= Builders<Issue>.Filter.Ne(x => x.Promoted, true); // Handle the field not existing as well as being set to false.
+					filter &= Builders<IssueDocument>.Filter.Ne(x => x.Promoted, true); // Handle the field not existing as well as being set to false.
 				}
 			}
 			return await _issues.Find(filter).SortByDescending(x => x.Id).Range(index, count).ToListAsync(cancellationToken);
@@ -1112,7 +1113,7 @@ namespace HordeServer.Issues
 		/// <param name="count"></param>
 		/// <param name="cancellationToken"></param>
 		/// <returns></returns>
-		static async Task<List<Issue>> PaginatedJoinAsync(IAsyncCursor<ProjectedIssueId> cursor, Func<IEnumerable<int>, int, int?, Task<List<Issue>>> nextStageFunc, int index, int? count, CancellationToken cancellationToken)
+		static async Task<List<IssueDocument>> PaginatedJoinAsync(IAsyncCursor<ProjectedIssueId> cursor, Func<IEnumerable<int>, int, int?, Task<List<IssueDocument>>> nextStageFunc, int index, int? count, CancellationToken cancellationToken)
 		{
 			if (count == null)
 			{
@@ -1121,10 +1122,10 @@ namespace HordeServer.Issues
 			}
 			else
 			{
-				List<Issue> results = new List<Issue>();
+				List<IssueDocument> results = new List<IssueDocument>();
 				while (await cursor.MoveNextAsync(cancellationToken) && results.Count < count.Value)
 				{
-					List<Issue> nextResults = await nextStageFunc(cursor.Current.Where(x => x._id != null).Select(x => x._id!.Value), 0, count.Value - results.Count);
+					List<IssueDocument> nextResults = await nextStageFunc(cursor.Current.Where(x => x._id != null).Select(x => x._id!.Value), 0, count.Value - results.Count);
 					int removeCount = Math.Min(index, nextResults.Count);
 					nextResults.RemoveRange(0, removeCount);
 					index -= removeCount;
@@ -1137,14 +1138,14 @@ namespace HordeServer.Issues
 		/// <inheritdoc/>
 		public async Task<IReadOnlyList<IIssue>> FindIssuesForChangesAsync(List<int> changes, CancellationToken cancellationToken)
 		{
-			List<int> issueIds = await (await _issueSuspects.DistinctAsync(x => x.IssueId, Builders<IssueSuspect>.Filter.In(x => x.Change, changes), cancellationToken: cancellationToken)).ToListAsync(cancellationToken);
-			return await _issues.Find(Builders<Issue>.Filter.In(x => x.Id, issueIds)).ToListAsync(cancellationToken);
+			List<int> issueIds = await (await _issueSuspects.DistinctAsync(x => x.IssueId, Builders<IssueSuspectDocument>.Filter.In(x => x.Change, changes), cancellationToken: cancellationToken)).ToListAsync(cancellationToken);
+			return await _issues.Find(Builders<IssueDocument>.Filter.In(x => x.Id, issueIds)).ToListAsync(cancellationToken);
 		}
 
 		/// <inheritdoc/>
 		public async Task<IIssue?> TryUpdateIssueAsync(IIssue issue, UserId? initiatedByUserId, IssueSeverity? newSeverity = null, string? newSummary = null, string? newUserSummary = null, string? newDescription = null, bool? newManuallyPromoted = null, UserId? newOwnerId = null, UserId? newNominatedById = null, bool? newAcknowledged = null, UserId? newDeclinedById = null, int? newFixChange = null, UserId? newResolvedById = null, List<ObjectId>? newExcludeSpanIds = null, DateTime? newLastSeenAt = null, string? newExternaIssueKey = null, UserId? newQuarantinedById = null, UserId? newForceClosedById = null, Uri? newWorkflowThreadUrl = null, CancellationToken cancellationToken = default)
 		{
-			Issue issueDocument = (Issue)issue;
+			IssueDocument issueDocument = (IssueDocument)issue;
 
 			if (newDeclinedById != null && newDeclinedById == issueDocument.OwnerId)
 			{
@@ -1158,40 +1159,40 @@ namespace HordeServer.Issues
 
 			DateTime utcNow = DateTime.UtcNow;
 
-			List<UpdateDefinition<Issue>> updates = new List<UpdateDefinition<Issue>>();
+			List<UpdateDefinition<IssueDocument>> updates = new List<UpdateDefinition<IssueDocument>>();
 			if (newSeverity != null)
 			{
-				updates.Add(Builders<Issue>.Update.Set(x => x.Severity, newSeverity.Value));
+				updates.Add(Builders<IssueDocument>.Update.Set(x => x.Severity, newSeverity.Value));
 			}
 			if (newSummary != null)
 			{
-				updates.Add(Builders<Issue>.Update.Set(x => x.Summary, newSummary));
+				updates.Add(Builders<IssueDocument>.Update.Set(x => x.Summary, newSummary));
 			}
 			if (newUserSummary != null)
 			{
 				if (newUserSummary.Length == 0)
 				{
-					updates.Add(Builders<Issue>.Update.Unset(x => x.UserSummary!));
+					updates.Add(Builders<IssueDocument>.Update.Unset(x => x.UserSummary!));
 				}
 				else
 				{
-					updates.Add(Builders<Issue>.Update.Set(x => x.UserSummary, newUserSummary));
+					updates.Add(Builders<IssueDocument>.Update.Set(x => x.UserSummary, newUserSummary));
 				}
 			}
 			if (newDescription != null)
 			{
 				if (newDescription.Length == 0)
 				{
-					updates.Add(Builders<Issue>.Update.Unset(x => x.Description));
+					updates.Add(Builders<IssueDocument>.Update.Unset(x => x.Description));
 				}
 				else
 				{
-					updates.Add(Builders<Issue>.Update.Set(x => x.Description, newDescription));
+					updates.Add(Builders<IssueDocument>.Update.Set(x => x.Description, newDescription));
 				}
 			}
 			if (newManuallyPromoted != null)
 			{
-				updates.Add(Builders<Issue>.Update.Set(x => x.ManuallyPromoted, newManuallyPromoted.Value));
+				updates.Add(Builders<IssueDocument>.Update.Set(x => x.ManuallyPromoted, newManuallyPromoted.Value));
 			}
 			if (newResolvedById != null)
 			{
@@ -1202,22 +1203,22 @@ namespace HordeServer.Issues
 			{
 				if (newOwnerId.Value == UserId.Empty)
 				{
-					updates.Add(Builders<Issue>.Update.Unset(x => x.OwnerId!));
-					updates.Add(Builders<Issue>.Update.Unset(x => x.NominatedAt!));
-					updates.Add(Builders<Issue>.Update.Unset(x => x.NominatedById!));
+					updates.Add(Builders<IssueDocument>.Update.Unset(x => x.OwnerId!));
+					updates.Add(Builders<IssueDocument>.Update.Unset(x => x.NominatedAt!));
+					updates.Add(Builders<IssueDocument>.Update.Unset(x => x.NominatedById!));
 				}
 				else
 				{
-					updates.Add(Builders<Issue>.Update.Set(x => x.OwnerId!, newOwnerId.Value));
+					updates.Add(Builders<IssueDocument>.Update.Set(x => x.OwnerId!, newOwnerId.Value));
 
-					updates.Add(Builders<Issue>.Update.Set(x => x.NominatedAt, DateTime.UtcNow));
+					updates.Add(Builders<IssueDocument>.Update.Set(x => x.NominatedAt, DateTime.UtcNow));
 					if (newNominatedById == null)
 					{
-						updates.Add(Builders<Issue>.Update.Unset(x => x.NominatedById!));
+						updates.Add(Builders<IssueDocument>.Update.Unset(x => x.NominatedById!));
 					}
 					else
 					{
-						updates.Add(Builders<Issue>.Update.Set(x => x.NominatedById, newNominatedById.Value));
+						updates.Add(Builders<IssueDocument>.Update.Set(x => x.NominatedById, newNominatedById.Value));
 					}
 					newAcknowledged ??= false;
 				}
@@ -1228,14 +1229,14 @@ namespace HordeServer.Issues
 				{
 					if (issueDocument.AcknowledgedAt == null)
 					{
-						updates.Add(Builders<Issue>.Update.Set(x => x.AcknowledgedAt, utcNow));
+						updates.Add(Builders<IssueDocument>.Update.Set(x => x.AcknowledgedAt, utcNow));
 					}
 				}
 				else
 				{
 					if (issueDocument.AcknowledgedAt != null)
 					{
-						updates.Add(Builders<Issue>.Update.Unset(x => x.AcknowledgedAt!));
+						updates.Add(Builders<IssueDocument>.Update.Unset(x => x.AcknowledgedAt!));
 					}
 				}
 			}
@@ -1243,11 +1244,11 @@ namespace HordeServer.Issues
 			{
 				if (newFixChange == 0)
 				{
-					updates.Add(Builders<Issue>.Update.Unset(x => x.FixChange!));
+					updates.Add(Builders<IssueDocument>.Update.Unset(x => x.FixChange!));
 				}
 				else
 				{
-					updates.Add(Builders<Issue>.Update.Set(x => x.FixChange, newFixChange));
+					updates.Add(Builders<IssueDocument>.Update.Set(x => x.FixChange, newFixChange));
 				}
 			}
 			if (newResolvedById != null)
@@ -1256,19 +1257,19 @@ namespace HordeServer.Issues
 				{
 					if (issueDocument.ResolvedAt == null || issueDocument.ResolvedById != newResolvedById)
 					{
-						updates.Add(Builders<Issue>.Update.Set(x => x.ResolvedAt, utcNow));
-						updates.Add(Builders<Issue>.Update.Set(x => x.ResolvedById, newResolvedById.Value));
+						updates.Add(Builders<IssueDocument>.Update.Set(x => x.ResolvedAt, utcNow));
+						updates.Add(Builders<IssueDocument>.Update.Set(x => x.ResolvedById, newResolvedById.Value));
 					}
 				}
 				else
 				{
 					if (issueDocument.ResolvedAt != null)
 					{
-						updates.Add(Builders<Issue>.Update.Unset(x => x.ResolvedAt!));
+						updates.Add(Builders<IssueDocument>.Update.Unset(x => x.ResolvedAt!));
 					}
 					if (issueDocument.ResolvedById != null)
 					{
-						updates.Add(Builders<Issue>.Update.Unset(x => x.ResolvedById!));
+						updates.Add(Builders<IssueDocument>.Update.Unset(x => x.ResolvedById!));
 					}
 				}
 			}
@@ -1279,58 +1280,58 @@ namespace HordeServer.Issues
 				{
 					newCombinedExcludeSpanIds = newCombinedExcludeSpanIds.Union(issue.ExcludeSpans).ToList();
 				}
-				updates.Add(Builders<Issue>.Update.Set(x => x.ExcludeSpans, newCombinedExcludeSpanIds));
+				updates.Add(Builders<IssueDocument>.Update.Set(x => x.ExcludeSpans, newCombinedExcludeSpanIds));
 			}
 			if (newLastSeenAt != null)
 			{
-				updates.Add(Builders<Issue>.Update.Set(x => x.LastSeenAt, newLastSeenAt.Value));
+				updates.Add(Builders<IssueDocument>.Update.Set(x => x.LastSeenAt, newLastSeenAt.Value));
 			}
 
 			if (newDeclinedById != null)
 			{
 				GetLogger(issue.Id).LogInformation("Declined by {UserId}", newDeclinedById.Value);
-				await _issueSuspects.UpdateManyAsync(x => x.IssueId == issue.Id && x.AuthorId == newDeclinedById.Value, Builders<IssueSuspect>.Update.Set(x => x.DeclinedAt, DateTime.UtcNow), null, cancellationToken);
+				await _issueSuspects.UpdateManyAsync(x => x.IssueId == issue.Id && x.AuthorId == newDeclinedById.Value, Builders<IssueSuspectDocument>.Update.Set(x => x.DeclinedAt, DateTime.UtcNow), null, cancellationToken);
 			}
 			if (newQuarantinedById != null)
 			{
 				if (newQuarantinedById.Value == UserId.Empty)
 				{
-					updates.Add(Builders<Issue>.Update.Unset(x => x.QuarantinedByUserId));
-					updates.Add(Builders<Issue>.Update.Unset(x => x.QuarantineTimeUtc));
+					updates.Add(Builders<IssueDocument>.Update.Unset(x => x.QuarantinedByUserId));
+					updates.Add(Builders<IssueDocument>.Update.Unset(x => x.QuarantineTimeUtc));
 				}
 				else
 				{
-					updates.Add(Builders<Issue>.Update.Set(x => x.QuarantinedByUserId!, newQuarantinedById.Value));
-					updates.Add(Builders<Issue>.Update.Set(x => x.QuarantineTimeUtc, DateTime.UtcNow));
+					updates.Add(Builders<IssueDocument>.Update.Set(x => x.QuarantinedByUserId!, newQuarantinedById.Value));
+					updates.Add(Builders<IssueDocument>.Update.Set(x => x.QuarantineTimeUtc, DateTime.UtcNow));
 				}
 			}
 			else if ((newResolvedById != null && newResolvedById.Value != UserId.Empty) || (newForceClosedById != null && newForceClosedById.Value != UserId.Empty))
 			{
 				// Clear quarantine if being resolved or if being force closed
-				updates.Add(Builders<Issue>.Update.Unset(x => x.QuarantinedByUserId));
-				updates.Add(Builders<Issue>.Update.Unset(x => x.QuarantineTimeUtc));
+				updates.Add(Builders<IssueDocument>.Update.Unset(x => x.QuarantinedByUserId));
+				updates.Add(Builders<IssueDocument>.Update.Unset(x => x.QuarantineTimeUtc));
 			}
 
 			if (newForceClosedById != null)
 			{
 				if (newForceClosedById.Value == UserId.Empty)
 				{
-					updates.Add(Builders<Issue>.Update.Unset(x => x.ForceClosedByUserId));
+					updates.Add(Builders<IssueDocument>.Update.Unset(x => x.ForceClosedByUserId));
 				}
 				else
 				{
-					updates.Add(Builders<Issue>.Update.Set(x => x.ForceClosedByUserId, newForceClosedById.Value));
+					updates.Add(Builders<IssueDocument>.Update.Set(x => x.ForceClosedByUserId, newForceClosedById.Value));
 				}
 			}
 
 			if (newExternaIssueKey != null)
 			{
-				updates.Add(Builders<Issue>.Update.Set(x => x.ExternalIssueKey, newExternaIssueKey.Length == 0 ? null : newExternaIssueKey));
+				updates.Add(Builders<IssueDocument>.Update.Set(x => x.ExternalIssueKey, newExternaIssueKey.Length == 0 ? null : newExternaIssueKey));
 			}
 
 			if (newWorkflowThreadUrl != null)
 			{
-				updates.Add(Builders<Issue>.Update.Set(x => x.WorkflowThreadUrl, newWorkflowThreadUrl.ToString().Length == 0 ? null : newWorkflowThreadUrl));
+				updates.Add(Builders<IssueDocument>.Update.Set(x => x.WorkflowThreadUrl, newWorkflowThreadUrl.ToString().Length == 0 ? null : newWorkflowThreadUrl));
 			}
 
 			if (updates.Count == 0)
@@ -1338,7 +1339,7 @@ namespace HordeServer.Issues
 				return issueDocument;
 			}
 
-			Issue? newIssue = await TryUpdateIssueAsync(issue, Builders<Issue>.Update.Combine(updates), cancellationToken);
+			IssueDocument? newIssue = await TryUpdateIssueAsync(issue, Builders<IssueDocument>.Update.Combine(updates), cancellationToken);
 			if (newIssue == null)
 			{
 				return null;
@@ -1351,21 +1352,21 @@ namespace HordeServer.Issues
 		/// <inheritdoc/>
 		public async Task<IIssue?> TryUpdateIssueDerivedDataAsync(IIssue issue, string newSummary, IssueSeverity newSeverity, List<NewIssueFingerprint> newFingerprints, List<NewIssueStream> newStreams, List<NewIssueSuspectData> newSuspects, DateTime? newResolvedAt, DateTime? newVerifiedAt, DateTime newLastSeenAt, CancellationToken cancellationToken)
 		{
-			Issue issueImpl = (Issue)issue;
+			IssueDocument issueImpl = (IssueDocument)issue;
 
 			// Update all the suspects for this issue
-			IReadOnlyList<IssueSuspect> oldSuspectImpls = await _issueSuspects.Find(x => x.IssueId == issue.Id).ToListAsync(cancellationToken);
-			List<IssueSuspect> newSuspectImpls = await UpdateIssueSuspectsAsync(issue.Id, oldSuspectImpls, newSuspects, newResolvedAt, cancellationToken);
+			IReadOnlyList<IssueSuspectDocument> oldSuspectImpls = await _issueSuspects.Find(x => x.IssueId == issue.Id).ToListAsync(cancellationToken);
+			List<IssueSuspectDocument> newSuspectImpls = await UpdateIssueSuspectsAsync(issue.Id, oldSuspectImpls, newSuspects, newResolvedAt, cancellationToken);
 
 			// Find the spans for this issue
-			List<IssueSpan> newSpans = await _issueSpans.Find(x => x.IssueId == issue.Id).ToListAsync(cancellationToken);
+			List<IssueSpanDocument> newSpans = await _issueSpans.Find(x => x.IssueId == issue.Id).ToListAsync(cancellationToken);
 
 			// Update the resolved time on any issues
 			List<ObjectId> updateSpanIds = newSpans.Where(x => x.ResolvedAt != newResolvedAt).Select(x => x.Id).ToList();
 			if (updateSpanIds.Count > 0)
 			{
-				FilterDefinition<IssueSpan> filter = Builders<IssueSpan>.Filter.In(x => x.Id, updateSpanIds);
-				await _issueSpans.UpdateManyAsync(filter, Builders<IssueSpan>.Update.Set(x => x.ResolvedAt, newResolvedAt), null, cancellationToken);
+				FilterDefinition<IssueSpanDocument> filter = Builders<IssueSpanDocument>.Filter.In(x => x.Id, updateSpanIds);
+				await _issueSpans.UpdateManyAsync(filter, Builders<IssueSpanDocument>.Update.Set(x => x.ResolvedAt, newResolvedAt), null, cancellationToken);
 			}
 
 			// Figure out if this issue should be promoted
@@ -1413,57 +1414,57 @@ namespace HordeServer.Issues
 			int newMaxSuspectChange = (newSuspects.Count > 0) ? newSuspects.Min(x => x.Change) : 0;
 
 			// Perform the actual update with this data
-			List<UpdateDefinition<Issue>> updates = new List<UpdateDefinition<Issue>>();
+			List<UpdateDefinition<IssueDocument>> updates = new List<UpdateDefinition<IssueDocument>>();
 			if (!String.Equals(issue.Summary, newSummary, StringComparison.Ordinal))
 			{
-				updates.Add(Builders<Issue>.Update.Set(x => x.Summary, newSummary));
+				updates.Add(Builders<IssueDocument>.Update.Set(x => x.Summary, newSummary));
 			}
 			if (issue.Severity != newSeverity)
 			{
-				updates.Add(Builders<Issue>.Update.Set(x => x.Severity, newSeverity));
+				updates.Add(Builders<IssueDocument>.Update.Set(x => x.Severity, newSeverity));
 			}
 			if (issue.Promoted != newPromoted)
 			{
-				updates.Add(Builders<Issue>.Update.Set(x => x.Promoted, newPromoted));
+				updates.Add(Builders<IssueDocument>.Update.Set(x => x.Promoted, newPromoted));
 			}
 			if (issue.Fingerprints.Count != newFingerprints.Count || !newFingerprints.Zip(issue.Fingerprints).All(x => x.First.Equals(x.Second)))
 			{
-				updates.Add(Builders<Issue>.Update.Set(x => x.Fingerprints, newFingerprints.Select(x => new IssueFingerprint(x))));
+				updates.Add(Builders<IssueDocument>.Update.Set(x => x.Fingerprints, newFingerprints.Select(x => new IssueFingerprintDocument(x))));
 			}
 			if (issue.Streams.Count != newStreams.Count || !newStreams.Zip(issue.Streams).All(x => x.First.StreamId == x.Second.StreamId && x.First.ContainsFix == x.Second.ContainsFix))
 			{
-				updates.Add(Builders<Issue>.Update.Set(x => x.Streams, newStreams.Select(x => new IssueStream(x))));
+				updates.Add(Builders<IssueDocument>.Update.Set(x => x.Streams, newStreams.Select(x => new IssueStreamDocument(x))));
 			}
 			if (issueImpl.MinSuspectChange != newMinSuspectChange)
 			{
-				updates.Add(Builders<Issue>.Update.Set(x => x.MinSuspectChange, newMinSuspectChange));
+				updates.Add(Builders<IssueDocument>.Update.Set(x => x.MinSuspectChange, newMinSuspectChange));
 			}
 			if (issueImpl.MaxSuspectChange != newMaxSuspectChange)
 			{
-				updates.Add(Builders<Issue>.Update.Set(x => x.MaxSuspectChange, newMaxSuspectChange));
+				updates.Add(Builders<IssueDocument>.Update.Set(x => x.MaxSuspectChange, newMaxSuspectChange));
 			}
 			if (issueImpl.DefaultOwnerId != newDefaultOwnerId)
 			{
-				updates.Add(Builders<Issue>.Update.Set(x => x.DefaultOwnerId, newDefaultOwnerId));
+				updates.Add(Builders<IssueDocument>.Update.Set(x => x.DefaultOwnerId, newDefaultOwnerId));
 			}
 			if (issue.ResolvedAt != newResolvedAt)
 			{
-				updates.Add(Builders<Issue>.Update.SetOrUnsetNull(x => x.ResolvedAt, newResolvedAt));
+				updates.Add(Builders<IssueDocument>.Update.SetOrUnsetNull(x => x.ResolvedAt, newResolvedAt));
 			}
 			if (newResolvedAt == null && issue.ResolvedById != null)
 			{
-				updates.Add(Builders<Issue>.Update.Unset(x => x.ResolvedById));
+				updates.Add(Builders<IssueDocument>.Update.Unset(x => x.ResolvedById));
 			}
 			if (issue.VerifiedAt != newVerifiedAt)
 			{
-				updates.Add(Builders<Issue>.Update.SetOrUnsetNull(x => x.VerifiedAt, newVerifiedAt));
+				updates.Add(Builders<IssueDocument>.Update.SetOrUnsetNull(x => x.VerifiedAt, newVerifiedAt));
 			}
 			if (issue.LastSeenAt != newLastSeenAt)
 			{
-				updates.Add(Builders<Issue>.Update.Set(x => x.LastSeenAt, newLastSeenAt));
+				updates.Add(Builders<IssueDocument>.Update.Set(x => x.LastSeenAt, newLastSeenAt));
 			}
 
-			Issue? newIssue = await TryUpdateIssueAsync(issue, Builders<Issue>.Update.Combine(updates), cancellationToken);
+			IssueDocument? newIssue = await TryUpdateIssueAsync(issue, Builders<IssueDocument>.Update.Combine(updates), cancellationToken);
 			if (newIssue != null)
 			{
 				await LogIssueChangesAsync(null, issueImpl, newIssue, cancellationToken);
@@ -1473,16 +1474,16 @@ namespace HordeServer.Issues
 			return null;
 		}
 
-		async Task<List<IssueSuspect>> UpdateIssueSuspectsAsync(int issueId, IReadOnlyList<IssueSuspect> oldSuspectImpls, List<NewIssueSuspectData> newSuspects, DateTime? resolvedAt, CancellationToken cancellationToken)
+		async Task<List<IssueSuspectDocument>> UpdateIssueSuspectsAsync(int issueId, IReadOnlyList<IssueSuspectDocument> oldSuspectImpls, List<NewIssueSuspectData> newSuspects, DateTime? resolvedAt, CancellationToken cancellationToken)
 		{
-			List<IssueSuspect> newSuspectImpls = new List<IssueSuspect>(oldSuspectImpls);
+			List<IssueSuspectDocument> newSuspectImpls = new List<IssueSuspectDocument>(oldSuspectImpls);
 
 			// Find the current list of suspects
 			HashSet<(UserId, int)> curSuspectKeys = new HashSet<(UserId, int)>(oldSuspectImpls.Select(x => (x.AuthorId, x.Change)));
-			List<IssueSuspect> createSuspects = newSuspects.Where(x => !curSuspectKeys.Contains((x.AuthorId, x.Change))).Select(x => new IssueSuspect(issueId, x, resolvedAt)).ToList();
+			List<IssueSuspectDocument> createSuspects = newSuspects.Where(x => !curSuspectKeys.Contains((x.AuthorId, x.Change))).Select(x => new IssueSuspectDocument(issueId, x, resolvedAt)).ToList();
 
 			HashSet<(UserId, int)> newSuspectKeys = new HashSet<(UserId, int)>(newSuspects.Select(x => (x.AuthorId, x.Change)));
-			List<IssueSuspect> deleteSuspects = oldSuspectImpls.Where(x => !newSuspectKeys.Contains((x.AuthorId, x.Change))).ToList();
+			List<IssueSuspectDocument> deleteSuspects = oldSuspectImpls.Where(x => !newSuspectKeys.Contains((x.AuthorId, x.Change))).ToList();
 
 			// Apply the suspect changes
 			if (createSuspects.Count > 0)
@@ -1492,14 +1493,14 @@ namespace HordeServer.Issues
 			}
 			if (deleteSuspects.Count > 0)
 			{
-				await _issueSuspects.DeleteManyAsync(Builders<IssueSuspect>.Filter.In(x => x.Id, deleteSuspects.Select(y => y.Id)), cancellationToken);
+				await _issueSuspects.DeleteManyAsync(Builders<IssueSuspectDocument>.Filter.In(x => x.Id, deleteSuspects.Select(y => y.Id)), cancellationToken);
 				newSuspectImpls.RemoveAll(x => !newSuspectKeys.Contains((x.AuthorId, x.Change)));
 			}
 
 			// Make sure all the remaining suspects have the correct resolved time
 			if (newSuspectImpls.Any(x => x.ResolvedAt != resolvedAt))
 			{
-				await _issueSuspects.UpdateManyAsync(Builders<IssueSuspect>.Filter.Eq(x => x.IssueId, issueId), Builders<IssueSuspect>.Update.Set(x => x.ResolvedAt, resolvedAt), null, cancellationToken);
+				await _issueSuspects.UpdateManyAsync(Builders<IssueSuspectDocument>.Filter.Eq(x => x.IssueId, issueId), Builders<IssueSuspectDocument>.Update.Set(x => x.ResolvedAt, resolvedAt), null, cancellationToken);
 			}
 			return newSuspectImpls;
 		}
@@ -1511,7 +1512,7 @@ namespace HordeServer.Issues
 		/// <inheritdoc/>
 		public async Task<IIssueSpan> AddSpanAsync(int issueId, NewIssueSpanData newSpan, CancellationToken cancellationToken)
 		{
-			IssueSpan span = new IssueSpan(issueId, newSpan);
+			IssueSpanDocument span = new IssueSpanDocument(issueId, newSpan);
 			await _issueSpans.InsertOneAsync(span, (InsertOneOptions?)null, cancellationToken);
 			SendTelemetry(span);
 			return span;
@@ -1539,45 +1540,45 @@ namespace HordeServer.Issues
 		/// <inheritdoc/>
 		public async Task<IIssueSpan?> GetSpanAsync(ObjectId spanId, CancellationToken cancellationToken)
 		{
-			return await _issueSpans.Find(Builders<IssueSpan>.Filter.Eq(x => x.Id, spanId)).FirstOrDefaultAsync(cancellationToken);
+			return await _issueSpans.Find(Builders<IssueSpanDocument>.Filter.Eq(x => x.Id, spanId)).FirstOrDefaultAsync(cancellationToken);
 		}
 
 		/// <inheritdoc/>
 		public async Task<IIssueSpan?> TryUpdateSpanAsync(IIssueSpan span, NewIssueStepData? newLastSuccess = null, NewIssueStepData? newFailure = null, NewIssueStepData? newNextSuccess = null, List<NewIssueSpanSuspectData>? newSuspects = null, int? newIssueId = null, CancellationToken cancellationToken = default)
 		{
-			List<UpdateDefinition<IssueSpan>> updates = new List<UpdateDefinition<IssueSpan>>();
+			List<UpdateDefinition<IssueSpanDocument>> updates = new List<UpdateDefinition<IssueSpanDocument>>();
 			if (newLastSuccess != null)
 			{
-				updates.Add(Builders<IssueSpan>.Update.Set(x => x.MinChange, newLastSuccess.Change));
-				updates.Add(Builders<IssueSpan>.Update.Set(x => x.LastSuccess, new IssueStep(span.Id, newLastSuccess)));
+				updates.Add(Builders<IssueSpanDocument>.Update.Set(x => x.MinChange, newLastSuccess.Change));
+				updates.Add(Builders<IssueSpanDocument>.Update.Set(x => x.LastSuccess, new IssueStepDocument(span.Id, newLastSuccess)));
 			}
 			if (newFailure != null)
 			{
 				if (newFailure.Change < span.FirstFailure.Change)
 				{
-					updates.Add(Builders<IssueSpan>.Update.Set(x => x.FirstFailure, new IssueStep(span.Id, newFailure)));
+					updates.Add(Builders<IssueSpanDocument>.Update.Set(x => x.FirstFailure, new IssueStepDocument(span.Id, newFailure)));
 				}
 				if (newFailure.Change >= span.LastFailure.Change)
 				{
-					updates.Add(Builders<IssueSpan>.Update.Set(x => x.LastFailure, new IssueStep(span.Id, newFailure)));
+					updates.Add(Builders<IssueSpanDocument>.Update.Set(x => x.LastFailure, new IssueStepDocument(span.Id, newFailure)));
 				}
 				if (newFailure.PromoteByDefault != span.PromoteByDefault && newFailure.Change >= span.LastFailure.Change)
 				{
-					updates.Add(Builders<IssueSpan>.Update.Set(x => x.PromoteByDefault, newFailure.PromoteByDefault));
+					updates.Add(Builders<IssueSpanDocument>.Update.Set(x => x.PromoteByDefault, newFailure.PromoteByDefault));
 				}
 			}
 			if (newNextSuccess != null)
 			{
-				updates.Add(Builders<IssueSpan>.Update.Set(x => x.MaxChange, newNextSuccess.Change));
-				updates.Add(Builders<IssueSpan>.Update.Set(x => x.NextSuccess, new IssueStep(span.Id, newNextSuccess)));
+				updates.Add(Builders<IssueSpanDocument>.Update.Set(x => x.MaxChange, newNextSuccess.Change));
+				updates.Add(Builders<IssueSpanDocument>.Update.Set(x => x.NextSuccess, new IssueStepDocument(span.Id, newNextSuccess)));
 			}
 			if (newSuspects != null)
 			{
-				updates.Add(Builders<IssueSpan>.Update.Set(x => x.Suspects, newSuspects.ConvertAll(x => new IssueSpanSuspect(x))));
+				updates.Add(Builders<IssueSpanDocument>.Update.Set(x => x.Suspects, newSuspects.ConvertAll(x => new IssueSpanSuspectDocument(x))));
 			}
 			if (newIssueId != null)
 			{
-				updates.Add(Builders<IssueSpan>.Update.Set(x => x.IssueId, newIssueId.Value));
+				updates.Add(Builders<IssueSpanDocument>.Update.Set(x => x.IssueId, newIssueId.Value));
 			}
 
 			if (updates.Count == 0)
@@ -1585,7 +1586,7 @@ namespace HordeServer.Issues
 				return span;
 			}
 
-			IssueSpan? newSpan = await TryUpdateSpanAsync(span, Builders<IssueSpan>.Update.Combine(updates), cancellationToken);
+			IssueSpanDocument? newSpan = await TryUpdateSpanAsync(span, Builders<IssueSpanDocument>.Update.Combine(updates), cancellationToken);
 			if (newSpan != null)
 			{
 				ILogger logger = GetLogger(newSpan.IssueId);
@@ -1615,7 +1616,7 @@ namespace HordeServer.Issues
 		/// <inheritdoc/>
 		public async Task<IReadOnlyList<IIssueSpan>> FindSpansAsync(IEnumerable<ObjectId> spanIds, CancellationToken cancellationToken)
 		{
-			return await _issueSpans.Find(Builders<IssueSpan>.Filter.In(x => x.Id, spanIds)).ToListAsync(cancellationToken);
+			return await _issueSpans.Find(Builders<IssueSpanDocument>.Filter.In(x => x.Id, spanIds)).ToListAsync(cancellationToken);
 		}
 
 		/// <inheritdoc/>
@@ -1627,39 +1628,39 @@ namespace HordeServer.Issues
 		/// <inheritdoc/>
 		public async Task<IReadOnlyList<IIssueSpan>> FindSpansAsync(IEnumerable<ObjectId>? spanIds, IEnumerable<int>? issueIds, StreamId? streamId, int? minChange, int? maxChange, bool? resolved, int? index, int? count, CancellationToken cancellationToken)
 		{
-			FilterDefinition<IssueSpan> filter = FilterDefinition<IssueSpan>.Empty;
+			FilterDefinition<IssueSpanDocument> filter = FilterDefinition<IssueSpanDocument>.Empty;
 
 			if (spanIds != null)
 			{
-				filter &= Builders<IssueSpan>.Filter.In(x => x.Id, spanIds);
+				filter &= Builders<IssueSpanDocument>.Filter.In(x => x.Id, spanIds);
 			}
 
 			if (streamId != null)
 			{
-				filter &= Builders<IssueSpan>.Filter.Eq(x => x.StreamId, streamId);
+				filter &= Builders<IssueSpanDocument>.Filter.Eq(x => x.StreamId, streamId);
 			}
 
 			if (issueIds != null)
 			{
-				filter &= Builders<IssueSpan>.Filter.In(x => x.IssueId, issueIds.Select<int, int?>(x => x));
+				filter &= Builders<IssueSpanDocument>.Filter.In(x => x.IssueId, issueIds.Select<int, int?>(x => x));
 			}
 			if (minChange != null)
 			{
-				filter &= Builders<IssueSpan>.Filter.Not(Builders<IssueSpan>.Filter.Lt(x => x.MaxChange, minChange.Value));
+				filter &= Builders<IssueSpanDocument>.Filter.Not(Builders<IssueSpanDocument>.Filter.Lt(x => x.MaxChange, minChange.Value));
 			}
 			if (maxChange != null)
 			{
-				filter &= Builders<IssueSpan>.Filter.Not(Builders<IssueSpan>.Filter.Gt(x => x.MinChange, maxChange.Value));
+				filter &= Builders<IssueSpanDocument>.Filter.Not(Builders<IssueSpanDocument>.Filter.Gt(x => x.MinChange, maxChange.Value));
 			}
 			if (resolved != null)
 			{
 				if (resolved.Value)
 				{
-					filter &= Builders<IssueSpan>.Filter.Ne(x => x.ResolvedAt, null);
+					filter &= Builders<IssueSpanDocument>.Filter.Ne(x => x.ResolvedAt, null);
 				}
 				else
 				{
-					filter &= Builders<IssueSpan>.Filter.Eq(x => x.ResolvedAt, null);
+					filter &= Builders<IssueSpanDocument>.Filter.Eq(x => x.ResolvedAt, null);
 				}
 			}
 
@@ -1673,7 +1674,7 @@ namespace HordeServer.Issues
 		/// <inheritdoc/>
 		public async Task<IIssueStep> AddStepAsync(ObjectId spanId, NewIssueStepData newStep, CancellationToken cancellationToken)
 		{
-			IssueStep step = new IssueStep(spanId, newStep);
+			IssueStepDocument step = new IssueStepDocument(spanId, newStep);
 			await _issueSteps.InsertOneAsync(step, (InsertOneOptions?)null, cancellationToken);
 			return step;
 		}
@@ -1681,21 +1682,21 @@ namespace HordeServer.Issues
 		/// <inheritdoc/>
 		public async Task<IReadOnlyList<IIssueStep>> FindStepsAsync(IEnumerable<ObjectId> spanIds, CancellationToken cancellationToken)
 		{
-			FilterDefinition<IssueStep> filter = Builders<IssueStep>.Filter.In(x => x.SpanId, spanIds);
+			FilterDefinition<IssueStepDocument> filter = Builders<IssueStepDocument>.Filter.In(x => x.SpanId, spanIds);
 			return await _issueSteps.Find(filter).ToListAsync(cancellationToken);
 		}
 
 		/// <inheritdoc/>
 		public async Task<IReadOnlyList<IIssueStep>> FindStepsAsync(JobId jobId, JobStepBatchId? batchId, JobStepId? stepId, CancellationToken cancellationToken)
 		{
-			FilterDefinition<IssueStep> filter = Builders<IssueStep>.Filter.Eq(x => x.JobId, jobId);
+			FilterDefinition<IssueStepDocument> filter = Builders<IssueStepDocument>.Filter.Eq(x => x.JobId, jobId);
 			if (batchId != null)
 			{
-				filter &= Builders<IssueStep>.Filter.Eq(x => x.BatchId, batchId.Value);
+				filter &= Builders<IssueStepDocument>.Filter.Eq(x => x.BatchId, batchId.Value);
 			}
 			if (stepId != null)
 			{
-				filter &= Builders<IssueStep>.Filter.Eq(x => x.StepId, stepId.Value);
+				filter &= Builders<IssueStepDocument>.Filter.Eq(x => x.StepId, stepId.Value);
 			}
 			return await _issueSteps.Find(filter).ToListAsync(cancellationToken);
 		}
diff --git a/Engine/Source/Programs/Shared/EpicGames.Horde/Issues/Handlers/ExternalIssueHandler.cs b/Engine/Source/Programs/Shared/EpicGames.Horde/Issues/Handlers/ExternalIssueHandler.cs
new file mode 100644
--- /dev/null
+++ b/Engine/Source/Programs/Shared/EpicGames.Horde/Issues/Handlers/ExternalIssueHandler.cs
@@ -0,0 +1,78 @@
+// Copyright Epic Games, Inc. All Rights Reserved.
+
+using System.Collections.Generic;
+using System.Text.Json;
+using System.Text.Json.Nodes;
+using EpicGames.Core;
+
+namespace EpicGames.Horde.Issues.Handlers
+{
+	using JsonObject = System.Text.Json.Nodes.JsonObject;
+
+	/// <summary>
+	/// Handler for log events with issue fingerprints embedded in the structured log data itself
+	/// </summary>
+	[IssueHandler]
+	public class ExternalIssueHandler : IssueHandler
+	{
+		static readonly JsonNodeOptions s_jsonNodeOptions = new JsonNodeOptions { PropertyNameCaseInsensitive = true };
+		static readonly JsonSerializerOptions? s_jsonSerializerOptions = new JsonSerializerOptions { PropertyNameCaseInsensitive = true };
+
+		readonly List<IssueEventGroup> _issues = new List<IssueEventGroup>();
+
+		/// <inheritdoc/>
+		public override int Priority => 20;
+
+		/// <inheritdoc/>
+		public override bool HandleEvent(IssueEvent issueEvent)
+		{
+			foreach (JsonLogEvent line in issueEvent.Lines)
+			{
+				IssueFingerprint? issueData = ParseIssueFingerprint(line);
+				if (issueData != null)
+				{
+					IssueEventGroup group = new IssueEventGroup(issueData.Type, issueData.SummaryTemplate, issueData.ChangeFilter);
+					if (issueData.Keys != null)
+					{
+						group.Keys.UnionWith(issueData.Keys);
+					}
+					if (issueData.Metadata != null)
+					{
+						group.Metadata.UnionWith(issueData.Metadata);
+					}
+					group.Events.Add(issueEvent);
+
+					_issues.Add(group);
+					return true;
+				}
+			}
+			return false;
+		}
+
+		/// <inheritdoc/>
+		public override IEnumerable<IssueEventGroup> GetIssues()
+			=> _issues;
+
+		static IssueFingerprint? ParseIssueFingerprint(JsonLogEvent line)
+		{
+			JsonNode? node = JsonObject.Parse(line.Data.Span, s_jsonNodeOptions);
+			if (node != null)
+			{
+				JsonObject? properties = node["properties"] as JsonObject;
+				if (properties != null)
+				{
+					JsonObject? issue = properties["$issue"] as JsonObject;
+					if (issue != null)
+					{
+						IssueFingerprint? fingerprint = JsonSerializer.Deserialize<IssueFingerprint>(issue, s_jsonSerializerOptions);
+						if (fingerprint != null)
+						{
+							return fingerprint;
+						}
+					}
+				}
+			}
+			return null;
+		}
+	}
+}
diff --git a/Engine/Source/Programs/Shared/EpicGames.Horde/Issues/IIssueFingerprint.cs b/Engine/Source/Programs/Shared/EpicGames.Horde/Issues/IIssueFingerprint.cs
new file mode 100644
--- /dev/null
+++ b/Engine/Source/Programs/Shared/EpicGames.Horde/Issues/IIssueFingerprint.cs
@@ -0,0 +1,106 @@
+// Copyright Epic Games, Inc. All Rights Reserved.
+
+using System;
+using System.Collections.Generic;
+using System.Linq;
+
+namespace EpicGames.Horde.Issues
+{
+	/// <summary>
+	/// Fingerprint for an issue
+	/// </summary>
+	public interface IIssueFingerprint
+	{
+		/// <summary>
+		/// The type of issue, which defines the handler to use for it
+		/// </summary>
+		public string Type { get; }
+
+		/// <summary>
+		/// Template string for the issue summary
+		/// </summary>
+		public string SummaryTemplate { get; }
+
+		/// <summary>
+		/// List of keys which identify this issue.
+		/// </summary>
+		public IReadOnlySet<IssueKey> Keys { get; }
+
+		/// <summary>
+		/// Set of keys which should trigger a negative match
+		/// </summary>
+		public IReadOnlySet<IssueKey>? RejectKeys { get; }
+
+		/// <summary>
+		/// Collection of additional metadata added by the handler
+		/// </summary>
+		public IReadOnlySet<IssueMetadata>? Metadata { get; }
+
+		/// <summary>
+		/// Filter for changes that should be included in this issue
+		/// </summary>
+		public string ChangeFilter { get; }
+	}
+
+	/// <summary>
+	/// Extension methods for <see cref="IIssueFingerprint"/>
+	/// </summary>
+	public static class IssueFingerprintExtensions
+	{
+		/// <summary>
+		/// Checks if a fingerprint matches another fingerprint
+		/// </summary>
+		/// <param name="fingerprint">The first fingerprint to compare</param>
+		/// <param name="other">The other fingerprint to compare to</param>
+		/// <returns>True is the fingerprints match</returns>
+		public static bool IsMatch(this IIssueFingerprint fingerprint, IIssueFingerprint other)
+		{
+			if (!fingerprint.Type.Equals(other.Type, StringComparison.Ordinal))
+			{
+				return false;
+			}
+			if (!fingerprint.SummaryTemplate.Equals(other.SummaryTemplate, StringComparison.Ordinal))
+			{
+				return false;
+			}
+			if (fingerprint.Keys.Count == 0)
+			{
+				if (other.Keys.Count > 0)
+				{
+					return false;
+				}
+			}
+			else
+			{
+				if (!fingerprint.Keys.Any(x => other.Keys.Contains(x)))
+				{
+					return false;
+				}
+			}
+			if (fingerprint.RejectKeys != null && fingerprint.RejectKeys.Any(x => other.Keys.Contains(x)))
+			{
+				return false;
+			}
+			return true;
+		}
+
+		/// <summary>
+		/// Checks if a fingerprint matches another fingerprint for creating a new span
+		/// </summary>
+		/// <param name="fingerprint">The first fingerprint to compare</param>
+		/// <param name="other">The other fingerprint to compare to</param>
+		/// <returns>True is the fingerprints match</returns>
+		public static bool IsMatchForNewSpan(this IIssueFingerprint fingerprint, IIssueFingerprint other)
+		{
+			if (!fingerprint.Type.Equals(other.Type, StringComparison.Ordinal))
+			{
+				return false;
+			}
+			if (fingerprint.RejectKeys != null && fingerprint.RejectKeys.Any(x => other.Keys.Contains(x)))
+			{
+				return false;
+			}
+			return true;
+		}
+	}
+}
diff --git a/Engine/Source/Programs/Shared/EpicGames.Horde/Issues/IssueExtensions.cs b/Engine/Source/Programs/Shared/EpicGames.Horde/Issues/IssueExtensions.cs
new file mode 100644
--- /dev/null
+++ b/Engine/Source/Programs/Shared/EpicGames.Horde/Issues/IssueExtensions.cs
@@ -0,0 +1,37 @@
+// Copyright Epic Games, Inc.All Rights Reserved.
+
+using System;
+using System.Collections.Generic;
+using EpicGames.Core;
+using Microsoft.Extensions.Logging;
+
+namespace EpicGames.Horde.Issues
+{
+	/// <summary>
+	/// Extension methods for logging issue data
+	/// </summary>
+	public static class IssueExtensions
+	{
+		/// <summary>
+		/// Adds an issue fingerprint to a log event
+		/// </summary>
+		/// <param name="logEvent">Log event to modify</param>
+		/// <param name="fingerprint">Fingerprint for the issue</param>
+		public static void AddIssueFingerprint(this LogEvent logEvent, IssueFingerprint fingerprint)
+			=> logEvent.AddProperty("@$issue", fingerprint);
+
+		/// <summary>
+		/// Enters a scope which annotates log messages with the supplied issue fingerprint
+		/// </summary>
+		/// <param name="logger">Logger device to operate on</param>
+		/// <param name="fingerprint">Fingerprint for the issue</param>
+		/// <returns>Disposable object for the lifetime of this scope</returns>
+		public static IDisposable? BeginIssueScope(this ILogger logger, IssueFingerprint fingerprint)
+		{
+			Dictionary<string, object> properties = new Dictionary<string, object>(1);
+			properties.Add("@$issue", fingerprint);
+
+			return logger.BeginScope(properties);
+		}
+	}
+}
diff --git a/Engine/Source/Programs/Shared/EpicGames.Horde/Issues/IssueFingerprint.cs b/Engine/Source/Programs/Shared/EpicGames.Horde/Issues/IssueFingerprint.cs
new file mode 100644
--- /dev/null
+++ b/Engine/Source/Programs/Shared/EpicGames.Horde/Issues/IssueFingerprint.cs
@@ -0,0 +1,93 @@
+// Copyright Epic Games, Inc.All Rights Reserved.
+
+using System;
+using System.Collections.Generic;
+using System.Text.Json.Serialization;
+
+#pragma warning disable CA2227 // Change 'x' to be read-only by removing the property setter
+
+namespace EpicGames.Horde.Issues
+{
+	/// <summary>
+	/// Fingerprint for an issue. Can be embedded into a structured log event under the "$issue" property to guide how Horde should group the event.
+	/// </summary>
+	public class IssueFingerprint : IIssueFingerprint
+	{
+		/// <summary>
+		/// The type of issue, which defines the handler to use for it
+		/// </summary>
+		[JsonPropertyName("type")]
+		public string Type { get; set; } = "None";
+
+		/// <summary>
+		/// Template string for the issue summary
+		/// </summary>
+		[JsonPropertyName("summary")]
+		public string SummaryTemplate { get; set; } = String.Empty;
+
+		/// <summary>
+		/// List of keys which identify this issue.
+		/// </summary>
+		[JsonIgnore]
+		public HashSet<IssueKey> Keys { get; set; } = new HashSet<IssueKey>();
+		IReadOnlySet<IssueKey> IIssueFingerprint.Keys => Keys;
+
+		/// <summary>
+		/// Set of keys which should trigger a negative match
+		/// </summary>
+		[JsonPropertyName("rejectKeys")]
+		public HashSet<IssueKey>? RejectKeys { get; set; }
+		IReadOnlySet<IssueKey>? IIssueFingerprint.RejectKeys => RejectKeys;
+
+		/// <summary>
+		/// Collection of additional metadata added by the handler
+		/// </summary>
+		[JsonIgnore]
+		public HashSet<IssueMetadata> Metadata { get; set; } = new HashSet<IssueMetadata>();
+		IReadOnlySet<IssueMetadata> IIssueFingerprint.Metadata => Metadata;
+
+		/// <summary>
+		/// Filter for changes that touch files which should be included in this issue
+		/// </summary>
+		[JsonPropertyName("filter")]
+		public string ChangeFilter { get; set; } = "...";
+
+		/// <summary>
+		/// Keys set which is null when empty
+		/// </summary>
+		[JsonPropertyName("keys")]
+		[JsonIgnore(Condition = JsonIgnoreCondition.WhenWritingNull)]
+		public HashSet<IssueKey>? SerializedKeys
+		{
+			get => (Keys.Count == 0) ? null : Keys;
+			set => Keys = value ?? new HashSet<IssueKey>();
+		}
+
+		/// <summary>
+		/// Metadata set which is null when empty
+		/// </summary>
+		[JsonPropertyName("meta")]
+		[JsonIgnore(Condition = JsonIgnoreCondition.WhenWritingNull)]
+		internal HashSet<IssueMetadata>? SerializedMetadata
+		{
+			get => (Metadata.Count == 0) ? null : Metadata;
+			set => Metadata = value ?? new HashSet<IssueMetadata>();
+		}
+
+		/// <summary>
+		/// Default constructor
+		/// </summary>
+		public IssueFingerprint()
+		{
+		}
+
+		/// <summary>
+		/// Constructor
+		/// </summary>
+		public IssueFingerprint(string type, string summaryTemplate)
+		{
+			Type = type;
+			SummaryTemplate = summaryTemplate;
+		}
+	}
+}
diff --git a/Engine/Source/Programs/Shared/EpicGames.Horde/Issues/IssueKey.cs b/Engine/Source/Programs/Shared/EpicGames.Horde/Issues/IssueKey.cs
--- a/Engine/Source/Programs/Shared/EpicGames.Horde/Issues/IssueKey.cs
+++ b/Engine/Source/Programs/Shared/EpicGames.Horde/Issues/IssueKey.cs
@@ -4,6 +4,7 @@ using System;
 using System.Collections.Generic;
 using System.Linq;
 using System.Text.Json;
+using System.Text.Json.Serialization;
 using EpicGames.Core;
 using EpicGames.Horde.Jobs.Templates;
 using EpicGames.Horde.Streams;
@@ -59,11 +60,13 @@ namespace EpicGames.Horde.Issues
 		/// <summary>
 		/// Type of the key
 		/// </summary>
+		[JsonConverter(typeof(JsonStringEnumConverter))]
 		public IssueKeyType Type { get; }
 
 		/// <summary>
 		/// Arbitrary string that can be used to discriminate between otherwise identical keys, limiting the issues that it can merge with.
 		/// </summary>
+		[JsonIgnore(Condition = JsonIgnoreCondition.WhenWritingNull)]
 		public string? Scope { get; }
 
 		/// <summary>
@@ -78,13 +81,24 @@ namespace EpicGames.Horde.Issues
 		/// <summary>
 		/// Constructor
 		/// </summary>
-		public IssueKey(string name, IssueKeyType type, string? scope)
+		[JsonConstructor]
+		public IssueKey(string name, IssueKeyType type, string? scope = null)
 		{
 			Name = name;
 			Type = type;
 			Scope = scope;
 		}
 
+		/// <summary>
+		/// Creates an issue key for a file
+		/// </summary>
+		public static IssueKey FromFile(string file, bool note = false) => new IssueKey(file, note? IssueKeyType.Note : IssueKeyType.File);
+
+		/// <summary>
+		/// Creates an issue key for a file
+		/// </summary>
+		public static IssueKey FromSymbol(string name) => new IssueKey(name, IssueKeyType.Symbol);
+
 		/// <summary>
 		/// Creates an issue key for a particular hash
 		/// </summary>
