Horde: Make update issues a read only on buld dialog.  As we now have workflows and explicit template settings, this setting adds confusion, though still useful for display/debugging purposes
#rnx
#skipci

diff --git a/Engine/Source/Programs/Horde/HordeDashboard/src/components/build/NewBuildV2.tsx b/Engine/Source/Programs/Horde/HordeDashboard/src/components/build/NewBuildV2.tsx
--- a/Engine/Source/Programs/Horde/HordeDashboard/src/components/build/NewBuildV2.tsx
+++ b/Engine/Source/Programs/Horde/HordeDashboard/src/components/build/NewBuildV2.tsx
@@ -73,7 +73,6 @@ class BuildOptions {
       this.autoSubmit = undefined;
       this.advJobName = undefined;
       this.advJobPriority = undefined;
-      this.advUpdateIssues = undefined;
       this.advAdditionalArgs = undefined;
       this.advTargets = undefined;
       this.template = undefined;
@@ -151,7 +150,6 @@ class BuildOptions {
    // advanced options
    advJobName?: string;
    advJobPriority?: Priority;
-   advUpdateIssues?: boolean;
    advAdditionalArgs?: string;
    advTargets?: string[];
 
@@ -204,7 +202,7 @@ class BuildOptions {
 
    get advancedModified(): boolean {
 
-      return !!((this.advJobPriority && this.advJobPriority !== Priority.Normal) || this.advUpdateIssues || !!this.advJobName || !!this.advAdditionalArgs || !!this.advTargets)
+      return !!((this.advJobPriority && this.advJobPriority !== Priority.Normal) || !!this.advJobName || !!this.advAdditionalArgs || !!this.advTargets)
    }
 
    change?: number;
@@ -1623,12 +1621,15 @@ const AdvancedPanel: React.FC = observer(() => {
       name: string;
    }
 
+   const updateIssues = options.jobDetails?.jobData?.updateIssues;
+
    const targetItems: TargetPickerItem[] = options.targets.map(t => {
       return {
          key: t,
          name: t
       }
    });
+
    return <Stack style={{
       height: height,
       position: 'relative',      
@@ -1647,15 +1648,11 @@ const AdvancedPanel: React.FC = observer(() => {
                }
                } />
             </Stack>
-            <Stack>
-               <Checkbox disabled={options.readOnly || options.template?.updateIssues} key={`key_adv_update_issues_${options.currentRenderKey}`}
+            {updateIssues !== undefined && options.readOnly && <Stack>
+               <Checkbox disabled key={`key_adv_update_issues_${options.currentRenderKey}`}
                   label="Update Build Health Issues"
-                  defaultChecked={options.template?.updateIssues ? true : options.advUpdateIssues}
-                  onChange={(ev, checked) => {
-                     options.advUpdateIssues = checked ? true : undefined;
-                     options.setChanged();
-                  }} />
-            </Stack>
+                  checked={updateIssues} />
+            </Stack>}
             <Stack>
                <TextField key={`key_adv_job_name_${options.currentRenderKey}`} disabled={options.readOnly} spellCheck={false} defaultValue={options.advJobName} label="Job Name" onChange={(ev, newValue) => {
                   options.advJobName = newValue;
@@ -1798,7 +1795,6 @@ const BuildModal: React.FC<{ setUseLegacyDialog: (value: boolean) => void }> = o
       }
 
       const templateId = template.id;
-      const updateIssues = template.updateIssues ? undefined : options.advUpdateIssues;
 
       let changeQueries: ChangeQueryConfig[] | undefined;
 
@@ -1826,8 +1822,7 @@ const BuildModal: React.FC<{ setUseLegacyDialog: (value: boolean) => void }> = o
          streamId: options.streamId,
          templateId: templateId,
          name: options.advJobName,
-         priority: options.advJobPriority,
-         updateIssues: updateIssues,
+         priority: options.advJobPriority,         
          changeQueries: changeQueries,
          parameters: options.parameters,
          additionalArguments: additionalArgs?.length ? additionalArgs : undefined,
