Horde: Adding server plugin info to version modal
#rnx
#skipci

diff --git a/Engine/Source/Programs/Horde/HordeDashboard/src/backend/Api.ts b/Engine/Source/Programs/Horde/HordeDashboard/src/backend/Api.ts
--- a/Engine/Source/Programs/Horde/HordeDashboard/src/backend/Api.ts
+++ b/Engine/Source/Programs/Horde/HordeDashboard/src/backend/Api.ts
@@ -4476,6 +4476,13 @@ export type UpdateGlobalConfigRequest = {
 
 }
 
+export type ServerPluginInfoResponse = {
+	name: string;
+	description?: string;
+	loaded: boolean,
+	version?: string;
+};
+
 export type GetServerInfoResponse = {
 
 	/// Server version info
@@ -4484,11 +4491,8 @@ export type GetServerInfoResponse = {
 	/// The current agent version
 	agentVersion?: string;
 
-	/// The operating system server is hosted on
-	osDescription: string;
-
-	/// whether the server is running in single instance mode
-	singleInstance: boolean;
+	/// The plugins on the server
+	plugins: ServerPluginInfoResponse[];
 }
 
 /// Information about a span within an issue
diff --git a/Engine/Source/Programs/Horde/HordeDashboard/src/components/VersionModal.tsx b/Engine/Source/Programs/Horde/HordeDashboard/src/components/VersionModal.tsx
--- a/Engine/Source/Programs/Horde/HordeDashboard/src/components/VersionModal.tsx
+++ b/Engine/Source/Programs/Horde/HordeDashboard/src/components/VersionModal.tsx
@@ -6,7 +6,7 @@ import { GetServerInfoResponse } from '../backend/Api';
 import { getHordeStyling } from '../styles/Styles';
 
 export const VersionModal: React.FC<{ show: boolean, onClose: () => void }> = ({ show, onClose }) => {
-   
+
    const [version, setVersion] = useState<{ serverInfo?: GetServerInfoResponse, querying?: boolean }>({});
    const { hordeClasses } = getHordeStyling();
 
@@ -27,7 +27,7 @@ export const VersionModal: React.FC<{ show: boolean, onClose: () => void }> = ({
 
          const serverInfo = await backend.getServerInfo();
 
-         setVersion({ serverInfo: serverInfo})
+         setVersion({ serverInfo: serverInfo })
       })()
 
    }
@@ -49,7 +49,7 @@ export const VersionModal: React.FC<{ show: boolean, onClose: () => void }> = ({
    } catch (reason) {
       console.log("Process env error:", reason);
    }
-   
+
    if (dashboardVersion) {
       versionItems.push({ name: "Dashboard", value: dashboardVersion });
    }
@@ -60,6 +60,17 @@ export const VersionModal: React.FC<{ show: boolean, onClose: () => void }> = ({
       versionItems.push({ name: "Agent", value: version.serverInfo?.agentVersion });
    }
 
+   if (version.serverInfo?.plugins?.length) {
+
+      versionItems.push({ name: "Plugins:", value: "" });
+
+      version.serverInfo?.plugins.forEach(p => {
+         versionItems.push({ name: p.name, value: p.version ?? "" });
+      })
+   
+   }
+
+
    const onRenderItemColumn = (item: GeneralItem, index?: number, columnIn?: IColumn) => {
 
       const column = columnIn!;
@@ -67,7 +78,7 @@ export const VersionModal: React.FC<{ show: boolean, onClose: () => void }> = ({
       // simple cases
       switch (column.name) {
          case 'Name':
-            return <Text variant="medium" styles={{ root: { fontFamily: "Horde Open Sans SemiBold" } }}>{item.name}:</Text>
+            return <Text variant="medium" styles={{ root: { fontFamily: "Horde Open Sans SemiBold" } }}>{item.name}</Text>
          case 'Value':
             if (item.value) {
                return <Text variant="medium" styles={{ root: { fontFamily: "Horde Open Sans SemiBold" } }}>{item.value}</Text>
