Horde: Make workflow issues linkable on stream summary
#rnx
#skipci

diff --git a/Engine/Source/Programs/Horde/HordeDashboard/src/components/StreamSummary.tsx b/Engine/Source/Programs/Horde/HordeDashboard/src/components/StreamSummary.tsx
--- a/Engine/Source/Programs/Horde/HordeDashboard/src/components/StreamSummary.tsx
+++ b/Engine/Source/Programs/Horde/HordeDashboard/src/components/StreamSummary.tsx
@@ -5,7 +5,7 @@ import { action, makeObservable, observable } from 'mobx';
 import { observer } from 'mobx-react-lite';
 import moment from 'moment-timezone';
 import React, { useEffect, useState } from 'react';
-import { Link, useLocation, useParams } from 'react-router-dom';
+import { Link, useLocation, useNavigate, useParams } from 'react-router-dom';
 import backend, { useBackend } from '../backend';
 import { FindIssueResponse, GetExternalIssueResponse, JobsTabData, TabType } from '../backend/Api';
 import dashboard from '../backend/Dashboard';
@@ -353,7 +353,7 @@ const sortIssueGroups = (itemsIn: HealthItem[], projectStore: ProjectStore): [He
                count: 1,
                headerText: tab.title,
                startIndex: curIndex,
-               isCollapsed: !handler.collapsedIssueGroups.get(key)
+               isCollapsed: handler.collapsedIssueGroups.get(key)
             });
 
             tabIssues.set(tab.title, []);
@@ -378,7 +378,7 @@ const sortIssueGroups = (itemsIn: HealthItem[], projectStore: ProjectStore): [He
          count: 1,
          headerText: name,
          startIndex: curIndex,
-         isCollapsed: !handler.collapsedIssueGroups.get(key)
+         isCollapsed: handler.collapsedIssueGroups.get(key)
       });
 
       tabIssues.set(name, []);
@@ -412,8 +412,9 @@ const HealthPanelIssues: React.FC<{ desktopAlerts?: boolean }> = observer(({ des
    const location = useLocation();
    const { projectStore } = useBackend();
    const [issueHistory, setIssueHistory] = useState(false);
-   const [currentPivot, setCurrentPivot] = useState("");
+   const [currentPivot, setCurrentPivot] = useState(query.get("workflow") === null ? "" : query.get("workflow")!);
    const { hordeClasses, detailClasses } = getHordeStyling();
+   const navigate = useNavigate();
 
    // subscribe
    if (handler.update) { }
@@ -688,7 +689,14 @@ const HealthPanelIssues: React.FC<{ desktopAlerts?: boolean }> = observer(({ des
                      if (!item || !item.props.itemKey) {
                         return;
                      }
-                     setCurrentPivot(item.props.itemKey.replace("issue_pivot_item_key_", ""));
+
+                     const search = new URLSearchParams(window.location.search);
+                     const workflowId = item.props.itemKey.replace("issue_pivot_item_key_", "");
+                     search.set("workflow", workflowId)
+                     const url = `${window.location.pathname}?` + search.toString();
+                     navigate(url, { replace: true })
+         
+                     setCurrentPivot(workflowId);
                   })}>
                   {pivotItems}
                </Pivot>
