Horde: Add doc pages for build health.

#rnx

diff --git a/Engine/Source/Programs/Horde/.editorconfig b/Engine/Source/Programs/Horde/.editorconfig
--- a/Engine/Source/Programs/Horde/.editorconfig
+++ b/Engine/Source/Programs/Horde/.editorconfig
@@ -2,6 +2,10 @@
 
 # EditorConfig is awesome:http://EditorConfig.org
 
+[*.md]
+indent_style = space
+indent_size = 2
+
 # Dotnet code style settings:
 [*.cs]
 
diff --git a/Engine/Source/Programs/Horde/Docs/Internals.md b/Engine/Source/Programs/Horde/Docs/Internals.md
--- a/Engine/Source/Programs/Horde/Docs/Internals.md
+++ b/Engine/Source/Programs/Horde/Docs/Internals.md
@@ -50,6 +50,7 @@ can be run as follows:
 
 ## Topics
 
+* [Build Health](Internals/BuildHealth.md)
 * [Storage](Internals/Storage.md)
 * [Leases](Internals/Leases.md)
 * [Logs](Internals/Logs.md)
diff --git a/Engine/Source/Programs/Horde/Docs/Internals/BuildHealth.md b/Engine/Source/Programs/Horde/Docs/Internals/BuildHealth.md
new file mode 100644
--- /dev/null
+++ b/Engine/Source/Programs/Horde/Docs/Internals/BuildHealth.md
@@ -0,0 +1,199 @@
+[Horde](../../README.md) > [Internals](../Internals.md) > Build Health
+
+# Build Health
+
+Epic has a large development team working on products that are continously shipping. We
+use a trunk-based development model with branches for each release, with
+commits to release branches automatically being merged back to the mainline.
+
+Maintaining a stable and healthy codebase is important to our ability to iterate quickly, 
+so we have a number of policies in place to help:
+
+* Local test suites that developers can run before changes are submitted.
+* The ability to run a suite of pre-submit tests on the build farm, through Horde's _preflight_ system.
+* An automated __Build Health__ system for identifying, grouping, and triaging build errors,
+  and notifying developers responsible for them as quickly as possible via _Slack_.
+
+The Build Health system can be thought of as similar to a traditional bug tracking database,
+creating __issues__ for build breakages, but is automated, context-aware and highly data
+driven - and tightly integrated into Horde's job dashboard and logs.
+
+This system is only as good as the data fed into it, so we are careful to annotate any
+errors with as much context as possible to allow the system to make smart decisions when
+things go wrong.
+
+## Anatomy of a build breakage
+
+Each branch is modeled as a linear sequence of commits. When a build failure occurs, we
+seek to identify which commit since the last successful build is responsible.
+
+There are many heuristics we can do this; looking for commits that modify a particular
+asset or source file that fails to compile, distinguishing between code/content errors and
+changes, looking for linker symbol names that match modified source file names, and so on.
+
+Additionally, we may have a single change that causes a widespread set of errors - a source
+file that fails to compile may produce completely different errors on MSVC vs Clang vs GCC,
+for example, a missing asset may cause cook errors for different platforms, and so on. These
+errors may occur at times during the build pipeline, and we want to get eyes on them as soon
+as possible - while also avoiding a long tail of spurious notifications that occur later.
+
+Finally, because we develop across multiple branches and merge frequently between them,
+we want to collate errors in different streams and trace back to the original commit that
+is deemed responsible.
+
+## Issues, Spans, and Fingerprints
+
+Horde creates a build health __issue__ for each detected breakage.
+
+Each issue contains multiple __spans__ - one for each stream the issue has been observed in -
+each of which identifies the last successful and first failing commit. A span is marked as
+resolved once the associated error no longer occurs in a stream, and the issue is marked as
+resolved once all its spans are resolved.
+
+Build errors are ingested and attached to spans via their __fingerprint__. Fingerprints
+provide a programmatic description of the error, and how it should be matched with other
+errors.
+
+### Contents of a Fingerprint
+
+A fingerprint contains the following data (see `IIssueFingerprint.cs`):
+
+* __Type__: A string identifying the type of issue. Fingerprints only match with other fingerprints of the same type.
+  One of the most commonly encountered issue types is `Compile`, indicating a compile error in one or more source files.
+* __SummaryTemplate__: A string that is used to populate an issue's description in the Horde dashboard and notifications. 
+  May contain placeholder strings that will be expanded when fingerprints are merged together
+  (see [Summary Templates](#summary-templates)).
+* __Keys__: Pieces of identifying data used to match and group fingerprints,
+  such as the name of a file associated with the error. Each key has a name and type.
+* __RejectKeys__: Set of keys that this fingerprint should _not_ match with.
+* __Metadata__: Arbitrary key/value data that can be aggregated for use in summary templates and for debugging.
+* __ChangeFilter__: A semicolon-delimited set of wildcards for the files that can be considered
+  responsible for this error, such as `*.c;*.cpp;*.h;*.hpp*` for C++ source files, for example.
+
+### Matching and merging fingerprints
+
+Fingerprints are matched by logic in `IIssueFingerprint.cs`. The pipeline is as follows:
+
+* A fingerprint is generated for each error or warning in a build step.
+* Each fingerprint is compared against the fingerprint of any unresolved spans in the same
+  stream. New fingerprints are considered to match an existing fingerprint if:
+    1. The type is identical
+    2. The summary template is identical
+    3. The new fingerprint has one key in common with the existing fingerprint, or neither fingerprint has any keys.
+    4. The new fingerprint does not contain any of the keys listed in the _rejectKeys_.
+* All the fingerprints which aren't matched with an existing span are grouped together and
+  added to new spans. These fingerprints are grouped using a more relaxed set of conditions:
+    1. The type is identical
+    2. Neither fingerprint has a key that is in other key's reject keys.
+* For any new spans, we find the list of suspect commits - following the merge history for that
+  commit to its originating stream. If any issue has a matching fingerprint, shares a suspect
+  commit, and does not yet have a span in the current stream, we add it to that issue. Otherwise a new issue is created.
+
+### Creating fingerprints
+
+Fingerprints may be created for errors in two ways:
+
+* By including fingerprint information directly in the [structured log event](StructuredLogging.md) at the
+  point that it is generated.
+* By post-processing [structured log events](StructuredLogging.md) in the Horde server once a build step 
+  completes.
+
+Including fingerprints directly in structured log events is preferred, since it
+allows for more extensibility and gives the most control over adding fine-grained context.
+
+#### Adding fingerprints to errors (C#)
+
+The `EpicGames.Horde` library includes extension methods for creating scopes that attach fingerprint
+metadata to log events, using the standard .NET `ILogger` interface. This interface is exposed to
+`AutomationTool` commands via the `CommandUtils.Logger` property.
+
+Fingerprints can be added as in the following example:
+
+  ```cs
+  IssueFingerprint fingerprint = new IssueFingerprint("Compile", "Compile {Severity} in {Files}", IssueChangeFilter.Code);
+  fingerprint.Keys.Add(IssueKey.FromFile("Foo.cpp"));
+
+  using (IDisposable scope = Logger.BeginIssueScope(fingerprint))
+  {
+      Logger.LogError("Compile errors in {File}", new FileReference("Foo.cpp"));
+  }
+  ```
+
+#### Adding fingerprints to errors (C++)
+
+// TODO
+
+#### Adding fingerprints via post-processing
+
+Post-processing is handled by the `IssueHandler` derived classes in the `EpicGames.Horde` project.
+
+Handlers are enumerated by the server by searching for classes with the `[IssueHandler]` attribute. A new
+instance of each handler is created for each build step containing events to be processed.
+
+Log events are passed to `IssueHandler.HandleEvent` in decreasing order of their `IssueHandler.Priorty` value,
+until a handler is found which returns `true`. Once all events have been processed, `IssueHandler.GetIssues`
+returns a list of the matched log events and their fingerprints.
+
+Note that handlers do not have to generate fingerprints for a log event; it can be useful
+for more specific errors (say, a compile error) to mask more vague errors (say, UnrealBuildTool 
+returning a failing exit code). As such, `HandleEvent` can return true for errors that it wants
+to exclude from the remaining handler pipeline.
+
+### Summary Templates
+
+The following placeholder variables can be used in summary templates:
+
+* `{Severity}`: Expands to the string 'Warnings' or 'Errors' depending on
+  the highest severity of diagnostic grouped into this issue.
+* `{Files}`: Expands to a list of filenames from keys with a type of `File`. Up
+  to 3 files are included.
+* `{Meta:Key}`: Expands to a list of metadata values associated with `Key`.
+
+## Workflows
+
+Workflows define how and where developers are notified about build health
+issues.
+
+Workflows are configured via the [`Workflows`](../Config/Schema/Streams.md) property in
+the stream configuration file, and jobs are configured to use them through the job
+template's [`WorkflowId`](../Config/Schema/Streams.md#templaterefconfig) property.
+
+Individual nodes within a job may be configured to use a workflow using an [annotation](#annotations) in the BuildGraph script.
+
+### Notifications
+
+Horde supports [Slack](https://www.slack.com) for issue notifications. By default, a
+Slack thread is created for each identified issue, and developers that are suspected of
+causing the issue are @mentioned to draw their attention to it.
+
+An interactive surface allows developers to acknowledge or decline issues from Slack
+without having to navigate to the Horde dashboard.
+
+When configuring Slack channels and users through the stream configuration file, most
+parameters require the Slack channel and user ID to be supplied rather than their name.
+
+### Reports
+
+Workflows can be configured to send out summary reports at defined intervals, showing a
+list of unresolved issues, links to the appropriate threads, and statistics for the
+health of the stream.
+
+The sending of reports can be configured through the `reportTimes` and `reportChannel` properties of the workflow configuration.
+
+### Annotations
+
+Handling of issues for individual job steps may be further configured using __node attributes__. Attributes can be specified for a node through the `Annotations` attribute on `Node` elements in BuildGraph scripts:
+
+  ```xml
+  <Node Name="Compile UnrealEditor Win64" Annotations="Workflow=my-workflow;BuildBlocker=true">
+  ```
+
+The supported set of annotations are defined in the `NodeAnnotations.cs` source file. The following are currently available:
+
+* `Workflow`: Workflow to use for triaging issues from this node
+* `CreateIssues`: Allows disabling creation of issues for this node, when set to `false`.
+* `AutoAssign`: Whether to automatically assign issues that could only be caused by one user, or have a well defined correlation with a modified file.
+* `AutoAssignToUser`: Automatically assign any issues in this step to the Perforce user given by the supplied argument.
+* `NotifySubmitters`: Whether to notify all submitters between a build suceeding and failing, allowing them to step forward and take ownership of an issue.
+* `IssueGroup`: Specifies a suffix to be appended to the `type` property of any issues created from this node, preventing them being merged with other issues.
+* `BuildBlocker`: Whether failures in this node should be considered a build blocker. Issues identified as build blockers have a special tag in Slack notifications.
