Horde: Add "Analyze" configurations to all plugin projects, and fix up any current violations.

#rnx

diff --git a/Engine/Source/Programs/Horde/HordeServer.Shared/Acls/AclConfig.cs b/Engine/Source/Programs/Horde/HordeServer.Shared/Acls/AclConfig.cs
--- a/Engine/Source/Programs/Horde/HordeServer.Shared/Acls/AclConfig.cs
+++ b/Engine/Source/Programs/Horde/HordeServer.Shared/Acls/AclConfig.cs
@@ -1,13 +1,9 @@
 // Copyright Epic Games, Inc. All Rights Reserved.
 
-using System;
-using System.Collections.Generic;
 using System.ComponentModel.DataAnnotations;
-using System.Linq;
 using System.Security.Claims;
 using System.Text.Json.Serialization;
 using EpicGames.Horde.Acls;
-using HordeServer.Server;
 using HordeServer.Utilities;
 
 #pragma warning disable CA2227 // Change 'X' to be read-only by removing the property setter
diff --git a/Engine/Source/Programs/Horde/HordeServer.Shared/Acls/IAclService.cs b/Engine/Source/Programs/Horde/HordeServer.Shared/Acls/IAclService.cs
--- a/Engine/Source/Programs/Horde/HordeServer.Shared/Acls/IAclService.cs
+++ b/Engine/Source/Programs/Horde/HordeServer.Shared/Acls/IAclService.cs
@@ -1,12 +1,7 @@
 // Copyright Epic Games, Inc. All Rights Reserved.
 
-using System;
-using System.Collections.Generic;
 using System.Diagnostics.CodeAnalysis;
-using System.Linq;
 using System.Security.Claims;
-using System.Text;
-using System.Threading.Tasks;
 using EpicGames.Horde.Acls;
 
 namespace HordeServer.Acls
@@ -41,6 +36,7 @@ namespace HordeServer.Acls
 		/// <summary>
 		/// Authorizes a user to perform a given action
 		/// </summary>
+		/// <param name="aclService">Instance of the ACL service</param>
 		/// <param name="scopeName">Name of the scope to auth against</param>
 		/// <param name="action">The action being performed</param>
 		/// <param name="user">The principal to validate</param>
@@ -50,6 +46,7 @@ namespace HordeServer.Acls
 		/// <summary>
 		/// Issues a bearer token with the given roles
 		/// </summary>
+		/// <param name="aclService">Instance of the ACL service</param>
 		/// <param name="claims">List of claims to include</param>
 		/// <param name="expiry">Time that the token expires</param>
 		/// <param name="cancellationToken">Cancellation token for the operation</param>
diff --git a/Engine/Source/Programs/Horde/HordeServer.Shared/Auditing/IAuditLog.cs b/Engine/Source/Programs/Horde/HordeServer.Shared/Auditing/IAuditLog.cs
--- a/Engine/Source/Programs/Horde/HordeServer.Shared/Auditing/IAuditLog.cs
+++ b/Engine/Source/Programs/Horde/HordeServer.Shared/Auditing/IAuditLog.cs
@@ -1,11 +1,7 @@
 // Copyright Epic Games, Inc. All Rights Reserved.
 
-using System;
-using System.Collections.Generic;
 using System.IO.Pipelines;
 using System.Text;
-using System.Threading;
-using System.Threading.Tasks;
 using Microsoft.Extensions.Logging;
 
 namespace HordeServer.Auditing
diff --git a/Engine/Source/Programs/Horde/HordeServer.Shared/ConfigVersion.cs b/Engine/Source/Programs/Horde/HordeServer.Shared/ConfigVersion.cs
--- a/Engine/Source/Programs/Horde/HordeServer.Shared/ConfigVersion.cs
+++ b/Engine/Source/Programs/Horde/HordeServer.Shared/ConfigVersion.cs
@@ -6,7 +6,7 @@ namespace HordeServer
 #pragma warning disable CA1069 // Enum member 'Latest' has same value as ...
 	/// <summary>
 	/// Global version number for running the server. As new features are introduced that require data migrations, this version number indicates the backwards compatibility functionality that must be enabled.
-	/// When adding a new version here, also add a message to <see cref="ConfigService.CreateSnapshotAsync"/> describing the steps that need to be taken to upgrade the deployment.
+	/// When adding a new version here, also add a message to ConfigService.CreateSnapshotAsync describing the steps that need to be taken to upgrade the deployment.
 	/// </summary>
 	public enum ConfigVersion
 	{
diff --git a/Engine/Source/Programs/Horde/HordeServer.Shared/Configuration/ConfigContext.cs b/Engine/Source/Programs/Horde/HordeServer.Shared/Configuration/ConfigContext.cs
--- a/Engine/Source/Programs/Horde/HordeServer.Shared/Configuration/ConfigContext.cs
+++ b/Engine/Source/Programs/Horde/HordeServer.Shared/Configuration/ConfigContext.cs
@@ -1,14 +1,9 @@
 // Copyright Epic Games, Inc. All Rights Reserved.
 
-using System;
-using System.Collections.Generic;
 using System.Diagnostics;
 using System.Diagnostics.CodeAnalysis;
 using System.Text.Json;
-using System.Threading;
-using System.Threading.Tasks;
 using EpicGames.Core;
-using EpicGames.Horde.Storage;
 using Microsoft.Extensions.Logging;
 
 namespace HordeServer.Configuration
diff --git a/Engine/Source/Programs/Horde/HordeServer.Shared/Configuration/ConfigException.cs b/Engine/Source/Programs/Horde/HordeServer.Shared/Configuration/ConfigException.cs
--- a/Engine/Source/Programs/Horde/HordeServer.Shared/Configuration/ConfigException.cs
+++ b/Engine/Source/Programs/Horde/HordeServer.Shared/Configuration/ConfigException.cs
@@ -1,8 +1,5 @@
 // Copyright Epic Games, Inc. All Rights Reserved.
 
-using System;
-using System.Collections.Generic;
-
 namespace HordeServer.Configuration
 {
 	/// <summary>
diff --git a/Engine/Source/Programs/Horde/HordeServer.Shared/Configuration/ConfigMacro.cs b/Engine/Source/Programs/Horde/HordeServer.Shared/Configuration/ConfigMacro.cs
--- a/Engine/Source/Programs/Horde/HordeServer.Shared/Configuration/ConfigMacro.cs
+++ b/Engine/Source/Programs/Horde/HordeServer.Shared/Configuration/ConfigMacro.cs
@@ -1,7 +1,5 @@
 // Copyright Epic Games, Inc. All Rights Reserved.
 
-using System;
-
 namespace HordeServer.Configuration
 {
 	/// <summary>
diff --git a/Engine/Source/Programs/Horde/HordeServer.Shared/Configuration/ConfigNode.cs b/Engine/Source/Programs/Horde/HordeServer.Shared/Configuration/ConfigNode.cs
--- a/Engine/Source/Programs/Horde/HordeServer.Shared/Configuration/ConfigNode.cs
+++ b/Engine/Source/Programs/Horde/HordeServer.Shared/Configuration/ConfigNode.cs
@@ -1,14 +1,10 @@
 // Copyright Epic Games, Inc. All Rights Reserved.
 
-using System;
-using System.Collections.Generic;
 using System.Diagnostics.CodeAnalysis;
 using System.Reflection;
 using System.Text.Json;
 using System.Text.Json.Nodes;
 using System.Text.Json.Serialization;
-using System.Threading;
-using System.Threading.Tasks;
 
 namespace HordeServer.Configuration
 {
@@ -52,7 +48,7 @@ namespace HordeServer.Configuration
 		/// <summary>
 		/// Default options for new json node objects
 		/// </summary>
-		internal static JsonNodeOptions s_defaultJsonNodeOptions { get; } = new JsonNodeOptions { PropertyNameCaseInsensitive = true };
+		protected static JsonNodeOptions DefaultJsonNodeOptions { get; } = new JsonNodeOptions { PropertyNameCaseInsensitive = true };
 
 		/// <summary>
 		/// Parses macro definitions from this object
@@ -126,7 +122,7 @@ namespace HordeServer.Configuration
 			}
 			else if (node is JsonObject obj)
 			{
-				JsonObject result = new JsonObject(s_defaultJsonNodeOptions);
+				JsonObject result = new JsonObject(DefaultJsonNodeOptions);
 				foreach ((string propertyName, JsonNode? propertyNode) in obj)
 				{
 					result[propertyName] = ExpandMacros(propertyNode, context);
@@ -327,7 +323,7 @@ namespace HordeServer.Configuration
 		/// <inheritdoc/>
 		public override async Task<JsonNode?> PreprocessAsync(JsonNode? node, JsonNode? existingNode, ConfigContext context, CancellationToken cancellationToken)
 		{
-			JsonObject? targetObject = ((JsonObject?)existingNode) ?? new JsonObject(s_defaultJsonNodeOptions);
+			JsonObject? targetObject = ((JsonObject?)existingNode) ?? new JsonObject(DefaultJsonNodeOptions);
 			foreach ((string key, JsonNode? element) in (JsonObject)node!)
 			{
 				context.EnterScope($"[{key}]");
@@ -559,7 +555,7 @@ namespace HordeServer.Configuration
 			}
 
 			// Ensure that the target object is valid so we can write properties into it
-			target ??= new JsonObject(s_defaultJsonNodeOptions);
+			target ??= new JsonObject(DefaultJsonNodeOptions);
 
 			// Parse all the macros for this scope
 			if (MacroScope)
diff --git a/Engine/Source/Programs/Horde/HordeServer.Shared/Configuration/ConfigObject.cs b/Engine/Source/Programs/Horde/HordeServer.Shared/Configuration/ConfigObject.cs
--- a/Engine/Source/Programs/Horde/HordeServer.Shared/Configuration/ConfigObject.cs
+++ b/Engine/Source/Programs/Horde/HordeServer.Shared/Configuration/ConfigObject.cs
@@ -1,9 +1,6 @@
 // Copyright Epic Games, Inc. All Rights Reserved.
 
-using System;
 using System.Collections;
-using System.Collections.Generic;
-using System.Linq;
 using System.Reflection;
 using EpicGames.Core;
 
diff --git a/Engine/Source/Programs/Horde/HordeServer.Shared/Configuration/IConfigFile.cs b/Engine/Source/Programs/Horde/HordeServer.Shared/Configuration/IConfigFile.cs
--- a/Engine/Source/Programs/Horde/HordeServer.Shared/Configuration/IConfigFile.cs
+++ b/Engine/Source/Programs/Horde/HordeServer.Shared/Configuration/IConfigFile.cs
@@ -1,8 +1,5 @@
 // Copyright Epic Games, Inc. All Rights Reserved.
 
-using System;
-using System.Threading;
-using System.Threading.Tasks;
 using HordeServer.Users;
 
 namespace HordeServer.Configuration
diff --git a/Engine/Source/Programs/Horde/HordeServer.Shared/Configuration/IConfigSource.cs b/Engine/Source/Programs/Horde/HordeServer.Shared/Configuration/IConfigSource.cs
--- a/Engine/Source/Programs/Horde/HordeServer.Shared/Configuration/IConfigSource.cs
+++ b/Engine/Source/Programs/Horde/HordeServer.Shared/Configuration/IConfigSource.cs
@@ -1,11 +1,6 @@
 // Copyright Epic Games, Inc. All Rights Reserved.
 
-using System;
 using System.Collections.Concurrent;
-using System.Collections.Generic;
-using System.IO;
-using System.Threading;
-using System.Threading.Tasks;
 using EpicGames.Core;
 using EpicGames.Horde.Users;
 using HordeServer.Users;
diff --git a/Engine/Source/Programs/Horde/HordeServer.Shared/HordeServer.Shared.csproj b/Engine/Source/Programs/Horde/HordeServer.Shared/HordeServer.Shared.csproj
--- a/Engine/Source/Programs/Horde/HordeServer.Shared/HordeServer.Shared.csproj
+++ b/Engine/Source/Programs/Horde/HordeServer.Shared/HordeServer.Shared.csproj
@@ -5,6 +5,16 @@
     <ImplicitUsings>enable</ImplicitUsings>
     <Nullable>enable</Nullable>
     <RootNamespace>HordeServer</RootNamespace>
+    <RunAnalyzersDuringBuild>False</RunAnalyzersDuringBuild>
+    <Configurations>Debug;Release;Analyze</Configurations>
+  </PropertyGroup>
+
+  <PropertyGroup Condition=" '$(Configuration)|$(Platform)' == 'Analyze|AnyCPU' ">
+    <RunAnalyzersDuringBuild>True</RunAnalyzersDuringBuild>
+    <TreatWarningsAsErrors>false</TreatWarningsAsErrors>
+    <EnforceCodeStyleInBuild>True</EnforceCodeStyleInBuild>
+    <GenerateDocumentationFile>True</GenerateDocumentationFile>
+    <DocumentationFile></DocumentationFile>
   </PropertyGroup>
 
   <ItemGroup>
diff --git a/Engine/Source/Programs/Horde/HordeServer.Shared/Plugins/IPluginConfig.cs b/Engine/Source/Programs/Horde/HordeServer.Shared/Plugins/IPluginConfig.cs
--- a/Engine/Source/Programs/Horde/HordeServer.Shared/Plugins/IPluginConfig.cs
+++ b/Engine/Source/Programs/Horde/HordeServer.Shared/Plugins/IPluginConfig.cs
@@ -1,12 +1,5 @@
 // Copyright Epic Games, Inc. All Rights Reserved.
 
-using System;
-using System.Collections.Generic;
-using System.Linq;
-using System.Text;
-using System.Text.Json;
-using System.Text.Json.Serialization;
-using System.Threading.Tasks;
 using HordeServer.Acls;
 
 namespace HordeServer.Plugins
@@ -19,12 +12,12 @@ namespace HordeServer.Plugins
 		/// <summary>
 		/// Called to fixup a plugin's configuration after deserialization
 		/// </summary>
-		/// <param name="options">Options for configuring the plugin</param>
+		/// <param name="configOptions">Options for configuring the plugin</param>
 		void PostLoad(PluginConfigOptions configOptions);
 	}
 
 	/// <summary>
-	/// Options passed to <see cref="IPluginConfig.PostLoad(AclConfig)"/>
+	/// Options passed to <see cref="IPluginConfig.PostLoad(PluginConfigOptions)"/>
 	/// </summary>
 	public record class PluginConfigOptions(ConfigVersion Version, IEnumerable<IPluginConfig> Plugins, AclConfig ParentAcl);
 
diff --git a/Engine/Source/Programs/Horde/HordeServer.Shared/Plugins/IPluginStartup.cs b/Engine/Source/Programs/Horde/HordeServer.Shared/Plugins/IPluginStartup.cs
--- a/Engine/Source/Programs/Horde/HordeServer.Shared/Plugins/IPluginStartup.cs
+++ b/Engine/Source/Programs/Horde/HordeServer.Shared/Plugins/IPluginStartup.cs
@@ -1,10 +1,6 @@
 // Copyright Epic Games, Inc. All Rights Reserved.
 
-using System;
 using Microsoft.AspNetCore.Builder;
-using Microsoft.AspNetCore.Hosting;
-using Microsoft.CodeAnalysis.CSharp.Syntax;
-using Microsoft.Extensions.Configuration;
 using Microsoft.Extensions.DependencyInjection;
 
 namespace HordeServer.Plugins
diff --git a/Engine/Source/Programs/Horde/HordeServer.Shared/Server/IMongoService.cs b/Engine/Source/Programs/Horde/HordeServer.Shared/Server/IMongoService.cs
--- a/Engine/Source/Programs/Horde/HordeServer.Shared/Server/IMongoService.cs
+++ b/Engine/Source/Programs/Horde/HordeServer.Shared/Server/IMongoService.cs
@@ -1,9 +1,5 @@
 // Copyright Epic Games, Inc. All Rights Reserved.
 
-using System;
-using System.Collections.Generic;
-using System.Threading;
-using System.Threading.Tasks;
 using HordeServer.Utilities;
 using Microsoft.Extensions.Diagnostics.HealthChecks;
 using MongoDB.Driver;
diff --git a/Engine/Source/Programs/Horde/HordeServer.Shared/Server/IRedisService.cs b/Engine/Source/Programs/Horde/HordeServer.Shared/Server/IRedisService.cs
--- a/Engine/Source/Programs/Horde/HordeServer.Shared/Server/IRedisService.cs
+++ b/Engine/Source/Programs/Horde/HordeServer.Shared/Server/IRedisService.cs
@@ -1,7 +1,5 @@
 // Copyright Epic Games, Inc. All Rights Reserved.
 
-using System;
-using System.Threading.Tasks;
 using EpicGames.Redis;
 using Microsoft.Extensions.Diagnostics.HealthChecks;
 using StackExchange.Redis;
diff --git a/Engine/Source/Programs/Horde/HordeServer.Shared/Server/MongoIndex.cs b/Engine/Source/Programs/Horde/HordeServer.Shared/Server/MongoIndex.cs
--- a/Engine/Source/Programs/Horde/HordeServer.Shared/Server/MongoIndex.cs
+++ b/Engine/Source/Programs/Horde/HordeServer.Shared/Server/MongoIndex.cs
@@ -1,7 +1,5 @@
 // Copyright Epic Games, Inc. All Rights Reserved.
 
-using System;
-using System.Collections.Generic;
 using System.Diagnostics;
 using System.Text;
 using MongoDB.Bson;
@@ -151,11 +149,17 @@ namespace HordeServer.Server
 	/// </summary>
 	public static class MongoIndexExtensions
 	{
+		/// <summary>
+		/// Adds a new index to the collection
+		/// </summary>
 		public static void Add<T>(this List<MongoIndex<T>> list, Func<IndexKeysDefinitionBuilder<T>, IndexKeysDefinition<T>> keyFunc, bool unique = false, bool sparse = false)
 		{
 			list.Add(MongoIndex.Create(keyFunc, unique, sparse));
 		}
 
+		/// <summary>
+		/// Adds a new index to the collection
+		/// </summary>
 		public static void Add<T>(this List<MongoIndex<T>> list, string name, Func<IndexKeysDefinitionBuilder<T>, IndexKeysDefinition<T>> keyFunc, bool unique = false, bool sparse = false)
 		{
 			list.Add(MongoIndex.Create(name, keyFunc, unique, sparse));
diff --git a/Engine/Source/Programs/Horde/HordeServer.Shared/Telemetry/ITelemetryWriter.cs b/Engine/Source/Programs/Horde/HordeServer.Shared/Telemetry/ITelemetryWriter.cs
--- a/Engine/Source/Programs/Horde/HordeServer.Shared/Telemetry/ITelemetryWriter.cs
+++ b/Engine/Source/Programs/Horde/HordeServer.Shared/Telemetry/ITelemetryWriter.cs
@@ -25,6 +25,7 @@ namespace HordeServer.Telemetry
 		/// Writes a telemetry event
 		/// </summary>
 		/// <param name="telemetryStoreId">The telemetry store to write to</param>
+		/// <param name="metadata">Metadata for the event</param>
 		/// <param name="payload">Event data to write</param>
 		void WriteEvent(TelemetryStoreId telemetryStoreId, TelemetryRecordMeta metadata, object payload);
 	}
diff --git a/Engine/Source/Programs/Horde/HordeServer.Shared/Users/IAvatarService.cs b/Engine/Source/Programs/Horde/HordeServer.Shared/Users/IAvatarService.cs
--- a/Engine/Source/Programs/Horde/HordeServer.Shared/Users/IAvatarService.cs
+++ b/Engine/Source/Programs/Horde/HordeServer.Shared/Users/IAvatarService.cs
@@ -1,8 +1,5 @@
 // Copyright Epic Games, Inc. All Rights Reserved.
 
-using System.Threading;
-using System.Threading.Tasks;
-
 namespace HordeServer.Users
 {
 	/// <summary>
diff --git a/Engine/Source/Programs/Horde/HordeServer.Shared/Users/IUser.cs b/Engine/Source/Programs/Horde/HordeServer.Shared/Users/IUser.cs
--- a/Engine/Source/Programs/Horde/HordeServer.Shared/Users/IUser.cs
+++ b/Engine/Source/Programs/Horde/HordeServer.Shared/Users/IUser.cs
@@ -1,8 +1,5 @@
 // Copyright Epic Games, Inc. All Rights Reserved.
 
-using System;
-using System.Collections.Generic;
-using System.Linq;
 using EpicGames.Core;
 using EpicGames.Horde.Jobs;
 using EpicGames.Horde.Jobs.Bisect;
diff --git a/Engine/Source/Programs/Horde/HordeServer.Shared/Users/IUserCollection.cs b/Engine/Source/Programs/Horde/HordeServer.Shared/Users/IUserCollection.cs
--- a/Engine/Source/Programs/Horde/HordeServer.Shared/Users/IUserCollection.cs
+++ b/Engine/Source/Programs/Horde/HordeServer.Shared/Users/IUserCollection.cs
@@ -1,9 +1,6 @@
 // Copyright Epic Games, Inc. All Rights Reserved.
 
-using System.Collections.Generic;
 using System.Security.Claims;
-using System.Threading;
-using System.Threading.Tasks;
 using EpicGames.Horde.Jobs;
 using EpicGames.Horde.Jobs.Bisect;
 using EpicGames.Horde.Users;
diff --git a/Engine/Source/Programs/Horde/HordeServer.Shared/Utilities/AsyncCachedValue.cs b/Engine/Source/Programs/Horde/HordeServer.Shared/Utilities/AsyncCachedValue.cs
--- a/Engine/Source/Programs/Horde/HordeServer.Shared/Utilities/AsyncCachedValue.cs
+++ b/Engine/Source/Programs/Horde/HordeServer.Shared/Utilities/AsyncCachedValue.cs
@@ -1,9 +1,6 @@
 // Copyright Epic Games, Inc. All Rights Reserved.
 
-using System;
 using System.Diagnostics;
-using System.Threading;
-using System.Threading.Tasks;
 using EpicGames.Core;
 
 namespace HordeServer.Utilities
diff --git a/Engine/Source/Programs/Horde/HordeServer.Shared/Utilities/Aws.cs b/Engine/Source/Programs/Horde/HordeServer.Shared/Utilities/Aws.cs
--- a/Engine/Source/Programs/Horde/HordeServer.Shared/Utilities/Aws.cs
+++ b/Engine/Source/Programs/Horde/HordeServer.Shared/Utilities/Aws.cs
@@ -1,8 +1,5 @@
 // Copyright Epic Games, Inc. All Rights Reserved.
 
-using System;
-using System.IO;
-
 namespace HordeServer.Utilities
 {
 	/// <summary>
diff --git a/Engine/Source/Programs/Horde/HordeServer.Shared/Utilities/CaseInsensitiveDictionary.cs b/Engine/Source/Programs/Horde/HordeServer.Shared/Utilities/CaseInsensitiveDictionary.cs
--- a/Engine/Source/Programs/Horde/HordeServer.Shared/Utilities/CaseInsensitiveDictionary.cs
+++ b/Engine/Source/Programs/Horde/HordeServer.Shared/Utilities/CaseInsensitiveDictionary.cs
@@ -1,8 +1,5 @@
 // Copyright Epic Games, Inc. All Rights Reserved.
 
-using System;
-using System.Collections.Generic;
-
 namespace HordeServer.Utilities
 {
 	/// <summary>
diff --git a/Engine/Source/Programs/Horde/HordeServer.Shared/Utilities/CollectionExtensions.cs b/Engine/Source/Programs/Horde/HordeServer.Shared/Utilities/CollectionExtensions.cs
--- a/Engine/Source/Programs/Horde/HordeServer.Shared/Utilities/CollectionExtensions.cs
+++ b/Engine/Source/Programs/Horde/HordeServer.Shared/Utilities/CollectionExtensions.cs
@@ -1,7 +1,5 @@
 // Copyright Epic Games, Inc. All Rights Reserved.
 
-using System.Collections.Generic;
-
 namespace HordeServer.Utilities
 {
 	/// <summary>
diff --git a/Engine/Source/Programs/Horde/HordeServer.Shared/Utilities/CustomFileCallbackResult.cs b/Engine/Source/Programs/Horde/HordeServer.Shared/Utilities/CustomFileCallbackResult.cs
--- a/Engine/Source/Programs/Horde/HordeServer.Shared/Utilities/CustomFileCallbackResult.cs
+++ b/Engine/Source/Programs/Horde/HordeServer.Shared/Utilities/CustomFileCallbackResult.cs
@@ -1,9 +1,6 @@
 // Copyright Epic Games, Inc. All Rights Reserved.
 
-using System;
-using System.IO;
 using System.Net.Mime;
-using System.Threading.Tasks;
 using Microsoft.AspNetCore.Mvc;
 using Microsoft.AspNetCore.Mvc.Infrastructure;
 using Microsoft.Extensions.DependencyInjection;
diff --git a/Engine/Source/Programs/Horde/HordeServer.Shared/Utilities/DebugEndpointAttribute.cs b/Engine/Source/Programs/Horde/HordeServer.Shared/Utilities/DebugEndpointAttribute.cs
--- a/Engine/Source/Programs/Horde/HordeServer.Shared/Utilities/DebugEndpointAttribute.cs
+++ b/Engine/Source/Programs/Horde/HordeServer.Shared/Utilities/DebugEndpointAttribute.cs
@@ -1,10 +1,8 @@
 // Copyright Epic Games, Inc. All Rights Reserved.
 
-using System;
 using Microsoft.AspNetCore.Mvc;
 using Microsoft.AspNetCore.Mvc.Filters;
 using Microsoft.Extensions.DependencyInjection;
-using Microsoft.Extensions.Options;
 
 namespace HordeServer.Utilities
 {
diff --git a/Engine/Source/Programs/Horde/HordeServer.Shared/Utilities/DefaultAclBuilder.cs b/Engine/Source/Programs/Horde/HordeServer.Shared/Utilities/DefaultAclBuilder.cs
--- a/Engine/Source/Programs/Horde/HordeServer.Shared/Utilities/DefaultAclBuilder.cs
+++ b/Engine/Source/Programs/Horde/HordeServer.Shared/Utilities/DefaultAclBuilder.cs
@@ -1,10 +1,5 @@
 // Copyright Epic Games, Inc. All Rights Reserved.
 
-using System;
-using System.Collections.Generic;
-using System.Linq;
-using System.Text;
-using System.Threading.Tasks;
 using EpicGames.Horde.Acls;
 using HordeServer.Acls;
 
@@ -15,9 +10,9 @@ namespace HordeServer.Utilities
 	/// </summary>
 	public class DefaultAclBuilder
 	{
-		List<AclEntryConfig> _roles = new List<AclEntryConfig>();
-		HashSet<AclAction> _readActions = new HashSet<AclAction>();
-		HashSet<AclAction> _writeActions = new HashSet<AclAction>();
+		readonly List<AclEntryConfig> _roles = new List<AclEntryConfig>();
+		readonly HashSet<AclAction> _readActions = new HashSet<AclAction>();
+		readonly HashSet<AclAction> _writeActions = new HashSet<AclAction>();
 
 		/// <summary>
 		/// Adds a custom role with a certain set of entitlements
@@ -27,11 +22,6 @@ namespace HordeServer.Utilities
 		public void AddCustomRole(AclClaimConfig claim, AclAction[] actions)
 			=> _roles.Add(new AclEntryConfig(claim, actions));
 
-		public void AddCustomRole(AclClaimConfig uploadToolsClaim, object[] values)
-		{
-			throw new NotImplementedException();
-		}
-
 		/// <summary>
 		/// Adds a default read operation that users can perform
 		/// </summary>
diff --git a/Engine/Source/Programs/Horde/HordeServer.Shared/Utilities/HordeClaimTypes.cs b/Engine/Source/Programs/Horde/HordeServer.Shared/Utilities/HordeClaimTypes.cs
--- a/Engine/Source/Programs/Horde/HordeServer.Shared/Utilities/HordeClaimTypes.cs
+++ b/Engine/Source/Programs/Horde/HordeServer.Shared/Utilities/HordeClaimTypes.cs
@@ -130,9 +130,6 @@ namespace HordeServer.Utilities
 		/// <summary>
 		/// Determines whether the given user can masquerade as a given user
 		/// </summary>
-		/// <param name="user"></param>
-		/// <param name="userId"></param>
-		/// <returns></returns>
 		public static bool AuthorizeAsUser(this AclConfig aclConfig, ClaimsPrincipal user, UserId userId)
 		{
 			UserId? currentUserId = user.GetUserId();
diff --git a/Engine/Source/Programs/Horde/HordeServer.Shared/Utilities/InlineFileStreamResult.cs b/Engine/Source/Programs/Horde/HordeServer.Shared/Utilities/InlineFileStreamResult.cs
--- a/Engine/Source/Programs/Horde/HordeServer.Shared/Utilities/InlineFileStreamResult.cs
+++ b/Engine/Source/Programs/Horde/HordeServer.Shared/Utilities/InlineFileStreamResult.cs
@@ -1,7 +1,6 @@
 // Copyright Epic Games, Inc. All Rights Reserved.
 
 using System.Net.Mime;
-using System.Threading.Tasks;
 using Microsoft.AspNetCore.Mvc;
 
 namespace HordeServer.Utilities
diff --git a/Engine/Source/Programs/Horde/HordeServer.Shared/Utilities/JsonConfigFile.cs b/Engine/Source/Programs/Horde/HordeServer.Shared/Utilities/JsonConfigFile.cs
--- a/Engine/Source/Programs/Horde/HordeServer.Shared/Utilities/JsonConfigFile.cs
+++ b/Engine/Source/Programs/Horde/HordeServer.Shared/Utilities/JsonConfigFile.cs
@@ -1,11 +1,8 @@
 ﻿// Copyright Epic Games, Inc. All Rights Reserved.
 
-using System;
 using System.Buffers;
 using System.Text.Json;
 using System.Text.Json.Nodes;
-using System.Threading;
-using System.Threading.Tasks;
 using EpicGames.Core;
 
 namespace HordeServer.Utilities
@@ -17,13 +14,22 @@ namespace HordeServer.Utilities
 	/// </summary>
 	public class JsonConfigFile
 	{
+		/// <summary>
+		/// The root of the config file
+		/// </summary>
 		public JsonObject Root { get; }
 
+		/// <summary>
+		/// Constructor
+		/// </summary>
 		public JsonConfigFile(JsonObject root)
 		{
 			Root = root;
 		}
 
+		/// <summary>
+		/// Reads a json config file from a file on disk
+		/// </summary>
 		public static async Task<JsonConfigFile> ReadAsync(FileReference file, CancellationToken cancellationToken = default)
 		{
 			byte[] data = await FileReference.ReadAllBytesAsync(file, cancellationToken);
@@ -31,6 +37,9 @@ namespace HordeServer.Utilities
 			return new JsonConfigFile(obj ?? new JsonObject());
 		}
 
+		/// <summary>
+		/// Writes the config file back out to a file
+		/// </summary>
 		public async Task WriteAsync(FileReference file, CancellationToken cancellationToken = default)
 		{
 			ArrayBufferWriter<byte> buffer = new ArrayBufferWriter<byte>();
@@ -41,6 +50,9 @@ namespace HordeServer.Utilities
 			await FileReference.WriteAllBytesAsync(file, buffer.WrittenMemory.ToArray(), cancellationToken);
 		}
 
+		/// <summary>
+		/// Adds a node with the given name to an object
+		/// </summary>
 		public static T FindOrAddNode<T>(JsonObject obj, string name, Func<T> factory) where T : JsonNode
 		{
 			JsonNode? node = obj[name];
@@ -61,6 +73,9 @@ namespace HordeServer.Utilities
 			return newTypedNode;
 		}
 
+		/// <summary>
+		/// Finds an existing element, or adds a new element to an array, with the given name.
+		/// </summary>
 		public static JsonObject FindOrAddElementByKey(JsonArray array, string key, string name)
 		{
 			foreach (JsonNode? element in array)
diff --git a/Engine/Source/Programs/Horde/HordeServer.Shared/Utilities/JsonUtils.cs b/Engine/Source/Programs/Horde/HordeServer.Shared/Utilities/JsonUtils.cs
--- a/Engine/Source/Programs/Horde/HordeServer.Shared/Utilities/JsonUtils.cs
+++ b/Engine/Source/Programs/Horde/HordeServer.Shared/Utilities/JsonUtils.cs
@@ -1,13 +1,9 @@
 // Copyright Epic Games, Inc. All Rights Reserved.
 
-using System.Diagnostics;
-using System.Globalization;
 using System.Text.Json;
-using System.Text.Json.Serialization;
 using EpicGames.Core;
 using EpicGames.Horde;
 using EpicGames.Horde.Jobs;
-using MongoDB.Bson;
 
 namespace HordeServer.Utilities
 {
diff --git a/Engine/Source/Programs/Horde/HordeServer.Shared/Utilities/MongoBufferedWriter.cs b/Engine/Source/Programs/Horde/HordeServer.Shared/Utilities/MongoBufferedWriter.cs
--- a/Engine/Source/Programs/Horde/HordeServer.Shared/Utilities/MongoBufferedWriter.cs
+++ b/Engine/Source/Programs/Horde/HordeServer.Shared/Utilities/MongoBufferedWriter.cs
@@ -1,13 +1,9 @@
 // Copyright Epic Games, Inc. All Rights Reserved.
 
-using System;
 using System.Collections.Concurrent;
-using System.Collections.Generic;
-using System.Threading;
-using System.Threading.Tasks;
 using EpicGames.Core;
-using MongoDB.Driver;
 using Microsoft.Extensions.Logging;
+using MongoDB.Driver;
 
 namespace HordeServer.Utilities
 {
@@ -26,6 +22,9 @@ namespace HordeServer.Utilities
 		readonly AsyncEvent _flushEvent = new AsyncEvent();
 		readonly ILogger _logger;
 
+		/// <summary>
+		/// Constructor
+		/// </summary>
 		public MongoBufferedWriter(IMongoCollection<TDocument> collection, ILogger logger)
 			: this(collection, 50, TimeSpan.FromSeconds(5.0), logger)
 		{
diff --git a/Engine/Source/Programs/Horde/HordeServer.Shared/Utilities/MongoExtensions.cs b/Engine/Source/Programs/Horde/HordeServer.Shared/Utilities/MongoExtensions.cs
--- a/Engine/Source/Programs/Horde/HordeServer.Shared/Utilities/MongoExtensions.cs
+++ b/Engine/Source/Programs/Horde/HordeServer.Shared/Utilities/MongoExtensions.cs
@@ -1,13 +1,8 @@
 // Copyright Epic Games, Inc. All Rights Reserved.
 
-using System;
-using System.Collections.Generic;
-using System.Linq;
 using System.Linq.Expressions;
 using System.Runtime.CompilerServices;
 using System.Text.Json;
-using System.Threading;
-using System.Threading.Tasks;
 using MongoDB.Bson;
 using MongoDB.Bson.Serialization;
 using MongoDB.Driver;
diff --git a/Engine/Source/Programs/Horde/HordeServer.Shared/Utilities/ObjectIdConverter.cs b/Engine/Source/Programs/Horde/HordeServer.Shared/Utilities/ObjectIdConverter.cs
--- a/Engine/Source/Programs/Horde/HordeServer.Shared/Utilities/ObjectIdConverter.cs
+++ b/Engine/Source/Programs/Horde/HordeServer.Shared/Utilities/ObjectIdConverter.cs
@@ -1,6 +1,5 @@
 // Copyright Epic Games, Inc. All Rights Reserved.
 
-using System;
 using System.ComponentModel;
 using System.Globalization;
 using System.Reflection;
@@ -118,11 +117,11 @@ namespace HordeServer.Utilities
 		readonly TConverter _converter = new TConverter();
 
 		/// <inheritdoc/>
-		public override TValue Read(ref Utf8JsonReader reader, Type typeToConvert, JsonSerializerOptions options) 
+		public override TValue Read(ref Utf8JsonReader reader, Type typeToConvert, JsonSerializerOptions options)
 			=> _converter.FromObjectId(ObjectId.Parse(reader.GetString()));
 
 		/// <inheritdoc/>
-		public override void Write(Utf8JsonWriter writer, TValue value, JsonSerializerOptions options) 
+		public override void Write(Utf8JsonWriter writer, TValue value, JsonSerializerOptions options)
 			=> writer.WriteStringValue(_converter.ToObjectId(value).ToString());
 	}
 
@@ -154,7 +153,7 @@ namespace HordeServer.Utilities
 		readonly TConverter _converter = new TConverter();
 
 		/// <inheritdoc/>
-		public override TValue Deserialize(BsonDeserializationContext context, BsonDeserializationArgs args) 
+		public override TValue Deserialize(BsonDeserializationContext context, BsonDeserializationArgs args)
 			=> _converter.FromObjectId(context.Reader.ReadObjectId());
 
 		/// <inheritdoc/>
diff --git a/Engine/Source/Programs/Horde/HordeServer.Shared/Utilities/PropertyFilter.cs b/Engine/Source/Programs/Horde/HordeServer.Shared/Utilities/PropertyFilter.cs
--- a/Engine/Source/Programs/Horde/HordeServer.Shared/Utilities/PropertyFilter.cs
+++ b/Engine/Source/Programs/Horde/HordeServer.Shared/Utilities/PropertyFilter.cs
@@ -1,8 +1,6 @@
 // Copyright Epic Games, Inc. All Rights Reserved.
 
-using System;
 using System.Collections;
-using System.Collections.Generic;
 using System.ComponentModel;
 using System.Globalization;
 using System.Reflection;
diff --git a/Engine/Source/Programs/Horde/HordeServer.Shared/Utilities/RedisExtensions.cs b/Engine/Source/Programs/Horde/HordeServer.Shared/Utilities/RedisExtensions.cs
--- a/Engine/Source/Programs/Horde/HordeServer.Shared/Utilities/RedisExtensions.cs
+++ b/Engine/Source/Programs/Horde/HordeServer.Shared/Utilities/RedisExtensions.cs
@@ -1,6 +1,5 @@
 // Copyright Epic Games, Inc. All Rights Reserved.
 
-using System.Threading.Tasks;
 using StackExchange.Redis;
 
 namespace HordeServer.Utilities;
diff --git a/Engine/Source/Programs/Horde/HordeServer.Shared/Utilities/StringIdConverter.cs b/Engine/Source/Programs/Horde/HordeServer.Shared/Utilities/StringIdConverter.cs
--- a/Engine/Source/Programs/Horde/HordeServer.Shared/Utilities/StringIdConverter.cs
+++ b/Engine/Source/Programs/Horde/HordeServer.Shared/Utilities/StringIdConverter.cs
@@ -1,6 +1,5 @@
 // Copyright Epic Games, Inc. All Rights Reserved.
 
-using System;
 using System.Reflection;
 using EpicGames.Core;
 using EpicGames.Horde;
@@ -41,6 +40,7 @@ namespace HordeServer.Utilities
 		}
 	}
 
+#pragma warning disable CS1591
 	/// <summary>
 	/// Surrogate type for serializing StringId types to ProtoBuf
 	/// </summary>
@@ -55,4 +55,5 @@ namespace HordeServer.Utilities
 		public static implicit operator TValue(StringIdProto<TValue, TConverter> source) => s_converter.FromStringId(new StringId(source.Id!));
 		public static implicit operator StringIdProto<TValue, TConverter>(TValue source) => new StringIdProto<TValue, TConverter> { Id = s_converter.ToStringId(source).ToString() };
 	}
+#pragma warning restore CS1591
 }
diff --git a/Engine/Source/Programs/Horde/HordeServer.Shared/Utilities/TryAuthorizeAttribute.cs b/Engine/Source/Programs/Horde/HordeServer.Shared/Utilities/TryAuthorizeAttribute.cs
--- a/Engine/Source/Programs/Horde/HordeServer.Shared/Utilities/TryAuthorizeAttribute.cs
+++ b/Engine/Source/Programs/Horde/HordeServer.Shared/Utilities/TryAuthorizeAttribute.cs
@@ -1,7 +1,5 @@
 // Copyright Epic Games, Inc. All Rights Reserved.
 
-using System;
-using System.Threading.Tasks;
 using Microsoft.AspNetCore.Authorization;
 using Microsoft.AspNetCore.Authorization.Policy;
 using Microsoft.AspNetCore.Mvc.Filters;
diff --git a/Engine/Source/Programs/Horde/HordeServer.Tests.Shared/ControllerIntegrationTest.cs b/Engine/Source/Programs/Horde/HordeServer.Tests.Shared/ControllerIntegrationTest.cs
--- a/Engine/Source/Programs/Horde/HordeServer.Tests.Shared/ControllerIntegrationTest.cs
+++ b/Engine/Source/Programs/Horde/HordeServer.Tests.Shared/ControllerIntegrationTest.cs
@@ -1,24 +1,12 @@
 // Copyright Epic Games, Inc. All Rights Reserved.
 
-using System;
-using System.Collections.Generic;
-using System.Net.Http;
-using System.Threading.Tasks;
 using Amazon.CloudWatch;
-using HordeServer.Agents;
-using HordeServer.Configuration;
-using HordeServer.Jobs;
-using HordeServer.Jobs.Graphs;
-using HordeServer.Jobs.Templates;
-using HordeServer.Perforce;
-using HordeServer.Plugins;
 using HordeServer.Server;
 using Microsoft.AspNetCore.Hosting;
 using Microsoft.AspNetCore.Mvc.Testing;
 using Microsoft.AspNetCore.TestHost;
 using Microsoft.Extensions.Configuration;
 using Microsoft.Extensions.DependencyInjection;
-using Microsoft.Extensions.Options;
 using Moq;
 using Serilog;
 
@@ -84,12 +72,12 @@ public class TestWebApplicationFactory<TStartup> : WebApplicationFactory<TStartu
 		{
 			builder.ConfigureTestServices(_configureServices);
 		}
-/*		
-		=>
-		{
-			collection.AddSingleton<IPerforceService, PerforceServiceStub>();
-			collection.AddSingleton<IAmazonCloudWatch>(x => cloudWatchMock.Object);
-		});*/
+		/*		
+				=>
+				{
+					collection.AddSingleton<IPerforceService, PerforceServiceStub>();
+					collection.AddSingleton<IAmazonCloudWatch>(x => cloudWatchMock.Object);
+				});*/
 	}
 }
 
diff --git a/Engine/Source/Programs/Horde/HordeServer.Tests.Shared/DatabaseIntegrationTest.cs b/Engine/Source/Programs/Horde/HordeServer.Tests.Shared/DatabaseIntegrationTest.cs
--- a/Engine/Source/Programs/Horde/HordeServer.Tests.Shared/DatabaseIntegrationTest.cs
+++ b/Engine/Source/Programs/Horde/HordeServer.Tests.Shared/DatabaseIntegrationTest.cs
@@ -1,11 +1,7 @@
 // Copyright Epic Games, Inc. All Rights Reserved.
 
-using System;
-using System.Collections.Generic;
 using System.Net;
 using System.Runtime.InteropServices;
-using System.Threading;
-using System.Threading.Tasks;
 using HordeServer.Server;
 using HordeServer.Utilities;
 using Microsoft.AspNetCore.Mvc;
@@ -227,7 +223,7 @@ namespace HordeServer.Tests
 					ss.MongoDatabaseName = _mongoInstance.DatabaseName;
 					ss.MongoConnectionString = _mongoInstance.ConnectionString;
 
-					MongoCommandTracer mongoTracer = new (OpenTelemetryTracers.MongoDb, _loggerFactory.CreateLogger<MongoCommandTracer>());
+					MongoCommandTracer mongoTracer = new(OpenTelemetryTracers.MongoDb, _loggerFactory.CreateLogger<MongoCommandTracer>());
 					_mongoService = new MongoService(Options.Create(ss), redisService, mongoTracer, OpenTelemetryTracers.Horde, _loggerFactory.CreateLogger<MongoService>(), _loggerFactory);
 				}
 			}
diff --git a/Engine/Source/Programs/Horde/HordeServer.Tests.Shared/DatabaseRunner.cs b/Engine/Source/Programs/Horde/HordeServer.Tests.Shared/DatabaseRunner.cs
--- a/Engine/Source/Programs/Horde/HordeServer.Tests.Shared/DatabaseRunner.cs
+++ b/Engine/Source/Programs/Horde/HordeServer.Tests.Shared/DatabaseRunner.cs
@@ -1,15 +1,11 @@
 // Copyright Epic Games, Inc. All Rights Reserved.
 
-using System;
 using System.Diagnostics;
-using System.IO;
-using System.Linq;
 using System.Net;
 using System.Net.NetworkInformation;
 using System.Net.Sockets;
 using System.Reflection;
 using System.Runtime.InteropServices;
-using System.Threading.Tasks;
 using EpicGames.Core;
 
 namespace HordeServer.Tests;
diff --git a/Engine/Source/Programs/Horde/HordeServer.Tests.Shared/HordeServer.Tests.Shared.csproj b/Engine/Source/Programs/Horde/HordeServer.Tests.Shared/HordeServer.Tests.Shared.csproj
--- a/Engine/Source/Programs/Horde/HordeServer.Tests.Shared/HordeServer.Tests.Shared.csproj
+++ b/Engine/Source/Programs/Horde/HordeServer.Tests.Shared/HordeServer.Tests.Shared.csproj
@@ -5,6 +5,16 @@
     <ImplicitUsings>enable</ImplicitUsings>
     <Nullable>enable</Nullable>
     <RootNamespace>HordeServer.Tests</RootNamespace>
+    <RunAnalyzersDuringBuild>False</RunAnalyzersDuringBuild>
+    <Configurations>Debug;Release;Analyze</Configurations>
+    <NoWarn>1701;1702;1591</NoWarn>
+  </PropertyGroup>
+
+  <PropertyGroup Condition=" '$(Configuration)|$(Platform)' == 'Analyze|AnyCPU' ">
+    <RunAnalyzersDuringBuild>True</RunAnalyzersDuringBuild>
+    <TreatWarningsAsErrors>false</TreatWarningsAsErrors>
+    <EnforceCodeStyleInBuild>True</EnforceCodeStyleInBuild>
+    <GenerateDocumentationFile>True</GenerateDocumentationFile>
   </PropertyGroup>
 
   <ItemGroup>
diff --git a/Engine/Source/Programs/Horde/HordeServer.Tests.Shared/ServerTestSetup.cs b/Engine/Source/Programs/Horde/HordeServer.Tests.Shared/ServerTestSetup.cs
--- a/Engine/Source/Programs/Horde/HordeServer.Tests.Shared/ServerTestSetup.cs
+++ b/Engine/Source/Programs/Horde/HordeServer.Tests.Shared/ServerTestSetup.cs
@@ -1,12 +1,6 @@
 // Copyright Epic Games, Inc. All Rights Reserved.
 
-using System;
-using System.Collections.Generic;
 using System.Diagnostics.Metrics;
-using System.Linq;
-using System.Reflection;
-using System.Text;
-using System.Threading.Tasks;
 using EpicGames.Core;
 using HordeCommon;
 using HordeServer.Accounts;
diff --git a/Engine/Source/Programs/Horde/HordeServer.Tests/HordeServer.Tests.csproj b/Engine/Source/Programs/Horde/HordeServer.Tests/HordeServer.Tests.csproj
--- a/Engine/Source/Programs/Horde/HordeServer.Tests/HordeServer.Tests.csproj
+++ b/Engine/Source/Programs/Horde/HordeServer.Tests/HordeServer.Tests.csproj
@@ -15,7 +15,6 @@
     <TreatWarningsAsErrors>false</TreatWarningsAsErrors>
     <EnforceCodeStyleInBuild>True</EnforceCodeStyleInBuild>
     <GenerateDocumentationFile>True</GenerateDocumentationFile>
-    <DocumentationFile></DocumentationFile>
   </PropertyGroup>
 
   <ItemGroup>
diff --git a/Engine/Source/Programs/Horde/Plugins/Analytics/HordeServer.Analytics.Tests/AnalyticsTestExtensions.cs b/Engine/Source/Programs/Horde/Plugins/Analytics/HordeServer.Analytics.Tests/AnalyticsTestExtensions.cs
--- a/Engine/Source/Programs/Horde/Plugins/Analytics/HordeServer.Analytics.Tests/AnalyticsTestExtensions.cs
+++ b/Engine/Source/Programs/Horde/Plugins/Analytics/HordeServer.Analytics.Tests/AnalyticsTestExtensions.cs
@@ -1,15 +1,6 @@
 // Copyright Epic Games, Inc. All Rights Reserved.
 
-using System;
-using System.Collections.Generic;
-using System.Linq;
-using System.Text;
-using System.Threading.Tasks;
 using HordeServer.Plugins;
-using HordeServer.Telemetry;
-using HordeServer.Telemetry.Metrics;
-using HordeServer.Telemetry.Sinks;
-using Microsoft.Extensions.DependencyInjection;
 
 namespace HordeServer.Analytics.Tests
 {
diff --git a/Engine/Source/Programs/Horde/Plugins/Analytics/HordeServer.Analytics.Tests/HordeServer.Analytics.Tests.csproj b/Engine/Source/Programs/Horde/Plugins/Analytics/HordeServer.Analytics.Tests/HordeServer.Analytics.Tests.csproj
--- a/Engine/Source/Programs/Horde/Plugins/Analytics/HordeServer.Analytics.Tests/HordeServer.Analytics.Tests.csproj
+++ b/Engine/Source/Programs/Horde/Plugins/Analytics/HordeServer.Analytics.Tests/HordeServer.Analytics.Tests.csproj
@@ -4,9 +4,18 @@
     <TargetFramework>net8.0</TargetFramework>
     <ImplicitUsings>enable</ImplicitUsings>
     <Nullable>enable</Nullable>
-
     <IsPackable>false</IsPackable>
     <IsTestProject>true</IsTestProject>
+    <RunAnalyzersDuringBuild>False</RunAnalyzersDuringBuild>
+    <Configurations>Debug;Release;Analyze</Configurations>
+    <NoWarn>1701;1702;1591</NoWarn>
+  </PropertyGroup>
+
+  <PropertyGroup Condition=" '$(Configuration)|$(Platform)' == 'Analyze|AnyCPU' ">
+    <RunAnalyzersDuringBuild>True</RunAnalyzersDuringBuild>
+    <TreatWarningsAsErrors>false</TreatWarningsAsErrors>
+    <EnforceCodeStyleInBuild>True</EnforceCodeStyleInBuild>
+    <GenerateDocumentationFile>True</GenerateDocumentationFile>
   </PropertyGroup>
 
   <ItemGroup>
diff --git a/Engine/Source/Programs/Horde/Plugins/Analytics/HordeServer.Analytics.Tests/Telemetry/TelemetryTests.cs b/Engine/Source/Programs/Horde/Plugins/Analytics/HordeServer.Analytics.Tests/Telemetry/TelemetryTests.cs
--- a/Engine/Source/Programs/Horde/Plugins/Analytics/HordeServer.Analytics.Tests/Telemetry/TelemetryTests.cs
+++ b/Engine/Source/Programs/Horde/Plugins/Analytics/HordeServer.Analytics.Tests/Telemetry/TelemetryTests.cs
@@ -1,11 +1,6 @@
 // Copyright Epic Games, Inc. All Rights Reserved.
 
-using System;
-using System.Collections.Generic;
-using System.Linq;
 using System.Text.Json;
-using System.Threading;
-using System.Threading.Tasks;
 using EpicGames.Horde.Telemetry;
 using EpicGames.Horde.Telemetry.Metrics;
 using HordeServer.Server;
@@ -16,7 +11,6 @@ using HordeServer.Tests;
 using HordeServer.Utilities;
 using Json.Path;
 using Microsoft.Extensions.DependencyInjection;
-using Microsoft.VisualStudio.TestTools.UnitTesting;
 
 namespace HordeServer.Analytics.Tests.Telemetry
 {
diff --git a/Engine/Source/Programs/Horde/Plugins/Analytics/HordeServer.Analytics/AnalyticsConfig.cs b/Engine/Source/Programs/Horde/Plugins/Analytics/HordeServer.Analytics/AnalyticsConfig.cs
--- a/Engine/Source/Programs/Horde/Plugins/Analytics/HordeServer.Analytics/AnalyticsConfig.cs
+++ b/Engine/Source/Programs/Horde/Plugins/Analytics/HordeServer.Analytics/AnalyticsConfig.cs
@@ -1,12 +1,7 @@
 // Copyright Epic Games, Inc. All Rights Reserved.
 
-using System;
-using System.Collections.Generic;
 using System.Diagnostics.CodeAnalysis;
-using System.Linq;
 using System.Security.Claims;
-using System.Text;
-using System.Threading.Tasks;
 using EpicGames.Horde.Acls;
 using EpicGames.Horde.Telemetry;
 using HordeServer.Acls;
@@ -30,6 +25,7 @@ namespace HordeServer
 		private AclConfig _parentAcl = null!;
 		private readonly Dictionary<TelemetryStoreId, TelemetryStoreConfig> _telemetryStoreLookup = new Dictionary<TelemetryStoreId, TelemetryStoreConfig>();
 
+		/// <inheritdoc/>
 		public void PostLoad(PluginConfigOptions configOptions)
 		{
 			_parentAcl = configOptions.ParentAcl;
@@ -46,14 +42,13 @@ namespace HordeServer
 		public bool Authorize(AclAction action, ClaimsPrincipal user)
 			=> _parentAcl.Authorize(action, user);
 
-
 		/// <summary>
 		/// Attempts to get configuration for a pool from this object
 		/// </summary>
 		/// <param name="telemetryStoreId">The pool identifier</param>
 		/// <param name="config">Configuration for the telemetry store</param>
 		/// <returns>True if the telemetry configuration was found</returns>
-		public bool TryGetTelemetryStore(TelemetryStoreId telemetryStoreId, [NotNullWhen(true)] out TelemetryStoreConfig? config) 
+		public bool TryGetTelemetryStore(TelemetryStoreId telemetryStoreId, [NotNullWhen(true)] out TelemetryStoreConfig? config)
 			=> _telemetryStoreLookup.TryGetValue(telemetryStoreId, out config);
 	}
 }
diff --git a/Engine/Source/Programs/Horde/Plugins/Analytics/HordeServer.Analytics/AnalyticsPlugin.cs b/Engine/Source/Programs/Horde/Plugins/Analytics/HordeServer.Analytics/AnalyticsPlugin.cs
--- a/Engine/Source/Programs/Horde/Plugins/Analytics/HordeServer.Analytics/AnalyticsPlugin.cs
+++ b/Engine/Source/Programs/Horde/Plugins/Analytics/HordeServer.Analytics/AnalyticsPlugin.cs
@@ -1,16 +1,10 @@
 ﻿// Copyright Epic Games, Inc. All Rights Reserved.
 
-using System;
-using System.Collections.Generic;
-using System.Linq;
-using System.Text;
-using System.Threading.Tasks;
 using HordeServer.Plugins;
 using HordeServer.Telemetry;
 using HordeServer.Telemetry.Metrics;
 using HordeServer.Telemetry.Sinks;
 using Microsoft.AspNetCore.Builder;
-using Microsoft.Extensions.Configuration;
 using Microsoft.Extensions.DependencyInjection;
 
 namespace HordeServer
diff --git a/Engine/Source/Programs/Horde/Plugins/Analytics/HordeServer.Analytics/HordeServer.Analytics.csproj b/Engine/Source/Programs/Horde/Plugins/Analytics/HordeServer.Analytics/HordeServer.Analytics.csproj
--- a/Engine/Source/Programs/Horde/Plugins/Analytics/HordeServer.Analytics/HordeServer.Analytics.csproj
+++ b/Engine/Source/Programs/Horde/Plugins/Analytics/HordeServer.Analytics/HordeServer.Analytics.csproj
@@ -5,6 +5,16 @@
     <ImplicitUsings>enable</ImplicitUsings>
     <Nullable>enable</Nullable>
     <RootNamespace>HordeServer</RootNamespace>
+    <RunAnalyzersDuringBuild>False</RunAnalyzersDuringBuild>
+    <Configurations>Debug;Release;Analyze</Configurations>
+  </PropertyGroup>
+
+  <PropertyGroup Condition=" '$(Configuration)|$(Platform)' == 'Analyze|AnyCPU' ">
+    <RunAnalyzersDuringBuild>True</RunAnalyzersDuringBuild>
+    <TreatWarningsAsErrors>false</TreatWarningsAsErrors>
+    <EnforceCodeStyleInBuild>True</EnforceCodeStyleInBuild>
+    <GenerateDocumentationFile>True</GenerateDocumentationFile>
+    <DocumentationFile></DocumentationFile>
   </PropertyGroup>
 
   <ItemGroup>
diff --git a/Engine/Source/Programs/Horde/Plugins/Analytics/HordeServer.Analytics/Telemetry/Metrics/MetricCollection.cs b/Engine/Source/Programs/Horde/Plugins/Analytics/HordeServer.Analytics/Telemetry/Metrics/MetricCollection.cs
--- a/Engine/Source/Programs/Horde/Plugins/Analytics/HordeServer.Analytics/Telemetry/Metrics/MetricCollection.cs
+++ b/Engine/Source/Programs/Horde/Plugins/Analytics/HordeServer.Analytics/Telemetry/Metrics/MetricCollection.cs
@@ -1,11 +1,6 @@
 // Copyright Epic Games, Inc. All Rights Reserved.
 
-using System;
-using System.Collections.Generic;
-using System.Linq;
 using System.Text.Json.Nodes;
-using System.Threading;
-using System.Threading.Tasks;
 using EpicGames.Core;
 using EpicGames.Horde.Telemetry;
 using EpicGames.Horde.Telemetry.Metrics;
diff --git a/Engine/Source/Programs/Horde/Plugins/Analytics/HordeServer.Analytics/Telemetry/Sinks/EpicTelemetrySink.cs b/Engine/Source/Programs/Horde/Plugins/Analytics/HordeServer.Analytics/Telemetry/Sinks/EpicTelemetrySink.cs
--- a/Engine/Source/Programs/Horde/Plugins/Analytics/HordeServer.Analytics/Telemetry/Sinks/EpicTelemetrySink.cs
+++ b/Engine/Source/Programs/Horde/Plugins/Analytics/HordeServer.Analytics/Telemetry/Sinks/EpicTelemetrySink.cs
@@ -1,14 +1,9 @@
 // Copyright Epic Games, Inc. All Rights Reserved.
 
-using System;
-using System.Collections.Generic;
 using System.Diagnostics;
-using System.Net.Http;
 using System.Net.Http.Headers;
 using System.Text.Json;
 using System.Text.Json.Serialization;
-using System.Threading;
-using System.Threading.Tasks;
 using EpicGames.Core;
 using EpicGames.Horde.Telemetry;
 using HordeServer.Utilities;
diff --git a/Engine/Source/Programs/Horde/Plugins/Analytics/HordeServer.Analytics/Telemetry/Sinks/MetricTelemetrySink.cs b/Engine/Source/Programs/Horde/Plugins/Analytics/HordeServer.Analytics/Telemetry/Sinks/MetricTelemetrySink.cs
--- a/Engine/Source/Programs/Horde/Plugins/Analytics/HordeServer.Analytics/Telemetry/Sinks/MetricTelemetrySink.cs
+++ b/Engine/Source/Programs/Horde/Plugins/Analytics/HordeServer.Analytics/Telemetry/Sinks/MetricTelemetrySink.cs
@@ -2,8 +2,6 @@
 
 using System.Text.Json;
 using System.Text.Json.Nodes;
-using System.Threading;
-using System.Threading.Tasks;
 using EpicGames.Horde.Telemetry;
 using HordeServer.Telemetry.Metrics;
 using HordeServer.Utilities;
diff --git a/Engine/Source/Programs/Horde/Plugins/Analytics/HordeServer.Analytics/Telemetry/Sinks/MongoTelemetrySink.cs b/Engine/Source/Programs/Horde/Plugins/Analytics/HordeServer.Analytics/Telemetry/Sinks/MongoTelemetrySink.cs
--- a/Engine/Source/Programs/Horde/Plugins/Analytics/HordeServer.Analytics/Telemetry/Sinks/MongoTelemetrySink.cs
+++ b/Engine/Source/Programs/Horde/Plugins/Analytics/HordeServer.Analytics/Telemetry/Sinks/MongoTelemetrySink.cs
@@ -1,10 +1,6 @@
 // Copyright Epic Games, Inc. All Rights Reserved.
 
-using System;
-using System.Linq;
 using System.Text.Json;
-using System.Threading;
-using System.Threading.Tasks;
 using EpicGames.Core;
 using EpicGames.Horde.Telemetry;
 using HordeServer.Server;
diff --git a/Engine/Source/Programs/Horde/Plugins/Analytics/HordeServer.Analytics/Telemetry/Sinks/NullTelemetrySink.cs b/Engine/Source/Programs/Horde/Plugins/Analytics/HordeServer.Analytics/Telemetry/Sinks/NullTelemetrySink.cs
--- a/Engine/Source/Programs/Horde/Plugins/Analytics/HordeServer.Analytics/Telemetry/Sinks/NullTelemetrySink.cs
+++ b/Engine/Source/Programs/Horde/Plugins/Analytics/HordeServer.Analytics/Telemetry/Sinks/NullTelemetrySink.cs
@@ -1,7 +1,5 @@
 // Copyright Epic Games, Inc. All Rights Reserved.
 
-using System.Threading;
-using System.Threading.Tasks;
 using EpicGames.Horde.Telemetry;
 
 namespace HordeServer.Telemetry.Sinks
diff --git a/Engine/Source/Programs/Horde/Plugins/Analytics/HordeServer.Analytics/Telemetry/TelemetryController.cs b/Engine/Source/Programs/Horde/Plugins/Analytics/HordeServer.Analytics/Telemetry/TelemetryController.cs
--- a/Engine/Source/Programs/Horde/Plugins/Analytics/HordeServer.Analytics/Telemetry/TelemetryController.cs
+++ b/Engine/Source/Programs/Horde/Plugins/Analytics/HordeServer.Analytics/Telemetry/TelemetryController.cs
@@ -1,15 +1,9 @@
 // Copyright Epic Games, Inc. All Rights Reserved.
 
-using System;
-using System.Collections.Generic;
 using System.ComponentModel.DataAnnotations;
-using System.Linq;
 using System.Text.Json.Nodes;
-using System.Threading;
-using System.Threading.Tasks;
 using EpicGames.Horde.Telemetry;
 using EpicGames.Horde.Telemetry.Metrics;
-using HordeServer.Server;
 using HordeServer.Telemetry.Metrics;
 using HordeServer.Utilities;
 using Microsoft.AspNetCore.Authorization;
diff --git a/Engine/Source/Programs/Horde/Plugins/Build/HordeServer.Build.Tests/Acls/AclTests.cs b/Engine/Source/Programs/Horde/Plugins/Build/HordeServer.Build.Tests/Acls/AclTests.cs
--- a/Engine/Source/Programs/Horde/Plugins/Build/HordeServer.Build.Tests/Acls/AclTests.cs
+++ b/Engine/Source/Programs/Horde/Plugins/Build/HordeServer.Build.Tests/Acls/AclTests.cs
@@ -1,6 +1,5 @@
 // Copyright Epic Games, Inc. All Rights Reserved.
 
-using System.Collections.Generic;
 using System.Security.Claims;
 using EpicGames.Horde.Acls;
 using EpicGames.Horde.Projects;
@@ -11,7 +10,6 @@ using HordeServer.Logs;
 using HordeServer.Projects;
 using HordeServer.Server;
 using HordeServer.Streams;
-using Microsoft.VisualStudio.TestTools.UnitTesting;
 
 namespace HordeServer.Tests.Acls
 {
diff --git a/Engine/Source/Programs/Horde/Plugins/Build/HordeServer.Build.Tests/Agents/AgentCollectionTests.cs b/Engine/Source/Programs/Horde/Plugins/Build/HordeServer.Build.Tests/Agents/AgentCollectionTests.cs
--- a/Engine/Source/Programs/Horde/Plugins/Build/HordeServer.Build.Tests/Agents/AgentCollectionTests.cs
+++ b/Engine/Source/Programs/Horde/Plugins/Build/HordeServer.Build.Tests/Agents/AgentCollectionTests.cs
@@ -1,14 +1,10 @@
 // Copyright Epic Games, Inc. All Rights Reserved.
 
-using System;
-using System.Collections.Generic;
-using System.Threading.Tasks;
 using EpicGames.Horde.Agents;
 using EpicGames.Horde.Agents.Leases;
 using EpicGames.Horde.Agents.Pools;
 using HordeServer.Agents;
 using HordeServer.Agents.Sessions;
-using Microsoft.VisualStudio.TestTools.UnitTesting;
 
 namespace HordeServer.Tests.Agents;
 
@@ -107,7 +103,7 @@ public class AgentCollectionTests : BuildTestSetup
 		await _agent.TryAddLeaseAsync(_lease2);
 		await UpdateAgentAsync();
 
-		await _agent.TryUpdateSessionAsync(new UpdateSessionOptions{ Leases = new List<AgentLease> { _lease1 } });
+		await _agent.TryUpdateSessionAsync(new UpdateSessionOptions { Leases = new List<AgentLease> { _lease1 } });
 
 		List<LeaseId> leases = await AgentCollection.FindActiveLeaseIdsAsync();
 
diff --git a/Engine/Source/Programs/Horde/Plugins/Build/HordeServer.Build.Tests/Agents/AgentControllerTest.cs b/Engine/Source/Programs/Horde/Plugins/Build/HordeServer.Build.Tests/Agents/AgentControllerTest.cs
--- a/Engine/Source/Programs/Horde/Plugins/Build/HordeServer.Build.Tests/Agents/AgentControllerTest.cs
+++ b/Engine/Source/Programs/Horde/Plugins/Build/HordeServer.Build.Tests/Agents/AgentControllerTest.cs
@@ -1,10 +1,8 @@
 // Copyright Epic Games, Inc. All Rights Reserved.
 
-using System.Threading.Tasks;
 using EpicGames.Horde.Agents;
 using HordeServer.Agents;
 using Microsoft.AspNetCore.Mvc;
-using Microsoft.VisualStudio.TestTools.UnitTesting;
 
 namespace HordeServer.Tests.Agents;
 
diff --git a/Engine/Source/Programs/Horde/Plugins/Build/HordeServer.Build.Tests/Agents/AgentServiceTest.cs b/Engine/Source/Programs/Horde/Plugins/Build/HordeServer.Build.Tests/Agents/AgentServiceTest.cs
--- a/Engine/Source/Programs/Horde/Plugins/Build/HordeServer.Build.Tests/Agents/AgentServiceTest.cs
+++ b/Engine/Source/Programs/Horde/Plugins/Build/HordeServer.Build.Tests/Agents/AgentServiceTest.cs
@@ -1,11 +1,6 @@
 // Copyright Epic Games, Inc. All Rights Reserved.
 
-using System;
-using System.Collections.Generic;
-using System.Linq;
 using System.Security.Claims;
-using System.Threading;
-using System.Threading.Tasks;
 using EpicGames.Core;
 using EpicGames.Horde.Agents;
 using EpicGames.Horde.Agents.Pools;
@@ -16,7 +11,6 @@ using HordeServer.Auditing;
 using HordeServer.Jobs;
 using HordeServer.Utilities;
 using Microsoft.AspNetCore.Mvc;
-using Microsoft.VisualStudio.TestTools.UnitTesting;
 
 namespace HordeServer.Tests.Agents;
 
diff --git a/Engine/Source/Programs/Horde/Plugins/Build/HordeServer.Build.Tests/Agents/Leases/LeasesControllerDbTest.cs b/Engine/Source/Programs/Horde/Plugins/Build/HordeServer.Build.Tests/Agents/Leases/LeasesControllerDbTest.cs
--- a/Engine/Source/Programs/Horde/Plugins/Build/HordeServer.Build.Tests/Agents/Leases/LeasesControllerDbTest.cs
+++ b/Engine/Source/Programs/Horde/Plugins/Build/HordeServer.Build.Tests/Agents/Leases/LeasesControllerDbTest.cs
@@ -1,8 +1,5 @@
 // Copyright Epic Games, Inc. All Rights Reserved.
 
-using System;
-using System.Collections.Generic;
-using System.Threading.Tasks;
 using EpicGames.Horde.Agents;
 using EpicGames.Horde.Agents.Leases;
 using EpicGames.Horde.Agents.Sessions;
@@ -10,7 +7,6 @@ using HordeServer.Agents.Leases;
 using HordeServer.Agents.Sessions;
 using HordeServer.Utilities;
 using Microsoft.AspNetCore.Mvc;
-using Microsoft.VisualStudio.TestTools.UnitTesting;
 
 namespace HordeServer.Tests.Agents.Leases
 {
diff --git a/Engine/Source/Programs/Horde/Plugins/Build/HordeServer.Build.Tests/Agents/Pools/PoolUpdateServiceTest.cs b/Engine/Source/Programs/Horde/Plugins/Build/HordeServer.Build.Tests/Agents/Pools/PoolUpdateServiceTest.cs
--- a/Engine/Source/Programs/Horde/Plugins/Build/HordeServer.Build.Tests/Agents/Pools/PoolUpdateServiceTest.cs
+++ b/Engine/Source/Programs/Horde/Plugins/Build/HordeServer.Build.Tests/Agents/Pools/PoolUpdateServiceTest.cs
@@ -1,17 +1,11 @@
 // Copyright Epic Games, Inc. All Rights Reserved.
 
-using System;
-using System.Collections.Generic;
-using System.Linq;
-using System.Threading;
-using System.Threading.Tasks;
 using EpicGames.Horde.Agents;
 using HordeServer.Agents;
 using HordeServer.Agents.Pools;
 using Microsoft.Extensions.DependencyInjection;
 using Microsoft.Extensions.Logging.Abstractions;
 using Microsoft.Extensions.Options;
-using Microsoft.VisualStudio.TestTools.UnitTesting;
 
 namespace HordeServer.Tests.Agents.Pools;
 
@@ -66,13 +60,13 @@ public class PoolUpdateServiceTest : BuildTestSetup
 		Assert.IsFalse(_disabledAgent.RequestShutdown);
 		Assert.IsFalse(_disabledAgentBeyondGracePeriod.RequestShutdown);
 	}
-	
+
 	[TestMethod]
 	public async Task ShutdownDisabledAgents_WithGlobalGracePeriod_RequestsShutdownAsync()
 	{
 		// Arrange
 		UpdateConfig(config => config.Plugins.GetBuildConfig().AgentShutdownIfDisabledGracePeriod = TimeSpan.FromHours(3.0));
-		
+
 		// Act
 		await _pus.ShutdownDisabledAgentsAsync(CancellationToken.None);
 		await RefreshAgentsAsync();
@@ -115,7 +109,7 @@ public class PoolUpdateServiceTest : BuildTestSetup
 		Assert.IsFalse(_disabledAgent.RequestShutdown);
 		Assert.IsFalse(_disabledAgentBeyondGracePeriod.RequestShutdown);
 	}
-	
+
 	[TestMethod]
 	[DataRow(false, 50)]
 	[DataRow(false, null)]
@@ -129,11 +123,11 @@ public class PoolUpdateServiceTest : BuildTestSetup
 	{
 		// Arrange
 		IAgent agent = await CreateAutoConformAgentAsync(100, autoConformThresholdsM);
-		
+
 		// Act
 		await _pus.AutoConformAgentsAsync(CancellationToken.None);
 		agent = (await AgentService.GetAgentAsync(agent.Id))!;
-		
+
 		// Assert
 		Assert.AreEqual(conformRequested, agent.RequestFullConform);
 	}
diff --git a/Engine/Source/Programs/Horde/Plugins/Build/HordeServer.Build.Tests/Agents/Pools/PoolsControllerTest.cs b/Engine/Source/Programs/Horde/Plugins/Build/HordeServer.Build.Tests/Agents/Pools/PoolsControllerTest.cs
--- a/Engine/Source/Programs/Horde/Plugins/Build/HordeServer.Build.Tests/Agents/Pools/PoolsControllerTest.cs
+++ b/Engine/Source/Programs/Horde/Plugins/Build/HordeServer.Build.Tests/Agents/Pools/PoolsControllerTest.cs
@@ -1,12 +1,8 @@
 ﻿// Copyright Epic Games, Inc. All Rights Reserved.
 
-using System;
-using System.Collections.Generic;
-using System.Threading.Tasks;
 using HordeServer.Agents.Fleet;
 using HordeServer.Agents.Pools;
 using Microsoft.AspNetCore.Mvc;
-using Microsoft.VisualStudio.TestTools.UnitTesting;
 
 namespace HordeServer.Tests.Agents.Pools
 {
diff --git a/Engine/Source/Programs/Horde/Plugins/Build/HordeServer.Build.Tests/Agents/Registration/RegistrationServiceTests.cs b/Engine/Source/Programs/Horde/Plugins/Build/HordeServer.Build.Tests/Agents/Registration/RegistrationServiceTests.cs
--- a/Engine/Source/Programs/Horde/Plugins/Build/HordeServer.Build.Tests/Agents/Registration/RegistrationServiceTests.cs
+++ b/Engine/Source/Programs/Horde/Plugins/Build/HordeServer.Build.Tests/Agents/Registration/RegistrationServiceTests.cs
@@ -1,14 +1,9 @@
 // Copyright Epic Games, Inc. All Rights Reserved.
 
-using System;
-using System.Collections.Generic;
 using System.Security.Cryptography;
-using System.Threading;
-using System.Threading.Tasks;
 using EpicGames.Horde.Agents;
 using HordeServer.Agents.Enrollment;
 using Microsoft.Extensions.DependencyInjection;
-using Microsoft.VisualStudio.TestTools.UnitTesting;
 
 namespace HordeServer.Tests.Agents.Registration
 {
diff --git a/Engine/Source/Programs/Horde/Plugins/Build/HordeServer.Build.Tests/Agents/Relay/AgentRelayClientTest.cs b/Engine/Source/Programs/Horde/Plugins/Build/HordeServer.Build.Tests/Agents/Relay/AgentRelayClientTest.cs
--- a/Engine/Source/Programs/Horde/Plugins/Build/HordeServer.Build.Tests/Agents/Relay/AgentRelayClientTest.cs
+++ b/Engine/Source/Programs/Horde/Plugins/Build/HordeServer.Build.Tests/Agents/Relay/AgentRelayClientTest.cs
@@ -1,16 +1,11 @@
 // Copyright Epic Games, Inc. All Rights Reserved.
 
-using System;
-using System.Collections.Generic;
-using System.Threading;
-using System.Threading.Tasks;
 using Grpc.Core;
 using Horde.Common.Rpc;
 using HordeServer.Agents.Relay;
 //using HordeServer.Tests.Server;
 using HordeServer.Utilities;
 using Microsoft.Extensions.Logging.Abstractions;
-using Microsoft.VisualStudio.TestTools.UnitTesting;
 
 namespace HordeServer.Tests.Agents.Relay;
 
diff --git a/Engine/Source/Programs/Horde/Plugins/Build/HordeServer.Build.Tests/Agents/Relay/NftablesTests.cs b/Engine/Source/Programs/Horde/Plugins/Build/HordeServer.Build.Tests/Agents/Relay/NftablesTests.cs
--- a/Engine/Source/Programs/Horde/Plugins/Build/HordeServer.Build.Tests/Agents/Relay/NftablesTests.cs
+++ b/Engine/Source/Programs/Horde/Plugins/Build/HordeServer.Build.Tests/Agents/Relay/NftablesTests.cs
@@ -1,9 +1,7 @@
 // Copyright Epic Games, Inc. All Rights Reserved.
 
-using System.Collections.Generic;
 using Horde.Common.Rpc;
 using HordeServer.Agents.Relay;
-using Microsoft.VisualStudio.TestTools.UnitTesting;
 
 namespace HordeServer.Tests.Agents.Relay;
 
diff --git a/Engine/Source/Programs/Horde/Plugins/Build/HordeServer.Build.Tests/Artifacts/BlockCacheTests.cs b/Engine/Source/Programs/Horde/Plugins/Build/HordeServer.Build.Tests/Artifacts/BlockCacheTests.cs
--- a/Engine/Source/Programs/Horde/Plugins/Build/HordeServer.Build.Tests/Artifacts/BlockCacheTests.cs
+++ b/Engine/Source/Programs/Horde/Plugins/Build/HordeServer.Build.Tests/Artifacts/BlockCacheTests.cs
@@ -1,14 +1,9 @@
 // Copyright Epic Games, Inc. All Rights Reserved.
 
-using System;
 using System.Buffers;
-using System.Collections.Generic;
-using System.Linq;
 using System.Security.Cryptography;
-using System.Threading.Tasks;
 using HordeServer.Artifacts;
 using Microsoft.Extensions.Logging;
-using Microsoft.VisualStudio.TestTools.UnitTesting;
 
 namespace HordeServer.Tests.Artifacts
 {
diff --git a/Engine/Source/Programs/Horde/Plugins/Build/HordeServer.Build.Tests/BuildTestSetup.cs b/Engine/Source/Programs/Horde/Plugins/Build/HordeServer.Build.Tests/BuildTestSetup.cs
--- a/Engine/Source/Programs/Horde/Plugins/Build/HordeServer.Build.Tests/BuildTestSetup.cs
+++ b/Engine/Source/Programs/Horde/Plugins/Build/HordeServer.Build.Tests/BuildTestSetup.cs
@@ -1,17 +1,11 @@
 // Copyright Epic Games, Inc. All Rights Reserved.
 
-using System;
-using System.Collections.Generic;
-using System.Net;
-using System.Net.Sockets;
 using System.Security.Claims;
-using System.Threading.Tasks;
 using Amazon.AutoScaling;
 using Amazon.CloudWatch;
 using Amazon.EC2;
 using EpicGames.Horde.Agents;
 using EpicGames.Horde.Agents.Pools;
-using HordeServer.Accounts;
 using HordeServer.Acls;
 using HordeServer.Agents;
 using HordeServer.Agents.Enrollment;
@@ -52,8 +46,6 @@ using Microsoft.AspNetCore.Http;
 using Microsoft.AspNetCore.Mvc;
 using Microsoft.Extensions.Configuration;
 using Microsoft.Extensions.DependencyInjection;
-using Microsoft.Extensions.Logging;
-using Microsoft.VisualStudio.TestTools.UnitTesting;
 using Moq;
 
 namespace HordeServer.Tests
diff --git a/Engine/Source/Programs/Horde/Plugins/Build/HordeServer.Build.Tests/Dashboard/DashboardTests.cs b/Engine/Source/Programs/Horde/Plugins/Build/HordeServer.Build.Tests/Dashboard/DashboardTests.cs
--- a/Engine/Source/Programs/Horde/Plugins/Build/HordeServer.Build.Tests/Dashboard/DashboardTests.cs
+++ b/Engine/Source/Programs/Horde/Plugins/Build/HordeServer.Build.Tests/Dashboard/DashboardTests.cs
@@ -1,10 +1,6 @@
 // Copyright Epic Games, Inc. All Rights Reserved.
 
-using System.Collections.Generic;
-using System.Linq;
-using System.Threading.Tasks;
 using EpicGames.Horde.Dashboard;
-using Microsoft.VisualStudio.TestTools.UnitTesting;
 
 namespace HordeServer.Tests.Dashboard
 {
diff --git a/Engine/Source/Programs/Horde/Plugins/Build/HordeServer.Build.Tests/Devices/DeviceServiceTests.cs b/Engine/Source/Programs/Horde/Plugins/Build/HordeServer.Build.Tests/Devices/DeviceServiceTests.cs
--- a/Engine/Source/Programs/Horde/Plugins/Build/HordeServer.Build.Tests/Devices/DeviceServiceTests.cs
+++ b/Engine/Source/Programs/Horde/Plugins/Build/HordeServer.Build.Tests/Devices/DeviceServiceTests.cs
@@ -1,10 +1,6 @@
 // Copyright Epic Games, Inc. All Rights Reserved.
 
-using System;
-using System.Collections.Generic;
-using System.Linq;
 using System.Security.Claims;
-using System.Threading.Tasks;
 using EpicGames.Core;
 using EpicGames.Horde.Devices;
 using EpicGames.Horde.Jobs;
@@ -19,7 +15,6 @@ using HordeServer.Users;
 using HordeServer.Utilities;
 using Microsoft.AspNetCore.Http;
 using Microsoft.AspNetCore.Mvc;
-using Microsoft.VisualStudio.TestTools.UnitTesting;
 using Moq;
 
 namespace HordeServer.Tests.Devices
diff --git a/Engine/Source/Programs/Horde/Plugins/Build/HordeServer.Build.Tests/Fixture.cs b/Engine/Source/Programs/Horde/Plugins/Build/HordeServer.Build.Tests/Fixture.cs
--- a/Engine/Source/Programs/Horde/Plugins/Build/HordeServer.Build.Tests/Fixture.cs
+++ b/Engine/Source/Programs/Horde/Plugins/Build/HordeServer.Build.Tests/Fixture.cs
@@ -1,7 +1,5 @@
 // Copyright Epic Games, Inc. All Rights Reserved.
 
-using System.Collections.Generic;
-using System.Threading.Tasks;
 using EpicGames.Core;
 using EpicGames.Horde.Agents.Pools;
 using EpicGames.Horde.Jobs;
diff --git a/Engine/Source/Programs/Horde/Plugins/Build/HordeServer.Build.Tests/Fleet/AwsRecyclingFleetManagerTest.cs b/Engine/Source/Programs/Horde/Plugins/Build/HordeServer.Build.Tests/Fleet/AwsRecyclingFleetManagerTest.cs
--- a/Engine/Source/Programs/Horde/Plugins/Build/HordeServer.Build.Tests/Fleet/AwsRecyclingFleetManagerTest.cs
+++ b/Engine/Source/Programs/Horde/Plugins/Build/HordeServer.Build.Tests/Fleet/AwsRecyclingFleetManagerTest.cs
@@ -1,19 +1,13 @@
 // Copyright Epic Games, Inc. All Rights Reserved.
 
-using System;
-using System.Collections.Generic;
-using System.Linq;
-using System.Threading;
-using System.Threading.Tasks;
 using Amazon.EC2;
 using Amazon.EC2.Model;
+using HordeCommon;
 using HordeServer.Agents;
 using HordeServer.Agents.Fleet;
 using HordeServer.Agents.Fleet.Providers;
 using HordeServer.Agents.Pools;
-using HordeCommon;
 using Microsoft.Extensions.Logging;
-using Microsoft.VisualStudio.TestTools.UnitTesting;
 using JsonSerializer = System.Text.Json.JsonSerializer;
 
 namespace HordeServer.Tests.Fleet
diff --git a/Engine/Source/Programs/Horde/Plugins/Build/HordeServer.Build.Tests/Fleet/AwsReuseFleetManagerTest.cs b/Engine/Source/Programs/Horde/Plugins/Build/HordeServer.Build.Tests/Fleet/AwsReuseFleetManagerTest.cs
--- a/Engine/Source/Programs/Horde/Plugins/Build/HordeServer.Build.Tests/Fleet/AwsReuseFleetManagerTest.cs
+++ b/Engine/Source/Programs/Horde/Plugins/Build/HordeServer.Build.Tests/Fleet/AwsReuseFleetManagerTest.cs
@@ -1,8 +1,5 @@
 // Copyright Epic Games, Inc. All Rights Reserved.
 
-using System.Collections.Generic;
-using System.Threading;
-using System.Threading.Tasks;
 using Amazon.EC2;
 using Amazon.EC2.Model;
 using HordeServer.Agents;
@@ -11,7 +8,6 @@ using HordeServer.Agents.Fleet.Providers;
 using HordeServer.Agents.Pools;
 using Microsoft.Extensions.DependencyInjection;
 using Microsoft.Extensions.Logging;
-using Microsoft.VisualStudio.TestTools.UnitTesting;
 
 namespace HordeServer.Tests.Fleet
 {
diff --git a/Engine/Source/Programs/Horde/Plugins/Build/HordeServer.Build.Tests/Fleet/FakeAmazonCloudWatch.cs b/Engine/Source/Programs/Horde/Plugins/Build/HordeServer.Build.Tests/Fleet/FakeAmazonCloudWatch.cs
--- a/Engine/Source/Programs/Horde/Plugins/Build/HordeServer.Build.Tests/Fleet/FakeAmazonCloudWatch.cs
+++ b/Engine/Source/Programs/Horde/Plugins/Build/HordeServer.Build.Tests/Fleet/FakeAmazonCloudWatch.cs
@@ -1,9 +1,6 @@
 // Copyright Epic Games, Inc. All Rights Reserved.
 
-using System.Collections.Generic;
 using System.Net;
-using System.Threading;
-using System.Threading.Tasks;
 using Amazon.CloudWatch;
 using Amazon.CloudWatch.Model;
 using Moq;
diff --git a/Engine/Source/Programs/Horde/Plugins/Build/HordeServer.Build.Tests/Fleet/LeaseUtilizationAwsMetricStrategyTests.cs b/Engine/Source/Programs/Horde/Plugins/Build/HordeServer.Build.Tests/Fleet/LeaseUtilizationAwsMetricStrategyTests.cs
--- a/Engine/Source/Programs/Horde/Plugins/Build/HordeServer.Build.Tests/Fleet/LeaseUtilizationAwsMetricStrategyTests.cs
+++ b/Engine/Source/Programs/Horde/Plugins/Build/HordeServer.Build.Tests/Fleet/LeaseUtilizationAwsMetricStrategyTests.cs
@@ -1,23 +1,18 @@
 // Copyright Epic Games, Inc. All Rights Reserved.
 
-using System;
-using System.Collections.Generic;
-using System.Linq;
-using System.Threading.Tasks;
 using Amazon.CloudWatch.Model;
 using EpicGames.Horde.Agents.Leases;
 using EpicGames.Horde.Agents.Pools;
 using EpicGames.Horde.Streams;
 using Google.Protobuf;
 using Google.Protobuf.WellKnownTypes;
+using HordeCommon.Rpc.Tasks;
 using HordeServer.Agents;
 using HordeServer.Agents.Fleet;
 using HordeServer.Agents.Leases;
 using HordeServer.Agents.Pools;
 using HordeServer.Utilities;
-using HordeCommon.Rpc.Tasks;
 using Microsoft.Extensions.Logging.Abstractions;
-using Microsoft.VisualStudio.TestTools.UnitTesting;
 
 namespace HordeServer.Tests.Fleet
 {
diff --git a/Engine/Source/Programs/Horde/Plugins/Build/HordeServer.Build.Tests/Fleet/LeaseUtilizationStrategyTest.cs b/Engine/Source/Programs/Horde/Plugins/Build/HordeServer.Build.Tests/Fleet/LeaseUtilizationStrategyTest.cs
--- a/Engine/Source/Programs/Horde/Plugins/Build/HordeServer.Build.Tests/Fleet/LeaseUtilizationStrategyTest.cs
+++ b/Engine/Source/Programs/Horde/Plugins/Build/HordeServer.Build.Tests/Fleet/LeaseUtilizationStrategyTest.cs
@@ -1,19 +1,15 @@
 // Copyright Epic Games, Inc. All Rights Reserved.
 
-using System;
-using System.Collections.Generic;
-using System.Threading.Tasks;
 using EpicGames.Horde.Agents.Leases;
 using EpicGames.Horde.Streams;
 using Google.Protobuf;
 using Google.Protobuf.WellKnownTypes;
+using HordeCommon.Rpc.Tasks;
 using HordeServer.Agents;
 using HordeServer.Agents.Fleet;
 using HordeServer.Agents.Leases;
 using HordeServer.Agents.Pools;
 using HordeServer.Utilities;
-using HordeCommon.Rpc.Tasks;
-using Microsoft.VisualStudio.TestTools.UnitTesting;
 
 namespace HordeServer.Tests.Fleet
 {
diff --git a/Engine/Source/Programs/Horde/Plugins/Build/HordeServer.Build.Tests/HordeServer.Build.Tests.csproj b/Engine/Source/Programs/Horde/Plugins/Build/HordeServer.Build.Tests/HordeServer.Build.Tests.csproj
--- a/Engine/Source/Programs/Horde/Plugins/Build/HordeServer.Build.Tests/HordeServer.Build.Tests.csproj
+++ b/Engine/Source/Programs/Horde/Plugins/Build/HordeServer.Build.Tests/HordeServer.Build.Tests.csproj
@@ -4,10 +4,19 @@
     <TargetFramework>net8.0</TargetFramework>
     <ImplicitUsings>enable</ImplicitUsings>
     <Nullable>enable</Nullable>
-
     <IsPackable>false</IsPackable>
     <IsTestProject>true</IsTestProject>
     <RootNamespace>HordeServer.Tests</RootNamespace>
+    <RunAnalyzersDuringBuild>False</RunAnalyzersDuringBuild>
+    <Configurations>Debug;Release;Analyze</Configurations>
+    <NoWarn>1701;1702;1591</NoWarn>
+  </PropertyGroup>
+
+  <PropertyGroup Condition=" '$(Configuration)|$(Platform)' == 'Analyze|AnyCPU' ">
+    <RunAnalyzersDuringBuild>True</RunAnalyzersDuringBuild>
+    <TreatWarningsAsErrors>false</TreatWarningsAsErrors>
+    <EnforceCodeStyleInBuild>True</EnforceCodeStyleInBuild>
+    <GenerateDocumentationFile>True</GenerateDocumentationFile>
   </PropertyGroup>
 
   <ItemGroup>
diff --git a/Engine/Source/Programs/Horde/Plugins/Build/HordeServer.Build.Tests/Issues/IssueServiceTests.cs b/Engine/Source/Programs/Horde/Plugins/Build/HordeServer.Build.Tests/Issues/IssueServiceTests.cs
--- a/Engine/Source/Programs/Horde/Plugins/Build/HordeServer.Build.Tests/Issues/IssueServiceTests.cs
+++ b/Engine/Source/Programs/Horde/Plugins/Build/HordeServer.Build.Tests/Issues/IssueServiceTests.cs
@@ -1,15 +1,10 @@
 // Copyright Epic Games, Inc. All Rights Reserved.
 
 extern alias JobDriver;
-using System;
 using System.Buffers;
-using System.Collections.Generic;
-using System.Linq;
 using System.Runtime.InteropServices;
 using System.Text;
 using System.Text.Json;
-using System.Threading;
-using System.Threading.Tasks;
 using EpicGames.Core;
 using EpicGames.Horde.Issues;
 using EpicGames.Horde.Jobs;
@@ -33,7 +28,6 @@ using HordeServer.Users;
 using JobDriver.JobDriver.Parser;
 using Microsoft.Extensions.Logging;
 using Microsoft.Extensions.Logging.Abstractions;
-using Microsoft.VisualStudio.TestTools.UnitTesting;
 using MongoDB.Bson;
 using MongoDB.Driver;
 using Moq;
@@ -1571,7 +1565,8 @@ namespace HordeServer.Tests.Issues
 		{
 			IJob job = CreateJob(_mainStreamId, 120, "Compile Test", _graph);
 
-			await WriteEventsAsync(job, 0, 0, logger => {
+			await WriteEventsAsync(job, 0, 0, logger =>
+			{
 				logger.LogInformation("Foo");
 
 				IssueFingerprint fingerprint = new IssueFingerprint("NewIssueType", "This is an issue with severity: {Severity}");
diff --git a/Engine/Source/Programs/Horde/Plugins/Build/HordeServer.Build.Tests/Issues/IssueTagTests.cs b/Engine/Source/Programs/Horde/Plugins/Build/HordeServer.Build.Tests/Issues/IssueTagTests.cs
--- a/Engine/Source/Programs/Horde/Plugins/Build/HordeServer.Build.Tests/Issues/IssueTagTests.cs
+++ b/Engine/Source/Programs/Horde/Plugins/Build/HordeServer.Build.Tests/Issues/IssueTagTests.cs
@@ -1,9 +1,6 @@
 // Copyright Epic Games, Inc. All Rights Reserved.
 
-using System;
-using System.Linq;
 using HordeServer.Issues;
-using Microsoft.VisualStudio.TestTools.UnitTesting;
 
 namespace HordeServer.Tests.Issues
 {
diff --git a/Engine/Source/Programs/Horde/Plugins/Build/HordeServer.Build.Tests/Jobs/JobExpirationTests.cs b/Engine/Source/Programs/Horde/Plugins/Build/HordeServer.Build.Tests/Jobs/JobExpirationTests.cs
--- a/Engine/Source/Programs/Horde/Plugins/Build/HordeServer.Build.Tests/Jobs/JobExpirationTests.cs
+++ b/Engine/Source/Programs/Horde/Plugins/Build/HordeServer.Build.Tests/Jobs/JobExpirationTests.cs
@@ -1,7 +1,5 @@
 // Copyright Epic Games, Inc. All Rights Reserved.
 
-using System;
-using System.Threading.Tasks;
 using EpicGames.Horde.Jobs.Templates;
 using EpicGames.Horde.Projects;
 using EpicGames.Horde.Streams;
@@ -12,7 +10,6 @@ using HordeServer.Projects;
 using HordeServer.Server;
 using HordeServer.Streams;
 using Microsoft.Extensions.DependencyInjection;
-using Microsoft.VisualStudio.TestTools.UnitTesting;
 
 namespace HordeServer.Tests.Jobs
 {
diff --git a/Engine/Source/Programs/Horde/Plugins/Build/HordeServer.Build.Tests/Jobs/JobServiceTests.cs b/Engine/Source/Programs/Horde/Plugins/Build/HordeServer.Build.Tests/Jobs/JobServiceTests.cs
--- a/Engine/Source/Programs/Horde/Plugins/Build/HordeServer.Build.Tests/Jobs/JobServiceTests.cs
+++ b/Engine/Source/Programs/Horde/Plugins/Build/HordeServer.Build.Tests/Jobs/JobServiceTests.cs
@@ -1,10 +1,6 @@
 // Copyright Epic Games, Inc. All Rights Reserved.
 
-using System.Collections.Generic;
-using System.Linq;
 using System.Text;
-using System.Threading;
-using System.Threading.Tasks;
 using EpicGames.Core;
 using EpicGames.Horde.Agents;
 using EpicGames.Horde.Agents.Pools;
@@ -23,7 +19,6 @@ using HordeServer.Server;
 using HordeServer.Streams;
 using HordeServer.Users;
 using Microsoft.CodeAnalysis;
-using Microsoft.VisualStudio.TestTools.UnitTesting;
 
 namespace HordeServer.Tests.Jobs
 {
@@ -254,7 +249,7 @@ namespace HordeServer.Tests.Jobs
 			IAgent? agent = await AgentService.CreateAgentAsync("TestAgent", false, "");
 			Assert.IsNotNull(agent);
 
-			agent = await agent.TryUpdateAsync(new UpdateAgentOptions{ Enabled = true, ExplicitPools = new List<PoolId> { new PoolId("win") } });
+			agent = await agent.TryUpdateAsync(new UpdateAgentOptions { Enabled = true, ExplicitPools = new List<PoolId> { new PoolId("win") } });
 			Assert.IsNotNull(agent);
 
 			await AgentService.CreateSessionAsync(agent, AgentStatus.Ok, new List<string>(), new Dictionary<string, int>(), null);
diff --git a/Engine/Source/Programs/Horde/Plugins/Build/HordeServer.Build.Tests/Jobs/JobsControllerDbTest.cs b/Engine/Source/Programs/Horde/Plugins/Build/HordeServer.Build.Tests/Jobs/JobsControllerDbTest.cs
--- a/Engine/Source/Programs/Horde/Plugins/Build/HordeServer.Build.Tests/Jobs/JobsControllerDbTest.cs
+++ b/Engine/Source/Programs/Horde/Plugins/Build/HordeServer.Build.Tests/Jobs/JobsControllerDbTest.cs
@@ -1,12 +1,9 @@
 // Copyright Epic Games, Inc. All Rights Reserved.
 
-using System.Collections.Generic;
-using System.Threading.Tasks;
 using EpicGames.Core;
 using EpicGames.Horde.Jobs;
 using HordeServer.Jobs;
 using Microsoft.AspNetCore.Mvc;
-using Microsoft.VisualStudio.TestTools.UnitTesting;
 
 namespace HordeServer.Tests.Jobs
 {
diff --git a/Engine/Source/Programs/Horde/Plugins/Build/HordeServer.Build.Tests/Jobs/Schedules/SchedulerTests.cs b/Engine/Source/Programs/Horde/Plugins/Build/HordeServer.Build.Tests/Jobs/Schedules/SchedulerTests.cs
--- a/Engine/Source/Programs/Horde/Plugins/Build/HordeServer.Build.Tests/Jobs/Schedules/SchedulerTests.cs
+++ b/Engine/Source/Programs/Horde/Plugins/Build/HordeServer.Build.Tests/Jobs/Schedules/SchedulerTests.cs
@@ -1,9 +1,5 @@
 // Copyright Epic Games, Inc. All Rights Reserved.
 
-using System;
-using System.Collections.Generic;
-using System.Linq;
-using System.Threading.Tasks;
 using EpicGames.Horde.Commits;
 using EpicGames.Horde.Jobs;
 using EpicGames.Horde.Jobs.Templates;
@@ -19,7 +15,6 @@ using HordeServer.Server;
 using HordeServer.Streams;
 using HordeServer.Users;
 using Microsoft.Extensions.DependencyInjection;
-using Microsoft.VisualStudio.TestTools.UnitTesting;
 
 namespace HordeServer.Tests.Jobs.Schedules
 {
@@ -42,7 +37,7 @@ namespace HordeServer.Tests.Jobs.Schedules
 
 			GlobalConfig globalConfig = new GlobalConfig();
 			globalConfig.Plugins.AddBuildConfig(buildConfig);
-			
+
 			SetConfig(globalConfig);
 
 			_template = TemplateCollection.GetOrAddAsync(new TemplateConfig { Name = "Test template" }).Result;
diff --git a/Engine/Source/Programs/Horde/Plugins/Build/HordeServer.Build.Tests/Jobs/TestData/TestDataTests.cs b/Engine/Source/Programs/Horde/Plugins/Build/HordeServer.Build.Tests/Jobs/TestData/TestDataTests.cs
--- a/Engine/Source/Programs/Horde/Plugins/Build/HordeServer.Build.Tests/Jobs/TestData/TestDataTests.cs
+++ b/Engine/Source/Programs/Horde/Plugins/Build/HordeServer.Build.Tests/Jobs/TestData/TestDataTests.cs
@@ -1,11 +1,6 @@
 // Copyright Epic Games, Inc. All Rights Reserved.
 
-using System;
-using System.Collections.Generic;
-using System.Linq;
 using System.Security.Claims;
-using System.Threading;
-using System.Threading.Tasks;
 using EpicGames.Horde;
 using EpicGames.Horde.Jobs;
 using EpicGames.Horde.Jobs.Templates;
@@ -15,16 +10,15 @@ using EpicGames.Horde.Projects;
 using EpicGames.Horde.Streams;
 using HordeServer.Jobs;
 using HordeServer.Jobs.Graphs;
+using HordeServer.Jobs.TestData;
 using HordeServer.Logs;
 using HordeServer.Projects;
 using HordeServer.Server;
 using HordeServer.Streams;
-using HordeServer.Jobs.TestData;
 using HordeServer.Users;
 using HordeServer.Utilities;
 using Microsoft.AspNetCore.Http;
 using Microsoft.AspNetCore.Mvc;
-using Microsoft.VisualStudio.TestTools.UnitTesting;
 using MongoDB.Bson;
 using Moq;
 
diff --git a/Engine/Source/Programs/Horde/Plugins/Build/HordeServer.Build.Tests/Perforce/PerforceLoadBalancerTests.cs b/Engine/Source/Programs/Horde/Plugins/Build/HordeServer.Build.Tests/Perforce/PerforceLoadBalancerTests.cs
--- a/Engine/Source/Programs/Horde/Plugins/Build/HordeServer.Build.Tests/Perforce/PerforceLoadBalancerTests.cs
+++ b/Engine/Source/Programs/Horde/Plugins/Build/HordeServer.Build.Tests/Perforce/PerforceLoadBalancerTests.cs
@@ -1,17 +1,10 @@
 // Copyright Epic Games, Inc. All Rights Reserved.
 
-using System;
-using System.Collections.Generic;
-using System.Linq;
 using System.Net;
-using System.Net.Http;
-using System.Threading;
-using System.Threading.Tasks;
 using HordeServer.Perforce;
 using HordeServer.Server;
 using Microsoft.Extensions.Diagnostics.HealthChecks;
 using Microsoft.Extensions.Logging;
-using Microsoft.VisualStudio.TestTools.UnitTesting;
 
 namespace HordeServer.Tests.Perforce;
 
diff --git a/Engine/Source/Programs/Horde/Plugins/Build/HordeServer.Build.Tests/Streams/StreamServiceTest.cs b/Engine/Source/Programs/Horde/Plugins/Build/HordeServer.Build.Tests/Streams/StreamServiceTest.cs
--- a/Engine/Source/Programs/Horde/Plugins/Build/HordeServer.Build.Tests/Streams/StreamServiceTest.cs
+++ b/Engine/Source/Programs/Horde/Plugins/Build/HordeServer.Build.Tests/Streams/StreamServiceTest.cs
@@ -1,10 +1,7 @@
 // Copyright Epic Games, Inc. All Rights Reserved.
 
-using System;
 using System.Globalization;
-using System.Threading.Tasks;
 using HordeServer.Streams;
-using Microsoft.VisualStudio.TestTools.UnitTesting;
 
 namespace HordeServer.Tests.Streams
 {
diff --git a/Engine/Source/Programs/Horde/Plugins/Build/HordeServer.Build/Agents/AgentReportService.cs b/Engine/Source/Programs/Horde/Plugins/Build/HordeServer.Build/Agents/AgentReportService.cs
--- a/Engine/Source/Programs/Horde/Plugins/Build/HordeServer.Build/Agents/AgentReportService.cs
+++ b/Engine/Source/Programs/Horde/Plugins/Build/HordeServer.Build/Agents/AgentReportService.cs
@@ -1,16 +1,11 @@
 // Copyright Epic Games, Inc. All Rights Reserved.
 
-using System;
-using System.Collections.Generic;
-using System.Threading;
-using System.Threading.Tasks;
 using EpicGames.Core;
 using EpicGames.Horde.Agents;
 using HordeServer.Jobs;
 using HordeServer.Notifications;
 using HordeServer.Server;
 using HordeServer.Utilities;
-using HordeCommon;
 using Microsoft.Extensions.Hosting;
 using Microsoft.Extensions.Logging;
 
diff --git a/Engine/Source/Programs/Horde/Plugins/Build/HordeServer.Build/Agents/Pools/PoolUpdateService.cs b/Engine/Source/Programs/Horde/Plugins/Build/HordeServer.Build/Agents/Pools/PoolUpdateService.cs
--- a/Engine/Source/Programs/Horde/Plugins/Build/HordeServer.Build/Agents/Pools/PoolUpdateService.cs
+++ b/Engine/Source/Programs/Horde/Plugins/Build/HordeServer.Build/Agents/Pools/PoolUpdateService.cs
@@ -1,15 +1,9 @@
 // Copyright Epic Games, Inc. All Rights Reserved.
 
-using System;
-using System.Collections.Generic;
-using System.Linq;
-using System.Threading;
-using System.Threading.Tasks;
 using EpicGames.Core;
 using EpicGames.Horde.Agents;
 using EpicGames.Horde.Agents.Pools;
 using HordeServer.Perforce;
-using HordeServer.Server;
 using HordeServer.Streams;
 using Microsoft.Extensions.Hosting;
 using Microsoft.Extensions.Logging;
@@ -111,7 +105,7 @@ namespace HordeServer.Agents.Pools
 			{
 				return false;
 			}
-			
+
 			DateTime expirationTime = agent.LastStatusChange.Value + gracePeriod.Value;
 			return _clock.UtcNow > expirationTime;
 		}
diff --git a/Engine/Source/Programs/Horde/Plugins/Build/HordeServer.Build/Artifacts/ArtifactExpirationService.cs b/Engine/Source/Programs/Horde/Plugins/Build/HordeServer.Build/Artifacts/ArtifactExpirationService.cs
--- a/Engine/Source/Programs/Horde/Plugins/Build/HordeServer.Build/Artifacts/ArtifactExpirationService.cs
+++ b/Engine/Source/Programs/Horde/Plugins/Build/HordeServer.Build/Artifacts/ArtifactExpirationService.cs
@@ -1,15 +1,9 @@
 // Copyright Epic Games, Inc. All Rights Reserved.
 
-using System;
-using System.Collections.Generic;
 using System.Diagnostics;
-using System.Linq;
-using System.Threading;
-using System.Threading.Tasks;
 using EpicGames.Core;
 using EpicGames.Horde.Storage;
 using EpicGames.Horde.Streams;
-using HordeServer.Server;
 using HordeServer.Storage;
 using Microsoft.Extensions.Hosting;
 using Microsoft.Extensions.Logging;
diff --git a/Engine/Source/Programs/Horde/Plugins/Build/HordeServer.Build/Artifacts/ArtifactsController.cs b/Engine/Source/Programs/Horde/Plugins/Build/HordeServer.Build/Artifacts/ArtifactsController.cs
--- a/Engine/Source/Programs/Horde/Plugins/Build/HordeServer.Build/Artifacts/ArtifactsController.cs
+++ b/Engine/Source/Programs/Horde/Plugins/Build/HordeServer.Build/Artifacts/ArtifactsController.cs
@@ -1,18 +1,13 @@
 // Copyright Epic Games, Inc. All Rights Reserved.
 
-using System;
 using System.Buffers;
 using System.Buffers.Binary;
-using System.Collections.Generic;
 using System.Diagnostics;
-using System.IO;
-using System.Linq;
 using System.Net;
 using System.Net.Mime;
 using System.Text.RegularExpressions;
-using System.Threading;
-using System.Threading.Tasks;
 using EpicGames.Core;
+using EpicGames.Horde.Acls;
 using EpicGames.Horde.Agents.Leases;
 using EpicGames.Horde.Artifacts;
 using EpicGames.Horde.Jobs;
@@ -21,14 +16,14 @@ using EpicGames.Horde.Storage.Bundles;
 using EpicGames.Horde.Storage.Nodes;
 using EpicGames.Horde.Streams;
 using Google.Protobuf.WellKnownTypes;
+using HordeCommon.Rpc.Tasks;
 using HordeServer.Acls;
+using HordeServer.Agents;
 using HordeServer.Agents.Leases;
 using HordeServer.Jobs;
-using HordeServer.Server;
 using HordeServer.Storage;
 using HordeServer.Streams;
 using HordeServer.Utilities;
-using HordeCommon.Rpc.Tasks;
 using Microsoft.AspNetCore.Authorization;
 using Microsoft.AspNetCore.Http;
 using Microsoft.AspNetCore.Http.Features;
@@ -36,10 +31,6 @@ using Microsoft.AspNetCore.Mvc;
 using Microsoft.AspNetCore.StaticFiles;
 using Microsoft.Extensions.Logging;
 using Microsoft.Extensions.Options;
-using EpicGames.Horde.Acls;
-using HordeServer.Agents;
-using System.Security.Claims;
-using HordeServer.Projects;
 
 namespace HordeServer.Artifacts
 {
@@ -59,13 +50,12 @@ namespace HordeServer.Artifacts
 		readonly IBlockCache _blockCache;
 		readonly IServerInfo _serverInfo;
 		readonly BuildConfig _buildConfig;
-		readonly StaticBuildConfig _staticBuildConfig;
 		readonly ILogger _logger;
 
 		/// <summary>
 		/// Constructor
 		/// </summary>
-		public ArtifactsController(IArtifactCollection artifactCollection, StorageService storageService, ILeaseCollection leaseCollection, IJobCollection jobCollection, IAclService aclService, UnsyncCache unsyncCache, IBlockCache blockCache, IServerInfo serverInfo, IOptionsSnapshot<BuildConfig> buildConfig, IOptions<StaticBuildConfig> staticBuildConfig, ILogger<ArtifactsController> logger)
+		public ArtifactsController(IArtifactCollection artifactCollection, StorageService storageService, ILeaseCollection leaseCollection, IJobCollection jobCollection, IAclService aclService, UnsyncCache unsyncCache, IBlockCache blockCache, IServerInfo serverInfo, IOptionsSnapshot<BuildConfig> buildConfig, ILogger<ArtifactsController> logger)
 		{
 			_artifactCollection = artifactCollection;
 			_storageService = storageService;
@@ -76,7 +66,6 @@ namespace HordeServer.Artifacts
 			_blockCache = blockCache;
 			_serverInfo = serverInfo;
 			_buildConfig = buildConfig.Value;
-			_staticBuildConfig = staticBuildConfig.Value;
 			_logger = logger;
 		}
 
diff --git a/Engine/Source/Programs/Horde/Plugins/Build/HordeServer.Build/Artifacts/BlockCache.cs b/Engine/Source/Programs/Horde/Plugins/Build/HordeServer.Build/Artifacts/BlockCache.cs
--- a/Engine/Source/Programs/Horde/Plugins/Build/HordeServer.Build/Artifacts/BlockCache.cs
+++ b/Engine/Source/Programs/Horde/Plugins/Build/HordeServer.Build/Artifacts/BlockCache.cs
@@ -1,17 +1,13 @@
 // Copyright Epic Games, Inc. All Rights Reserved.
 
-using System;
 using System.Buffers;
 using System.Buffers.Binary;
 using System.Collections.Concurrent;
-using System.Collections.Generic;
 using System.Diagnostics;
-using System.IO;
 using System.IO.Hashing;
 using System.IO.MemoryMappedFiles;
 using System.Numerics;
 using System.Text;
-using System.Threading;
 using EpicGames.Core;
 using Microsoft.Extensions.Logging;
 using Microsoft.Extensions.Logging.Abstractions;
@@ -463,7 +459,7 @@ namespace HordeServer.Artifacts
 			blockState.Touch();
 
 			ReadOnlySequenceBuilder<byte> builder = new ReadOnlySequenceBuilder<byte>();
-			for(int idx = 0; idx < blockState.BlockIdxs.Length; idx++)
+			for (int idx = 0; idx < blockState.BlockIdxs.Length; idx++)
 			{
 				int blockIdx = blockState.BlockIdxs[idx];
 
diff --git a/Engine/Source/Programs/Horde/Plugins/Build/HordeServer.Build/Artifacts/IArtifact.cs b/Engine/Source/Programs/Horde/Plugins/Build/HordeServer.Build/Artifacts/IArtifact.cs
--- a/Engine/Source/Programs/Horde/Plugins/Build/HordeServer.Build/Artifacts/IArtifact.cs
+++ b/Engine/Source/Programs/Horde/Plugins/Build/HordeServer.Build/Artifacts/IArtifact.cs
@@ -1,8 +1,5 @@
 // Copyright Epic Games, Inc. All Rights Reserved.
 
-using System;
-using System.Collections.Generic;
-using EpicGames.Horde.Acls;
 using EpicGames.Horde.Artifacts;
 using EpicGames.Horde.Storage;
 using EpicGames.Horde.Streams;
diff --git a/Engine/Source/Programs/Horde/Plugins/Build/HordeServer.Build/BuildConfig.cs b/Engine/Source/Programs/Horde/Plugins/Build/HordeServer.Build/BuildConfig.cs
--- a/Engine/Source/Programs/Horde/Plugins/Build/HordeServer.Build/BuildConfig.cs
+++ b/Engine/Source/Programs/Horde/Plugins/Build/HordeServer.Build/BuildConfig.cs
@@ -1,15 +1,10 @@
 ﻿// Copyright Epic Games, Inc. All Rights Reserved.
 
-using System;
-using System.Collections.Generic;
 using System.ComponentModel.DataAnnotations;
 using System.Diagnostics;
 using System.Diagnostics.CodeAnalysis;
-using System.Linq;
 using System.Security.Claims;
-using System.Text;
 using System.Text.Json.Serialization;
-using System.Threading.Tasks;
 using EpicGames.Horde.Acls;
 using EpicGames.Horde.Agents.Pools;
 using EpicGames.Horde.Artifacts;
@@ -32,6 +27,9 @@ using HordeServer.Utilities;
 
 namespace HordeServer
 {
+	/// <summary>
+	/// Configuration for the build plugin
+	/// </summary>
 	public class BuildConfig : IPluginConfig
 	{
 		/// <summary>
diff --git a/Engine/Source/Programs/Horde/Plugins/Build/HordeServer.Build/BuildPlugin.cs b/Engine/Source/Programs/Horde/Plugins/Build/HordeServer.Build/BuildPlugin.cs
--- a/Engine/Source/Programs/Horde/Plugins/Build/HordeServer.Build/BuildPlugin.cs
+++ b/Engine/Source/Programs/Horde/Plugins/Build/HordeServer.Build/BuildPlugin.cs
@@ -44,6 +44,9 @@ namespace HordeServer
 		readonly IServerInfo _serverInfo;
 		readonly StaticBuildConfig _staticConfig;
 
+		/// <summary>
+		/// Constructor
+		/// </summary>
 		public BuildPlugin(IServerInfo serverInfo, StaticBuildConfig staticConfig)
 		{
 			_serverInfo = serverInfo;
@@ -88,7 +91,7 @@ namespace HordeServer
 			services.AddSingleton<ICommitService, CommitService>();
 
 			services.AddSingleton<DeviceService>();
-			services.AddSingleton<IBlockCache>(sp => CreateBlockCache(sp));
+			services.AddSingleton<IBlockCache>(sp => CreateBlockCache());
 			services.AddSingleton<TestDataService>();
 
 			services.AddSingleton<IssueService>();
@@ -209,7 +212,7 @@ namespace HordeServer
 			}
 		}
 
-		BlockCache CreateBlockCache(IServiceProvider serviceProvider)
+		BlockCache CreateBlockCache()
 		{
 			DirectoryReference cacheDir = DirectoryReference.Combine(_serverInfo.DataDir, String.IsNullOrEmpty(_staticConfig.BlockCacheDir) ? "BlockCache" : _staticConfig.BlockCacheDir);
 			return BlockCache.Create(cacheDir, (int)(_staticConfig.BlockCacheSizeBytes / (1024 * 1024 * 1024)));
diff --git a/Engine/Source/Programs/Horde/Plugins/Build/HordeServer.Build/Commands/RelayCommand.cs b/Engine/Source/Programs/Horde/Plugins/Build/HordeServer.Build/Commands/RelayCommand.cs
--- a/Engine/Source/Programs/Horde/Plugins/Build/HordeServer.Build/Commands/RelayCommand.cs
+++ b/Engine/Source/Programs/Horde/Plugins/Build/HordeServer.Build/Commands/RelayCommand.cs
@@ -1,10 +1,6 @@
 // Copyright Epic Games, Inc. All Rights Reserved.
 
-using System;
 using System.ComponentModel;
-using System.Linq;
-using System.Threading;
-using System.Threading.Tasks;
 using EpicGames.Core;
 using Grpc.Net.Client;
 using Horde.Common.Rpc;
diff --git a/Engine/Source/Programs/Horde/Plugins/Build/HordeServer.Build/Commands/Test/ReplicationCommand.cs b/Engine/Source/Programs/Horde/Plugins/Build/HordeServer.Build/Commands/Test/ReplicationCommand.cs
--- a/Engine/Source/Programs/Horde/Plugins/Build/HordeServer.Build/Commands/Test/ReplicationCommand.cs
+++ b/Engine/Source/Programs/Horde/Plugins/Build/HordeServer.Build/Commands/Test/ReplicationCommand.cs
@@ -1,13 +1,9 @@
 // Copyright Epic Games, Inc. All Rights Reserved.
 
-using System;
-using System.Threading.Tasks;
 using EpicGames.Core;
 using EpicGames.Horde.Replicators;
 using EpicGames.Horde.Streams;
-using HordeServer.Configuration;
 using HordeServer.Replicators;
-using HordeServer.Server;
 using HordeServer.Streams;
 using Microsoft.Extensions.Configuration;
 using Microsoft.Extensions.DependencyInjection;
diff --git a/Engine/Source/Programs/Horde/Plugins/Build/HordeServer.Build/Devices/DeviceReportService.cs b/Engine/Source/Programs/Horde/Plugins/Build/HordeServer.Build/Devices/DeviceReportService.cs
--- a/Engine/Source/Programs/Horde/Plugins/Build/HordeServer.Build/Devices/DeviceReportService.cs
+++ b/Engine/Source/Programs/Horde/Plugins/Build/HordeServer.Build/Devices/DeviceReportService.cs
@@ -1,10 +1,5 @@
 // Copyright Epic Games, Inc. All Rights Reserved.
 
-using System;
-using System.Collections.Generic;
-using System.Linq;
-using System.Threading;
-using System.Threading.Tasks;
 using EpicGames.Core;
 using EpicGames.Horde.Devices;
 using HordeServer.Jobs;
@@ -139,7 +134,6 @@ namespace HordeServer.Devices
 		readonly IClock _clock;
 		readonly ITicker _ticker;
 		readonly IServerInfo _serverInfo;
-		readonly IOptionsMonitor<BuildConfig> _buildConfig;
 		readonly IOptions<StaticBuildConfig> _staticBuildConfig;
 		readonly ILogger<DeviceReportService> _logger;
 
@@ -148,7 +142,7 @@ namespace HordeServer.Devices
 		/// <summary>
 		/// Constructor
 		/// </summary>
-		public DeviceReportService(IMongoService mongoService, DeviceService deviceService, INotificationService notificationService, IClock clock, IServerInfo serverInfo, IOptionsMonitor<BuildConfig> buildConfig, IOptions<StaticBuildConfig> staticBuildConfig, ILogger<DeviceReportService> logger)
+		public DeviceReportService(IMongoService mongoService, DeviceService deviceService, INotificationService notificationService, IClock clock, IServerInfo serverInfo, IOptions<StaticBuildConfig> staticBuildConfig, ILogger<DeviceReportService> logger)
 		{
 			_state = new SingletonDocument<DeviceReportState>(mongoService);
 			_deviceService = deviceService;
@@ -156,7 +150,6 @@ namespace HordeServer.Devices
 			_clock = clock;
 			_serverInfo = serverInfo;
 			_ticker = clock.AddSharedTicker<DeviceReportService>(TimeSpan.FromMinutes(5.0), TickAsync, logger);
-			_buildConfig = buildConfig;
 			_staticBuildConfig = staticBuildConfig;
 			_logger = logger;
 		}
diff --git a/Engine/Source/Programs/Horde/Plugins/Build/HordeServer.Build/Devices/IDeviceCollection.cs b/Engine/Source/Programs/Horde/Plugins/Build/HordeServer.Build/Devices/IDeviceCollection.cs
--- a/Engine/Source/Programs/Horde/Plugins/Build/HordeServer.Build/Devices/IDeviceCollection.cs
+++ b/Engine/Source/Programs/Horde/Plugins/Build/HordeServer.Build/Devices/IDeviceCollection.cs
@@ -1,8 +1,5 @@
 // Copyright Epic Games, Inc. All Rights Reserved.
 
-using System;
-using System.Collections.Generic;
-using System.Threading.Tasks;
 using EpicGames.Horde.Devices;
 using EpicGames.Horde.Jobs;
 using EpicGames.Horde.Users;
diff --git a/Engine/Source/Programs/Horde/Plugins/Build/HordeServer.Build/Issues/External/JiraService.cs b/Engine/Source/Programs/Horde/Plugins/Build/HordeServer.Build/Issues/External/JiraService.cs
--- a/Engine/Source/Programs/Horde/Plugins/Build/HordeServer.Build/Issues/External/JiraService.cs
+++ b/Engine/Source/Programs/Horde/Plugins/Build/HordeServer.Build/Issues/External/JiraService.cs
@@ -1,17 +1,10 @@
 // Copyright Epic Games, Inc. All Rights Reserved.	
 
-using System;
 using System.Collections.Concurrent;
-using System.Collections.Generic;
-using System.Linq;
-using System.Net.Http;
 using System.Net.Http.Headers;
 using System.Text;
 using System.Text.Json;
-using System.Threading;
-using System.Threading.Tasks;
 using EpicGames.Core;
-using HordeServer.Server;
 using HordeServer.Streams;
 using HordeServer.Users;
 using Microsoft.Extensions.Hosting;
diff --git a/Engine/Source/Programs/Horde/Plugins/Build/HordeServer.Build/Issues/IssueCollection.cs b/Engine/Source/Programs/Horde/Plugins/Build/HordeServer.Build/Issues/IssueCollection.cs
--- a/Engine/Source/Programs/Horde/Plugins/Build/HordeServer.Build/Issues/IssueCollection.cs
+++ b/Engine/Source/Programs/Horde/Plugins/Build/HordeServer.Build/Issues/IssueCollection.cs
@@ -1,13 +1,8 @@
 // Copyright Epic Games, Inc. All Rights Reserved.
 
-using System;
-using System.Collections.Generic;
 using System.Data;
 using System.Diagnostics;
 using System.Diagnostics.CodeAnalysis;
-using System.Linq;
-using System.Threading;
-using System.Threading.Tasks;
 using EpicGames.Core;
 using EpicGames.Horde.Issues;
 using EpicGames.Horde.Jobs;
@@ -274,7 +269,7 @@ namespace HordeServer.Issues
 			public List<string>? LegacyChangeFilter
 			{
 				get => null;
-				set => ChangeFilter = (value == null)? null : String.Join(";", value);
+				set => ChangeFilter = (value == null) ? null : String.Join(";", value);
 			}
 
 			string IIssueFingerprint.ChangeFilter => ChangeFilter ?? String.Join(";", GetLegacyHandlerInfo(Type).ChangeFilter);
diff --git a/Engine/Source/Programs/Horde/Plugins/Build/HordeServer.Build/Issues/IssueReportService.cs b/Engine/Source/Programs/Horde/Plugins/Build/HordeServer.Build/Issues/IssueReportService.cs
--- a/Engine/Source/Programs/Horde/Plugins/Build/HordeServer.Build/Issues/IssueReportService.cs
+++ b/Engine/Source/Programs/Horde/Plugins/Build/HordeServer.Build/Issues/IssueReportService.cs
@@ -1,10 +1,5 @@
 // Copyright Epic Games, Inc. All Rights Reserved.
 
-using System;
-using System.Collections.Generic;
-using System.Linq;
-using System.Threading;
-using System.Threading.Tasks;
 using EpicGames.Core;
 using EpicGames.Horde.Issues;
 using EpicGames.Horde.Jobs;
@@ -16,7 +11,6 @@ using HordeServer.Notifications;
 using HordeServer.Server;
 using HordeServer.Streams;
 using HordeServer.Utilities;
-using HordeCommon;
 using Microsoft.Extensions.Hosting;
 using Microsoft.Extensions.Logging;
 using Microsoft.Extensions.Options;
diff --git a/Engine/Source/Programs/Horde/Plugins/Build/HordeServer.Build/Issues/IssueService.cs b/Engine/Source/Programs/Horde/Plugins/Build/HordeServer.Build/Issues/IssueService.cs
--- a/Engine/Source/Programs/Horde/Plugins/Build/HordeServer.Build/Issues/IssueService.cs
+++ b/Engine/Source/Programs/Horde/Plugins/Build/HordeServer.Build/Issues/IssueService.cs
@@ -1,14 +1,9 @@
 // Copyright Epic Games, Inc. All Rights Reserved.
 
-using System;
-using System.Collections.Generic;
 using System.Globalization;
-using System.Linq;
 using System.Reflection;
 using System.Text;
 using System.Text.RegularExpressions;
-using System.Threading;
-using System.Threading.Tasks;
 using EpicGames.Core;
 using EpicGames.Horde.Issues;
 using EpicGames.Horde.Jobs;
@@ -20,7 +15,6 @@ using HordeServer.Commits;
 using HordeServer.Jobs;
 using HordeServer.Jobs.Graphs;
 using HordeServer.Logs;
-using HordeServer.Server;
 using HordeServer.Streams;
 using HordeServer.Users;
 using Microsoft.Extensions.DependencyInjection;
@@ -702,7 +696,7 @@ namespace HordeServer.Issues
 			handlers.SortBy(x => -x.Priority);
 
 			// Create all the issue definitions by passing each log event to the handlers in order until one attaches it to an issue
-			List <ILogAnchor> stepAnchors = await log.GetAnchorsAsync(cancellationToken: cancellationToken);
+			List<ILogAnchor> stepAnchors = await log.GetAnchorsAsync(cancellationToken: cancellationToken);
 			foreach (ILogAnchor stepAnchor in stepAnchors)
 			{
 				ILogEventData stepEventData = await stepAnchor.GetDataAsync(cancellationToken);
