Horde: Add "Analyze" configurations to all plugin projects, and fix up any current violations.

diff --git a/Engine/Source/Programs/Horde/Plugins/Build/HordeServer.Build/Issues/IssuesController.cs b/Engine/Source/Programs/Horde/Plugins/Build/HordeServer.Build/Issues/IssuesController.cs
--- a/Engine/Source/Programs/Horde/Plugins/Build/HordeServer.Build/Issues/IssuesController.cs
+++ b/Engine/Source/Programs/Horde/Plugins/Build/HordeServer.Build/Issues/IssuesController.cs
@@ -1,10 +1,5 @@
 // Copyright Epic Games, Inc. All Rights Reserved.
 
-using System;
-using System.Collections.Generic;
-using System.Linq;
-using System.Threading;
-using System.Threading.Tasks;
 using EpicGames.Core;
 using EpicGames.Horde.Issues;
 using EpicGames.Horde.Jobs;
@@ -17,7 +12,6 @@ using HordeServer.Issues.External;
 using HordeServer.Jobs;
 using HordeServer.Jobs.Graphs;
 using HordeServer.Logs;
-using HordeServer.Server;
 using HordeServer.Streams;
 using HordeServer.Users;
 using HordeServer.Utilities;
@@ -535,7 +529,7 @@ namespace HordeServer.Issues
 			response.QuarantineTimeUtc = details.QuarantineTimeUtc;
 			response.ForceClosedByUserInfo = details.ForceClosedBy?.ToThinApiResponse();
 			response.WorkflowThreadUrl = issue.WorkflowThreadUrl;
-			
+
 			if (issue.Fingerprints != null && issue.Fingerprints.Count > 0)
 			{
 				response.FingerprintDescription = String.Join(", ", issue.Fingerprints.Select(x =>
diff --git a/Engine/Source/Programs/Horde/Plugins/Build/HordeServer.Build/Jobs/Bisect/IBisectTask.cs b/Engine/Source/Programs/Horde/Plugins/Build/HordeServer.Build/Jobs/Bisect/IBisectTask.cs
--- a/Engine/Source/Programs/Horde/Plugins/Build/HordeServer.Build/Jobs/Bisect/IBisectTask.cs
+++ b/Engine/Source/Programs/Horde/Plugins/Build/HordeServer.Build/Jobs/Bisect/IBisectTask.cs
@@ -1,6 +1,5 @@
 // Copyright Epic Games, Inc. All Rights Reserved.
 
-using System.Collections.Generic;
 using EpicGames.Horde.Commits;
 using EpicGames.Horde.Jobs;
 using EpicGames.Horde.Jobs.Bisect;
diff --git a/Engine/Source/Programs/Horde/Plugins/Build/HordeServer.Build/Jobs/Graphs/GraphCollection.cs b/Engine/Source/Programs/Horde/Plugins/Build/HordeServer.Build/Jobs/Graphs/GraphCollection.cs
--- a/Engine/Source/Programs/Horde/Plugins/Build/HordeServer.Build/Jobs/Graphs/GraphCollection.cs
+++ b/Engine/Source/Programs/Horde/Plugins/Build/HordeServer.Build/Jobs/Graphs/GraphCollection.cs
@@ -1,17 +1,12 @@
 // Copyright Epic Games, Inc. All Rights Reserved.
 
-using System;
-using System.Collections.Generic;
 using System.Diagnostics;
-using System.Linq;
-using System.Threading;
-using System.Threading.Tasks;
 using EpicGames.Core;
 using EpicGames.Horde.Artifacts;
 using EpicGames.Horde.Jobs;
+using HordeCommon;
 using HordeServer.Jobs.Templates;
 using HordeServer.Server;
-using HordeCommon;
 using Microsoft.Extensions.Caching.Memory;
 using MongoDB.Bson;
 using MongoDB.Bson.Serialization.Attributes;
diff --git a/Engine/Source/Programs/Horde/Plugins/Build/HordeServer.Build/Jobs/JobCollection.cs b/Engine/Source/Programs/Horde/Plugins/Build/HordeServer.Build/Jobs/JobCollection.cs
--- a/Engine/Source/Programs/Horde/Plugins/Build/HordeServer.Build/Jobs/JobCollection.cs
+++ b/Engine/Source/Programs/Horde/Plugins/Build/HordeServer.Build/Jobs/JobCollection.cs
@@ -1,15 +1,9 @@
 // Copyright Epic Games, Inc. All Rights Reserved.
 
-using System;
-using System.Collections.Generic;
 using System.Diagnostics;
 using System.Diagnostics.CodeAnalysis;
-using System.IO;
-using System.Linq;
 using System.Runtime.CompilerServices;
 using System.Text;
-using System.Threading;
-using System.Threading.Tasks;
 using EpicGames.Core;
 using EpicGames.Horde.Agents;
 using EpicGames.Horde.Agents.Leases;
@@ -1044,7 +1038,7 @@ namespace HordeServer.Jobs
 			if (!String.IsNullOrEmpty(cancellationReason))
 			{
 				jobDocument.CancellationReason = cancellationReason;
-				updates.Add(updateBuilder.Set(x => x.CancellationReason, jobDocument.CancellationReason));				
+				updates.Add(updateBuilder.Set(x => x.CancellationReason, jobDocument.CancellationReason));
 			}
 
 			if (notificationTriggerId != null)
@@ -1269,7 +1263,7 @@ namespace HordeServer.Jobs
 							// Update the request abort status
 							if (newAbortRequested != null && step.AbortRequested == false)
 							{
-								step.AbortRequested = newAbortRequested.Value;								
+								step.AbortRequested = newAbortRequested.Value;
 								updates.Add(updateBuilder.Set(x => x.Batches[batchIdx].Steps[stepIdx].AbortRequested, step.AbortRequested));
 
 								// If the step is pending, and not running on an agent, set to aborted
diff --git a/Engine/Source/Programs/Horde/Plugins/Build/HordeServer.Build/Jobs/JobExpirationService.cs b/Engine/Source/Programs/Horde/Plugins/Build/HordeServer.Build/Jobs/JobExpirationService.cs
--- a/Engine/Source/Programs/Horde/Plugins/Build/HordeServer.Build/Jobs/JobExpirationService.cs
+++ b/Engine/Source/Programs/Horde/Plugins/Build/HordeServer.Build/Jobs/JobExpirationService.cs
@@ -1,11 +1,6 @@
 ﻿// Copyright Epic Games, Inc. All Rights Reserved.
 
-using System;
-using System.Collections.Generic;
-using System.Threading;
-using System.Threading.Tasks;
 using EpicGames.Core;
-using HordeServer.Server;
 using HordeServer.Streams;
 using Microsoft.Extensions.Hosting;
 using Microsoft.Extensions.Logging;
diff --git a/Engine/Source/Programs/Horde/Plugins/Build/HordeServer.Build/Jobs/JobRpcService.cs b/Engine/Source/Programs/Horde/Plugins/Build/HordeServer.Build/Jobs/JobRpcService.cs
--- a/Engine/Source/Programs/Horde/Plugins/Build/HordeServer.Build/Jobs/JobRpcService.cs
+++ b/Engine/Source/Programs/Horde/Plugins/Build/HordeServer.Build/Jobs/JobRpcService.cs
@@ -1,16 +1,10 @@
 // Copyright Epic Games, Inc. All Rights Reserved.
 
-using System;
-using System.Collections.Generic;
 using System.Globalization;
-using System.Linq;
-using System.Net.Http;
 using System.Security.Claims;
 using System.Text;
 using System.Text.Json;
 using System.Text.Json.Serialization;
-using System.Threading;
-using System.Threading.Tasks;
 using EpicGames.Core;
 using EpicGames.Horde;
 using EpicGames.Horde.Agents;
diff --git a/Engine/Source/Programs/Horde/Plugins/Build/HordeServer.Build/Jobs/JobTaskSource.cs b/Engine/Source/Programs/Horde/Plugins/Build/HordeServer.Build/Jobs/JobTaskSource.cs
--- a/Engine/Source/Programs/Horde/Plugins/Build/HordeServer.Build/Jobs/JobTaskSource.cs
+++ b/Engine/Source/Programs/Horde/Plugins/Build/HordeServer.Build/Jobs/JobTaskSource.cs
@@ -1,14 +1,9 @@
 // Copyright Epic Games, Inc. All Rights Reserved.
 
-using System;
-using System.Collections.Generic;
 using System.Diagnostics;
 using System.Diagnostics.CodeAnalysis;
 using System.Globalization;
-using System.Linq;
 using System.Text;
-using System.Threading;
-using System.Threading.Tasks;
 using EpicGames.Core;
 using EpicGames.Horde.Agents;
 using EpicGames.Horde.Agents.Leases;
@@ -20,6 +15,9 @@ using EpicGames.Horde.Storage;
 using EpicGames.Horde.Streams;
 using Google.Protobuf;
 using Google.Protobuf.WellKnownTypes;
+using HordeCommon;
+using HordeCommon.Rpc.Messages;
+using HordeCommon.Rpc.Tasks;
 using HordeServer.Acls;
 using HordeServer.Agents;
 using HordeServer.Agents.Pools;
@@ -27,15 +25,11 @@ using HordeServer.Jobs.Bisect;
 using HordeServer.Jobs.Graphs;
 using HordeServer.Logs;
 using HordeServer.Perforce;
-using HordeServer.Server;
 using HordeServer.Storage;
 using HordeServer.Streams;
 using HordeServer.Tasks;
 using HordeServer.Ugs;
 using HordeServer.Utilities;
-using HordeCommon;
-using HordeCommon.Rpc.Messages;
-using HordeCommon.Rpc.Tasks;
 using Microsoft.Extensions.Hosting;
 using Microsoft.Extensions.Logging;
 using Microsoft.Extensions.Options;
@@ -187,12 +181,12 @@ namespace HordeServer.Jobs
 				Agent = agent;
 			}
 		}
-		
+
 		/// <summary>
 		/// Use lazy initialization of agent service to prevent circular dependency as AgentService depend on this class
 		/// </summary>
 		readonly Lazy<AgentService> _agentService;
-		
+
 		readonly IAclService _aclService;
 		readonly IStreamCollection _streamCollection;
 		readonly ILogCollection _logCollection;
@@ -369,7 +363,7 @@ namespace HordeServer.Jobs
 		internal async ValueTask TickAsync(CancellationToken cancellationToken)
 		{
 			using TelemetrySpan span = _tracer.StartActiveSpan($"{nameof(JobTaskSource)}.{nameof(TickAsync)}");
-			
+
 			// Set the NewBatchIdToQueueItem member, so we capture any updated jobs during the DB query.
 			lock (_lockObject)
 			{
@@ -377,7 +371,7 @@ namespace HordeServer.Jobs
 			}
 
 			using TelemetrySpan setupSpan = _tracer.StartActiveSpan($"{nameof(JobTaskSource)}.{nameof(TickAsync)}.Setup");
-			
+
 			// Query all the current streams
 			BuildConfig buildConfig = _buildConfig.CurrentValue;
 			IReadOnlyList<IStream> streamsList = await _streamCollection.GetAsync(buildConfig.Streams, cancellationToken);
@@ -400,7 +394,7 @@ namespace HordeServer.Jobs
 			span.SetAttribute("numAgents", agents.Count);
 			span.SetAttribute("numPools", pools.Count);
 			span.SetAttribute("numStreams", streamsList.Count);
-			
+
 			for (int idx = 0; idx < newJobs.Count; idx++)
 			{
 				IJob? newJob = newJobs[idx];
@@ -776,15 +770,15 @@ namespace HordeServer.Jobs
 			IJob job = item._job;
 			IJobStepBatch batch = item.Batch;
 			IAgent agent = waiter.Agent;
-			LeaseId leaseId = new (BinaryIdUtils.CreateNew()); // Generate a new unique id for the lease
-			
+			LeaseId leaseId = new(BinaryIdUtils.CreateNew()); // Generate a new unique id for the lease
+
 			using IDisposable logScope = _logger
 				.WithProperty("JobId", job.Id.ToString())
 				.WithProperty("BatchId", batch.Id.ToString())
 				.WithProperty("AgentId", agent.Id.ToString())
 				.WithProperty("LeaseId", leaseId.ToString())
 				.BeginScope();
-			
+
 			_logger.LogInformation("Assigning job to waiter");
 
 			// Allocate a log ID but hold off creating the actual log file until the lease has been accepted
@@ -912,7 +906,7 @@ namespace HordeServer.Jobs
 			task.BatchId = batch.Id.ToString();
 			task.LogId = logId.ToString();
 			task.JobName = leaseName.ToString();
-			task.JobOptions = (job.JobOptions != null)? GetRpcJobOptions(job.JobOptions) : null;
+			task.JobOptions = (job.JobOptions != null) ? GetRpcJobOptions(job.JobOptions) : null;
 			task.NamespaceId = namespaceId.ToString();
 			task.StoragePrefix = storagePrefix;
 			task.Token = await _aclService.IssueBearerTokenAsync(claims, null, cancellationToken);
diff --git a/Engine/Source/Programs/Horde/Plugins/Build/HordeServer.Build/Jobs/Schedules/ScheduleService.cs b/Engine/Source/Programs/Horde/Plugins/Build/HordeServer.Build/Jobs/Schedules/ScheduleService.cs
--- a/Engine/Source/Programs/Horde/Plugins/Build/HordeServer.Build/Jobs/Schedules/ScheduleService.cs
+++ b/Engine/Source/Programs/Horde/Plugins/Build/HordeServer.Build/Jobs/Schedules/ScheduleService.cs
@@ -1,12 +1,7 @@
 // Copyright Epic Games, Inc. All Rights Reserved.
 
-using System;
-using System.Collections.Generic;
 using System.Diagnostics;
-using System.Linq;
 using System.Text.RegularExpressions;
-using System.Threading;
-using System.Threading.Tasks;
 using EpicGames.Core;
 using EpicGames.Horde.Commits;
 using EpicGames.Horde.Jobs;
@@ -21,7 +16,6 @@ using HordeServer.Jobs.Templates;
 using HordeServer.Server;
 using HordeServer.Streams;
 using HordeServer.Utilities;
-using HordeCommon;
 using Microsoft.Extensions.Hosting;
 using Microsoft.Extensions.Logging;
 using Microsoft.Extensions.Options;
diff --git a/Engine/Source/Programs/Horde/Plugins/Build/HordeServer.Build/Jobs/Templates/TemplateCollection.cs b/Engine/Source/Programs/Horde/Plugins/Build/HordeServer.Build/Jobs/Templates/TemplateCollection.cs
--- a/Engine/Source/Programs/Horde/Plugins/Build/HordeServer.Build/Jobs/Templates/TemplateCollection.cs
+++ b/Engine/Source/Programs/Horde/Plugins/Build/HordeServer.Build/Jobs/Templates/TemplateCollection.cs
@@ -1,9 +1,5 @@
 // Copyright Epic Games, Inc. All Rights Reserved.
 
-using System;
-using System.Collections.Generic;
-using System.Linq;
-using System.Threading.Tasks;
 using EpicGames.Core;
 using EpicGames.Horde.Jobs;
 using HordeServer.Server;
diff --git a/Engine/Source/Programs/Horde/Plugins/Build/HordeServer.Build/Jobs/Templates/TemplateConfig.cs b/Engine/Source/Programs/Horde/Plugins/Build/HordeServer.Build/Jobs/Templates/TemplateConfig.cs
--- a/Engine/Source/Programs/Horde/Plugins/Build/HordeServer.Build/Jobs/Templates/TemplateConfig.cs
+++ b/Engine/Source/Programs/Horde/Plugins/Build/HordeServer.Build/Jobs/Templates/TemplateConfig.cs
@@ -1,9 +1,6 @@
 // Copyright Epic Games, Inc. All Rights Reserved.
 
-using System;
-using System.Collections.Generic;
 using System.ComponentModel.DataAnnotations;
-using System.IO;
 using System.Text.Json.Serialization;
 using EpicGames.Core;
 using EpicGames.Horde;
diff --git a/Engine/Source/Programs/Horde/Plugins/Build/HordeServer.Build/Jobs/Timing/IJobTimingCollection.cs b/Engine/Source/Programs/Horde/Plugins/Build/HordeServer.Build/Jobs/Timing/IJobTimingCollection.cs
--- a/Engine/Source/Programs/Horde/Plugins/Build/HordeServer.Build/Jobs/Timing/IJobTimingCollection.cs
+++ b/Engine/Source/Programs/Horde/Plugins/Build/HordeServer.Build/Jobs/Timing/IJobTimingCollection.cs
@@ -1,7 +1,5 @@
 // Copyright Epic Games, Inc. All Rights Reserved.
 
-using System.Collections.Generic;
-using System.Threading.Tasks;
 using EpicGames.Horde.Jobs;
 using MongoDB.Bson.Serialization.Attributes;
 
diff --git a/Engine/Source/Programs/Horde/Plugins/Build/HordeServer.Build/Jobs/Timing/JobTimingCollection.cs b/Engine/Source/Programs/Horde/Plugins/Build/HordeServer.Build/Jobs/Timing/JobTimingCollection.cs
--- a/Engine/Source/Programs/Horde/Plugins/Build/HordeServer.Build/Jobs/Timing/JobTimingCollection.cs
+++ b/Engine/Source/Programs/Horde/Plugins/Build/HordeServer.Build/Jobs/Timing/JobTimingCollection.cs
@@ -1,9 +1,6 @@
 // Copyright Epic Games, Inc. All Rights Reserved.
 
-using System.Collections.Generic;
 using System.Diagnostics.CodeAnalysis;
-using System.Linq;
-using System.Threading.Tasks;
 using EpicGames.Horde.Jobs;
 using HordeServer.Server;
 using HordeServer.Utilities;
diff --git a/Engine/Source/Programs/Horde/Plugins/Build/HordeServer.Build/Notifications/NotificationTriggerCollection.cs b/Engine/Source/Programs/Horde/Plugins/Build/HordeServer.Build/Notifications/NotificationTriggerCollection.cs
--- a/Engine/Source/Programs/Horde/Plugins/Build/HordeServer.Build/Notifications/NotificationTriggerCollection.cs
+++ b/Engine/Source/Programs/Horde/Plugins/Build/HordeServer.Build/Notifications/NotificationTriggerCollection.cs
@@ -1,9 +1,5 @@
 // Copyright Epic Games, Inc. All Rights Reserved.
 
-using System.Collections.Generic;
-using System.Linq;
-using System.Threading;
-using System.Threading.Tasks;
 using EpicGames.Horde.Users;
 using HordeServer.Server;
 using HordeServer.Users;
diff --git a/Engine/Source/Programs/Horde/Plugins/Build/HordeServer.Build/Notifications/Sinks/SlackNotificationSink.cs b/Engine/Source/Programs/Horde/Plugins/Build/HordeServer.Build/Notifications/Sinks/SlackNotificationSink.cs
--- a/Engine/Source/Programs/Horde/Plugins/Build/HordeServer.Build/Notifications/Sinks/SlackNotificationSink.cs
+++ b/Engine/Source/Programs/Horde/Plugins/Build/HordeServer.Build/Notifications/Sinks/SlackNotificationSink.cs
@@ -1,20 +1,13 @@
 // Copyright Epic Games, Inc. All Rights Reserved.
 
-using System;
-using System.Collections.Generic;
 using System.Diagnostics;
 using System.Diagnostics.CodeAnalysis;
 using System.Globalization;
-using System.IO;
-using System.Linq;
-using System.Net.Http;
 using System.Net.WebSockets;
 using System.Text;
 using System.Text.Json;
 using System.Text.Json.Serialization;
 using System.Text.RegularExpressions;
-using System.Threading;
-using System.Threading.Tasks;
 using EpicGames.Core;
 using EpicGames.Horde;
 using EpicGames.Horde.Agents;
diff --git a/Engine/Source/Programs/Horde/Plugins/Build/HordeServer.Build/Perforce/PerforceConfigSource.cs b/Engine/Source/Programs/Horde/Plugins/Build/HordeServer.Build/Perforce/PerforceConfigSource.cs
--- a/Engine/Source/Programs/Horde/Plugins/Build/HordeServer.Build/Perforce/PerforceConfigSource.cs
+++ b/Engine/Source/Programs/Horde/Plugins/Build/HordeServer.Build/Perforce/PerforceConfigSource.cs
@@ -1,11 +1,5 @@
 // Copyright Epic Games, Inc. All Rights Reserved.
 
-using System;
-using System.Collections.Generic;
-using System.IO;
-using System.Linq;
-using System.Threading;
-using System.Threading.Tasks;
 using EpicGames.Core;
 using EpicGames.Horde.Users;
 using EpicGames.Perforce;
diff --git a/Engine/Source/Programs/Horde/Plugins/Build/HordeServer.Build/Perforce/PerforceLoadBalancer.cs b/Engine/Source/Programs/Horde/Plugins/Build/HordeServer.Build/Perforce/PerforceLoadBalancer.cs
--- a/Engine/Source/Programs/Horde/Plugins/Build/HordeServer.Build/Perforce/PerforceLoadBalancer.cs
+++ b/Engine/Source/Programs/Horde/Plugins/Build/HordeServer.Build/Perforce/PerforceLoadBalancer.cs
@@ -1,25 +1,18 @@
 // Copyright Epic Games, Inc. All Rights Reserved.
 
-using System;
-using System.Collections.Generic;
 using System.Globalization;
-using System.Linq;
 using System.Net;
-using System.Net.Http;
 using System.Text;
 using System.Text.Json;
-using System.Threading;
-using System.Threading.Tasks;
 using EpicGames.Core;
 using EpicGames.Perforce;
 using Google.Protobuf.WellKnownTypes;
+using HordeCommon.Rpc.Messages;
+using HordeCommon.Rpc.Tasks;
 using HordeServer.Agents;
 using HordeServer.Agents.Leases;
 using HordeServer.Server;
 using HordeServer.Utilities;
-using HordeCommon;
-using HordeCommon.Rpc.Messages;
-using HordeCommon.Rpc.Tasks;
 using Microsoft.Extensions.Diagnostics.HealthChecks;
 using Microsoft.Extensions.Hosting;
 using Microsoft.Extensions.Logging;
@@ -495,7 +488,7 @@ namespace HordeServer.Perforce
 			{
 				HealthStatus.Unhealthy => "One or more Perforce servers are unhealthy. Check Perforce servers page for details.",
 				HealthStatus.Degraded => "One or more Perforce servers are degraded. Check Perforce servers page for details.",
-				HealthStatus.Healthy => null, 
+				HealthStatus.Healthy => null,
 				_ => throw new ArgumentOutOfRangeException($"Unknown health status: {result}")
 			};
 
diff --git a/Engine/Source/Programs/Horde/Plugins/Build/HordeServer.Build/Perforce/PerforceService.cs b/Engine/Source/Programs/Horde/Plugins/Build/HordeServer.Build/Perforce/PerforceService.cs
--- a/Engine/Source/Programs/Horde/Plugins/Build/HordeServer.Build/Perforce/PerforceService.cs
+++ b/Engine/Source/Programs/Horde/Plugins/Build/HordeServer.Build/Perforce/PerforceService.cs
@@ -1,25 +1,18 @@
 // Copyright Epic Games, Inc. All Rights Reserved.
 
-using System;
-using System.Collections.Generic;
 using System.Diagnostics.CodeAnalysis;
 using System.Globalization;
-using System.IO;
-using System.Linq;
 using System.Net;
 using System.Runtime.CompilerServices;
 using System.Runtime.InteropServices;
 using System.Text;
 using System.Text.RegularExpressions;
-using System.Threading;
-using System.Threading.Tasks;
 using EpicGames.Core;
 using EpicGames.Horde.Commits;
 using EpicGames.Horde.Streams;
 using EpicGames.Horde.Users;
 using EpicGames.Perforce;
 using HordeServer.Commits;
-using HordeServer.Server;
 using HordeServer.Streams;
 using HordeServer.Users;
 using HordeServer.Utilities;
diff --git a/Engine/Source/Programs/Horde/Plugins/Build/HordeServer.Build/Perforce/PerforceServiceCache.cs b/Engine/Source/Programs/Horde/Plugins/Build/HordeServer.Build/Perforce/PerforceServiceCache.cs
--- a/Engine/Source/Programs/Horde/Plugins/Build/HordeServer.Build/Perforce/PerforceServiceCache.cs
+++ b/Engine/Source/Programs/Horde/Plugins/Build/HordeServer.Build/Perforce/PerforceServiceCache.cs
@@ -1,13 +1,7 @@
 // Copyright Epic Games, Inc. All Rights Reserved.
 
-using System;
-using System.Collections.Generic;
-using System.IO;
-using System.Linq;
 using System.Runtime.CompilerServices;
 using System.Text;
-using System.Threading;
-using System.Threading.Tasks;
 using EpicGames.Core;
 using EpicGames.Horde.Commits;
 using EpicGames.Horde.Streams;
diff --git a/Engine/Source/Programs/Horde/Plugins/Build/HordeServer.Build/Perforce/PerforceWorkspace.cs b/Engine/Source/Programs/Horde/Plugins/Build/HordeServer.Build/Perforce/PerforceWorkspace.cs
--- a/Engine/Source/Programs/Horde/Plugins/Build/HordeServer.Build/Perforce/PerforceWorkspace.cs
+++ b/Engine/Source/Programs/Horde/Plugins/Build/HordeServer.Build/Perforce/PerforceWorkspace.cs
@@ -1,17 +1,11 @@
 // Copyright Epic Games, Inc. All Rights Reserved.
 
-using System;
-using System.Collections.Generic;
 using System.Diagnostics.CodeAnalysis;
-using System.Linq;
-using System.Threading;
-using System.Threading.Tasks;
 using EpicGames.Horde.Agents.Pools;
+using HordeCommon.Rpc.Messages;
 using HordeServer.Agents;
 using HordeServer.Agents.Pools;
-using HordeServer.Server;
 using HordeServer.Streams;
-using HordeCommon.Rpc.Messages;
 
 namespace HordeServer.Perforce
 {
diff --git a/Engine/Source/Programs/Horde/Plugins/Build/HordeServer.Build/Projects/ProjectConfig.cs b/Engine/Source/Programs/Horde/Plugins/Build/HordeServer.Build/Projects/ProjectConfig.cs
--- a/Engine/Source/Programs/Horde/Plugins/Build/HordeServer.Build/Projects/ProjectConfig.cs
+++ b/Engine/Source/Programs/Horde/Plugins/Build/HordeServer.Build/Projects/ProjectConfig.cs
@@ -1,11 +1,8 @@
 // Copyright Epic Games, Inc. All Rights Reserved.
 
-using System;
-using System.Collections.Generic;
 using System.ComponentModel.DataAnnotations;
 using System.Diagnostics;
 using System.Security.Claims;
-using System.Text.Json.Serialization;
 using EpicGames.Core;
 using EpicGames.Horde.Acls;
 using EpicGames.Horde.Jobs;
@@ -14,7 +11,6 @@ using EpicGames.Horde.Telemetry;
 using HordeServer.Acls;
 using HordeServer.Agents.Pools;
 using HordeServer.Configuration;
-using HordeServer.Server;
 using HordeServer.Streams;
 using HordeServer.Utilities;
 
diff --git a/Engine/Source/Programs/Horde/Plugins/Build/HordeServer.Build/Projects/ProjectsController.cs b/Engine/Source/Programs/Horde/Plugins/Build/HordeServer.Build/Projects/ProjectsController.cs
--- a/Engine/Source/Programs/Horde/Plugins/Build/HordeServer.Build/Projects/ProjectsController.cs
+++ b/Engine/Source/Programs/Horde/Plugins/Build/HordeServer.Build/Projects/ProjectsController.cs
@@ -1,11 +1,8 @@
 // Copyright Epic Games, Inc. All Rights Reserved.
 
-using System.Collections.Generic;
-using System.Linq;
 using System.Text.RegularExpressions;
 using EpicGames.Horde.Projects;
 using HordeServer.Configuration;
-using HordeServer.Server;
 using HordeServer.Streams;
 using HordeServer.Utilities;
 using Microsoft.AspNetCore.Authorization;
diff --git a/Engine/Source/Programs/Horde/Plugins/Build/HordeServer.Build/Replicators/PerforceReplicator.cs b/Engine/Source/Programs/Horde/Plugins/Build/HordeServer.Build/Replicators/PerforceReplicator.cs
--- a/Engine/Source/Programs/Horde/Plugins/Build/HordeServer.Build/Replicators/PerforceReplicator.cs
+++ b/Engine/Source/Programs/Horde/Plugins/Build/HordeServer.Build/Replicators/PerforceReplicator.cs
@@ -1,14 +1,9 @@
 // Copyright Epic Games, Inc. All Rights Reserved.
 
-using System;
 using System.Buffers.Binary;
-using System.Collections.Generic;
 using System.Diagnostics;
-using System.Linq;
 using System.Security.Cryptography;
 using System.Text;
-using System.Threading;
-using System.Threading.Tasks;
 using EpicGames.Core;
 using EpicGames.Horde.Replicators;
 using EpicGames.Horde.Storage;
@@ -18,7 +13,6 @@ using EpicGames.Horde.Streams;
 using EpicGames.Perforce;
 using HordeServer.Commits;
 using HordeServer.Perforce;
-using HordeServer.Server;
 using HordeServer.Storage;
 using HordeServer.Streams;
 using Microsoft.Extensions.Logging;
@@ -30,9 +24,24 @@ namespace HordeServer.Replicators
 	/// </summary>
 	public class PerforceReplicationOptions
 	{
+		/// <summary>
+		/// Whether to replicate content, or just metadata
+		/// </summary>
 		public bool IncludeContent { get; set; }
+
+		/// <summary>
+		/// 
+		/// </summary>
 		public BundleOptions TreeOptions { get; set; } = new BundleOptions();
+
+		/// <summary>
+		/// 
+		/// </summary>
 		public ChunkingOptions ChunkingOptions { get; set; } = new ChunkingOptions();
+
+		/// <summary>
+		/// 
+		/// </summary>
 		public RefOptions RefOptions { get; set; } = new RefOptions();
 	}
 
diff --git a/Engine/Source/Programs/Horde/Plugins/Build/HordeServer.Build/Replicators/ReplicatorsController.cs b/Engine/Source/Programs/Horde/Plugins/Build/HordeServer.Build/Replicators/ReplicatorsController.cs
--- a/Engine/Source/Programs/Horde/Plugins/Build/HordeServer.Build/Replicators/ReplicatorsController.cs
+++ b/Engine/Source/Programs/Horde/Plugins/Build/HordeServer.Build/Replicators/ReplicatorsController.cs
@@ -1,13 +1,9 @@
 ﻿// Copyright Epic Games, Inc. All Rights Reserved.
 
-using System.Collections.Generic;
 using System.Text;
-using System.Threading;
-using System.Threading.Tasks;
 using EpicGames.Horde.Replicators;
 using EpicGames.Horde.Streams;
 using HordeServer.Projects;
-using HordeServer.Server;
 using HordeServer.Streams;
 using HordeServer.Utilities;
 using Microsoft.AspNetCore.Authorization;
diff --git a/Engine/Source/Programs/Horde/Plugins/Build/HordeServer.Build/StaticBuildConfig.cs b/Engine/Source/Programs/Horde/Plugins/Build/HordeServer.Build/StaticBuildConfig.cs
--- a/Engine/Source/Programs/Horde/Plugins/Build/HordeServer.Build/StaticBuildConfig.cs
+++ b/Engine/Source/Programs/Horde/Plugins/Build/HordeServer.Build/StaticBuildConfig.cs
@@ -1,11 +1,6 @@
 ﻿// Copyright Epic Games, Inc. All Rights Reserved.
 
-using System;
-using System.Collections.Generic;
 using System.ComponentModel;
-using System.Linq;
-using System.Text;
-using System.Threading.Tasks;
 using EpicGames.Core;
 using EpicGames.Horde;
 using EpicGames.Horde.Storage.Bundles;
@@ -14,6 +9,9 @@ using EpicGames.Perforce;
 
 namespace HordeServer
 {
+	/// <summary>
+	/// Static configuration for the build plugin
+	/// </summary>
 	public class StaticBuildConfig
 	{
 		/// <summary>
diff --git a/Engine/Source/Programs/Horde/Plugins/Build/HordeServer.Build/Streams/StreamConfig.cs b/Engine/Source/Programs/Horde/Plugins/Build/HordeServer.Build/Streams/StreamConfig.cs
--- a/Engine/Source/Programs/Horde/Plugins/Build/HordeServer.Build/Streams/StreamConfig.cs
+++ b/Engine/Source/Programs/Horde/Plugins/Build/HordeServer.Build/Streams/StreamConfig.cs
@@ -1,11 +1,8 @@
 // Copyright Epic Games, Inc. All Rights Reserved.
 
-using System;
-using System.Collections.Generic;
 using System.ComponentModel.DataAnnotations;
 using System.Diagnostics;
 using System.Diagnostics.CodeAnalysis;
-using System.Linq;
 using System.Security.Claims;
 using System.Text.Json;
 using System.Text.Json.Serialization;
@@ -30,7 +27,6 @@ using HordeServer.Jobs.Graphs;
 using HordeServer.Jobs.Templates;
 using HordeServer.Projects;
 using HordeServer.Replicators;
-using HordeServer.Server;
 using HordeServer.Utilities;
 
 #pragma warning disable CA2227 // Change x to be read-only by removing the property setter
diff --git a/Engine/Source/Programs/Horde/Plugins/Build/HordeServer.Build/Streams/StreamsController.cs b/Engine/Source/Programs/Horde/Plugins/Build/HordeServer.Build/Streams/StreamsController.cs
--- a/Engine/Source/Programs/Horde/Plugins/Build/HordeServer.Build/Streams/StreamsController.cs
+++ b/Engine/Source/Programs/Horde/Plugins/Build/HordeServer.Build/Streams/StreamsController.cs
@@ -1,10 +1,5 @@
 // Copyright Epic Games, Inc. All Rights Reserved.
 
-using System;
-using System.Collections.Generic;
-using System.Linq;
-using System.Threading;
-using System.Threading.Tasks;
 using EpicGames.Core;
 using EpicGames.Horde.Commits;
 using EpicGames.Horde.Jobs;
@@ -23,7 +18,6 @@ using Microsoft.AspNetCore.Authorization;
 using Microsoft.AspNetCore.Mvc;
 using Microsoft.Extensions.Options;
 using OpenTelemetry.Trace;
-using TimeZoneConverter;
 
 namespace HordeServer.Streams
 {
@@ -47,7 +41,7 @@ namespace HordeServer.Streams
 		/// <summary>
 		/// Constructor
 		/// </summary>
-		public StreamsController(IStreamCollection streamCollection, ICommitService commitService, ITemplateCollection templateCollection, IJobStepRefCollection jobStepRefCollection, IUserCollection userCollection, Tracer tracer, IClock clock, IOptionsSnapshot<BuildConfig> buildConfig, IOptions<StaticBuildConfig> staticBuildConfig)
+		public StreamsController(IStreamCollection streamCollection, ICommitService commitService, ITemplateCollection templateCollection, IJobStepRefCollection jobStepRefCollection, IUserCollection userCollection, Tracer tracer, IClock clock, IOptionsSnapshot<BuildConfig> buildConfig)
 		{
 			_streamCollection = streamCollection;
 			_commitService = commitService;
diff --git a/Engine/Source/Programs/Horde/Plugins/Build/HordeServer.Build/Tasks/ConformTaskSource.cs b/Engine/Source/Programs/Horde/Plugins/Build/HordeServer.Build/Tasks/ConformTaskSource.cs
--- a/Engine/Source/Programs/Horde/Plugins/Build/HordeServer.Build/Tasks/ConformTaskSource.cs
+++ b/Engine/Source/Programs/Horde/Plugins/Build/HordeServer.Build/Tasks/ConformTaskSource.cs
@@ -1,10 +1,5 @@
 // Copyright Epic Games, Inc. All Rights Reserved.
 
-using System;
-using System.Collections.Generic;
-using System.Linq;
-using System.Threading;
-using System.Threading.Tasks;
 using EpicGames.Core;
 using EpicGames.Horde.Agents;
 using EpicGames.Horde.Agents.Leases;
@@ -12,6 +7,8 @@ using EpicGames.Horde.Jobs;
 using EpicGames.Horde.Logs;
 using Google.Protobuf;
 using Google.Protobuf.WellKnownTypes;
+using HordeCommon.Rpc.Messages;
+using HordeCommon.Rpc.Tasks;
 using HordeServer.Agents;
 using HordeServer.Agents.Pools;
 using HordeServer.Jobs;
@@ -19,8 +16,6 @@ using HordeServer.Logs;
 using HordeServer.Perforce;
 using HordeServer.Server;
 using HordeServer.Utilities;
-using HordeCommon.Rpc.Messages;
-using HordeCommon.Rpc.Tasks;
 using Microsoft.Extensions.Hosting;
 using Microsoft.Extensions.Logging;
 using Microsoft.Extensions.Options;
diff --git a/Engine/Source/Programs/Horde/Plugins/Build/HordeServer.Build/Ugs/IUgsMetadataCollection.cs b/Engine/Source/Programs/Horde/Plugins/Build/HordeServer.Build/Ugs/IUgsMetadataCollection.cs
--- a/Engine/Source/Programs/Horde/Plugins/Build/HordeServer.Build/Ugs/IUgsMetadataCollection.cs
+++ b/Engine/Source/Programs/Horde/Plugins/Build/HordeServer.Build/Ugs/IUgsMetadataCollection.cs
@@ -1,10 +1,5 @@
 // Copyright Epic Games, Inc. All Rights Reserved.
 
-using System;
-using System.Collections.Generic;
-using System.Threading;
-using System.Threading.Tasks;
-
 namespace HordeServer.Ugs
 {
 	/// <summary>
diff --git a/Engine/Source/Programs/Horde/Plugins/Build/HordeServer.Build/Ugs/UgsController.cs b/Engine/Source/Programs/Horde/Plugins/Build/HordeServer.Build/Ugs/UgsController.cs
--- a/Engine/Source/Programs/Horde/Plugins/Build/HordeServer.Build/Ugs/UgsController.cs
+++ b/Engine/Source/Programs/Horde/Plugins/Build/HordeServer.Build/Ugs/UgsController.cs
@@ -1,10 +1,5 @@
 // Copyright Epic Games, Inc. All Rights Reserved.
 
-using System;
-using System.Collections.Generic;
-using System.Linq;
-using System.Threading;
-using System.Threading.Tasks;
 using EpicGames.Horde.Issues;
 using EpicGames.Horde.Users;
 using HordeServer.Issues;
@@ -13,7 +8,6 @@ using HordeServer.Users;
 using HordeServer.Utilities;
 using Microsoft.AspNetCore.Mvc;
 using Microsoft.Extensions.Logging;
-using Microsoft.Extensions.Options;
 using MongoDB.Bson;
 
 namespace HordeServer.Ugs
@@ -30,20 +24,18 @@ namespace HordeServer.Ugs
 		readonly IUserCollection _userCollection;
 		readonly ILogCollection _logCollection;
 		readonly IServerInfo _serverInfo;
-		readonly StaticBuildConfig _settings;
 		readonly ILogger _logger;
 
 		/// <summary>
 		/// Constructor
 		/// </summary>
-		public UgsController(IssueService issueService, IUgsMetadataCollection ugsMetadataCollection, IUserCollection userCollection, ILogCollection logCollection, IServerInfo serverInfo, IOptionsMonitor<StaticBuildConfig> optionsMonitor, ILogger<UgsController> logger)
+		public UgsController(IssueService issueService, IUgsMetadataCollection ugsMetadataCollection, IUserCollection userCollection, ILogCollection logCollection, IServerInfo serverInfo, ILogger<UgsController> logger)
 		{
 			_issueService = issueService;
 			_ugsMetadataCollection = ugsMetadataCollection;
 			_userCollection = userCollection;
 			_logCollection = logCollection;
 			_serverInfo = serverInfo;
-			_settings = optionsMonitor.CurrentValue;
 			_logger = logger;
 		}
 
diff --git a/Engine/Source/Programs/Horde/Plugins/Build/HordeServer.Build/Utilities/OpenTelemetryExtensions.cs b/Engine/Source/Programs/Horde/Plugins/Build/HordeServer.Build/Utilities/OpenTelemetryExtensions.cs
--- a/Engine/Source/Programs/Horde/Plugins/Build/HordeServer.Build/Utilities/OpenTelemetryExtensions.cs
+++ b/Engine/Source/Programs/Horde/Plugins/Build/HordeServer.Build/Utilities/OpenTelemetryExtensions.cs
@@ -1,35 +1,10 @@
 // Copyright Epic Games, Inc. All Rights Reserved.
 
-using System;
-using System.Collections.Generic;
-using System.IO;
-using System.Linq;
-using System.Runtime.CompilerServices;
-using System.Text;
-using System.Threading;
-using System.Threading.Tasks;
 using EpicGames.Core;
-using EpicGames.Horde.Commits;
 using EpicGames.Horde.Jobs;
 using EpicGames.Horde.Jobs.Templates;
 using EpicGames.Horde.Streams;
-using EpicGames.Horde.Users;
-using EpicGames.Perforce;
-using EpicGames.Redis;
-using HordeServer.Commits;
-using HordeServer.Server;
-using HordeServer.Streams;
-using HordeServer.Users;
-using HordeServer.Utilities;
-using Microsoft.Extensions.Hosting;
-using Microsoft.Extensions.Logging;
-using Microsoft.Extensions.Options;
-using MongoDB.Bson;
-using MongoDB.Bson.Serialization.Attributes;
-using MongoDB.Bson.Serialization.Options;
-using MongoDB.Driver;
 using OpenTelemetry.Trace;
-using StackExchange.Redis;
 
 namespace HordeServer.Utilities
 {
diff --git a/Engine/Source/Programs/Horde/Plugins/Compute/HordeServer.Compute.Tests/Compute/RedisTaskSchedulerTest.cs b/Engine/Source/Programs/Horde/Plugins/Compute/HordeServer.Compute.Tests/Compute/RedisTaskSchedulerTest.cs
--- a/Engine/Source/Programs/Horde/Plugins/Compute/HordeServer.Compute.Tests/Compute/RedisTaskSchedulerTest.cs
+++ b/Engine/Source/Programs/Horde/Plugins/Compute/HordeServer.Compute.Tests/Compute/RedisTaskSchedulerTest.cs
@@ -1,12 +1,8 @@
 // Copyright Epic Games, Inc. All Rights Reserved.
 
-using System;
-using System.Threading;
-using System.Threading.Tasks;
 using HordeServer.Compute;
 using HordeServer.Server;
 using Microsoft.Extensions.Logging;
-using Microsoft.VisualStudio.TestTools.UnitTesting;
 using StackExchange.Redis;
 
 namespace HordeServer.Tests.Compute
diff --git a/Engine/Source/Programs/Horde/Plugins/Compute/HordeServer.Compute.Tests/Compute/TunnelServiceTests.cs b/Engine/Source/Programs/Horde/Plugins/Compute/HordeServer.Compute.Tests/Compute/TunnelServiceTests.cs
--- a/Engine/Source/Programs/Horde/Plugins/Compute/HordeServer.Compute.Tests/Compute/TunnelServiceTests.cs
+++ b/Engine/Source/Programs/Horde/Plugins/Compute/HordeServer.Compute.Tests/Compute/TunnelServiceTests.cs
@@ -1,19 +1,12 @@
 // Copyright Epic Games, Inc. All Rights Reserved.
 
-using System;
-using System.Collections.Generic;
-using System.IO;
 using System.Net;
 using System.Net.Sockets;
 using System.Text;
-using System.Threading;
-using System.Threading.Tasks;
 using EpicGames.Core;
 using EpicGames.Horde.Compute.Clients;
 using HordeServer.Compute;
-using Microsoft.Extensions.Logging;
 using Microsoft.Extensions.Options;
-using Microsoft.VisualStudio.TestTools.UnitTesting;
 
 namespace HordeServer.Tests.Compute
 {
diff --git a/Engine/Source/Programs/Horde/Plugins/Compute/HordeServer.Compute.Tests/ComputeTestSetup.cs b/Engine/Source/Programs/Horde/Plugins/Compute/HordeServer.Compute.Tests/ComputeTestSetup.cs
--- a/Engine/Source/Programs/Horde/Plugins/Compute/HordeServer.Compute.Tests/ComputeTestSetup.cs
+++ b/Engine/Source/Programs/Horde/Plugins/Compute/HordeServer.Compute.Tests/ComputeTestSetup.cs
@@ -1,16 +1,10 @@
 // Copyright Epic Games, Inc. All Rights Reserved.
 
-using System;
-using System.Collections.Generic;
-using System.Linq;
-using System.Text;
-using System.Threading.Tasks;
 using HordeServer.Acls;
 using HordeServer.Compute;
 using HordeServer.Logs;
 using HordeServer.Storage;
 using Microsoft.Extensions.DependencyInjection;
-using Microsoft.Extensions.Logging;
 
 namespace HordeServer.Tests
 {
diff --git a/Engine/Source/Programs/Horde/Plugins/Compute/HordeServer.Compute.Tests/Logs/LogIndexingTests.cs b/Engine/Source/Programs/Horde/Plugins/Compute/HordeServer.Compute.Tests/Logs/LogIndexingTests.cs
--- a/Engine/Source/Programs/Horde/Plugins/Compute/HordeServer.Compute.Tests/Logs/LogIndexingTests.cs
+++ b/Engine/Source/Programs/Horde/Plugins/Compute/HordeServer.Compute.Tests/Logs/LogIndexingTests.cs
@@ -1,12 +1,6 @@
 // Copyright Epic Games, Inc. All Rights Reserved.
 
-using System;
-using System.Collections.Generic;
-using System.IO;
-using System.Linq;
 using System.Text;
-using System.Threading;
-using System.Threading.Tasks;
 using EpicGames.Horde.Jobs;
 using EpicGames.Horde.Logs;
 using EpicGames.Horde.Storage;
@@ -14,7 +8,6 @@ using HordeServer.Jobs;
 using HordeServer.Logs;
 using HordeServer.Storage;
 using Microsoft.Extensions.Logging.Abstractions;
-using Microsoft.VisualStudio.TestTools.UnitTesting;
 
 namespace HordeServer.Tests.Logs
 {
diff --git a/Engine/Source/Programs/Horde/Plugins/Compute/HordeServer.Compute.Tests/Logs/LogServiceTests.cs b/Engine/Source/Programs/Horde/Plugins/Compute/HordeServer.Compute.Tests/Logs/LogServiceTests.cs
--- a/Engine/Source/Programs/Horde/Plugins/Compute/HordeServer.Compute.Tests/Logs/LogServiceTests.cs
+++ b/Engine/Source/Programs/Horde/Plugins/Compute/HordeServer.Compute.Tests/Logs/LogServiceTests.cs
@@ -1,13 +1,7 @@
 // Copyright Epic Games, Inc. All Rights Reserved.
 
-using System;
-using System.Collections.Generic;
-using System.IO;
-using System.Linq;
 using System.Security.Claims;
 using System.Text;
-using System.Threading;
-using System.Threading.Tasks;
 using EpicGames.Core;
 using EpicGames.Horde.Agents.Sessions;
 using EpicGames.Horde.Jobs;
@@ -16,7 +10,6 @@ using HordeServer.Agents.Sessions;
 using HordeServer.Jobs;
 using HordeServer.Logs;
 using HordeServer.Utilities;
-using Microsoft.VisualStudio.TestTools.UnitTesting;
 
 namespace HordeServer.Tests.Logs
 {
diff --git a/Engine/Source/Programs/Horde/Plugins/Compute/HordeServer.Compute/Agents/AgentCollection.cs b/Engine/Source/Programs/Horde/Plugins/Compute/HordeServer.Compute/Agents/AgentCollection.cs
--- a/Engine/Source/Programs/Horde/Plugins/Compute/HordeServer.Compute/Agents/AgentCollection.cs
+++ b/Engine/Source/Programs/Horde/Plugins/Compute/HordeServer.Compute/Agents/AgentCollection.cs
@@ -1,13 +1,8 @@
 // Copyright Epic Games, Inc. All Rights Reserved.
 
-using System;
-using System.Collections.Generic;
 using System.Collections.ObjectModel;
 using System.Diagnostics.CodeAnalysis;
-using System.Linq;
 using System.Linq.Expressions;
-using System.Threading;
-using System.Threading.Tasks;
 using EpicGames.Core;
 using EpicGames.Horde.Agents;
 using EpicGames.Horde.Agents.Leases;
@@ -15,9 +10,9 @@ using EpicGames.Horde.Agents.Pools;
 using EpicGames.Horde.Agents.Sessions;
 using EpicGames.Redis;
 using Google.Protobuf.WellKnownTypes;
+using HordeCommon.Rpc.Tasks;
 using HordeServer.Auditing;
 using HordeServer.Server;
-using HordeCommon.Rpc.Tasks;
 using MongoDB.Bson;
 using MongoDB.Bson.Serialization.Attributes;
 using MongoDB.Driver;
diff --git a/Engine/Source/Programs/Horde/Plugins/Compute/HordeServer.Compute/Agents/AgentService.cs b/Engine/Source/Programs/Horde/Plugins/Compute/HordeServer.Compute/Agents/AgentService.cs
--- a/Engine/Source/Programs/Horde/Plugins/Compute/HordeServer.Compute/Agents/AgentService.cs
+++ b/Engine/Source/Programs/Horde/Plugins/Compute/HordeServer.Compute/Agents/AgentService.cs
@@ -1,13 +1,8 @@
 // Copyright Epic Games, Inc. All Rights Reserved.
 
-using System;
-using System.Collections.Generic;
 using System.Diagnostics;
 using System.Diagnostics.Metrics;
-using System.Linq;
 using System.Security.Claims;
-using System.Threading;
-using System.Threading.Tasks;
 using EpicGames.Core;
 using EpicGames.Horde.Agents;
 using EpicGames.Horde.Agents.Leases;
@@ -16,6 +11,8 @@ using EpicGames.Horde.Agents.Sessions;
 using EpicGames.Redis;
 using EpicGames.Serialization;
 using Google.Protobuf.WellKnownTypes;
+using HordeCommon.Rpc;
+using HordeCommon.Rpc.Messages;
 using HordeServer.Acls;
 using HordeServer.Agents.Leases;
 using HordeServer.Agents.Pools;
@@ -24,8 +21,6 @@ using HordeServer.Auditing;
 using HordeServer.Server;
 using HordeServer.Tasks;
 using HordeServer.Utilities;
-using HordeCommon.Rpc;
-using HordeCommon.Rpc.Messages;
 using Microsoft.Extensions.Hosting;
 using Microsoft.Extensions.Logging;
 using OpenTelemetry.Trace;
@@ -86,19 +81,19 @@ namespace HordeServer.Agents
 		readonly ITicker _ticker;
 		readonly ITicker _sharedTicker;
 
-		readonly RedisStringKey<AgentRateTable> _agentRateTableData = new ("agent-rates");
+		readonly RedisStringKey<AgentRateTable> _agentRateTableData = new("agent-rates");
 
 		/// <summary>Lazily updated costs for different agent types</summary>
 		readonly AsyncCachedValue<AgentRateTable?> _cachedRates;
 
 		/// <summary>Lazily updated list of current pools</summary>
 		readonly AsyncCachedValue<IReadOnlyList<IPoolConfig>> _cachedPools;
-		
+
 		/// <summary>Manually updated cached list of current agents</summary>
 		IReadOnlyDictionary<AgentId, IAgent>? _cachedAgents;
 
 		/// <summary>All the agents currently performing a long poll for work on this server</summary>
-		readonly Dictionary<AgentId, CancellationTokenSource> _waitingAgents = new ();
+		readonly Dictionary<AgentId, CancellationTokenSource> _waitingAgents = new();
 
 		/// <summary>OpenTelemetry measurements (gauges)</summary>
 		IEnumerable<Measurement<int>> _measurements = new List<Measurement<int>>();
@@ -278,7 +273,7 @@ namespace HordeServer.Agents
 		{
 			return Agents.FindAsync(poolId, modifiedAfter, property, null, null, includeDeleted, index, count, true, cancellationToken);
 		}
-		
+
 		/// <summary>
 		/// Get all agents from local in-memory cache
 		/// </summary>
@@ -1177,7 +1172,7 @@ namespace HordeServer.Agents
 			{
 				agents[agent.Id] = agent;
 			}
-			
+
 			Interlocked.Exchange(ref _cachedAgents, agents);
 		}
 
diff --git a/Engine/Source/Programs/Horde/Plugins/Compute/HordeServer.Compute/Agents/AgentsController.cs b/Engine/Source/Programs/Horde/Plugins/Compute/HordeServer.Compute/Agents/AgentsController.cs
--- a/Engine/Source/Programs/Horde/Plugins/Compute/HordeServer.Compute/Agents/AgentsController.cs
+++ b/Engine/Source/Programs/Horde/Plugins/Compute/HordeServer.Compute/Agents/AgentsController.cs
@@ -1,10 +1,5 @@
 // Copyright Epic Games, Inc. All Rights Reserved.
 
-using System;
-using System.Collections.Generic;
-using System.Linq;
-using System.Threading;
-using System.Threading.Tasks;
 using EpicGames.Core;
 using EpicGames.Horde.Agents;
 using EpicGames.Horde.Agents.Leases;
@@ -71,12 +66,12 @@ namespace HordeServer.Agents
 		public async Task<ActionResult<List<object>>> FindAgentsAsync([FromQuery] PoolId? poolId = null, [FromQuery] Condition? condition = null, [FromQuery] bool includeDeleted = false, [FromQuery] int? index = null, [FromQuery] int? count = null, [FromQuery] DateTimeOffset? modifiedAfter = null, [FromQuery] PropertyFilter? filter = null)
 		{
 			using TelemetrySpan span = _tracer.StartActiveSpan($"{nameof(AgentsController)}.{nameof(FindAgentsAsync)}");
-			
+
 			if (!_computeConfig.Value.Authorize(AgentAclAction.ListAgents, User))
 			{
 				return Forbid(AgentAclAction.ListAgents);
 			}
-			
+
 			IEnumerable<IAgent> agentsEnumerable = (await _agentService.GetCachedAgentsAsync(HttpContext.RequestAborted))
 				.Where(x => poolId == null || x.Pools.Contains(poolId.Value))
 				.Where(x => modifiedAfter == null || modifiedAfter.Value.UtcDateTime >= x.UpdateTime);
@@ -85,10 +80,10 @@ namespace HordeServer.Agents
 			{
 				agentsEnumerable = agentsEnumerable.Where(x => !x.Deleted);
 			}
-			
+
 			agentsEnumerable = index != null ? agentsEnumerable.Skip(index.Value) : agentsEnumerable;
 			agentsEnumerable = count != null ? agentsEnumerable.Take(count.Value) : agentsEnumerable;
-			
+
 			List<IAgent> agents = agentsEnumerable.ToList();
 			List<object> responses = [];
 			using TelemetrySpan filterSpan = _tracer.StartActiveSpan($"FilterAgentResponses");
@@ -101,7 +96,7 @@ namespace HordeServer.Agents
 					}
 				}
 			}
-			
+
 			span.SetAttribute("NumAgents", agents.Count);
 			span.SetAttribute("NumResponses", responses.Count);
 			return responses;
diff --git a/Engine/Source/Programs/Horde/Plugins/Compute/HordeServer.Compute/Agents/Enrollment/EnrollmentController.cs b/Engine/Source/Programs/Horde/Plugins/Compute/HordeServer.Compute/Agents/Enrollment/EnrollmentController.cs
--- a/Engine/Source/Programs/Horde/Plugins/Compute/HordeServer.Compute/Agents/Enrollment/EnrollmentController.cs
+++ b/Engine/Source/Programs/Horde/Plugins/Compute/HordeServer.Compute/Agents/Enrollment/EnrollmentController.cs
@@ -1,12 +1,7 @@
 // Copyright Epic Games, Inc. All Rights Reserved.
 
-using System;
-using System.Collections.Generic;
-using System.Threading;
-using System.Threading.Tasks;
 using EpicGames.Core;
 using EpicGames.Horde.Agents.Enrollment;
-using HordeServer.Server;
 using HordeServer.Utilities;
 using Microsoft.AspNetCore.Authorization;
 using Microsoft.AspNetCore.Mvc;
diff --git a/Engine/Source/Programs/Horde/Plugins/Compute/HordeServer.Compute/Agents/Fleet/FleetManagerFactory.cs b/Engine/Source/Programs/Horde/Plugins/Compute/HordeServer.Compute/Agents/Fleet/FleetManagerFactory.cs
--- a/Engine/Source/Programs/Horde/Plugins/Compute/HordeServer.Compute/Agents/Fleet/FleetManagerFactory.cs
+++ b/Engine/Source/Programs/Horde/Plugins/Compute/HordeServer.Compute/Agents/Fleet/FleetManagerFactory.cs
@@ -1,8 +1,6 @@
 // Copyright Epic Games, Inc. All Rights Reserved.
 
-using System;
 using System.Diagnostics.Metrics;
-using System.IO;
 using System.Text.Json;
 using Amazon.AutoScaling;
 using Amazon.EC2;
diff --git a/Engine/Source/Programs/Horde/Plugins/Compute/HordeServer.Compute/Agents/Fleet/IPoolSizeStrategy.cs b/Engine/Source/Programs/Horde/Plugins/Compute/HordeServer.Compute/Agents/Fleet/IPoolSizeStrategy.cs
--- a/Engine/Source/Programs/Horde/Plugins/Compute/HordeServer.Compute/Agents/Fleet/IPoolSizeStrategy.cs
+++ b/Engine/Source/Programs/Horde/Plugins/Compute/HordeServer.Compute/Agents/Fleet/IPoolSizeStrategy.cs
@@ -1,10 +1,6 @@
 // Copyright Epic Games, Inc. All Rights Reserved.
 
-using System;
-using System.Collections.Generic;
 using System.Text.Json;
-using System.Threading;
-using System.Threading.Tasks;
 using HordeServer.Agents.Pools;
 
 namespace HordeServer.Agents.Fleet
@@ -22,7 +18,6 @@ namespace HordeServer.Agents.Fleet
 
 		/// <summary>
 		/// Strategy based on size of job build queue
-		/// <see cref="JobQueueStrategy"/> 
 		/// </summary>
 		JobQueue,
 
diff --git a/Engine/Source/Programs/Horde/Plugins/Compute/HordeServer.Compute/Agents/Fleet/LeaseUtilizationStrategy.cs b/Engine/Source/Programs/Horde/Plugins/Compute/HordeServer.Compute/Agents/Fleet/LeaseUtilizationStrategy.cs
--- a/Engine/Source/Programs/Horde/Plugins/Compute/HordeServer.Compute/Agents/Fleet/LeaseUtilizationStrategy.cs
+++ b/Engine/Source/Programs/Horde/Plugins/Compute/HordeServer.Compute/Agents/Fleet/LeaseUtilizationStrategy.cs
@@ -1,19 +1,14 @@
 // Copyright Epic Games, Inc. All Rights Reserved.
 
-using System;
-using System.Collections.Generic;
-using System.Linq;
 using System.Text;
-using System.Threading;
-using System.Threading.Tasks;
 using EpicGames.Core;
 using EpicGames.Horde.Agents;
 using EpicGames.Horde.Agents.Pools;
 using Google.Protobuf.WellKnownTypes;
+using HordeCommon.Rpc.Tasks;
 using HordeServer.Agents.Leases;
 using HordeServer.Agents.Pools;
 using HordeServer.Utilities;
-using HordeCommon.Rpc.Tasks;
 using Microsoft.Extensions.Caching.Memory;
 using OpenTelemetry.Trace;
 
diff --git a/Engine/Source/Programs/Horde/Plugins/Compute/HordeServer.Compute/Agents/IAgent.cs b/Engine/Source/Programs/Horde/Plugins/Compute/HordeServer.Compute/Agents/IAgent.cs
--- a/Engine/Source/Programs/Horde/Plugins/Compute/HordeServer.Compute/Agents/IAgent.cs
+++ b/Engine/Source/Programs/Horde/Plugins/Compute/HordeServer.Compute/Agents/IAgent.cs
@@ -1,12 +1,7 @@
 // Copyright Epic Games, Inc. All Rights Reserved.
 
-using System;
-using System.Collections.Generic;
 using System.Diagnostics.CodeAnalysis;
-using System.Linq;
 using System.Security.Cryptography;
-using System.Threading;
-using System.Threading.Tasks;
 using EpicGames.Core;
 using EpicGames.Horde.Agents;
 using EpicGames.Horde.Agents.Leases;
@@ -17,11 +12,8 @@ using EpicGames.Horde.Compute;
 using EpicGames.Horde.Logs;
 using EpicGames.Horde.Streams;
 using EpicGames.Horde.Tools;
-using Google.Protobuf.WellKnownTypes;
-using HordeServer.Server;
 using HordeCommon.Rpc;
 using HordeCommon.Rpc.Messages;
-using HordeCommon.Rpc.Tasks;
 using MongoDB.Bson;
 using MongoDB.Bson.Serialization.Attributes;
 
diff --git a/Engine/Source/Programs/Horde/Plugins/Compute/HordeServer.Compute/Agents/Pools/PoolMessages.cs b/Engine/Source/Programs/Horde/Plugins/Compute/HordeServer.Compute/Agents/Pools/PoolMessages.cs
--- a/Engine/Source/Programs/Horde/Plugins/Compute/HordeServer.Compute/Agents/Pools/PoolMessages.cs
+++ b/Engine/Source/Programs/Horde/Plugins/Compute/HordeServer.Compute/Agents/Pools/PoolMessages.cs
@@ -1,9 +1,6 @@
 // Copyright Epic Games, Inc. All Rights Reserved.
 
-using System;
-using System.Collections.Generic;
 using System.ComponentModel.DataAnnotations;
-using System.Linq;
 using EpicGames.Horde.Agents;
 using EpicGames.Horde.Agents.Pools;
 using EpicGames.Horde.Common;
@@ -513,7 +510,7 @@ namespace HordeServer.Agents.Pools
 		/// The minimum number of idle agents in this pool, if autoscaling is enabled
 		/// </summary>
 		public int? NumReserveAgents { get; set; }
-		
+
 		/// <summary>
 		/// Time to wait before shutting down an agent that has been disabled
 		/// </summary>
diff --git a/Engine/Source/Programs/Horde/Plugins/Compute/HordeServer.Compute/Agents/Software/AgentSoftwareController.cs b/Engine/Source/Programs/Horde/Plugins/Compute/HordeServer.Compute/Agents/Software/AgentSoftwareController.cs
--- a/Engine/Source/Programs/Horde/Plugins/Compute/HordeServer.Compute/Agents/Software/AgentSoftwareController.cs
+++ b/Engine/Source/Programs/Horde/Plugins/Compute/HordeServer.Compute/Agents/Software/AgentSoftwareController.cs
@@ -1,11 +1,6 @@
 // Copyright Epic Games, Inc. All Rights Reserved.
 
-using System;
-using System.IO;
-using System.Threading;
-using System.Threading.Tasks;
 using EpicGames.Core;
-using HordeServer.Server;
 using HordeServer.Tools;
 using Microsoft.AspNetCore.Authorization;
 using Microsoft.AspNetCore.Mvc;
diff --git a/Engine/Source/Programs/Horde/Plugins/Compute/HordeServer.Compute/Agents/Utilization/IUtilizationDataCollection.cs b/Engine/Source/Programs/Horde/Plugins/Compute/HordeServer.Compute/Agents/Utilization/IUtilizationDataCollection.cs
--- a/Engine/Source/Programs/Horde/Plugins/Compute/HordeServer.Compute/Agents/Utilization/IUtilizationDataCollection.cs
+++ b/Engine/Source/Programs/Horde/Plugins/Compute/HordeServer.Compute/Agents/Utilization/IUtilizationDataCollection.cs
@@ -1,10 +1,5 @@
 // Copyright Epic Games, Inc. All Rights Reserved.
 
-using System;
-using System.Collections.Generic;
-using System.Threading;
-using System.Threading.Tasks;
-
 namespace HordeServer.Agents.Utilization
 {
 	/// <summary>
diff --git a/Engine/Source/Programs/Horde/Plugins/Compute/HordeServer.Compute/Aws/AwsCloudWatchMetricExporter.cs b/Engine/Source/Programs/Horde/Plugins/Compute/HordeServer.Compute/Aws/AwsCloudWatchMetricExporter.cs
--- a/Engine/Source/Programs/Horde/Plugins/Compute/HordeServer.Compute/Aws/AwsCloudWatchMetricExporter.cs
+++ b/Engine/Source/Programs/Horde/Plugins/Compute/HordeServer.Compute/Aws/AwsCloudWatchMetricExporter.cs
@@ -1,11 +1,6 @@
 // Copyright Epic Games, Inc. All Rights Reserved.
 
-using System;
-using System.Collections.Generic;
 using System.Net;
-using System.Net.Http;
-using System.Threading;
-using System.Threading.Tasks;
 using Amazon.CloudWatch;
 using Amazon.CloudWatch.Model;
 using HordeServer.Compute;
diff --git a/Engine/Source/Programs/Horde/Plugins/Compute/HordeServer.Compute/Compute/ComputeController.cs b/Engine/Source/Programs/Horde/Plugins/Compute/HordeServer.Compute/Compute/ComputeController.cs
--- a/Engine/Source/Programs/Horde/Plugins/Compute/HordeServer.Compute/Compute/ComputeController.cs
+++ b/Engine/Source/Programs/Horde/Plugins/Compute/HordeServer.Compute/Compute/ComputeController.cs
@@ -1,17 +1,11 @@
 // Copyright Epic Games, Inc. All Rights Reserved.
 
-using System;
-using System.Collections.Generic;
-using System.Linq;
 using System.Net;
-using System.Threading;
-using System.Threading.Tasks;
 using EpicGames.Core;
 using EpicGames.Horde.Agents;
 using EpicGames.Horde.Agents.Pools;
 using EpicGames.Horde.Compute;
 using HordeServer.Agents;
-using HordeServer.Server;
 using HordeServer.Utilities;
 using Microsoft.AspNetCore.Authorization;
 using Microsoft.AspNetCore.Mvc;
@@ -41,7 +35,7 @@ namespace HordeServer.Compute
 			_computeConfig = computeConfig;
 			_tracer = tracer;
 		}
-		
+
 		/// <summary>
 		/// Find the most suitable cluster given a compute assignment request
 		/// </summary>
@@ -60,13 +54,13 @@ namespace HordeServer.Compute
 				IPAddress? clientIp = String.IsNullOrEmpty(forwardedForHeader)
 					? HttpContext.Connection.RemoteIpAddress
 					: IPAddress.Parse(forwardedForHeader);
-				
+
 				span.SetAttribute("clientIp", clientIp?.ToString());
 				requesterIp = ComputeService.ResolveRequesterIp(clientIp, request.Connection?.PreferPublicIp, request.Connection?.ClientPublicIp);
 				span.SetAttribute("requesterIp", requesterIp.ToString());
 				ClusterId clusterId = ComputeService.FindBestComputeClusterId(_computeConfig.Value, requesterIp);
 				span.SetAttribute("clusterId", clusterId.ToString());
-				
+
 				GetClusterResponse response = new() { ClusterId = clusterId };
 				return response;
 			}
diff --git a/Engine/Source/Programs/Horde/Plugins/Compute/HordeServer.Compute/Compute/ComputeService.cs b/Engine/Source/Programs/Horde/Plugins/Compute/HordeServer.Compute/Compute/ComputeService.cs
--- a/Engine/Source/Programs/Horde/Plugins/Compute/HordeServer.Compute/Compute/ComputeService.cs
+++ b/Engine/Source/Programs/Horde/Plugins/Compute/HordeServer.Compute/Compute/ComputeService.cs
@@ -1,15 +1,10 @@
 // Copyright Epic Games, Inc. All Rights Reserved.
 
-using System;
-using System.Collections.Generic;
 using System.Diagnostics.Metrics;
-using System.Linq;
 using System.Net;
 using System.Security.Cryptography;
 using System.Text;
 using System.Text.Json;
-using System.Threading;
-using System.Threading.Tasks;
 using EpicGames.Core;
 using EpicGames.Horde.Agents;
 using EpicGames.Horde.Agents.Leases;
@@ -23,13 +18,13 @@ using EpicGames.Horde.Users;
 using Google.Protobuf;
 using Google.Protobuf.WellKnownTypes;
 using Horde.Common.Rpc;
+using HordeCommon.Rpc.Tasks;
 using HordeServer.Agents;
 using HordeServer.Agents.Leases;
 using HordeServer.Agents.Relay;
 using HordeServer.Logs;
 using HordeServer.Server;
 using HordeServer.Utilities;
-using HordeCommon.Rpc.Tasks;
 using Microsoft.Extensions.Hosting;
 using Microsoft.Extensions.Logging;
 using Microsoft.Extensions.Options;
@@ -80,7 +75,7 @@ namespace HordeServer.Compute
 		/// Cluster ID
 		/// </summary>
 		public ClusterId ClusterId { get; }
-		
+
 		/// <summary>
 		/// User allocating the resource
 		/// </summary>
@@ -376,7 +371,7 @@ namespace HordeServer.Compute
 				span.SetAttribute($"req.res.{name}.min", resReq.Min);
 				span.SetAttribute($"req.res.{name}.max", resReq.Max);
 			}
-			
+
 			if (!_computeConfig.CurrentValue.TryGetComputeCluster(arp.ClusterId, out ComputeClusterConfig? clusterConfig))
 			{
 				throw new ArgumentException($"Cluster '{arp.ClusterId}' not found");
@@ -776,7 +771,7 @@ namespace HordeServer.Compute
 
 			throw new Exception("Unable to resolve a suitable connection mode for compute task");
 		}
-		
+
 		/// <summary>
 		/// Resolve IP of the requester
 		/// </summary>
@@ -794,15 +789,15 @@ namespace HordeServer.Compute
 					requesterIp = publicIp;
 				}
 			}
-			
+
 			if (requesterIp == null)
 			{
 				throw new ComputeServiceException("Unable to determine IP of requester");
 			}
-			
+
 			return requesterIp;
 		}
-		
+
 		/// <summary>
 		/// Find best compute cluster for given parameters
 		/// </summary>
@@ -816,12 +811,12 @@ namespace HordeServer.Compute
 			{
 				throw new ComputeServiceException("Unable to find a matching network config");
 			}
-			
+
 			if (networkConfig.ComputeId == null)
 			{
 				throw new ComputeServiceException($"Network config '{networkConfig.Id}' has no compute ID set");
 			}
-			
+
 			return new ClusterId(networkConfig.ComputeId);
 		}
 
diff --git a/Engine/Source/Programs/Horde/Plugins/Compute/HordeServer.Compute/Compute/ComputeTaskSource.cs b/Engine/Source/Programs/Horde/Plugins/Compute/HordeServer.Compute/Compute/ComputeTaskSource.cs
--- a/Engine/Source/Programs/Horde/Plugins/Compute/HordeServer.Compute/Compute/ComputeTaskSource.cs
+++ b/Engine/Source/Programs/Horde/Plugins/Compute/HordeServer.Compute/Compute/ComputeTaskSource.cs
@@ -1,21 +1,15 @@
 // Copyright Epic Games, Inc. All Rights Reserved.
 
-using System;
-using System.Collections.Generic;
-using System.Linq;
 using System.Net;
-using System.Threading;
-using System.Threading.Tasks;
 using EpicGames.Horde.Agents;
 using EpicGames.Horde.Agents.Leases;
 using EpicGames.Horde.Compute;
 using Google.Protobuf.WellKnownTypes;
+using HordeCommon.Rpc.Tasks;
 using HordeServer.Agents;
 using HordeServer.Agents.Leases;
 using HordeServer.Agents.Relay;
-using HordeServer.Server;
 using HordeServer.Tasks;
-using HordeCommon.Rpc.Tasks;
 using Microsoft.Extensions.Logging;
 using Microsoft.Extensions.Options;
 
diff --git a/Engine/Source/Programs/Horde/Plugins/Compute/HordeServer.Compute/Compute/RedisTaskScheduler.cs b/Engine/Source/Programs/Horde/Plugins/Compute/HordeServer.Compute/Compute/RedisTaskScheduler.cs
--- a/Engine/Source/Programs/Horde/Plugins/Compute/HordeServer.Compute/Compute/RedisTaskScheduler.cs
+++ b/Engine/Source/Programs/Horde/Plugins/Compute/HordeServer.Compute/Compute/RedisTaskScheduler.cs
@@ -1,12 +1,7 @@
 // Copyright Epic Games, Inc. All Rights Reserved.
 
-using System;
-using System.Collections.Generic;
 using System.Diagnostics;
-using System.Linq;
-using System.Threading;
 using System.Threading.Channels;
-using System.Threading.Tasks;
 using EpicGames.Redis;
 using Microsoft.Extensions.Logging;
 using StackExchange.Redis;
@@ -375,6 +370,12 @@ namespace HordeServer.Compute
 			return keys.ToList();
 		}
 
+		/// <summary>
+		/// 
+		/// </summary>
+		/// <param name="predicate"></param>
+		/// <param name="token"></param>
+		/// <returns></returns>
 		public async Task<int> GetNumQueuedTasksAsync(Func<TQueueId, ValueTask<bool>> predicate, CancellationToken token = default)
 		{
 			HashSet<TQueueId> queueIds = new(await _redisConnectionPool.GetDatabase().SetMembersAsync(_queueIndex));
diff --git a/Engine/Source/Programs/Horde/Plugins/Compute/HordeServer.Compute/Compute/TunnelService.cs b/Engine/Source/Programs/Horde/Plugins/Compute/HordeServer.Compute/Compute/TunnelService.cs
--- a/Engine/Source/Programs/Horde/Plugins/Compute/HordeServer.Compute/Compute/TunnelService.cs
+++ b/Engine/Source/Programs/Horde/Plugins/Compute/HordeServer.Compute/Compute/TunnelService.cs
@@ -1,12 +1,7 @@
 // Copyright Epic Games, Inc. All Rights Reserved.
 
-using System;
-using System.Collections.Generic;
-using System.IO;
 using System.Net;
 using System.Net.Sockets;
-using System.Threading;
-using System.Threading.Tasks;
 using EpicGames.Core;
 using EpicGames.Horde.Compute.Clients;
 using Microsoft.Extensions.Hosting;
diff --git a/Engine/Source/Programs/Horde/Plugins/Compute/HordeServer.Compute/ComputeConfig.cs b/Engine/Source/Programs/Horde/Plugins/Compute/HordeServer.Compute/ComputeConfig.cs
--- a/Engine/Source/Programs/Horde/Plugins/Compute/HordeServer.Compute/ComputeConfig.cs
+++ b/Engine/Source/Programs/Horde/Plugins/Compute/HordeServer.Compute/ComputeConfig.cs
@@ -1,15 +1,9 @@
 // Copyright Epic Games, Inc. All Rights Reserved.
 
-using System;
-using System.Collections.Generic;
 using System.Diagnostics;
 using System.Diagnostics.CodeAnalysis;
-using System.Linq;
 using System.Net;
 using System.Security.Claims;
-using System.Text;
-using System.Text.Json.Serialization;
-using System.Threading.Tasks;
 using EpicGames.Horde.Acls;
 using EpicGames.Horde.Agents.Pools;
 using EpicGames.Horde.Common;
diff --git a/Engine/Source/Programs/Horde/Plugins/Compute/HordeServer.Compute/Logs/ILog.cs b/Engine/Source/Programs/Horde/Plugins/Compute/HordeServer.Compute/Logs/ILog.cs
--- a/Engine/Source/Programs/Horde/Plugins/Compute/HordeServer.Compute/Logs/ILog.cs
+++ b/Engine/Source/Programs/Horde/Plugins/Compute/HordeServer.Compute/Logs/ILog.cs
@@ -1,12 +1,7 @@
 // Copyright Epic Games, Inc. All Rights Reserved.
 
-using System.Collections.Generic;
-using System.IO;
-using System.Linq;
 using System.Security.Claims;
 using System.Text.Json;
-using System.Threading;
-using System.Threading.Tasks;
 using EpicGames.Core;
 using EpicGames.Horde.Acls;
 using EpicGames.Horde.Agents.Leases;
diff --git a/Engine/Source/Programs/Horde/Plugins/Compute/HordeServer.Compute/Logs/LogsController.cs b/Engine/Source/Programs/Horde/Plugins/Compute/HordeServer.Compute/Logs/LogsController.cs
--- a/Engine/Source/Programs/Horde/Plugins/Compute/HordeServer.Compute/Logs/LogsController.cs
+++ b/Engine/Source/Programs/Horde/Plugins/Compute/HordeServer.Compute/Logs/LogsController.cs
@@ -1,21 +1,14 @@
 // Copyright Epic Games, Inc. All Rights Reserved.
 
-using System;
-using System.Collections.Generic;
-using System.IO;
-using System.Linq;
 using System.Security.Claims;
 using System.Text;
 using System.Text.Json;
-using System.Threading;
-using System.Threading.Tasks;
 using EpicGames.Core;
 using EpicGames.Horde.Acls;
 using EpicGames.Horde.Logs;
 using EpicGames.Horde.Storage;
 using HordeServer.Agents;
 using HordeServer.Agents.Sessions;
-using HordeServer.Server;
 using HordeServer.Storage;
 using HordeServer.Utilities;
 using Microsoft.AspNetCore.Authorization;
diff --git a/Engine/Source/Programs/Horde/Plugins/Compute/HordeServer.Compute/Tasks/RestartTaskSource.cs b/Engine/Source/Programs/Horde/Plugins/Compute/HordeServer.Compute/Tasks/RestartTaskSource.cs
--- a/Engine/Source/Programs/Horde/Plugins/Compute/HordeServer.Compute/Tasks/RestartTaskSource.cs
+++ b/Engine/Source/Programs/Horde/Plugins/Compute/HordeServer.Compute/Tasks/RestartTaskSource.cs
@@ -1,16 +1,14 @@
 // Copyright Epic Games, Inc. All Rights Reserved.
 
-using System.Threading;
-using System.Threading.Tasks;
 using EpicGames.Horde.Agents.Leases;
 using EpicGames.Horde.Jobs;
 using EpicGames.Horde.Logs;
 using Google.Protobuf;
 using Google.Protobuf.WellKnownTypes;
+using HordeCommon.Rpc.Tasks;
 using HordeServer.Agents;
 using HordeServer.Logs;
 using HordeServer.Utilities;
-using HordeCommon.Rpc.Tasks;
 
 namespace HordeServer.Tasks
 {
diff --git a/Engine/Source/Programs/Horde/Plugins/Compute/HordeServer.Compute/Tasks/ShutdownTaskSource.cs b/Engine/Source/Programs/Horde/Plugins/Compute/HordeServer.Compute/Tasks/ShutdownTaskSource.cs
--- a/Engine/Source/Programs/Horde/Plugins/Compute/HordeServer.Compute/Tasks/ShutdownTaskSource.cs
+++ b/Engine/Source/Programs/Horde/Plugins/Compute/HordeServer.Compute/Tasks/ShutdownTaskSource.cs
@@ -1,16 +1,14 @@
 // Copyright Epic Games, Inc. All Rights Reserved.
 
-using System.Threading;
-using System.Threading.Tasks;
 using EpicGames.Horde.Agents.Leases;
 using EpicGames.Horde.Jobs;
 using EpicGames.Horde.Logs;
 using Google.Protobuf;
 using Google.Protobuf.WellKnownTypes;
+using HordeCommon.Rpc.Tasks;
 using HordeServer.Agents;
 using HordeServer.Logs;
 using HordeServer.Utilities;
-using HordeCommon.Rpc.Tasks;
 
 namespace HordeServer.Tasks
 {
diff --git a/Engine/Source/Programs/Horde/Plugins/Compute/HordeServer.Compute/Tasks/UpgradeTaskSource.cs b/Engine/Source/Programs/Horde/Plugins/Compute/HordeServer.Compute/Tasks/UpgradeTaskSource.cs
--- a/Engine/Source/Programs/Horde/Plugins/Compute/HordeServer.Compute/Tasks/UpgradeTaskSource.cs
+++ b/Engine/Source/Programs/Horde/Plugins/Compute/HordeServer.Compute/Tasks/UpgradeTaskSource.cs
@@ -1,9 +1,6 @@
 // Copyright Epic Games, Inc. All Rights Reserved.
 
-using System;
 using System.Text;
-using System.Threading;
-using System.Threading.Tasks;
 using EpicGames.Core;
 using EpicGames.Horde.Agents.Leases;
 using EpicGames.Horde.Jobs;
@@ -11,12 +8,11 @@ using EpicGames.Horde.Logs;
 using EpicGames.Horde.Tools;
 using Google.Protobuf;
 using Google.Protobuf.WellKnownTypes;
+using HordeCommon.Rpc.Tasks;
 using HordeServer.Agents;
 using HordeServer.Logs;
-using HordeServer.Server;
 using HordeServer.Tools;
 using HordeServer.Utilities;
-using HordeCommon.Rpc.Tasks;
 using Microsoft.Extensions.Options;
 
 namespace HordeServer.Tasks
diff --git a/Engine/Source/Programs/Horde/Plugins/Ddc/HordeServer.Ddc.Tests/Ddc/Functional/ContentId/CompressedBlobTests.cs b/Engine/Source/Programs/Horde/Plugins/Ddc/HordeServer.Ddc.Tests/Ddc/Functional/ContentId/CompressedBlobTests.cs
--- a/Engine/Source/Programs/Horde/Plugins/Ddc/HordeServer.Ddc.Tests/Ddc/Functional/ContentId/CompressedBlobTests.cs
+++ b/Engine/Source/Programs/Horde/Plugins/Ddc/HordeServer.Ddc.Tests/Ddc/Functional/ContentId/CompressedBlobTests.cs
@@ -1,16 +1,10 @@
 // Copyright Epic Games, Inc. All Rights Reserved.
 
-using System;
-using System.Collections.Generic;
-using System.IO;
 using System.Net;
-using System.Net.Http;
 using System.Net.Http.Headers;
 using System.Net.Http.Json;
 using System.Net.Mime;
 using System.Text;
-using System.Threading;
-using System.Threading.Tasks;
 using Blake3;
 using EpicGames.AspNet;
 using EpicGames.Core;
@@ -23,7 +17,6 @@ using HordeServer.Server;
 using HordeServer.Storage;
 using Microsoft.Extensions.DependencyInjection;
 using Microsoft.Extensions.Options;
-using Microsoft.VisualStudio.TestTools.UnitTesting;
 using Serilog;
 
 namespace HordeServer.Tests.Ddc.FunctionalTests.CompressedBlobs
diff --git a/Engine/Source/Programs/Horde/Plugins/Ddc/HordeServer.Ddc.Tests/Ddc/Functional/Objects/ReferencesTests.cs b/Engine/Source/Programs/Horde/Plugins/Ddc/HordeServer.Ddc.Tests/Ddc/Functional/Objects/ReferencesTests.cs
--- a/Engine/Source/Programs/Horde/Plugins/Ddc/HordeServer.Ddc.Tests/Ddc/Functional/Objects/ReferencesTests.cs
+++ b/Engine/Source/Programs/Horde/Plugins/Ddc/HordeServer.Ddc.Tests/Ddc/Functional/Objects/ReferencesTests.cs
@@ -1,19 +1,12 @@
 // Copyright Epic Games, Inc. All Rights Reserved.
 
-using System;
-using System.Collections.Generic;
-using System.IO;
-using System.Linq;
 using System.Net;
-using System.Net.Http;
 using System.Net.Http.Headers;
 using System.Net.Http.Json;
 using System.Net.Mime;
 using System.Text;
 using System.Text.Json;
 using System.Text.Json.Nodes;
-using System.Threading;
-using System.Threading.Tasks;
 using EpicGames.AspNet;
 using EpicGames.Core;
 using EpicGames.Horde.Storage;
@@ -27,7 +20,6 @@ using HordeServer.Storage;
 using Microsoft.AspNetCore.Mvc;
 using Microsoft.Extensions.DependencyInjection;
 using Microsoft.Extensions.Options;
-using Microsoft.VisualStudio.TestTools.UnitTesting;
 
 namespace HordeServer.Tests.Ddc.FunctionalTests.References
 {
@@ -1044,7 +1036,7 @@ namespace HordeServer.Tests.Ddc.FunctionalTests.References
 		public async Task GetMissingContentIdRecordAsync()
 		{
 			CancellationToken cancellationToken = CancellationToken.None;
-			
+
 			string blobContents = $"This is a blob in {nameof(GetMissingContentIdRecordAsync)}";
 			byte[] blobData = Encoding.ASCII.GetBytes(blobContents);
 			BlobId uncompressedHash = BlobId.FromBlob(blobData);
diff --git a/Engine/Source/Programs/Horde/Plugins/Ddc/HordeServer.Ddc.Tests/Ddc/UnitTests/CompactBinaryTests.cs b/Engine/Source/Programs/Horde/Plugins/Ddc/HordeServer.Ddc.Tests/Ddc/UnitTests/CompactBinaryTests.cs
--- a/Engine/Source/Programs/Horde/Plugins/Ddc/HordeServer.Ddc.Tests/Ddc/UnitTests/CompactBinaryTests.cs
+++ b/Engine/Source/Programs/Horde/Plugins/Ddc/HordeServer.Ddc.Tests/Ddc/UnitTests/CompactBinaryTests.cs
@@ -1,11 +1,7 @@
 // Copyright Epic Games, Inc. All Rights Reserved.
 
-using System.Collections.Generic;
-using System.IO;
-using System.Linq;
 using EpicGames.Core;
 using EpicGames.Serialization;
-using Microsoft.VisualStudio.TestTools.UnitTesting;
 
 namespace HordeServer.Tests.Ddc.UnitTests
 {
diff --git a/Engine/Source/Programs/Horde/Plugins/Ddc/HordeServer.Ddc.Tests/Ddc/UnitTests/CompressedBufferTests.cs b/Engine/Source/Programs/Horde/Plugins/Ddc/HordeServer.Ddc.Tests/Ddc/UnitTests/CompressedBufferTests.cs
--- a/Engine/Source/Programs/Horde/Plugins/Ddc/HordeServer.Ddc.Tests/Ddc/UnitTests/CompressedBufferTests.cs
+++ b/Engine/Source/Programs/Horde/Plugins/Ddc/HordeServer.Ddc.Tests/Ddc/UnitTests/CompressedBufferTests.cs
@@ -1,13 +1,9 @@
 // Copyright Epic Games, Inc. All Rights Reserved.
 
-using System.IO;
 using System.Text;
-using System.Threading;
-using System.Threading.Tasks;
 using EpicGames.Core;
 using HordeServer.Ddc;
 using Microsoft.Extensions.Options;
-using Microsoft.VisualStudio.TestTools.UnitTesting;
 using Moq;
 using OpenTelemetry.Trace;
 
diff --git a/Engine/Source/Programs/Horde/Plugins/Ddc/HordeServer.Ddc/Ddc/BlobContents.cs b/Engine/Source/Programs/Horde/Plugins/Ddc/HordeServer.Ddc/Ddc/BlobContents.cs
--- a/Engine/Source/Programs/Horde/Plugins/Ddc/HordeServer.Ddc/Ddc/BlobContents.cs
+++ b/Engine/Source/Programs/Horde/Plugins/Ddc/HordeServer.Ddc/Ddc/BlobContents.cs
@@ -1,9 +1,5 @@
 ﻿// Copyright Epic Games, Inc. All Rights Reserved.
 
-using System;
-using System.IO;
-using System.Threading.Tasks;
-
 #pragma warning disable CS1591
 
 namespace HordeServer.Ddc
diff --git a/Engine/Source/Programs/Horde/Plugins/Ddc/HordeServer.Ddc/Ddc/BlobId.cs b/Engine/Source/Programs/Horde/Plugins/Ddc/HordeServer.Ddc/Ddc/BlobId.cs
--- a/Engine/Source/Programs/Horde/Plugins/Ddc/HordeServer.Ddc/Ddc/BlobId.cs
+++ b/Engine/Source/Programs/Horde/Plugins/Ddc/HordeServer.Ddc/Ddc/BlobId.cs
@@ -1,14 +1,10 @@
 ﻿// Copyright Epic Games, Inc. All Rights Reserved.
 
-using System;
 using System.Buffers;
 using System.ComponentModel;
 using System.Globalization;
-using System.IO;
 using System.Text.Json;
 using System.Text.Json.Serialization;
-using System.Threading;
-using System.Threading.Tasks;
 using Blake3;
 using EpicGames.Core;
 using EpicGames.Serialization;
diff --git a/Engine/Source/Programs/Horde/Plugins/Ddc/HordeServer.Ddc/Ddc/BlobsController.cs b/Engine/Source/Programs/Horde/Plugins/Ddc/HordeServer.Ddc/Ddc/BlobsController.cs
--- a/Engine/Source/Programs/Horde/Plugins/Ddc/HordeServer.Ddc/Ddc/BlobsController.cs
+++ b/Engine/Source/Programs/Horde/Plugins/Ddc/HordeServer.Ddc/Ddc/BlobsController.cs
@@ -1,14 +1,9 @@
 // Copyright Epic Games, Inc. All Rights Reserved.
 
-using System;
 using System.Collections.Concurrent;
-using System.Collections.Generic;
 using System.ComponentModel.DataAnnotations;
-using System.IO;
-using System.Linq;
 using System.Net;
 using System.Net.Mime;
-using System.Threading.Tasks;
 using EpicGames.AspNet;
 using EpicGames.Horde.Acls;
 using EpicGames.Horde.Storage;
diff --git a/Engine/Source/Programs/Horde/Plugins/Ddc/HordeServer.Ddc/Ddc/BufferedPayload.cs b/Engine/Source/Programs/Horde/Plugins/Ddc/HordeServer.Ddc/Ddc/BufferedPayload.cs
--- a/Engine/Source/Programs/Horde/Plugins/Ddc/HordeServer.Ddc/Ddc/BufferedPayload.cs
+++ b/Engine/Source/Programs/Horde/Plugins/Ddc/HordeServer.Ddc/Ddc/BufferedPayload.cs
@@ -1,9 +1,5 @@
 // Copyright Epic Games, Inc. All Rights Reserved.
 
-using System;
-using System.IO;
-using System.Threading;
-using System.Threading.Tasks;
 using Microsoft.AspNetCore.Http;
 using Microsoft.Extensions.Options;
 using OpenTelemetry.Trace;
diff --git a/Engine/Source/Programs/Horde/Plugins/Ddc/HordeServer.Ddc/Ddc/CompressedBlobController.cs b/Engine/Source/Programs/Horde/Plugins/Ddc/HordeServer.Ddc/Ddc/CompressedBlobController.cs
--- a/Engine/Source/Programs/Horde/Plugins/Ddc/HordeServer.Ddc/Ddc/CompressedBlobController.cs
+++ b/Engine/Source/Programs/Horde/Plugins/Ddc/HordeServer.Ddc/Ddc/CompressedBlobController.cs
@@ -1,13 +1,9 @@
 ﻿// Copyright Epic Games, Inc. All Rights Reserved.
 
 using System.Collections.Concurrent;
-using System.Collections.Generic;
 using System.ComponentModel.DataAnnotations;
-using System.Linq;
 using System.Net;
 using System.Net.Mime;
-using System.Threading;
-using System.Threading.Tasks;
 using EpicGames.AspNet;
 using EpicGames.Horde.Storage;
 using EpicGames.Serialization;
diff --git a/Engine/Source/Programs/Horde/Plugins/Ddc/HordeServer.Ddc/Ddc/ContentHash.cs b/Engine/Source/Programs/Horde/Plugins/Ddc/HordeServer.Ddc/Ddc/ContentHash.cs
--- a/Engine/Source/Programs/Horde/Plugins/Ddc/HordeServer.Ddc/Ddc/ContentHash.cs
+++ b/Engine/Source/Programs/Horde/Plugins/Ddc/HordeServer.Ddc/Ddc/ContentHash.cs
@@ -1,9 +1,7 @@
 ﻿// Copyright Epic Games, Inc. All Rights Reserved.
 
-using System;
 using System.ComponentModel;
 using System.Globalization;
-using System.IO;
 using System.Text.Json;
 using System.Text.Json.Serialization;
 using Blake3;
diff --git a/Engine/Source/Programs/Horde/Plugins/Ddc/HordeServer.Ddc/Ddc/FormatResolver.cs b/Engine/Source/Programs/Horde/Plugins/Ddc/HordeServer.Ddc/Ddc/FormatResolver.cs
--- a/Engine/Source/Programs/Horde/Plugins/Ddc/HordeServer.Ddc/Ddc/FormatResolver.cs
+++ b/Engine/Source/Programs/Horde/Plugins/Ddc/HordeServer.Ddc/Ddc/FormatResolver.cs
@@ -1,7 +1,5 @@
 // Copyright Epic Games, Inc. All Rights Reserved.
 
-using System;
-using System.Linq;
 using System.Net.Mime;
 using EpicGames.AspNet;
 using Microsoft.AspNetCore.Http;
diff --git a/Engine/Source/Programs/Horde/Plugins/Ddc/HordeServer.Ddc/Ddc/RequestHelper.cs b/Engine/Source/Programs/Horde/Plugins/Ddc/HordeServer.Ddc/Ddc/RequestHelper.cs
--- a/Engine/Source/Programs/Horde/Plugins/Ddc/HordeServer.Ddc/Ddc/RequestHelper.cs
+++ b/Engine/Source/Programs/Horde/Plugins/Ddc/HordeServer.Ddc/Ddc/RequestHelper.cs
@@ -1,8 +1,6 @@
 // Copyright Epic Games, Inc. All Rights Reserved.
 
-using System.Linq;
 using System.Security.Claims;
-using System.Threading.Tasks;
 using EpicGames.Horde.Acls;
 using EpicGames.Horde.Storage;
 using HordeServer.Storage;
diff --git a/Engine/Source/Programs/Horde/Plugins/Ddc/HordeServer.Ddc/Ddc/RequestUtil.cs b/Engine/Source/Programs/Horde/Plugins/Ddc/HordeServer.Ddc/Ddc/RequestUtil.cs
--- a/Engine/Source/Programs/Horde/Plugins/Ddc/HordeServer.Ddc/Ddc/RequestUtil.cs
+++ b/Engine/Source/Programs/Horde/Plugins/Ddc/HordeServer.Ddc/Ddc/RequestUtil.cs
@@ -1,8 +1,5 @@
 ﻿// Copyright Epic Games, Inc. All Rights Reserved.
 
-using System;
-using System.IO;
-using System.Threading.Tasks;
 using Microsoft.AspNetCore.Http;
 using Microsoft.Extensions.DependencyInjection;
 using OpenTelemetry.Trace;
diff --git a/Engine/Source/Programs/Horde/Plugins/Ddc/HordeServer.Ddc/Ddc/RequiredContentTypeAttribute.cs b/Engine/Source/Programs/Horde/Plugins/Ddc/HordeServer.Ddc/Ddc/RequiredContentTypeAttribute.cs
--- a/Engine/Source/Programs/Horde/Plugins/Ddc/HordeServer.Ddc/Ddc/RequiredContentTypeAttribute.cs
+++ b/Engine/Source/Programs/Horde/Plugins/Ddc/HordeServer.Ddc/Ddc/RequiredContentTypeAttribute.cs
@@ -1,7 +1,5 @@
 ﻿// Copyright Epic Games, Inc. All Rights Reserved.
 
-using System;
-using System.Linq;
 using Microsoft.AspNetCore.Mvc.ActionConstraints;
 using Microsoft.Extensions.Primitives;
 
diff --git a/Engine/Source/Programs/Horde/Plugins/Ddc/HordeServer.Ddc/Ddc/StreamUtils.cs b/Engine/Source/Programs/Horde/Plugins/Ddc/HordeServer.Ddc/Ddc/StreamUtils.cs
--- a/Engine/Source/Programs/Horde/Plugins/Ddc/HordeServer.Ddc/Ddc/StreamUtils.cs
+++ b/Engine/Source/Programs/Horde/Plugins/Ddc/HordeServer.Ddc/Ddc/StreamUtils.cs
@@ -1,8 +1,5 @@
 ﻿// Copyright Epic Games, Inc. All Rights Reserved.
 
-using System.IO;
-using System.Threading;
-using System.Threading.Tasks;
 using EpicGames.AspNet;
 using Microsoft.AspNetCore.Http;
 
diff --git a/Engine/Source/Programs/Horde/Plugins/Ddc/HordeServer.Ddc/HordeServer.Ddc.csproj b/Engine/Source/Programs/Horde/Plugins/Ddc/HordeServer.Ddc/HordeServer.Ddc.csproj
--- a/Engine/Source/Programs/Horde/Plugins/Ddc/HordeServer.Ddc/HordeServer.Ddc.csproj
+++ b/Engine/Source/Programs/Horde/Plugins/Ddc/HordeServer.Ddc/HordeServer.Ddc.csproj
@@ -1,9 +1,19 @@
 ﻿<Project Sdk="Microsoft.NET.Sdk">
 
-    <PropertyGroup>
+  <PropertyGroup>
     <TargetFramework>net8.0</TargetFramework>
     <ImplicitUsings>enable</ImplicitUsings>
     <Nullable>enable</Nullable>
+    <RunAnalyzersDuringBuild>False</RunAnalyzersDuringBuild>
+    <Configurations>Debug;Release;Analyze</Configurations>
+  </PropertyGroup>
+
+  <PropertyGroup Condition=" '$(Configuration)|$(Platform)' == 'Analyze|AnyCPU' ">
+    <RunAnalyzersDuringBuild>True</RunAnalyzersDuringBuild>
+    <TreatWarningsAsErrors>false</TreatWarningsAsErrors>
+    <EnforceCodeStyleInBuild>True</EnforceCodeStyleInBuild>
+    <GenerateDocumentationFile>True</GenerateDocumentationFile>
+    <DocumentationFile></DocumentationFile>
   </PropertyGroup>
 
   <ItemGroup>
diff --git a/Engine/Source/Programs/Horde/Plugins/Secrets/HordeServer.Secrets/Secrets/ISecretProvider.cs b/Engine/Source/Programs/Horde/Plugins/Secrets/HordeServer.Secrets/Secrets/ISecretProvider.cs
--- a/Engine/Source/Programs/Horde/Plugins/Secrets/HordeServer.Secrets/Secrets/ISecretProvider.cs
+++ b/Engine/Source/Programs/Horde/Plugins/Secrets/HordeServer.Secrets/Secrets/ISecretProvider.cs
@@ -1,8 +1,5 @@
 // Copyright Epic Games, Inc. All Rights Reserved.
 
-using System.Threading;
-using System.Threading.Tasks;
-
 namespace HordeServer.Secrets
 {
 	/// <summary>
diff --git a/Engine/Source/Programs/Horde/Plugins/Secrets/HordeServer.Secrets/Secrets/SecretConfig.cs b/Engine/Source/Programs/Horde/Plugins/Secrets/HordeServer.Secrets/Secrets/SecretConfig.cs
--- a/Engine/Source/Programs/Horde/Plugins/Secrets/HordeServer.Secrets/Secrets/SecretConfig.cs
+++ b/Engine/Source/Programs/Horde/Plugins/Secrets/HordeServer.Secrets/Secrets/SecretConfig.cs
@@ -1,12 +1,9 @@
 // Copyright Epic Games, Inc. All Rights Reserved.
 
-using System;
-using System.Collections.Generic;
 using System.Security.Claims;
 using EpicGames.Horde.Acls;
 using EpicGames.Horde.Secrets;
 using HordeServer.Acls;
-using HordeServer.Server;
 
 #pragma warning disable CA2227 // Change x to be read-only by removing the property setter
 
diff --git a/Engine/Source/Programs/Horde/Plugins/Secrets/HordeServer.Secrets/Secrets/SecretsController.cs b/Engine/Source/Programs/Horde/Plugins/Secrets/HordeServer.Secrets/Secrets/SecretsController.cs
--- a/Engine/Source/Programs/Horde/Plugins/Secrets/HordeServer.Secrets/Secrets/SecretsController.cs
+++ b/Engine/Source/Programs/Horde/Plugins/Secrets/HordeServer.Secrets/Secrets/SecretsController.cs
@@ -1,11 +1,6 @@
 // Copyright Epic Games, Inc. All Rights Reserved.
 
-using System.Collections.Generic;
-using System.Linq;
-using System.Threading;
-using System.Threading.Tasks;
 using EpicGames.Horde.Secrets;
-using HordeServer.Server;
 using HordeServer.Utilities;
 using Microsoft.AspNetCore.Authorization;
 using Microsoft.AspNetCore.Mvc;
diff --git a/Engine/Source/Programs/Horde/Plugins/Secrets/HordeServer.Secrets/SecretsConfig.cs b/Engine/Source/Programs/Horde/Plugins/Secrets/HordeServer.Secrets/SecretsConfig.cs
--- a/Engine/Source/Programs/Horde/Plugins/Secrets/HordeServer.Secrets/SecretsConfig.cs
+++ b/Engine/Source/Programs/Horde/Plugins/Secrets/HordeServer.Secrets/SecretsConfig.cs
@@ -4,7 +4,6 @@ using System.Diagnostics.CodeAnalysis;
 using System.Security.Claims;
 using EpicGames.Horde.Acls;
 using EpicGames.Horde.Secrets;
-using HordeServer.Acls;
 using HordeServer.Plugins;
 using HordeServer.Secrets;
 
diff --git a/Engine/Source/Programs/Horde/Plugins/Storage/HordeServer.Storage.Tests/Storage/StorageControllerTests.cs b/Engine/Source/Programs/Horde/Plugins/Storage/HordeServer.Storage.Tests/Storage/StorageControllerTests.cs
--- a/Engine/Source/Programs/Horde/Plugins/Storage/HordeServer.Storage.Tests/Storage/StorageControllerTests.cs
+++ b/Engine/Source/Programs/Horde/Plugins/Storage/HordeServer.Storage.Tests/Storage/StorageControllerTests.cs
@@ -1,18 +1,5 @@
 ﻿// Copyright Epic Games, Inc. All Rights Reserved.
 
-using System;
-using System.Linq;
-using System.Threading.Tasks;
-using EpicGames.Core;
-using EpicGames.Horde.Storage;
-using EpicGames.Horde.Storage.Backends;
-using EpicGames.Horde.Storage.Bundles;
-using HordeServer.Storage;
-using Microsoft.Extensions.DependencyInjection;
-using Microsoft.Extensions.Logging;
-using Microsoft.VisualStudio.TestTools.UnitTesting;
-using MongoDB.Driver;
-
 namespace HordeServer.Tests.Storage
 {
 #if false
diff --git a/Engine/Source/Programs/Horde/Plugins/Storage/HordeServer.Storage/HordeServer.Storage.csproj b/Engine/Source/Programs/Horde/Plugins/Storage/HordeServer.Storage/HordeServer.Storage.csproj
--- a/Engine/Source/Programs/Horde/Plugins/Storage/HordeServer.Storage/HordeServer.Storage.csproj
+++ b/Engine/Source/Programs/Horde/Plugins/Storage/HordeServer.Storage/HordeServer.Storage.csproj
@@ -1,9 +1,19 @@
 ﻿<Project Sdk="Microsoft.NET.Sdk">
 
-    <PropertyGroup>
+  <PropertyGroup>
     <TargetFramework>net8.0</TargetFramework>
     <ImplicitUsings>enable</ImplicitUsings>
     <Nullable>enable</Nullable>
+    <RunAnalyzersDuringBuild>False</RunAnalyzersDuringBuild>
+    <Configurations>Debug;Release;Analyze</Configurations>
+  </PropertyGroup>
+
+  <PropertyGroup Condition=" '$(Configuration)|$(Platform)' == 'Analyze|AnyCPU' ">
+    <RunAnalyzersDuringBuild>True</RunAnalyzersDuringBuild>
+    <TreatWarningsAsErrors>false</TreatWarningsAsErrors>
+    <EnforceCodeStyleInBuild>True</EnforceCodeStyleInBuild>
+    <GenerateDocumentationFile>True</GenerateDocumentationFile>
+    <DocumentationFile></DocumentationFile>
   </PropertyGroup>
 
   <ItemGroup>
diff --git a/Engine/Source/Programs/Horde/Plugins/Storage/HordeServer.Storage/Storage/StorageService.cs b/Engine/Source/Programs/Horde/Plugins/Storage/HordeServer.Storage/Storage/StorageService.cs
--- a/Engine/Source/Programs/Horde/Plugins/Storage/HordeServer.Storage/Storage/StorageService.cs
+++ b/Engine/Source/Programs/Horde/Plugins/Storage/HordeServer.Storage/Storage/StorageService.cs
@@ -1,16 +1,9 @@
 // Copyright Epic Games, Inc. All Rights Reserved.
 
-using System;
 using System.Buffers;
-using System.Collections.Generic;
 using System.ComponentModel;
 using System.Diagnostics;
-using System.IO;
-using System.Linq;
-using System.Threading;
-using System.Threading.Tasks;
 using EpicGames.Core;
-using EpicGames.Horde;
 using EpicGames.Horde.Storage;
 using EpicGames.Horde.Storage.Bundles;
 using EpicGames.Horde.Storage.ObjectStores;
@@ -665,12 +658,13 @@ namespace HordeServer.Storage
 			using TelemetrySpan span = _tracer.StartActiveSpan($"{nameof(StorageService)}.{nameof(TickBlobsAsync)}");
 
 			GcState gcState = await _gcState.GetAsync(cancellationToken);
-			if(gcState.Reset)
+			if (gcState.Reset)
 			{
 				_logger.LogInformation("Resetting scan for new blobs...");
-				gcState = await _gcState.UpdateAsync(x => { 
-					x.Reset = false; 
-					x.LastImportBlobInfoId = ObjectId.Empty; 
+				gcState = await _gcState.UpdateAsync(x =>
+				{
+					x.Reset = false;
+					x.LastImportBlobInfoId = ObjectId.Empty;
 				}, cancellationToken);
 			}
 
diff --git a/Engine/Source/Programs/Horde/Plugins/Storage/HordeServer.Storage/StorageConfig.cs b/Engine/Source/Programs/Horde/Plugins/Storage/HordeServer.Storage/StorageConfig.cs
--- a/Engine/Source/Programs/Horde/Plugins/Storage/HordeServer.Storage/StorageConfig.cs
+++ b/Engine/Source/Programs/Horde/Plugins/Storage/HordeServer.Storage/StorageConfig.cs
@@ -1,11 +1,8 @@
 // Copyright Epic Games, Inc. All Rights Reserved.
 
-using System;
-using System.Collections.Generic;
 using System.ComponentModel.DataAnnotations;
 using System.Diagnostics;
 using System.Diagnostics.CodeAnalysis;
-using System.IO;
 using System.Reflection;
 using System.Security.Claims;
 using System.Text.Json;
@@ -15,7 +12,6 @@ using EpicGames.Horde.Acls;
 using EpicGames.Horde.Storage;
 using HordeServer.Acls;
 using HordeServer.Plugins;
-using HordeServer.Server;
 using HordeServer.Storage.ObjectStores;
 using HordeServer.Utilities;
 
@@ -365,7 +361,6 @@ namespace HordeServer.Storage
 		/// <summary>
 		/// Callback once the configuration has been read from disk
 		/// </summary>
-		/// <param name="backendConfig"></param>
 		public void PostLoad(AclConfig parentAcl)
 		{
 			Acl.PostLoad(parentAcl, $"namespace:{Id}");
diff --git a/Engine/Source/Programs/Horde/Plugins/Storage/HordeServer.Storage/StoragePlugin.cs b/Engine/Source/Programs/Horde/Plugins/Storage/HordeServer.Storage/StoragePlugin.cs
--- a/Engine/Source/Programs/Horde/Plugins/Storage/HordeServer.Storage/StoragePlugin.cs
+++ b/Engine/Source/Programs/Horde/Plugins/Storage/HordeServer.Storage/StoragePlugin.cs
@@ -1,10 +1,5 @@
 // Copyright Epic Games, Inc. All Rights Reserved.
 
-using System;
-using System.Collections.Generic;
-using System.Linq;
-using System.Text;
-using System.Threading.Tasks;
 using EpicGames.Core;
 using EpicGames.Horde.Storage;
 using EpicGames.Horde.Storage.Bundles;
@@ -27,6 +22,9 @@ namespace HordeServer.Storage
 		readonly IServerInfo _serverInfo;
 		readonly StaticStorageConfig _staticConfig;
 
+		/// <summary>
+		/// Constructor
+		/// </summary>
 		public StoragePlugin(IServerInfo serverInfo, StaticStorageConfig staticConfig)
 		{
 			_serverInfo = serverInfo;
diff --git a/Engine/Source/Programs/Horde/Plugins/Tools/HordeServer.Tools.Tests/Tools/ToolTests.cs b/Engine/Source/Programs/Horde/Plugins/Tools/HordeServer.Tools.Tests/Tools/ToolTests.cs
--- a/Engine/Source/Programs/Horde/Plugins/Tools/HordeServer.Tools.Tests/Tools/ToolTests.cs
+++ b/Engine/Source/Programs/Horde/Plugins/Tools/HordeServer.Tools.Tests/Tools/ToolTests.cs
@@ -1,20 +1,13 @@
 // Copyright Epic Games, Inc. All Rights Reserved.
 
-using System;
-using System.IO;
 using System.IO.Compression;
-using System.Linq;
 using System.Text;
-using System.Threading;
-using System.Threading.Tasks;
 using EpicGames.Horde.Tools;
+using HordeCommon;
 using HordeServer.Server;
 using HordeServer.Storage;
 using HordeServer.Tools;
-using HordeCommon;
 using Microsoft.Extensions.DependencyInjection;
-using Microsoft.VisualStudio.TestTools.UnitTesting;
-using HordeServer.Plugins;
 
 namespace HordeServer.Tests.Tools
 {
@@ -148,8 +141,8 @@ namespace HordeServer.Tests.Tools
 
 		public static async Task<byte[]> CreateZipFileDataAsync(string fileName, string fileData)
 		{
-			using MemoryStream stream = new ();
-			using (ZipArchive archive = new (stream, ZipArchiveMode.Create))
+			using MemoryStream stream = new();
+			using (ZipArchive archive = new(stream, ZipArchiveMode.Create))
 			{
 				ZipArchiveEntry entry = archive.CreateEntry(fileName);
 				await using (Stream entryStream = entry.Open())
diff --git a/Engine/Source/Programs/Horde/Plugins/Tools/HordeServer.Tools.Tests/Tools/ToolsControllerTests.cs b/Engine/Source/Programs/Horde/Plugins/Tools/HordeServer.Tools.Tests/Tools/ToolsControllerTests.cs
--- a/Engine/Source/Programs/Horde/Plugins/Tools/HordeServer.Tools.Tests/Tools/ToolsControllerTests.cs
+++ b/Engine/Source/Programs/Horde/Plugins/Tools/HordeServer.Tools.Tests/Tools/ToolsControllerTests.cs
@@ -1,11 +1,7 @@
 ﻿// Copyright Epic Games, Inc. All Rights Reserved.
 
-using System.Collections.Generic;
-using System.IO;
 using System.Net;
-using System.Net.Http;
 using System.Net.Http.Headers;
-using System.Threading.Tasks;
 using EpicGames.Horde.Server;
 using EpicGames.Horde.Tools;
 using HordeServer.Acls;
@@ -18,7 +14,6 @@ using HordeServer.Tools;
 using HordeServer.Users;
 using HordeServer.Utilities;
 using Microsoft.Extensions.DependencyInjection;
-using Microsoft.VisualStudio.TestTools.UnitTesting;
 
 namespace HordeServer.Tests.Tools;
 
@@ -49,13 +44,13 @@ public class ToolsControllerTests
 		GlobalConfig globalConfig = new();
 		globalConfig.Plugins.AddStorageConfig(storageConfig);
 		globalConfig.Plugins.AddToolsConfig(toolsConfig);
-		
+
 		ServerSettings serverSettings = new() { AuthMethod = AuthMethod.Horde };
 		globalConfig.PostLoad(serverSettings, pluginCollection.LoadedPlugins, Array.Empty<IDefaultAclModifier>());
 
 		Dictionary<string, string> settings = new() { { "Horde:AuthMethod", AuthMethod.Horde.ToString() } };
 		await using FakeHordeWebApp app = new(settings: settings);
-		
+
 		ConfigService configService = app.ServiceProvider.GetRequiredService<ConfigService>();
 		IToolCollection tools = app.ServiceProvider.GetRequiredService<IToolCollection>();
 		IServiceAccountCollection serviceAccounts = app.ServiceProvider.GetRequiredService<IServiceAccountCollection>();
@@ -64,7 +59,7 @@ public class ToolsControllerTests
 
 		List<IUserClaim> claims = [new UserClaim("http://epicgames.com/ue/horde/role", "agent")];
 		(IServiceAccount _, string token) = await serviceAccounts.CreateAsync(new CreateServiceAccountOptions("myDesc", claims));
-		
+
 		// Create tool and deployment
 		using MemoryStream ms = new(await ToolTests.CreateZipFileDataAsync("foo.txt", "foo content"));
 		ITool? tool = await tools.GetAsync(toolId);
diff --git a/Engine/Source/Programs/Horde/Plugins/Tools/HordeServer.Tools/Tools/ToolConfig.cs b/Engine/Source/Programs/Horde/Plugins/Tools/HordeServer.Tools/Tools/ToolConfig.cs
--- a/Engine/Source/Programs/Horde/Plugins/Tools/HordeServer.Tools/Tools/ToolConfig.cs
+++ b/Engine/Source/Programs/Horde/Plugins/Tools/HordeServer.Tools/Tools/ToolConfig.cs
@@ -1,17 +1,12 @@
 // Copyright Epic Games, Inc. All Rights Reserved.
 
-using System;
-using System.Collections.Generic;
 using System.ComponentModel.DataAnnotations;
 using System.Diagnostics;
 using System.Security.Claims;
-using System.Text.Json.Serialization;
 using EpicGames.Horde.Acls;
 using EpicGames.Horde.Storage;
 using EpicGames.Horde.Tools;
 using HordeServer.Acls;
-using HordeServer.Plugins;
-using HordeServer.Server;
 using HordeServer.Storage;
 
 #pragma warning disable CA2227 // Change x to be read-only by removing the property setter
diff --git a/Engine/Source/Programs/Horde/Plugins/Tools/HordeServer.Tools/ToolsConfig.cs b/Engine/Source/Programs/Horde/Plugins/Tools/HordeServer.Tools/ToolsConfig.cs
--- a/Engine/Source/Programs/Horde/Plugins/Tools/HordeServer.Tools/ToolsConfig.cs
+++ b/Engine/Source/Programs/Horde/Plugins/Tools/HordeServer.Tools/ToolsConfig.cs
@@ -2,7 +2,6 @@
 
 using System.Diagnostics.CodeAnalysis;
 using EpicGames.Horde.Tools;
-using HordeServer.Acls;
 using HordeServer.Plugins;
 
 #pragma warning disable CA2227 // Change x to be read-only by removing the property setter
