Horde: Fix installer after plugin refactor.

diff --git a/Engine/Source/Programs/Horde/BuildHorde.xml b/Engine/Source/Programs/Horde/BuildHorde.xml
--- a/Engine/Source/Programs/Horde/BuildHorde.xml
+++ b/Engine/Source/Programs/Horde/BuildHorde.xml
@@ -123,7 +123,7 @@
 	<Agent Name="HordeInstallerTools" Type="LicenseeWin64;Win64">
 		<Node Name="Build HordeCmd" Produces="#HordeCmdFiles">
 			<Delete Files="$(ServerToolsStagingDir_HordeCmd)/..."/>
-			<DotNet	Arguments="publish &quot;$(RootDir)/Engine/Source/Programs/Horde/Horde/Horde.csproj&quot; --output &quot;$(ServerToolsStagingDir_HordeCmd)&quot; --runtime win-x64 --self-contained $(VersionArguments)" DotNetPath="$(DotNetPath)" />
+			<DotNet	Arguments="publish &quot;$(RootDir)/Engine/Source/Programs/Horde/HordeCmd/HordeCmd.csproj&quot; --output &quot;$(ServerToolsStagingDir_HordeCmd)&quot; --runtime win-x64 --self-contained $(VersionArguments)" DotNetPath="$(DotNetPath)" />
 			<Sign Files="$(ServerToolsStagingDir_HordeCmd)/Horde.exe;$(ServerToolsStagingDir_HordeCmd)/Horde.dll;$(ServerToolsStagingDir_HordeCmd)/EpicGames.*.dll" If="'$(SignOutputs)' == 'true'"/>
 			<Tag Files="$(ServerToolsStagingDir_HordeCmd)/..." With="#HordeCmdFiles"/>
 		</Node>
@@ -148,8 +148,8 @@
 			<Sign Files="$(ServerToolsStagingDir_HordeAgentWin64)/JobDriver/JobDriver.exe;$(ServerToolsStagingDir_HordeAgentWin64)/JobDriver/JobDriver.dll;$(ServerToolsStagingDir_HordeAgentWin64)/JobDriver/EpicGames.*.dll" If="'$(SignOutputs)' == 'true'"/>
 
 			<!-- Build the tray app -->
-			<DotNet Arguments="publish &quot;$(RootDir)/Engine/Source/Programs/Horde/HordeAgent.TrayApp/HordeAgent.TrayApp.csproj&quot; --output &quot;$(ServerToolsStagingDir_HordeAgentWin64)/TrayApp&quot; --runtime win7-x64 --self-contained $(VersionArguments)"/>
-			<Sign Files="$(ServerToolsStagingDir_HordeAgentWin64)/TrayApp/HordeAgent.TrayApp.exe;$(ServerToolsStagingDir_HordeAgentWin64)/TrayApp/EpicGames.*.dll" If="'$(SignOutputs)' == 'true'"/>
+			<DotNet Arguments="publish &quot;$(RootDir)/Engine/Source/Programs/Horde/Horde.Agent.TrayApp/Horde.Agent.TrayApp.csproj&quot; --output &quot;$(ServerToolsStagingDir_HordeAgentWin64)/TrayApp&quot; --runtime win7-x64 --self-contained $(VersionArguments)"/>
+			<Sign Files="$(ServerToolsStagingDir_HordeAgentWin64)/TrayApp/Horde.Agent.TrayApp.exe;$(ServerToolsStagingDir_HordeAgentWin64)/TrayApp/EpicGames.*.dll" If="'$(SignOutputs)' == 'true'"/>
 
 			<Delete Files="$(ServerToolsStagingDir_HordeAgentWin64)/appsettings.Local.json"/>
 			<Tag Files="$(ServerToolsStagingDir_HordeAgentWin64)/..." With="#HordeAgentWin64Files"/>
@@ -226,7 +226,7 @@
 
 			<!-- Create the agent file list -->
 			<Delete Files="$(StagingDir)/Agent-Bulk/..."/>
-			<Tag Files="#HordeAgentWin64Files" Except="HordeAgent.exe;HordeAgent.TrayApp.exe" With="#AgentFiles"/>
+			<Tag Files="#HordeAgentWin64Files" Except="HordeAgent.exe;Horde.Agent.TrayApp.exe" With="#AgentFiles"/>
 			<Copy Files="#AgentFiles" From="$(ServerToolsStagingDir_HordeAgentWin64)" To="$(StagingDir)/Agent-Bulk/Agent"/>
 
 			<Spawn Exe="$(WixDir)/heat.exe" Arguments="dir &quot;$(StagingDir)/Agent-Bulk&quot; -cg HordeAgentFiles -scom -sreg -gg -dr InstallDir -srd -var var.SourceDir -out &quot;$(InstallerDir)/HordeAgentFiles.wxs&quot;"/>
@@ -267,7 +267,7 @@
 
 				<!-- Separate copy of the command line tool for installation with the server on Windows -->
 				<Delete Files="$(StagingDir)/Tools/..."/>
-				<DotNet Arguments="publish &quot;$(RootDir)/Engine/Source/Programs/Horde/Horde/Horde.csproj&quot; --output &quot;$(StagingDir)/Tools&quot; --runtime win-x64 --self-contained $(VersionArguments)" DotNetPath="$(DotNetPath)" />
+				<DotNet Arguments="publish &quot;$(RootDir)/Engine/Source/Programs/Horde/HordeCmd/HordeCmd.csproj&quot; --output &quot;$(StagingDir)/Tools&quot; --runtime win-x64 --self-contained $(VersionArguments)" DotNetPath="$(DotNetPath)" />
 
 				<!-- Copy the Dashboard -->
 				<Copy Files="#DashboardOutputFiles" From="$(StagingDir)/Dashboard/..." To="$(StagingDir)/Server/$(ServerArch)/DashboardApp/..." />
