Horde: Make storage debug controller public.

#rnx

diff --git a/Engine/Source/Programs/Horde/Plugins/Storage/HordeServer.Storage/Debug/StorageDebugController.cs b/Engine/Source/Programs/Horde/Plugins/Storage/HordeServer.Storage/Debug/StorageDebugController.cs
--- a/Engine/Source/Programs/Horde/Plugins/Storage/HordeServer.Storage/Debug/StorageDebugController.cs
+++ b/Engine/Source/Programs/Horde/Plugins/Storage/HordeServer.Storage/Debug/StorageDebugController.cs
@@ -20,7 +20,7 @@ namespace HordeServer.Debug
 	[Authorize]
 	[DebugEndpoint]
 	[Tags("Debug")]
-	class StorageDebugController : HordeControllerBase
+	public class StorageDebugController : HordeControllerBase
 	{
 		readonly IServiceProvider _serviceProvider;
 		readonly IOptionsSnapshot<StorageConfig> _storageConfig;
