Horde: Fix corner case with tables and markdown links to schemas
#rnx

diff --git a/Engine/Source/Programs/Horde/HordeDashboard/src/base/components/Markdown/Markdown.tsx b/Engine/Source/Programs/Horde/HordeDashboard/src/base/components/Markdown/Markdown.tsx
--- a/Engine/Source/Programs/Horde/HordeDashboard/src/base/components/Markdown/Markdown.tsx
+++ b/Engine/Source/Programs/Horde/HordeDashboard/src/base/components/Markdown/Markdown.tsx
@@ -54,8 +54,13 @@ const getStyles: IStyleFunction<IMarkdownStyleProps, IMarkdownStyles> = () => {
 
             },
             'a > code': {
-
-               padding: "3px 6px !important"
+               
+            },
+            'td > a': {
+               fontFamily: "Horde Cousine Regular",
+               fontSize: "12px",
+               backgroundColor: dashboard.darktheme ? "#25282A" : "#f2f2f2",
+               padding: "3px 0px 2px 6px !important"
 
             },
             'td > code': {
diff --git a/Engine/Source/Programs/Horde/HordeServer/Commands/Generate/ConfigDocsCommand.cs b/Engine/Source/Programs/Horde/HordeServer/Commands/Generate/ConfigDocsCommand.cs
--- a/Engine/Source/Programs/Horde/HordeServer/Commands/Generate/ConfigDocsCommand.cs
+++ b/Engine/Source/Programs/Horde/HordeServer/Commands/Generate/ConfigDocsCommand.cs
@@ -255,7 +255,7 @@ namespace HordeServer.Commands.Generate
 					}
 					else
 					{
-						return $"[`{type.Name}`]({GetLink(type, typeNameToLink)})";
+						return $"[{type.Name}]({GetLink(type, typeNameToLink)})";
 					}
 				case JsonSchemaObject obj:
 					if (obj.AdditionalProperties != null)
@@ -268,7 +268,7 @@ namespace HordeServer.Commands.Generate
 					}
 					else
 					{
-						return $"[`{type.Name}`]({GetLink(type, typeNameToLink)})";
+						return $"[{type.Name}]({GetLink(type, typeNameToLink)})";
 					}
 				default:
 					return type.GetType().Name;
