Horde: Updating schema documentation

#rnx

diff --git a/Engine/Source/Programs/Horde/Docs/Config/Schema/Dashboard.md b/Engine/Source/Programs/Horde/Docs/Config/Schema/Dashboard.md
--- a/Engine/Source/Programs/Horde/Docs/Config/Schema/Dashboard.md
+++ b/Engine/Source/Programs/Horde/Docs/Config/Schema/Dashboard.md
@@ -13,11 +13,11 @@ Name | Description
 `showPerforceServers` | `boolean`<br>Show the Perforce server option on the server menu
 `showDeviceManager` | `boolean`<br>Show the device manager on the server menu
 `showTests` | `boolean`<br>Show automated tests on the server menu
-`agentCategories` | [`DashboardAgentCategoryConfig`](#dashboardagentcategoryconfig)`[]`<br>Configuration for different agent pages
-`poolCategories` | [`DashboardPoolCategoryConfig`](#dashboardpoolcategoryconfig)`[]`<br>Configuration for different pool pages
-`analytics` | [`TelemetryViewConfig`](#telemetryviewconfig)`[]`<br>Configuration for telemetry views
-`include` | [`ConfigInclude`](#configinclude)`[]`<br>Includes for other configuration files
-`macros` | [`ConfigMacro`](#configmacro)`[]`<br>Macros within this configuration
+`agentCategories` | [DashboardAgentCategoryConfig](#dashboardagentcategoryconfig)`[]`<br>Configuration for different agent pages
+`poolCategories` | [DashboardPoolCategoryConfig](#dashboardpoolcategoryconfig)`[]`<br>Configuration for different pool pages
+`analytics` | [TelemetryViewConfig](#telemetryviewconfig)`[]`<br>Configuration for telemetry views
+`include` | [ConfigInclude](#configinclude)`[]`<br>Includes for other configuration files
+`macros` | [ConfigMacro](#configmacro)`[]`<br>Macros within this configuration
 
 ## DashboardAgentCategoryConfig
 
@@ -46,8 +46,8 @@ Name | Description
 `id` | `string`<br>Identifier for the view
 `name` | `string`<br>The name of the view
 `telemetryStoreId` | `string`<br>The telemetry store this view uses
-`variables` | [`TelemetryVariableConfig`](#telemetryvariableconfig)`[]`<br>The variables used to filter the view data
-`categories` | [`TelemetryCategoryConfig`](#telemetrycategoryconfig)`[]`<br>The categories contained within the view
+`variables` | [TelemetryVariableConfig](#telemetryvariableconfig)`[]`<br>The variables used to filter the view data
+`categories` | [TelemetryCategoryConfig](#telemetrycategoryconfig)`[]`<br>The categories contained within the view
 
 ## TelemetryVariableConfig
 
@@ -66,7 +66,7 @@ A chart categody, will be displayed on the dashbord under an associated pivot
 Name | Description
 ---- | -----------
 `name` | `string`<br>The name of the category
-`charts` | [`TelemetryChartConfig`](#telemetrychartconfig)`[]`<br>The charts contained within the category
+`charts` | [TelemetryChartConfig](#telemetrychartconfig)`[]`<br>The charts contained within the category
 
 ## TelemetryChartConfig
 
@@ -75,9 +75,9 @@ Telemetry chart configuraton
 Name | Description
 ---- | -----------
 `name` | `string`<br>The name of the chart, will be displayed on the dashboard
-`display` | [`TelemetryMetricUnitType`](#telemetrymetricunittype-enum)<br>The unit to display
-`graph` | [`TelemetryMetricGraphType`](#telemetrymetricgraphtype-enum)<br>The graph type
-`metrics` | [`TelemetryChartMetricConfig`](#telemetrychartmetricconfig)`[]`<br>List of configured metrics
+`display` | [TelemetryMetricUnitType](#telemetrymetricunittype-enum)<br>The unit to display
+`graph` | [TelemetryMetricGraphType](#telemetrymetricgraphtype-enum)<br>The graph type
+`metrics` | [TelemetryChartMetricConfig](#telemetrychartmetricconfig)`[]`<br>List of configured metrics
 `min` | `integer`<br>The min unit value for clamping chart
 `max` | `integer`<br>The max unit value for clamping chart
 
diff --git a/Engine/Source/Programs/Horde/Docs/Config/Schema/Globals.md b/Engine/Source/Programs/Horde/Docs/Config/Schema/Globals.md
--- a/Engine/Source/Programs/Horde/Docs/Config/Schema/Globals.md
+++ b/Engine/Source/Programs/Horde/Docs/Config/Schema/Globals.md
@@ -7,13 +7,13 @@ Global configuration
 Name | Description
 ---- | -----------
 `version` | `integer`<br>Version number for the server. Values are indicated by the .
-`include` | [`ConfigInclude`](#configinclude)`[]`<br>Other paths to include
-`macros` | [`ConfigMacro`](#configmacro)`[]`<br>Macros within the global scope
-`dashboard` | [`DashboardConfig`](Dashboard.md)<br>Settings for the dashboard
-`downtime` | [`ScheduledDowntime`](#scheduleddowntime)`[]`<br>List of scheduled downtime
+`include` | [ConfigInclude](#configinclude)`[]`<br>Other paths to include
+`macros` | [ConfigMacro](#configmacro)`[]`<br>Macros within the global scope
+`dashboard` | [DashboardConfig](Dashboard.md)<br>Settings for the dashboard
+`downtime` | [ScheduledDowntime](#scheduleddowntime)`[]`<br>List of scheduled downtime
 `plugins` | `object`<br>Plugin config objects
 `parameters` | `object`<br>General parameters for other tools. Can be queried through the api/v1/parameters endpoint.
-`acl` | [`AclConfig`](#aclconfig)<br>Access control list
+`acl` | [AclConfig](#aclconfig)<br>Access control list
 
 ## ConfigInclude
 
@@ -40,7 +40,7 @@ Name | Description
 ---- | -----------
 `startTime` | `string`<br>Start time
 `finishTime` | `string`<br>Finish time
-`frequency` | [`ScheduledDowntimeFrequency`](#scheduleddowntimefrequency-enum)<br>Frequency that the window repeats
+`frequency` | [ScheduledDowntimeFrequency](#scheduleddowntimefrequency-enum)<br>Frequency that the window repeats
 
 ## ScheduledDowntimeFrequency (Enum)
 
@@ -58,7 +58,7 @@ Config settings for analytics
 
 Name | Description
 ---- | -----------
-`stores` | [`TelemetryStoreConfig`](Telemetry.md)`[]`<br>Metrics to aggregate on the Horde server
+`stores` | [TelemetryStoreConfig](Telemetry.md)`[]`<br>Metrics to aggregate on the Horde server
 
 ## BuildConfig
 
@@ -66,12 +66,12 @@ Configuration for the build plugin
 
 Name | Description
 ---- | -----------
-`perforceClusters` | [`PerforceCluster`](#perforcecluster)`[]`<br>List of Perforce clusters
-`devices` | [`DeviceConfig`](#deviceconfig)<br>Device configuration
+`perforceClusters` | [PerforceCluster](#perforcecluster)`[]`<br>List of Perforce clusters
+`devices` | [DeviceConfig](#deviceconfig)<br>Device configuration
 `maxConformCount` | `integer`<br>Maximum number of conforms to run at once
 `agentShutdownIfDisabledGracePeriod` | `string`<br>Time to wait before shutting down an agent that has been disabled Used if no value is set on the actual pool.
-`artifactTypes` | [`ArtifactTypeConfig`](#artifacttypeconfig)`[]`<br>Configuration for different artifact types
-`projects` | [`ProjectConfig`](Projects.md)`[]`<br>List of projects
+`artifactTypes` | [ArtifactTypeConfig](#artifacttypeconfig)`[]`<br>Configuration for different artifact types
+`projects` | [ProjectConfig](Projects.md)`[]`<br>List of projects
 `enableConformTasks` | `boolean`<br>Whether to allow conform tasks to run
 `issueFixedTag` | `string`<br>Commit tag to use for marking issues as fixed
 
@@ -85,9 +85,9 @@ Name | Description
 `serviceAccount` | `string`<br>Username for Horde to log in to this server. Will use the first account specified below if not overridden.
 `canImpersonate` | `boolean`<br>Whether the service account can impersonate other users
 `supportsPartitionedWorkspaces` | `boolean`<br>Whether to use partitioned workspaces on this server
-`servers` | [`PerforceServer`](#perforceserver)`[]`<br>List of servers
-`credentials` | [`PerforceCredentials`](#perforcecredentials)`[]`<br>List of server credentials
-`autoSdk` | [`AutoSdkWorkspace`](#autosdkworkspace)`[]`<br>List of autosdk streams
+`servers` | [PerforceServer](#perforceserver)`[]`<br>List of servers
+`credentials` | [PerforceCredentials](#perforcecredentials)`[]`<br>List of server credentials
+`autoSdk` | [AutoSdkWorkspace](#autosdkworkspace)`[]`<br>List of autosdk streams
 
 ## PerforceServer
 
@@ -129,8 +129,8 @@ Configuration for devices
 
 Name | Description
 ---- | -----------
-`platforms` | [`DevicePlatformConfig`](#deviceplatformconfig)`[]`<br>List of device platforms
-`pools` | [`DevicePoolConfig`](#devicepoolconfig)`[]`<br>List of device pools
+`platforms` | [DevicePlatformConfig](#deviceplatformconfig)`[]`<br>List of device platforms
+`pools` | [DevicePoolConfig](#devicepoolconfig)`[]`<br>List of device pools
 
 ## DevicePlatformConfig
 
@@ -152,7 +152,7 @@ Name | Description
 ---- | -----------
 `id` | `string`<br>The id for this platform
 `name` | `string`<br>The name of the pool
-`poolType` | [`DevicePoolType`](#devicepooltype-enum)<br>The type of the pool
+`poolType` | [DevicePoolType](#devicepooltype-enum)<br>The type of the pool
 `projectIds` | `string[]`<br>List of project ids associated with pool
 
 ## DevicePoolType (Enum)
@@ -174,7 +174,7 @@ Name | Description
 `keepCount` | `integer`<br>Number of artifacts to retain
 `keepDays` | `integer`<br>Number of days to retain artifacts of this type
 `type` | `string`<br>Name of the artifact type
-`acl` | [`AclConfig`](#aclconfig)<br>Acl for the artifact type
+`acl` | [AclConfig](#aclconfig)<br>Acl for the artifact type
 
 ## AclConfig
 
@@ -182,8 +182,8 @@ Parameters to update an ACL
 
 Name | Description
 ---- | -----------
-`entries` | [`AclEntryConfig`](#aclentryconfig)`[]`<br>Entries to replace the existing ACL
-`profiles` | [`AclProfileConfig`](#aclprofileconfig)`[]`<br>Defines profiles which allow grouping sets of actions into named collections
+`entries` | [AclEntryConfig](#aclentryconfig)`[]`<br>Entries to replace the existing ACL
+`profiles` | [AclProfileConfig](#aclprofileconfig)`[]`<br>Defines profiles which allow grouping sets of actions into named collections
 `inherit` | `boolean`<br>Whether to inherit permissions from the parent ACL
 `exceptions` | `string[]`<br>List of exceptions to the inherited setting
 
@@ -193,7 +193,7 @@ Individual entry in an ACL
 
 Name | Description
 ---- | -----------
-`claim` | [`AclClaimConfig`](#aclclaimconfig)<br>Name of the user or group
+`claim` | [AclClaimConfig](#aclclaimconfig)<br>Name of the user or group
 `actions` | `string[]`<br>Array of actions to allow
 `profiles` | `string[]`<br>List of profiles to grant
 
@@ -223,13 +223,13 @@ Configuration for the compute system
 
 Name | Description
 ---- | -----------
-`acl` | [`AclConfig`](#aclconfig)<br>Inherited root acl
-`versionEnum` | [`ConfigVersion`](#configversion-enum)<br>Config version number
-`rates` | [`AgentRateConfig`](#agentrateconfig)`[]`<br>List of costs of a particular agent type
-`clusters` | [`ComputeClusterConfig`](#computeclusterconfig)`[]`<br>List of compute profiles
-`pools` | [`PoolConfig`](#poolconfig)`[]`<br>List of pools
-`software` | [`AgentSoftwareConfig`](#agentsoftwareconfig)`[]`<br>List of costs of a particular agent type
-`networks` | [`NetworkConfig`](#networkconfig)`[]`<br>List of networks
+`acl` | [AclConfig](#aclconfig)<br>Inherited root acl
+`versionEnum` | [ConfigVersion](#configversion-enum)<br>Config version number
+`rates` | [AgentRateConfig](#agentrateconfig)`[]`<br>List of costs of a particular agent type
+`clusters` | [ComputeClusterConfig](#computeclusterconfig)`[]`<br>List of compute profiles
+`pools` | [PoolConfig](#poolconfig)`[]`<br>List of pools
+`software` | [AgentSoftwareConfig](#agentsoftwareconfig)`[]`<br>List of costs of a particular agent type
+`networks` | [NetworkConfig](#networkconfig)`[]`<br>List of networks
 
 ## ConfigVersion (Enum)
 
@@ -263,7 +263,7 @@ Name | Description
 `requestBucketId` | `string`<br>Name of the input bucket
 `responseBucketId` | `string`<br>Name of the output bucket
 `condition` | `string`<br>Filter for agents to include
-`acl` | [`AclConfig`](#aclconfig)<br>Access control list
+`acl` | [AclConfig](#aclconfig)<br>Access control list
 
 ## PoolConfig
 
@@ -276,7 +276,7 @@ Name | Description
 `name` | `string`<br>Name of the pool
 `condition` | `string`<br>Condition for agents to automatically be included in this pool
 `properties` | `string` `->` `string`<br>Arbitrary properties related to this pool
-`color` | [`PoolColor`](#poolcolor-enum)<br>Color to use for this pool on the dashboard
+`color` | [PoolColor](#poolcolor-enum)<br>Color to use for this pool on the dashboard
 `enableAutoscaling` | `boolean`<br>Whether to enable autoscaling for this pool
 `minAgents` | `integer`<br>The minimum number of agents to keep in the pool
 `numReserveAgents` | `integer`<br>The minimum number of idle agents to hold in reserve
@@ -284,12 +284,12 @@ Name | Description
 `scaleOutCooldown` | `string`<br>Cooldown time between scale-out events
 `scaleInCooldown` | `string`<br>Cooldown time between scale-in events
 `shutdownIfDisabledGracePeriod` | `string`<br>Time to wait before shutting down an agent that has been disabled
-`sizeStrategy` | [`PoolSizeStrategy`](#poolsizestrategy-enum)<br>
-`sizeStrategies` | [`PoolSizeStrategyInfo`](#poolsizestrategyinfo)`[]`<br>List of pool sizing strategies for this pool. The first strategy with a matching condition will be picked.
-`fleetManagers` | [`FleetManagerInfo`](#fleetmanagerinfo)`[]`<br>List of fleet managers for this pool. The first strategy with a matching condition will be picked. If empty or no conditions match, a default fleet manager will be used.
-`leaseUtilizationSettings` | [`LeaseUtilizationSettings`](#leaseutilizationsettings)<br>Settings for lease utilization pool sizing strategy (if used)
-`jobQueueSettings` | [`JobQueueSettings`](#jobqueuesettings)<br>Settings for job queue pool sizing strategy (if used)
-`computeQueueAwsMetricSettings` | [`ComputeQueueAwsMetricSettings`](#computequeueawsmetricsettings)<br>Settings for job queue pool sizing strategy (if used)
+`sizeStrategy` | [PoolSizeStrategy](#poolsizestrategy-enum)<br>
+`sizeStrategies` | [PoolSizeStrategyInfo](#poolsizestrategyinfo)`[]`<br>List of pool sizing strategies for this pool. The first strategy with a matching condition will be picked.
+`fleetManagers` | [FleetManagerInfo](#fleetmanagerinfo)`[]`<br>List of fleet managers for this pool. The first strategy with a matching condition will be picked. If empty or no conditions match, a default fleet manager will be used.
+`leaseUtilizationSettings` | [LeaseUtilizationSettings](#leaseutilizationsettings)<br>Settings for lease utilization pool sizing strategy (if used)
+`jobQueueSettings` | [JobQueueSettings](#jobqueuesettings)<br>Settings for job queue pool sizing strategy (if used)
+`computeQueueAwsMetricSettings` | [ComputeQueueAwsMetricSettings](#computequeueawsmetricsettings)<br>Settings for job queue pool sizing strategy (if used)
 
 ## PoolColor (Enum)
 
@@ -321,7 +321,7 @@ Metadata for configuring and picking a pool sizing strategy
 
 Name | Description
 ---- | -----------
-`type` | [`PoolSizeStrategy`](#poolsizestrategy-enum)<br>Strategy implementation to use
+`type` | [PoolSizeStrategy](#poolsizestrategy-enum)<br>Strategy implementation to use
 `condition` | `string`<br>Condition if this strategy should be enabled (right now, using date/time as a distinguishing factor)
 `config` | `string`<br>Configuration for the strategy, serialized as JSON
 `extraAgentCount` | `integer`<br>Integer to add after pool size has been calculated. Can also be negative.
@@ -332,7 +332,7 @@ Metadata for configuring and picking a fleet manager
 
 Name | Description
 ---- | -----------
-`type` | [`FleetManagerType`](#fleetmanagertype-enum)<br>Fleet manager type implementation to use
+`type` | [FleetManagerType](#fleetmanagertype-enum)<br>Fleet manager type implementation to use
 `condition` | `string`<br>Condition if this strategy should be enabled (right now, using date/time as a distinguishing factor)
 `config` | `string`<br>Configuration for the strategy, serialized as JSON
 
@@ -412,7 +412,7 @@ Configuration for the secrets system
 
 Name | Description
 ---- | -----------
-`secrets` | [`SecretConfig`](#secretconfig)`[]`<br>List of secrets
+`secrets` | [SecretConfig](#secretconfig)`[]`<br>List of secrets
 
 ## SecretConfig
 
@@ -422,8 +422,8 @@ Name | Description
 ---- | -----------
 `id` | `string`<br>Identifier for this secret
 `data` | `string` `->` `string`<br>Key/value pairs associated with this secret
-`sources` | [`ExternalSecretConfig`](#externalsecretconfig)`[]`<br>Providers to source key/value pairs from
-`acl` | [`AclConfig`](#aclconfig)<br>Defines access to this particular secret
+`sources` | [ExternalSecretConfig](#externalsecretconfig)`[]`<br>Providers to source key/value pairs from
+`acl` | [AclConfig](#aclconfig)<br>Defines access to this particular secret
 
 ## ExternalSecretConfig
 
@@ -432,7 +432,7 @@ Configuration for an external secret provider
 Name | Description
 ---- | -----------
 `provider` | `string`<br>Name of the provider to use
-`format` | [`ExternalSecretFormat`](#externalsecretformat-enum)<br>Format of the secret
+`format` | [ExternalSecretFormat](#externalsecretformat-enum)<br>Format of the secret
 `key` | `string`<br>Optional key indicating the parameter to set in the resulting data array. Required if if  is .
 `path` | `string`<br>Optional value indicating what to fetch from the provider
 
@@ -453,8 +453,8 @@ Name | Description
 ---- | -----------
 `enableGc` | `boolean`<br>Whether to enable garbage collection
 `enableGcVerification` | `boolean`<br>Whether to enable garbage collection in verification mode (nothing deleted, just logging on access to deleted blobs)
-`backends` | [`BackendConfig`](#backendconfig)`[]`<br>List of storage backends
-`namespaces` | [`NamespaceConfig`](#namespaceconfig)`[]`<br>List of namespaces for storage
+`backends` | [BackendConfig](#backendconfig)`[]`<br>List of storage backends
+`namespaces` | [NamespaceConfig](#namespaceconfig)`[]`<br>List of namespaces for storage
 
 ## BackendConfig
 
@@ -465,11 +465,11 @@ Name | Description
 `id` | `string`<br>The storage backend ID
 `base` | `string`<br>Base backend to copy default settings from
 `secondary` | `string`<br>Specifies another backend to read from if an object is not found in this one. Can be used when migrating data from one backend to another.
-`type` | [`StorageBackendType`](#storagebackendtype-enum)<br>
+`type` | [StorageBackendType](#storagebackendtype-enum)<br>
 `baseDir` | `string`<br>
 `awsBucketName` | `string`<br>Name of the bucket to use
 `awsBucketPath` | `string`<br>Base path within the bucket
-`awsCredentials` | [`AwsCredentialsType`](#awscredentialstype-enum)<br>Type of credentials to use
+`awsCredentials` | [AwsCredentialsType](#awscredentialstype-enum)<br>Type of credentials to use
 `awsRole` | `string`<br>ARN of a role to assume
 `awsProfile` | `string`<br>The AWS profile to read credentials form
 `awsRegion` | `string`<br>Region to connect to
@@ -512,7 +512,7 @@ Name | Description
 `gcFrequencyHrs` | `number`<br>How frequently to run garbage collection, in hours.
 `gcDelayHrs` | `number`<br>How long to keep newly uploaded orphanned objects before allowing them to be deleted, in hours.
 `enableAliases` | `boolean`<br>Support querying exports by their aliases
-`acl` | [`AclConfig`](#aclconfig)<br>Access list for this namespace
+`acl` | [AclConfig](#aclconfig)<br>Access list for this namespace
 
 ## ToolsConfig
 
@@ -520,7 +520,7 @@ Configuration for the tools system
 
 Name | Description
 ---- | -----------
-`tools` | [`ToolConfig`](#toolconfig)`[]`<br>Tool configurations
+`tools` | [ToolConfig](#toolconfig)`[]`<br>Tool configurations
 
 ## ToolConfig
 
@@ -538,4 +538,4 @@ Name | Description
 `showInUgs` | `boolean`<br>Whether to show this tool for download in the UGS tools menu
 `showInDashboard` | `boolean`<br>Whether to show this tool for download in the dashboard
 `namespaceId` | `string`<br>Default namespace for new deployments of this tool
-`acl` | [`AclConfig`](#aclconfig)<br>Permissions for the tool
+`acl` | [AclConfig](#aclconfig)<br>Permissions for the tool
diff --git a/Engine/Source/Programs/Horde/Docs/Config/Schema/Projects.md b/Engine/Source/Programs/Horde/Docs/Config/Schema/Projects.md
--- a/Engine/Source/Programs/Horde/Docs/Config/Schema/Projects.md
+++ b/Engine/Source/Programs/Horde/Docs/Config/Schema/Projects.md
@@ -9,19 +9,19 @@ Name | Description
 `id` | `string`<br>The project id
 `name` | `string`<br>Name for the new project
 `path` | `string`<br>Direct include path for the project config. For backwards compatibility with old config files when including from a GlobalConfig object.
-`include` | [`ConfigInclude`](#configinclude)`[]`<br>Includes for other configuration files
-`macros` | [`ConfigMacro`](#configmacro)`[]`<br>Macros within the global scope
+`include` | [ConfigInclude](#configinclude)`[]`<br>Includes for other configuration files
+`macros` | [ConfigMacro](#configmacro)`[]`<br>Macros within the global scope
 `order` | `integer`<br>Order of this project on the dashboard
 `logo` | `string`<br>Path to the project logo
 `logoDarkTheme` | `string`<br>Optional path to the project logo for the dark theme
-`pools` | [`PoolConfig`](#poolconfig)`[]`<br>List of pools for this project
-`categories` | [`ProjectCategoryConfig`](#projectcategoryconfig)`[]`<br>Categories to include in this project
-`jobOptions` | [`JobOptions`](#joboptions)<br>Default settings for executing jobs
-`workspaceTypes` | `string` `->` [`WorkspaceConfig`](#workspaceconfig)<br>Default workspace types for streams These are added to the list of each stream's workspace types.
+`pools` | [PoolConfig](#poolconfig)`[]`<br>List of pools for this project
+`categories` | [ProjectCategoryConfig](#projectcategoryconfig)`[]`<br>Categories to include in this project
+`jobOptions` | [JobOptions](#joboptions)<br>Default settings for executing jobs
+`workspaceTypes` | `string` `->` [WorkspaceConfig](#workspaceconfig)<br>Default workspace types for streams These are added to the list of each stream's workspace types.
 `telemetryStoreId` | `string`<br>Telemetry store for Horde data for this project
-`streams` | [`StreamConfig`](Streams.md)`[]`<br>List of streams
-`artifactTypes` | [`ArtifactTypeAclConfig`](#artifacttypeaclconfig)`[]`<br>Permissions for artifact types
-`acl` | [`AclConfig`](#aclconfig)<br>Acl entries
+`streams` | [StreamConfig](Streams.md)`[]`<br>List of streams
+`artifactTypes` | [ArtifactTypeAclConfig](#artifacttypeaclconfig)`[]`<br>Permissions for artifact types
+`acl` | [AclConfig](#aclconfig)<br>Acl entries
 
 ## ConfigInclude
 
@@ -51,7 +51,7 @@ Name | Description
 `name` | `string`<br>Name of the pool
 `condition` | `string`<br>Condition for agents to automatically be included in this pool
 `properties` | `string` `->` `string`<br>Arbitrary properties related to this pool
-`color` | [`PoolColor`](#poolcolor-enum)<br>Color to use for this pool on the dashboard
+`color` | [PoolColor](#poolcolor-enum)<br>Color to use for this pool on the dashboard
 `enableAutoscaling` | `boolean`<br>Whether to enable autoscaling for this pool
 `minAgents` | `integer`<br>The minimum number of agents to keep in the pool
 `numReserveAgents` | `integer`<br>The minimum number of idle agents to hold in reserve
@@ -59,12 +59,12 @@ Name | Description
 `scaleOutCooldown` | `string`<br>Cooldown time between scale-out events
 `scaleInCooldown` | `string`<br>Cooldown time between scale-in events
 `shutdownIfDisabledGracePeriod` | `string`<br>Time to wait before shutting down an agent that has been disabled
-`sizeStrategy` | [`PoolSizeStrategy`](#poolsizestrategy-enum)<br>
-`sizeStrategies` | [`PoolSizeStrategyInfo`](#poolsizestrategyinfo)`[]`<br>List of pool sizing strategies for this pool. The first strategy with a matching condition will be picked.
-`fleetManagers` | [`FleetManagerInfo`](#fleetmanagerinfo)`[]`<br>List of fleet managers for this pool. The first strategy with a matching condition will be picked. If empty or no conditions match, a default fleet manager will be used.
-`leaseUtilizationSettings` | [`LeaseUtilizationSettings`](#leaseutilizationsettings)<br>Settings for lease utilization pool sizing strategy (if used)
-`jobQueueSettings` | [`JobQueueSettings`](#jobqueuesettings)<br>Settings for job queue pool sizing strategy (if used)
-`computeQueueAwsMetricSettings` | [`ComputeQueueAwsMetricSettings`](#computequeueawsmetricsettings)<br>Settings for job queue pool sizing strategy (if used)
+`sizeStrategy` | [PoolSizeStrategy](#poolsizestrategy-enum)<br>
+`sizeStrategies` | [PoolSizeStrategyInfo](#poolsizestrategyinfo)`[]`<br>List of pool sizing strategies for this pool. The first strategy with a matching condition will be picked.
+`fleetManagers` | [FleetManagerInfo](#fleetmanagerinfo)`[]`<br>List of fleet managers for this pool. The first strategy with a matching condition will be picked. If empty or no conditions match, a default fleet manager will be used.
+`leaseUtilizationSettings` | [LeaseUtilizationSettings](#leaseutilizationsettings)<br>Settings for lease utilization pool sizing strategy (if used)
+`jobQueueSettings` | [JobQueueSettings](#jobqueuesettings)<br>Settings for job queue pool sizing strategy (if used)
+`computeQueueAwsMetricSettings` | [ComputeQueueAwsMetricSettings](#computequeueawsmetricsettings)<br>Settings for job queue pool sizing strategy (if used)
 
 ## PoolColor (Enum)
 
@@ -96,7 +96,7 @@ Metadata for configuring and picking a pool sizing strategy
 
 Name | Description
 ---- | -----------
-`type` | [`PoolSizeStrategy`](#poolsizestrategy-enum)<br>Strategy implementation to use
+`type` | [PoolSizeStrategy](#poolsizestrategy-enum)<br>Strategy implementation to use
 `condition` | `string`<br>Condition if this strategy should be enabled (right now, using date/time as a distinguishing factor)
 `config` | `string`<br>Configuration for the strategy, serialized as JSON
 `extraAgentCount` | `integer`<br>Integer to add after pool size has been calculated. Can also be negative.
@@ -107,7 +107,7 @@ Metadata for configuring and picking a fleet manager
 
 Name | Description
 ---- | -----------
-`type` | [`FleetManagerType`](#fleetmanagertype-enum)<br>Fleet manager type implementation to use
+`type` | [FleetManagerType](#fleetmanagertype-enum)<br>Fleet manager type implementation to use
 `condition` | `string`<br>Condition if this strategy should be enabled (right now, using date/time as a distinguishing factor)
 `config` | `string`<br>Configuration for the strategy, serialized as JSON
 
@@ -178,7 +178,7 @@ Name | Description
 `useWine` | `boolean`<br>Whether to execute using Wine emulation on Linux
 `runInSeparateProcess` | `boolean`<br>Executes the job lease in a separate process
 `workspaceMaterializer` | `string`<br>What workspace materializer to use in WorkspaceExecutor. Will override any value from workspace config.
-`container` | [`JobContainerOptions`](#jobcontaineroptions)<br>Options for executing a job inside a container
+`container` | [JobContainerOptions](#jobcontaineroptions)<br>Options for executing a job inside a container
 `expireAfterDays` | `integer`<br>Number of days after which to expire jobs
 `driver` | `string`<br>Name of the driver to use
 
@@ -221,7 +221,7 @@ ACL configuration for an artifact type
 Name | Description
 ---- | -----------
 `type` | `string`<br>Name of the artifact type
-`acl` | [`AclConfig`](#aclconfig)<br>Acl for the artifact type
+`acl` | [AclConfig](#aclconfig)<br>Acl for the artifact type
 
 ## AclConfig
 
@@ -229,8 +229,8 @@ Parameters to update an ACL
 
 Name | Description
 ---- | -----------
-`entries` | [`AclEntryConfig`](#aclentryconfig)`[]`<br>Entries to replace the existing ACL
-`profiles` | [`AclProfileConfig`](#aclprofileconfig)`[]`<br>Defines profiles which allow grouping sets of actions into named collections
+`entries` | [AclEntryConfig](#aclentryconfig)`[]`<br>Entries to replace the existing ACL
+`profiles` | [AclProfileConfig](#aclprofileconfig)`[]`<br>Defines profiles which allow grouping sets of actions into named collections
 `inherit` | `boolean`<br>Whether to inherit permissions from the parent ACL
 `exceptions` | `string[]`<br>List of exceptions to the inherited setting
 
@@ -240,7 +240,7 @@ Individual entry in an ACL
 
 Name | Description
 ---- | -----------
-`claim` | [`AclClaimConfig`](#aclclaimconfig)<br>Name of the user or group
+`claim` | [AclClaimConfig](#aclclaimconfig)<br>Name of the user or group
 `actions` | `string[]`<br>Array of actions to allow
 `profiles` | `string[]`<br>List of profiles to grant
 
diff --git a/Engine/Source/Programs/Horde/Docs/Config/Schema/Streams.md b/Engine/Source/Programs/Horde/Docs/Config/Schema/Streams.md
--- a/Engine/Source/Programs/Horde/Docs/Config/Schema/Streams.md
+++ b/Engine/Source/Programs/Horde/Docs/Config/Schema/Streams.md
@@ -8,8 +8,8 @@ Name | Description
 ---- | -----------
 `id` | `string`<br>Identifier for the stream
 `path` | `string`<br>Direct include path for the stream config. For backwards compatibility with old config files when including from a ProjectConfig object.
-`include` | [`ConfigInclude`](#configinclude)`[]`<br>Includes for other configuration files
-`macros` | [`ConfigMacro`](#configmacro)`[]`<br>Macros within this stream
+`include` | [ConfigInclude](#configinclude)`[]`<br>Includes for other configuration files
+`macros` | [ConfigMacro](#configmacro)`[]`<br>Macros within this stream
 `name` | `string`<br>Name of the stream
 `clusterName` | `string`<br>The perforce cluster containing the stream
 `order` | `integer`<br>Order for this stream
@@ -17,24 +17,24 @@ Name | Description
 `notificationChannel` | `string`<br>Notification channel for all jobs in this stream
 `notificationChannelFilter` | `string`<br>Notification channel filter for this template. Can be Success, Failure, or Warnings.
 `triageChannel` | `string`<br>Channel to post issue triage notifications
-`jobOptions` | [`JobOptions`](#joboptions)<br>Default settings for executing jobs
+`jobOptions` | [JobOptions](#joboptions)<br>Default settings for executing jobs
 `telemetryStoreId` | `string`<br>Telemetry store for Horde data for this stream
 `autoSdkView` | `string[]`<br>View for the AutoSDK paths to sync. If null, the whole thing will be synced.
 `defaultPreflightTemplate` | `string`<br>Legacy name for the default preflight template
-`defaultPreflight` | [`DefaultPreflightConfig`](#defaultpreflightconfig)<br>Default template for running preflights
-`commitTags` | [`CommitTagConfig`](#committagconfig)`[]`<br>List of tags to apply to commits. Allows fast searching and classification of different commit types (eg. code vs content).
-`tabs` | [`TabConfig`](#tabconfig)`[]`<br>List of tabs to show for the new stream
+`defaultPreflight` | [DefaultPreflightConfig](#defaultpreflightconfig)<br>Default template for running preflights
+`commitTags` | [CommitTagConfig](#committagconfig)`[]`<br>List of tags to apply to commits. Allows fast searching and classification of different commit types (eg. code vs content).
+`tabs` | [TabConfig](#tabconfig)`[]`<br>List of tabs to show for the new stream
 `environment` | `string` `->` `string`<br>Global environment variables for all agents in this stream
-`agentTypes` | `string` `->` [`AgentConfig`](#agentconfig)<br>Map of agent name to type
-`workspaceTypes` | `string` `->` [`WorkspaceConfig`](#workspaceconfig)<br>Map of workspace name to type
-`templates` | [`TemplateRefConfig`](#templaterefconfig)`[]`<br>List of templates to create
-`acl` | [`AclConfig`](#aclconfig)<br>Custom permissions for this object
+`agentTypes` | `string` `->` [AgentConfig](#agentconfig)<br>Map of agent name to type
+`workspaceTypes` | `string` `->` [WorkspaceConfig](#workspaceconfig)<br>Map of workspace name to type
+`templates` | [TemplateRefConfig](#templaterefconfig)`[]`<br>List of templates to create
+`acl` | [AclConfig](#aclconfig)<br>Custom permissions for this object
 `pausedUntil` | `string`<br>Pause stream builds until specified date
 `pauseComment` | `string`<br>Reason for pausing builds of the stream
-`replicators` | [`ReplicatorConfig`](#replicatorconfig)`[]`<br>Configuration for workers to replicate commit data into Horde Storage.
-`workflows` | [`WorkflowConfig`](#workflowconfig)`[]`<br>Workflows for dealing with new issues
-`tokens` | [`TokenConfig`](#tokenconfig)`[]`<br>Tokens to create for each job step
-`artifactTypes` | [`ArtifactTypeAclConfig`](#artifacttypeaclconfig)`[]`<br>Permissions for artifact types
+`replicators` | [ReplicatorConfig](#replicatorconfig)`[]`<br>Configuration for workers to replicate commit data into Horde Storage.
+`workflows` | [WorkflowConfig](#workflowconfig)`[]`<br>Workflows for dealing with new issues
+`tokens` | [TokenConfig](#tokenconfig)`[]`<br>Tokens to create for each job step
+`artifactTypes` | [ArtifactTypeAclConfig](#artifacttypeaclconfig)`[]`<br>Permissions for artifact types
 
 ## ConfigInclude
 
@@ -63,7 +63,7 @@ Name | Description
 `useWine` | `boolean`<br>Whether to execute using Wine emulation on Linux
 `runInSeparateProcess` | `boolean`<br>Executes the job lease in a separate process
 `workspaceMaterializer` | `string`<br>What workspace materializer to use in WorkspaceExecutor. Will override any value from workspace config.
-`container` | [`JobContainerOptions`](#jobcontaineroptions)<br>Options for executing a job inside a container
+`container` | [JobContainerOptions](#jobcontaineroptions)<br>Options for executing a job inside a container
 `expireAfterDays` | `integer`<br>Number of days after which to expire jobs
 `driver` | `string`<br>Name of the driver to use
 
@@ -85,7 +85,7 @@ Specifies defaults for running a preflight
 Name | Description
 ---- | -----------
 `templateId` | `string`<br>The template id to query
-`change` | [`ChangeQueryConfig`](#changequeryconfig)<br>Query for the change to use
+`change` | [ChangeQueryConfig](#changequeryconfig)<br>Query for the change to use
 
 ## ChangeQueryConfig
 
@@ -97,7 +97,7 @@ Name | Description
 `condition` | `string`<br>Condition to evaluate before deciding to use this query. May query tags in a preflight.
 `templateId` | `string`<br>The template id to query
 `target` | `string`<br>The target to query
-`outcomes` | [`JobStepOutcome`](#jobstepoutcome-enum)`[]`<br>Whether to match a job that produced warnings
+`outcomes` | [JobStepOutcome](#jobstepoutcome-enum)`[]`<br>Whether to match a job that produced warnings
 `commitTag` | `string`<br>Finds the last commit with this tag
 
 ## JobStepOutcome (Enum)
@@ -129,12 +129,12 @@ Name | Description
 ---- | -----------
 `title` | `string`<br>Title of this page
 `type` | `string`<br>Type of this tab
-`style` | [`TabStyle`](#tabstyle-enum)<br>Presentation style for this page
+`style` | [TabStyle](#tabstyle-enum)<br>Presentation style for this page
 `showNames` | `boolean`<br>Whether to show job names on this page
 `showPreflights` | `boolean`<br>Whether to show all user preflights
 `jobNames` | `string[]`<br>Names of jobs to include on this page. If there is only one name specified, the name column does not need to be displayed.
 `templates` | `string[]`<br>List of job template names to show on this page.
-`columns` | [`TabColumnConfig`](#tabcolumnconfig)`[]`<br>Columns to display for different types of aggregates
+`columns` | [TabColumnConfig](#tabcolumnconfig)`[]`<br>Columns to display for different types of aggregates
 
 ## TabStyle (Enum)
 
@@ -151,7 +151,7 @@ Describes a column to display on the jobs page
 
 Name | Description
 ---- | -----------
-`type` | [`TabColumnType`](#tabcolumntype-enum)<br>The type of column
+`type` | [TabColumnType](#tabcolumntype-enum)<br>The type of column
 `heading` | `string`<br>Heading for this column
 `category` | `string`<br>Category of aggregates to display in this column. If null, includes any aggregate not matched by another column.
 `parameter` | `string`<br>Parameter to show in this column
@@ -177,7 +177,7 @@ Name | Description
 `workspace` | `string`<br>Name of the workspace to sync
 `tempStorageDir` | `string`<br>Path to the temporary storage dir
 `environment` | `string` `->` `string`<br>Environment variables to be set when executing the job
-`tokens` | [`TokenConfig`](#tokenconfig)`[]`<br>Tokens to allocate for this agent type
+`tokens` | [TokenConfig](#tokenconfig)`[]`<br>Tokens to allocate for this agent type
 
 ## TokenConfig
 
@@ -226,22 +226,22 @@ Name | Description
 `triageChannel` | `string`<br>Triage channel for this template. Overrides the stream channel if set.
 `workflowId` | `string`<br>Workflow to user for this stream
 `annotations` | `string` `->` `string`<br>Default annotations to apply to nodes in this template
-`schedule` | [`ScheduleConfig`](#scheduleconfig)<br>Schedule to execute this template
-`chainedJobs` | [`ChainedJobTemplateConfig`](#chainedjobtemplateconfig)`[]`<br>List of chained job triggers
-`acl` | [`AclConfig`](#aclconfig)<br>The ACL for this template
+`schedule` | [ScheduleConfig](#scheduleconfig)<br>Schedule to execute this template
+`chainedJobs` | [ChainedJobTemplateConfig](#chainedjobtemplateconfig)`[]`<br>List of chained job triggers
+`acl` | [AclConfig](#aclconfig)<br>The ACL for this template
 `name` | `string`<br>Name for the new template
 `description` | `string`<br>Description for the template
-`priority` | [`Priority`](#priority-enum)<br>Default priority for this job
+`priority` | [Priority](#priority-enum)<br>Default priority for this job
 `allowPreflights` | `boolean`<br>Whether to allow preflights of this template
 `updateIssues` | `boolean`<br>Whether issues should be updated for all jobs using this template
 `promoteIssuesByDefault` | `boolean`<br>Whether issues should be promoted by default for this template, promoted issues will generate user notifications
 `initialAgentType` | `string`<br>Initial agent type to parse the buildgraph script on
 `submitNewChange` | `string`<br>Path to a file within the stream to submit to generate a new changelist for jobs
 `submitDescription` | `string`<br>Description for new changelists
-`defaultChange` | [`ChangeQueryConfig`](#changequeryconfig)`[]`<br>Default change to build at. Each object has a condition parameter which can evaluated by the server to determine which change to use.
+`defaultChange` | [ChangeQueryConfig](#changequeryconfig)`[]`<br>Default change to build at. Each object has a condition parameter which can evaluated by the server to determine which change to use.
 `arguments` | `string[]`<br>Fixed arguments for the new job
-`parameters` | [`TextParameterData`](#textparameterdata)/[`ListParameterData`](#listparameterdata)/[`BoolParameterData`](#boolparameterdata)`[]`<br>Parameters for this template
-`jobOptions` | [`JobOptions`](#joboptions)<br>Default settings for jobs
+`parameters` | [TextParameterData](#textparameterdata)/[ListParameterData](#listparameterdata)/[BoolParameterData](#boolparameterdata)`[]`<br>Parameters for this template
+`jobOptions` | [JobOptions](#joboptions)<br>Default settings for jobs
 
 ## ScheduleConfig
 
@@ -249,17 +249,17 @@ Parameters to create a new schedule
 
 Name | Description
 ---- | -----------
-`claims` | [`AclClaimConfig`](#aclclaimconfig)`[]`<br>Roles to impersonate for this schedule
+`claims` | [AclClaimConfig](#aclclaimconfig)`[]`<br>Roles to impersonate for this schedule
 `enabled` | `boolean`<br>Whether the schedule should be enabled
 `maxActive` | `integer`<br>Maximum number of builds that can be active at once
 `maxChanges` | `integer`<br>Maximum number of changes the schedule can fall behind head revision. If greater than zero, builds will be triggered for every submitted changelist until the backlog is this size.
 `requireSubmittedChange` | `boolean`<br>Whether the build requires a change to be submitted
-`gate` | [`ScheduleGateConfig`](#schedulegateconfig)<br>Gate allowing the schedule to trigger
+`gate` | [ScheduleGateConfig](#schedulegateconfig)<br>Gate allowing the schedule to trigger
 `commits` | `string[]`<br>Commit tags for this schedule
-`filter` | [`ChangeContentFlags`](#changecontentflags-enum)`[]`<br>The types of changes to run for
+`filter` | [ChangeContentFlags](#changecontentflags-enum)`[]`<br>The types of changes to run for
 `files` | `string[]`<br>Files that should cause the job to trigger
 `templateParameters` | `string` `->` `string`<br>Parameters for the template
-`patterns` | [`SchedulePatternConfig`](#schedulepatternconfig)`[]`<br>New patterns for the schedule
+`patterns` | [SchedulePatternConfig](#schedulepatternconfig)`[]`<br>New patterns for the schedule
 
 ## AclClaimConfig
 
@@ -294,7 +294,7 @@ Parameters to create a new schedule
 
 Name | Description
 ---- | -----------
-`daysOfWeek` | [`DayOfWeek`](#dayofweek-enum)`[]`<br>Days of the week to run this schedule on. If null, the schedule will run every day.
+`daysOfWeek` | [DayOfWeek](#dayofweek-enum)`[]`<br>Days of the week to run this schedule on. If null, the schedule will run every day.
 `minTime` | `string`<br>Time during the day for the first schedule to trigger. Measured in minutes from midnight.
 `maxTime` | `string`<br>Time during the day for the last schedule to trigger. Measured in minutes from midnight.
 `interval` | `string`<br>Interval between each schedule triggering
@@ -327,8 +327,8 @@ Parameters to update an ACL
 
 Name | Description
 ---- | -----------
-`entries` | [`AclEntryConfig`](#aclentryconfig)`[]`<br>Entries to replace the existing ACL
-`profiles` | [`AclProfileConfig`](#aclprofileconfig)`[]`<br>Defines profiles which allow grouping sets of actions into named collections
+`entries` | [AclEntryConfig](#aclentryconfig)`[]`<br>Entries to replace the existing ACL
+`profiles` | [AclProfileConfig](#aclprofileconfig)`[]`<br>Defines profiles which allow grouping sets of actions into named collections
 `inherit` | `boolean`<br>Whether to inherit permissions from the parent ACL
 `exceptions` | `string[]`<br>List of exceptions to the inherited setting
 
@@ -338,7 +338,7 @@ Individual entry in an ACL
 
 Name | Description
 ---- | -----------
-`claim` | [`AclClaimConfig`](#aclclaimconfig)<br>Name of the user or group
+`claim` | [AclClaimConfig](#aclclaimconfig)<br>Name of the user or group
 `actions` | `string[]`<br>Array of actions to allow
 `profiles` | `string[]`<br>List of profiles to grant
 
@@ -392,8 +392,8 @@ Name | Description
 ---- | -----------
 `type` | List<br>Type discriminator
 `label` | `string`<br>Label to display next to this parameter. Defaults to the parameter name.
-`style` | [`ListParameterStyle`](#listparameterstyle-enum)<br>The type of list parameter
-`items` | [`ListParameterItemData`](#listparameteritemdata)`[]`<br>List of values to display in the list
+`style` | [ListParameterStyle](#listparameterstyle-enum)<br>The type of list parameter
+`items` | [ListParameterItemData](#listparameteritemdata)`[]`<br>List of values to display in the list
 `toolTip` | `string`<br>Tool tip text to display
 
 ## ListParameterStyle (Enum)
@@ -475,7 +475,7 @@ Name | Description
 `inviteRestrictedUsers` | `boolean`<br>Uses the admin.conversations.invite API to invite users to the channel
 `skipWhenEmpty` | `boolean`<br>Skips sending reports when there are no active issues.
 `annotations` | `string` `->` `string`<br>Additional node annotations implicit in this workflow
-`externalIssues` | [`ExternalIssueConfig`](#externalissueconfig)<br>External issue tracking configuration for this workflow
+`externalIssues` | [ExternalIssueConfig](#externalissueconfig)<br>External issue tracking configuration for this workflow
 `issueHandlers` | `string[]`<br>Additional issue handlers enabled for this workflow
 
 ## ExternalIssueConfig
@@ -495,4 +495,4 @@ ACL configuration for an artifact type
 Name | Description
 ---- | -----------
 `type` | `string`<br>Name of the artifact type
-`acl` | [`AclConfig`](#aclconfig)<br>Acl for the artifact type
+`acl` | [AclConfig](#aclconfig)<br>Acl for the artifact type
diff --git a/Engine/Source/Programs/Horde/Docs/Config/Schema/Telemetry.md b/Engine/Source/Programs/Horde/Docs/Config/Schema/Telemetry.md
--- a/Engine/Source/Programs/Horde/Docs/Config/Schema/Telemetry.md
+++ b/Engine/Source/Programs/Horde/Docs/Config/Schema/Telemetry.md
@@ -7,10 +7,10 @@ Config for metrics
 Name | Description
 ---- | -----------
 `id` | `string`<br>Identifier for this store
-`acl` | [`AclConfig`](#aclconfig)<br>Permissions for this store
-`metrics` | [`MetricConfig`](#metricconfig)`[]`<br>Metrics to aggregate on the Horde server
-`include` | [`ConfigInclude`](#configinclude)`[]`<br>Includes for other configuration files
-`macros` | [`ConfigMacro`](#configmacro)`[]`<br>Macros within this configuration
+`acl` | [AclConfig](#aclconfig)<br>Permissions for this store
+`metrics` | [MetricConfig](#metricconfig)`[]`<br>Metrics to aggregate on the Horde server
+`include` | [ConfigInclude](#configinclude)`[]`<br>Includes for other configuration files
+`macros` | [ConfigMacro](#configmacro)`[]`<br>Macros within this configuration
 
 ## AclConfig
 
@@ -18,8 +18,8 @@ Parameters to update an ACL
 
 Name | Description
 ---- | -----------
-`entries` | [`AclEntryConfig`](#aclentryconfig)`[]`<br>Entries to replace the existing ACL
-`profiles` | [`AclProfileConfig`](#aclprofileconfig)`[]`<br>Defines profiles which allow grouping sets of actions into named collections
+`entries` | [AclEntryConfig](#aclentryconfig)`[]`<br>Entries to replace the existing ACL
+`profiles` | [AclProfileConfig](#aclprofileconfig)`[]`<br>Defines profiles which allow grouping sets of actions into named collections
 `inherit` | `boolean`<br>Whether to inherit permissions from the parent ACL
 `exceptions` | `string[]`<br>List of exceptions to the inherited setting
 
@@ -29,7 +29,7 @@ Individual entry in an ACL
 
 Name | Description
 ---- | -----------
-`claim` | [`AclClaimConfig`](#aclclaimconfig)<br>Name of the user or group
+`claim` | [AclClaimConfig](#aclclaimconfig)<br>Name of the user or group
 `actions` | `string[]`<br>Array of actions to allow
 `profiles` | `string[]`<br>List of profiles to grant
 
@@ -63,7 +63,7 @@ Name | Description
 `filter` | `string`<br>Filter expression to evaluate to determine which events to include. This query is evaluated against an array.
 `property` | `string`<br>Property to aggregate
 `groupBy` | `string`<br>Property to group by. Specified as a comma-separated list of JSON path expressions.
-`function` | [`AggregationFunction`](#aggregationfunction-enum)<br>How to aggregate samples for this metric
+`function` | [AggregationFunction](#aggregationfunction-enum)<br>How to aggregate samples for this metric
 `percentile` | `integer`<br>For the percentile function, specifies the percentile to measure
 `interval` | `string`<br>Interval for each metric. Supports times such as "2d", "1h", "1h30m", "20s".
 
diff --git a/Engine/Source/Programs/Horde/Docs/Deployment/AgentSettings.md b/Engine/Source/Programs/Horde/Docs/Deployment/AgentSettings.md
--- a/Engine/Source/Programs/Horde/Docs/Deployment/AgentSettings.md
+++ b/Engine/Source/Programs/Horde/Docs/Deployment/AgentSettings.md
@@ -6,15 +6,15 @@ All Horde-specific settings are stored in a root object called `Horde`. Other .N
 
 Name | Description
 ---- | -----------
-`serverProfiles` | `string` `->` [`ServerProfile`](#serverprofile)<br>Known servers to connect to
+`serverProfiles` | `string` `->` [ServerProfile](#serverprofile)<br>Known servers to connect to
 `server` | `string`<br>The default server, unless overridden from the command line
 `name` | `string`<br>Name of agent to report as when connecting to server. By default, the computer's hostname will be used.
 `installed` | `boolean`<br>Whether the server is running in 'installed' mode. In this mode, on Windows, the default data directory will use the common application data folder (C:\ProgramData\Epic\Horde), and configuration data will be read from here and the registry. This setting is overridden to false for local builds from appsettings.Local.json.
 `ephemeral` | `boolean`<br>Whether agent should register as being ephemeral. Doing so will not persist any long-lived data on the server and once disconnected it's assumed to have been deleted permanently. Ideal for short-lived agents, such as spot instances on AWS EC2.
-`workingDir` | [`DirectoryReference`](#directoryreference)<br>Working directory for leases and jobs (i.e where files from Perforce will be checked out)
-`logsDir` | [`DirectoryReference`](#directoryreference)<br>Directory where agent and lease logs are written
+`workingDir` | [DirectoryReference](#directoryreference)<br>Working directory for leases and jobs (i.e where files from Perforce will be checked out)
+`logsDir` | [DirectoryReference](#directoryreference)<br>Directory where agent and lease logs are written
 `shareMountingEnabled` | `boolean`<br>Whether to mount the specified list of network shares
-`shares` | [`MountNetworkShare`](#mountnetworkshare)`[]`<br>List of network shares to mount
+`shares` | [MountNetworkShare](#mountnetworkshare)`[]`<br>List of network shares to mount
 `wineExecutablePath` | `string`<br>Path to Wine executable. If null, execution under Wine is disabled
 `containerEngineExecutablePath` | `string`<br>Path to container engine executable, such as /usr/bin/podman. If null, execution of compute workloads inside a container is disabled
 `writeStepOutputToLogger` | `boolean`<br>Whether to write step output to the logging device
@@ -45,7 +45,7 @@ Representation of an absolute directory path. Allows fast hashing and comparison
 
 Name | Description
 ---- | -----------
-`parentDirectory` | [`DirectoryReference`](#directoryreference)<br>Gets the directory containing this object
+`parentDirectory` | [DirectoryReference](#directoryreference)<br>Gets the directory containing this object
 `fullName` | `string`<br>The path to this object. Stored as an absolute path, with O/S preferred separator characters, and no trailing slash for directories.
 
 ## MountNetworkShare
diff --git a/Engine/Source/Programs/Horde/Docs/Deployment/ServerSettings.md b/Engine/Source/Programs/Horde/Docs/Deployment/ServerSettings.md
--- a/Engine/Source/Programs/Horde/Docs/Deployment/ServerSettings.md
+++ b/Engine/Source/Programs/Horde/Docs/Deployment/ServerSettings.md
@@ -6,7 +6,7 @@ All Horde-specific settings are stored in a root object called `Horde`. Other .N
 
 Name | Description
 ---- | -----------
-`runModes` | [`RunMode`](#runmode-enum)`[]`<br>Modes that the server should run in. Runmodes can be used in a multi-server deployment to limit the operations that a particular instance will try to perform.
+`runModes` | [RunMode](#runmode-enum)`[]`<br>Modes that the server should run in. Runmodes can be used in a multi-server deployment to limit the operations that a particular instance will try to perform.
 `dataDir` | `string`<br>Override the data directory used by Horde. Defaults to C:\ProgramData\HordeServer on Windows, {AppDir}/Data on other platforms.
 `installed` | `boolean`<br>Whether the server is running in 'installed' mode. In this mode, on Windows, the default data directory will use the common application data folder (C:\ProgramData\Epic\Horde), and configuration data will be read from here and the registry. This setting is overridden to false for local builds from appsettings.Local.json.
 `httpPort` | `integer`<br>Main port for serving HTTP.
@@ -22,7 +22,7 @@ Name | Description
 `databaseReadOnlyMode` | `boolean`<br>
 `shutdownMemoryThreshold` | `integer`<br>Shutdown the current server process if memory usage reaches this threshold (specified in MB)<br>Usually set to 80-90% of available memory to avoid CLR heap using all of it. If a memory leak was to occur, it's usually better to restart the process rather than to let the GC work harder and harder trying to recoup memory.<br>Should only be used when multiple server processes are running behind a load balancer and one can be safely restarted automatically by the underlying process handler (Docker, Kubernetes, AWS ECS, Supervisor etc). The shutdown behaves similar to receiving a SIGTERM and will wait for outstanding requests to finish.
 `serverPrivateCert` | `string`<br>Optional PFX certificate to use for encrypting agent SSL traffic. This can be a self-signed certificate, as long as it's trusted by agents.
-`authMethod` | [`AuthMethod`](#authmethod-enum)<br>Issuer for tokens from the auth provider
+`authMethod` | [AuthMethod](#authmethod-enum)<br>Issuer for tokens from the auth provider
 `oidcProfileName` | `string`<br>Optional profile name to report through the /api/v1/server/auth endpoint. Allows sharing auth tokens between providers configured through the same profile name in OidcToken.exe config files.
 `oidcAuthority` | `string`<br>Issuer for tokens from the auth provider
 `oidcAudience` | `string`<br>Audience for validating externally issued tokens
@@ -63,8 +63,8 @@ Name | Description
 `configPath` | `string`<br>Path to the root config file. Relative to the server.json file by default.
 `forceConfigUpdateOnStartup` | `boolean`<br>Forces configuration data to be read and updated as part of appplication startup, rather than on a schedule. Useful when running locally.
 `openBrowser` | `boolean`<br>Whether to open a browser on startup
-`featureFlags` | [`FeatureFlagSettings`](#featureflagsettings)<br>Experimental features to enable on the server.
-`openTelemetry` | [`OpenTelemetrySettings`](#opentelemetrysettings)<br>Options for OpenTelemetry
+`featureFlags` | [FeatureFlagSettings](#featureflagsettings)<br>Experimental features to enable on the server.
+`openTelemetry` | [OpenTelemetrySettings](#opentelemetrysettings)<br>Options for OpenTelemetry
 
 ## RunMode (Enum)
 
@@ -106,7 +106,7 @@ Name | Description
 `enableDatadogCompatibility` | `boolean`<br>Whether to enrich and format telemetry to fit presentation in Datadog
 `attributes` | `string` `->` `string`<br>Extra attributes to set
 `enableConsoleExporter` | `boolean`<br>Whether to enable the console exporter (for debugging purposes)
-`protocolExporters` | `string` `->` [`OpenTelemetryProtocolExporterSettings`](#opentelemetryprotocolexportersettings)<br>Protocol exporters (key is a unique and arbitrary name)
+`protocolExporters` | `string` `->` [OpenTelemetryProtocolExporterSettings](#opentelemetryprotocolexportersettings)<br>Protocol exporters (key is a unique and arbitrary name)
 
 ## OpenTelemetryProtocolExporterSettings
 
