Horde: Serialize config for pool size strategies and fleet managers as actual json objects, rather than json objects embedded in strings. Supports reading from json objects embedded in strings for backwards compatibility.

diff --git a/Engine/Source/Programs/Horde/Plugins/Build/HordeServer.Build.Tests/Agents/Pools/PoolsControllerTest.cs b/Engine/Source/Programs/Horde/Plugins/Build/HordeServer.Build.Tests/Agents/Pools/PoolsControllerTest.cs
--- a/Engine/Source/Programs/Horde/Plugins/Build/HordeServer.Build.Tests/Agents/Pools/PoolsControllerTest.cs
+++ b/Engine/Source/Programs/Horde/Plugins/Build/HordeServer.Build.Tests/Agents/Pools/PoolsControllerTest.cs
@@ -1,5 +1,6 @@
 ﻿// Copyright Epic Games, Inc. All Rights Reserved.
 
+using EpicGames.Horde.Common;
 using HordeServer.Agents.Fleet;
 using HordeServer.Agents.Pools;
 using Microsoft.AspNetCore.Mvc;
@@ -47,10 +48,10 @@ namespace HordeServer.Tests.Agents.Pools
 				Name = "Pool1",
 				ScaleOutCooldown = TimeSpan.FromSeconds(111),
 				ScaleInCooldown = TimeSpan.FromSeconds(222),
-				SizeStrategies = new List<PoolSizeStrategyInfo> { new() { Type = PoolSizeStrategy.JobQueue, Condition = "dayOfWeek == 'monday'", Config = @"{""ScaleOutFactor"": 22.0, ""ScaleInFactor"": 33.0}", ExtraAgentCount = 567 } },
+				SizeStrategies = new List<PoolSizeStrategyInfo> { new(PoolSizeStrategy.JobQueue, "dayOfWeek == 'monday'", @"{""ScaleOutFactor"": 22.0, ""ScaleInFactor"": 33.0}", extraAgentCount: 567) },
 				SizeStrategy = PoolSizeStrategy.JobQueue,
 				JobQueueSettings = new JobQueueSettings(0.35, 0.85),
-				FleetManagers = new List<FleetManagerInfo> { new() { Type = FleetManagerType.AwsReuse, Condition = "dayOfWeek == 'monday'", Config = "{}" } },
+				FleetManagers = new List<FleetManagerInfo> { new() { Type = FleetManagerType.AwsReuse, Condition = "dayOfWeek == 'monday'" } },
 			};
 #pragma warning restore CS0618
 
diff --git a/Engine/Source/Programs/Horde/Plugins/Build/HordeServer.Build.Tests/Fleet/FleetServiceTest.cs b/Engine/Source/Programs/Horde/Plugins/Build/HordeServer.Build.Tests/Fleet/FleetServiceTest.cs
--- a/Engine/Source/Programs/Horde/Plugins/Build/HordeServer.Build.Tests/Fleet/FleetServiceTest.cs
+++ b/Engine/Source/Programs/Horde/Plugins/Build/HordeServer.Build.Tests/Fleet/FleetServiceTest.cs
@@ -2,6 +2,7 @@
 
 extern alias HordeAgent;
 using System.Text.Json;
+using System.Text.Json.Nodes;
 using HordeServer.Agents;
 using HordeServer.Agents.Fleet;
 using HordeServer.Agents.Fleet.Providers;
@@ -83,7 +84,7 @@ namespace HordeServer.Tests.Fleet
 			_fleetManager = fleetManager;
 		}
 
-		public IFleetManager CreateFleetManager(FleetManagerType type, string? config = null)
+		public IFleetManager CreateFleetManager(FleetManagerType type, JsonObject? config = null)
 		{
 			return _fleetManager;
 		}
@@ -307,7 +308,24 @@ namespace HordeServer.Tests.Fleet
 			await CreateStrategyAsync(new PoolSizeStrategyInfo(PoolSizeStrategy.LeaseUtilization, null, "{}"));
 			await CreateStrategyAsync(new PoolSizeStrategyInfo(PoolSizeStrategy.LeaseUtilization, null, "  {} "));
 
-			await Assert.ThrowsExceptionAsync<JsonException>(() => CreateStrategyAsync(new PoolSizeStrategyInfo(PoolSizeStrategy.LeaseUtilization, null, "BAD_JSON")));
+			await ThrowsDerivedExceptionAsync<JsonException>(() => CreateStrategyAsync(new PoolSizeStrategyInfo(PoolSizeStrategy.LeaseUtilization, null, "BAD_JSON")));
+		}
+
+		static async Task<bool> ThrowsDerivedExceptionAsync<T>(Func<Task> func) where T : Exception
+		{
+			try
+			{
+				await func();
+				return false;
+			}
+			catch (T)
+			{
+				return true;
+			}
+			catch
+			{
+				return false;
+			}
 		}
 
 		[TestMethod]
diff --git a/Engine/Source/Programs/Horde/Plugins/Compute/HordeServer.Compute/Agents/Fleet/FleetManagerFactory.cs b/Engine/Source/Programs/Horde/Plugins/Compute/HordeServer.Compute/Agents/Fleet/FleetManagerFactory.cs
--- a/Engine/Source/Programs/Horde/Plugins/Compute/HordeServer.Compute/Agents/Fleet/FleetManagerFactory.cs
+++ b/Engine/Source/Programs/Horde/Plugins/Compute/HordeServer.Compute/Agents/Fleet/FleetManagerFactory.cs
@@ -2,6 +2,7 @@
 
 using System.Diagnostics.Metrics;
 using System.Text.Json;
+using System.Text.Json.Nodes;
 using Amazon.AutoScaling;
 using Amazon.EC2;
 using EpicGames.Core;
@@ -13,6 +14,8 @@ using OpenTelemetry.Trace;
 
 namespace HordeServer.Agents.Fleet;
 
+using JsonObject = System.Text.Json.Nodes.JsonObject;
+
 /// <summary>
 /// Interface for fleet manager factory
 /// </summary>
@@ -25,7 +28,7 @@ public interface IFleetManagerFactory
 	/// <param name="config">Config as a serialized JSON string</param>
 	/// <returns>An instantiated fleet manager with parameters loaded from config</returns>
 	/// <exception cref="ArgumentException">If fleet manager could not be instantiated</exception>
-	public IFleetManager CreateFleetManager(FleetManagerType type, string? config = null);
+	public IFleetManager CreateFleetManager(FleetManagerType type, JsonObject? config = null);
 }
 
 /// <summary>
@@ -59,12 +62,12 @@ public sealed class FleetManagerFactory : IFleetManagerFactory
 	}
 
 	/// <inheritdoc/>
-	public IFleetManager CreateFleetManager(FleetManagerType type, string? config)
+	public IFleetManager CreateFleetManager(FleetManagerType type, JsonObject? config)
 	{
 		return type switch
 		{
 			FleetManagerType.Default =>
-				CreateFleetManager(_staticComputeConfig.CurrentValue.FleetManagerV2, _staticComputeConfig.CurrentValue.FleetManagerV2Config ?? "{}"),
+				CreateFleetManager(_staticComputeConfig.CurrentValue.FleetManagerV2, _staticComputeConfig.CurrentValue.FleetManagerV2Config ?? new JsonObject()),
 			FleetManagerType.NoOp =>
 				new NoOpFleetManager(_loggerFactory.CreateLogger<NoOpFleetManager>()),
 			FleetManagerType.Aws =>
@@ -79,16 +82,16 @@ public sealed class FleetManagerFactory : IFleetManagerFactory
 		};
 	}
 
-	private static T DeserializeSettings<T>(string? config)
+	private static T DeserializeSettings<T>(JsonObject? config)
 	{
-		if (String.IsNullOrEmpty(config))
+		if (config == null)
 		{
-			config = "{}";
+			config = new JsonObject();
 		}
 
 		try
 		{
-			T? settings = JsonSerializer.Deserialize<T>(config);
+			T? settings = JsonSerializer.Deserialize<T>((JsonNode)config, new JsonSerializerOptions { AllowTrailingCommas = true, PropertyNameCaseInsensitive = true, ReadCommentHandling = JsonCommentHandling.Skip });
 			if (settings == null)
 			{
 				throw new InvalidDataException($"Unable to deserialize");
diff --git a/Engine/Source/Programs/Horde/Plugins/Compute/HordeServer.Compute/Agents/Fleet/FleetService.cs b/Engine/Source/Programs/Horde/Plugins/Compute/HordeServer.Compute/Agents/Fleet/FleetService.cs
--- a/Engine/Source/Programs/Horde/Plugins/Compute/HordeServer.Compute/Agents/Fleet/FleetService.cs
+++ b/Engine/Source/Programs/Horde/Plugins/Compute/HordeServer.Compute/Agents/Fleet/FleetService.cs
@@ -15,6 +15,8 @@ using OpenTelemetry.Trace;
 
 namespace HordeServer.Agents.Fleet
 {
+	using JsonObject = System.Text.Json.Nodes.JsonObject;
+
 	/// <summary>
 	/// Parameters required for calculating pool size
 	/// </summary>
@@ -344,7 +346,7 @@ namespace HordeServer.Agents.Fleet
 				}
 			}
 
-			return _fleetManagerFactory.CreateFleetManager(FleetManagerType.Default, "{}");
+			return _fleetManagerFactory.CreateFleetManager(FleetManagerType.Default, new JsonObject());
 		}
 
 		/// <summary>
diff --git a/Engine/Source/Programs/Horde/Plugins/Compute/HordeServer.Compute/Agents/Fleet/IPoolSizeStrategy.cs b/Engine/Source/Programs/Horde/Plugins/Compute/HordeServer.Compute/Agents/Fleet/IPoolSizeStrategy.cs
--- a/Engine/Source/Programs/Horde/Plugins/Compute/HordeServer.Compute/Agents/Fleet/IPoolSizeStrategy.cs
+++ b/Engine/Source/Programs/Horde/Plugins/Compute/HordeServer.Compute/Agents/Fleet/IPoolSizeStrategy.cs
@@ -1,6 +1,7 @@
 // Copyright Epic Games, Inc. All Rights Reserved.
 
 using System.Text.Json;
+using System.Text.Json.Nodes;
 using HordeServer.Agents.Pools;
 
 namespace HordeServer.Agents.Fleet
@@ -127,7 +128,7 @@ namespace HordeServer.Agents.Fleet
 		/// <summary>
 		/// Create a new <see cref="IPoolSizeStrategy"/> instance from the given configuration
 		/// </summary>
-		IPoolSizeStrategy Create(string config);
+		IPoolSizeStrategy Create(JsonObject config);
 	}
 
 	/// <summary>
@@ -145,7 +146,7 @@ namespace HordeServer.Agents.Fleet
 			=> Type = type;
 
 		/// <inheritdoc/>
-		public IPoolSizeStrategy Create(string? config)
+		public IPoolSizeStrategy Create(JsonObject? config)
 			=> Create(DeserializeConfig<TConfig>(config));
 
 		/// <summary>
@@ -153,10 +154,9 @@ namespace HordeServer.Agents.Fleet
 		/// </summary>
 		public abstract IPoolSizeStrategy Create(TConfig config);
 
-		private static T DeserializeConfig<T>(string? json)
+		private static T DeserializeConfig<T>(JsonObject? json)
 		{
-			json = String.IsNullOrEmpty(json) ? "{}" : json;
-			T? config = JsonSerializer.Deserialize<T>(json, new JsonSerializerOptions { PropertyNameCaseInsensitive = true });
+			T? config = JsonSerializer.Deserialize<T>(json ?? new JsonObject(), new JsonSerializerOptions { PropertyNameCaseInsensitive = true });
 			if (config == null)
 			{
 				throw new ArgumentException("Unable to deserialize config: " + json);
@@ -185,7 +185,7 @@ namespace HordeServer.Agents.Fleet
 	{
 		public PoolSizeStrategy Type => PoolSizeStrategy.NoOp;
 
-		public IPoolSizeStrategy Create(string config)
+		public IPoolSizeStrategy Create(JsonObject? config)
 			=> new NoOpPoolSizeStrategy();
 	}
 
diff --git a/Engine/Source/Programs/Horde/Plugins/Compute/HordeServer.Compute/Agents/Pools/IPool.cs b/Engine/Source/Programs/Horde/Plugins/Compute/HordeServer.Compute/Agents/Pools/IPool.cs
--- a/Engine/Source/Programs/Horde/Plugins/Compute/HordeServer.Compute/Agents/Pools/IPool.cs
+++ b/Engine/Source/Programs/Horde/Plugins/Compute/HordeServer.Compute/Agents/Pools/IPool.cs
@@ -1,6 +1,7 @@
 // Copyright Epic Games, Inc. All Rights Reserved.
 
 using System.Text.Json;
+using System.Text.Json.Nodes;
 using System.Text.Json.Serialization;
 using EpicGames.Horde.Common;
 using HordeServer.Agents.Fleet;
@@ -25,7 +26,8 @@ namespace HordeServer.Agents.Pools
 		/// <summary>
 		/// Configuration for the strategy, serialized as JSON
 		/// </summary>
-		public string Config { get; set; } = "";
+		[JsonConverter(typeof(JsonObjectOrStringConverter))]
+		public JsonObject Config { get; set; } = new JsonObject();
 
 		/// <summary>
 		/// Integer to add after pool size has been calculated. Can also be negative.
@@ -35,6 +37,7 @@ namespace HordeServer.Agents.Pools
 		/// <summary>
 		/// Empty constructor for JSON serialization
 		/// </summary>
+		[JsonConstructor]
 		public PoolSizeStrategyInfo()
 		{
 		}
@@ -42,17 +45,28 @@ namespace HordeServer.Agents.Pools
 		/// <summary>
 		/// Constructor
 		/// </summary>
-		public PoolSizeStrategyInfo(PoolSizeStrategy type, Condition? condition, string? config, int extraAgentCount = 0)
+		public PoolSizeStrategyInfo(PoolSizeStrategy type, Condition? condition, int extraAgentCount = 0)
 		{
-			config = String.IsNullOrWhiteSpace(config) ? "{}" : config;
-
-			// Try deserializing to ensure the config is valid JSON
-			// Config can be null due to JSON serializer calling the constructor
-			JsonSerializer.Deserialize<dynamic>(config);
+			Type = type;
+			Condition = condition;
+			ExtraAgentCount = extraAgentCount;
+		}
 
+		/// <summary>
+		/// Constructor
+		/// </summary>
+		public PoolSizeStrategyInfo(PoolSizeStrategy type, Condition? condition, string config, int extraAgentCount = 0)
+		{
 			Type = type;
 			Condition = condition;
-			Config = config;
+			if (String.IsNullOrEmpty(config))
+			{
+				Config = new JsonObject();
+			}
+			else
+			{
+				Config = (JsonObject)JsonNode.Parse(config, new JsonNodeOptions { PropertyNameCaseInsensitive = true }, new JsonDocumentOptions { AllowTrailingCommas = true, CommentHandling = JsonCommentHandling.Skip })!;
+			}
 			ExtraAgentCount = extraAgentCount;
 		}
 	}
@@ -75,7 +89,8 @@ namespace HordeServer.Agents.Pools
 		/// <summary>
 		/// Configuration for the strategy, serialized as JSON
 		/// </summary>
-		public string? Config { get; set; }
+		[JsonConverter(typeof(JsonObjectOrStringConverter))]
+		public JsonObject? Config { get; set; }
 
 		/// <summary>
 		/// Empty constructor for BSON/JSON serialization
@@ -88,17 +103,58 @@ namespace HordeServer.Agents.Pools
 		/// Constructor
 		/// </summary>
 		[JsonConstructor]
-		public FleetManagerInfo(FleetManagerType type, Condition? condition, string? config)
+		public FleetManagerInfo(FleetManagerType type, Condition? condition)
 		{
-			config = String.IsNullOrWhiteSpace(config) ? "{}" : config;
-			JsonSerializer.Deserialize<dynamic>(config);
+			Type = type;
+			Condition = condition;
+		}
 
+		/// <summary>
+		/// Constructor
+		/// </summary>
+		public FleetManagerInfo(FleetManagerType type, Condition? condition, string config)
+		{
 			Type = type;
 			Condition = condition;
-			Config = config;
+			if (String.IsNullOrEmpty(config))
+			{
+				Config = new JsonObject();
+			}
+			else
+			{
+				Config = (JsonObject?)JsonNode.Parse(config, new JsonNodeOptions { PropertyNameCaseInsensitive = true }, new JsonDocumentOptions { AllowTrailingCommas = true });
+			}
 		}
 	}
 
+	/// <summary>
+	/// Allows parsing a string or json object into a native JsonObject instance
+	/// </summary>
+	class JsonObjectOrStringConverter : JsonConverter<JsonObject>
+	{
+		static readonly JsonNodeOptions s_nodeOptions = new JsonNodeOptions { PropertyNameCaseInsensitive = true };
+		static readonly JsonDocumentOptions s_documentOptions = new JsonDocumentOptions { AllowTrailingCommas = true, CommentHandling = JsonCommentHandling.Skip };
+
+		/// <inheritdoc/>
+		public override JsonObject? Read(ref Utf8JsonReader reader, Type typeToConvert, JsonSerializerOptions options)
+		{
+			JsonObject? result = null;
+			if (reader.TokenType == JsonTokenType.String)
+			{
+				result = JsonNode.Parse(reader.GetString() ?? "{}", s_nodeOptions, s_documentOptions) as JsonObject;
+			}
+			else if (reader.TokenType == JsonTokenType.StartObject)
+			{
+				result = JsonNode.Parse(ref reader, s_nodeOptions) as JsonObject;
+			}
+			return result;
+		}
+
+		/// <inheritdoc/>
+		public override void Write(Utf8JsonWriter writer, JsonObject value, JsonSerializerOptions options)
+			=> JsonSerializer.Serialize(writer, value, options);
+	}
+
 	/// <summary>
 	/// Pool of machines
 	/// </summary>
diff --git a/Engine/Source/Programs/Horde/Plugins/Compute/HordeServer.Compute/Agents/Pools/PoolMessages.cs b/Engine/Source/Programs/Horde/Plugins/Compute/HordeServer.Compute/Agents/Pools/PoolMessages.cs
--- a/Engine/Source/Programs/Horde/Plugins/Compute/HordeServer.Compute/Agents/Pools/PoolMessages.cs
+++ b/Engine/Source/Programs/Horde/Plugins/Compute/HordeServer.Compute/Agents/Pools/PoolMessages.cs
@@ -1,6 +1,7 @@
 // Copyright Epic Games, Inc. All Rights Reserved.
 
 using System.ComponentModel.DataAnnotations;
+using System.Text.Json.Nodes;
 using EpicGames.Horde.Agents;
 using EpicGames.Horde.Agents.Pools;
 using EpicGames.Horde.Common;
@@ -26,7 +27,7 @@ namespace HordeServer.Agents.Pools
 		/// <summary>
 		/// Configuration for the strategy, serialized as JSON
 		/// </summary>
-		public string Config { get; set; } = "";
+		public JsonObject Config { get; set; } = new JsonObject();
 
 		/// <summary>
 		/// Integer to add after pool size has been calculated. Can also be negative.
@@ -48,7 +49,7 @@ namespace HordeServer.Agents.Pools
 		{
 			Type = strategy.Type;
 			Condition = strategy.Condition;
-			Config = strategy.Config;
+			Config = (JsonObject)strategy.Config.DeepClone();
 			ExtraAgentCount = strategy.ExtraAgentCount;
 		}
 
@@ -58,7 +59,7 @@ namespace HordeServer.Agents.Pools
 		/// <returns></returns>
 		public PoolSizeStrategyInfo Convert()
 		{
-			return new PoolSizeStrategyInfo(Type, Condition, Config, ExtraAgentCount);
+			return new PoolSizeStrategyInfo(Type, Condition, ExtraAgentCount) { Config = (JsonObject)Config.DeepClone() };
 		}
 	}
 
@@ -78,7 +79,7 @@ namespace HordeServer.Agents.Pools
 		/// <summary>
 		/// Configuration for the strategy, serialized as JSON
 		/// </summary>
-		public string? Config { get; set; }
+		public JsonObject? Config { get; set; }
 
 		/// <summary>
 		/// Empty constructor for JSON serialization
@@ -104,7 +105,7 @@ namespace HordeServer.Agents.Pools
 		/// <returns></returns>
 		public FleetManagerInfo Convert()
 		{
-			return new FleetManagerInfo(Type, Condition, Config);
+			return new FleetManagerInfo(Type, Condition) { Config = (JsonObject?)Config?.DeepClone() };
 		}
 	}
 
diff --git a/Engine/Source/Programs/Horde/Plugins/Compute/HordeServer.Compute/StaticComputeConfig.cs b/Engine/Source/Programs/Horde/Plugins/Compute/HordeServer.Compute/StaticComputeConfig.cs
--- a/Engine/Source/Programs/Horde/Plugins/Compute/HordeServer.Compute/StaticComputeConfig.cs
+++ b/Engine/Source/Programs/Horde/Plugins/Compute/HordeServer.Compute/StaticComputeConfig.cs
@@ -1,6 +1,9 @@
 ﻿// Copyright Epic Games, Inc. All Rights Reserved.
 
+using System.Text.Json.Nodes;
+using System.Text.Json.Serialization;
 using HordeServer.Agents.Fleet;
+using HordeServer.Agents.Pools;
 
 namespace HordeServer
 {
@@ -27,7 +30,8 @@ namespace HordeServer
 		/// <summary>
 		/// Config for the fleet manager (serialized JSON)
 		/// </summary>
-		public string? FleetManagerV2Config { get; set; }
+		[JsonConverter(typeof(JsonObjectOrStringConverter))]
+		public JsonObject? FleetManagerV2Config { get; set; }
 
 		/// <summary>
 		/// AWS SQS queue URLs where lifecycle events from EC2 auto-scaling are received
