Horde: Add a separate flag for marking an issue as fixed as a systemic issue, rather than having to pass a negative value as a fix changelist.

Still supports passing/returning negative fix changelists as well for now, but will be removed in future.

diff --git a/Engine/Source/Programs/Horde/Plugins/Build/HordeServer.Build/Issues/IIssue.cs b/Engine/Source/Programs/Horde/Plugins/Build/HordeServer.Build/Issues/IIssue.cs
--- a/Engine/Source/Programs/Horde/Plugins/Build/HordeServer.Build/Issues/IIssue.cs
+++ b/Engine/Source/Programs/Horde/Plugins/Build/HordeServer.Build/Issues/IIssue.cs
@@ -107,6 +107,11 @@ namespace HordeServer.Issues
 		/// </summary>
 		public int? FixChange { get; }
 
+		/// <summary>
+		/// Whether the issue is marked fixed as systemic
+		/// </summary>
+		public bool FixedSystemic { get; }
+
 		//		/// <summary>
 		//		/// The first stream that encountered the error. The fix will be considered failed if an error after FixChange occurs in this stream. 
 		//		/// </summary>
diff --git a/Engine/Source/Programs/Horde/Plugins/Build/HordeServer.Build/Issues/IIssueCollection.cs b/Engine/Source/Programs/Horde/Plugins/Build/HordeServer.Build/Issues/IIssueCollection.cs
--- a/Engine/Source/Programs/Horde/Plugins/Build/HordeServer.Build/Issues/IIssueCollection.cs
+++ b/Engine/Source/Programs/Horde/Plugins/Build/HordeServer.Build/Issues/IIssueCollection.cs
@@ -493,6 +493,7 @@ namespace HordeServer.Issues
 		/// <param name="newAcknowledged">Whether the issue has been acknowledged</param>
 		/// <param name="newDeclinedById">Name of a user that has declined the issue</param>
 		/// <param name="newFixChange">Fix changelist for the issue. Pass 0 to clear the fix changelist, -1 for systemic issue.</param>
+		/// <param name="newFixedSystemic">Whether the issue should be marked fixed as a systemic issue</param>
 		/// <param name="newResolvedById">User that resolved the issue (may be ObjectId.Empty to clear)</param>
 		/// <param name="newExcludeSpanIds">List of span ids to exclude from this issue</param>
 		/// <param name="newLastSeenAt"></param>
@@ -502,7 +503,7 @@ namespace HordeServer.Issues
 		/// <param name="newWorkflowThreadUrl">The workflow thread url associated with the issue</param>
 		/// <param name="cancellationToken">Cancellation token for the operation</param>
 		/// <returns>True if the issue was updated</returns>
-		Task<IIssue?> TryUpdateIssueAsync(IIssue issue, UserId? initiatedByUserId, IssueSeverity? newSeverity = null, string? newSummary = null, string? newUserSummary = null, string? newDescription = null, bool? newPromoted = null, UserId? newOwnerId = null, UserId? newNominatedById = null, bool? newAcknowledged = null, UserId? newDeclinedById = null, int? newFixChange = null, UserId? newResolvedById = null, List<ObjectId>? newExcludeSpanIds = null, DateTime? newLastSeenAt = null, string? newExternalIssueKey = null, UserId? newQuarantinedById = null, UserId? newForceClosedById = null, Uri? newWorkflowThreadUrl = null, CancellationToken cancellationToken = default);
+		Task<IIssue?> TryUpdateIssueAsync(IIssue issue, UserId? initiatedByUserId, IssueSeverity? newSeverity = null, string? newSummary = null, string? newUserSummary = null, string? newDescription = null, bool? newPromoted = null, UserId? newOwnerId = null, UserId? newNominatedById = null, bool? newAcknowledged = null, UserId? newDeclinedById = null, int? newFixChange = null, bool? newFixedSystemic = null, UserId? newResolvedById = null, List<ObjectId>? newExcludeSpanIds = null, DateTime? newLastSeenAt = null, string? newExternalIssueKey = null, UserId? newQuarantinedById = null, UserId? newForceClosedById = null, Uri? newWorkflowThreadUrl = null, CancellationToken cancellationToken = default);
 
 		/// <summary>
 		/// Updates derived data for an issue (ie. data computed from the spans attached to it). Also clears the issue's 'modified' state.
diff --git a/Engine/Source/Programs/Horde/Plugins/Build/HordeServer.Build/Issues/IssueCollection.cs b/Engine/Source/Programs/Horde/Plugins/Build/HordeServer.Build/Issues/IssueCollection.cs
--- a/Engine/Source/Programs/Horde/Plugins/Build/HordeServer.Build/Issues/IssueCollection.cs
+++ b/Engine/Source/Programs/Horde/Plugins/Build/HordeServer.Build/Issues/IssueCollection.cs
@@ -99,6 +99,10 @@ namespace HordeServer.Issues
 
 			[BsonIgnoreIfNull]
 			public int? FixChange { get; set; }
+			int? IIssue.FixChange => FixChange ?? (FixedSystemic ? -1 : null);
+
+			[BsonIgnoreIfDefault, BsonDefaultValue(false)]
+			public bool FixedSystemic { get; set; }
 
 			public List<IssueStreamDocument> Streams { get; set; } = new List<IssueStreamDocument>();
 
@@ -754,6 +758,7 @@ namespace HordeServer.Issues
 					CreatedAt = issue.CreatedAt,
 					OwnerId = issue.OwnerId,
 					FixChange = issue.FixChange,
+					FixedSystemic = issue.FixedSystemic,
 					LastSeenAt = issue.LastSeenAt,
 					NominatedAt = issue.NominatedAt,
 					NominatedById = issue.NominatedById,
@@ -848,6 +853,17 @@ namespace HordeServer.Issues
 					issueLogger.LogInformation("Issue was marked as fixed in {Change}", newIssue.FixChange);
 				}
 			}
+			if (newIssue.FixedSystemic != oldIssue.FixedSystemic)
+			{
+				if (newIssue.FixChange == 0)
+				{
+					issueLogger.LogInformation("Issue was marked as not fixed");
+				}
+				else
+				{
+					issueLogger.LogInformation("Issue was marked as fixed in {Change}", newIssue.FixChange);
+				}
+			}
 			if (newIssue.ResolvedById != oldIssue.ResolvedById)
 			{
 				if (newIssue.ResolvedById == null)
@@ -1138,7 +1154,7 @@ namespace HordeServer.Issues
 		}
 
 		/// <inheritdoc/>
-		public async Task<IIssue?> TryUpdateIssueAsync(IIssue issue, UserId? initiatedByUserId, IssueSeverity? newSeverity = null, string? newSummary = null, string? newUserSummary = null, string? newDescription = null, bool? newManuallyPromoted = null, UserId? newOwnerId = null, UserId? newNominatedById = null, bool? newAcknowledged = null, UserId? newDeclinedById = null, int? newFixChange = null, UserId? newResolvedById = null, List<ObjectId>? newExcludeSpanIds = null, DateTime? newLastSeenAt = null, string? newExternaIssueKey = null, UserId? newQuarantinedById = null, UserId? newForceClosedById = null, Uri? newWorkflowThreadUrl = null, CancellationToken cancellationToken = default)
+		public async Task<IIssue?> TryUpdateIssueAsync(IIssue issue, UserId? initiatedByUserId, IssueSeverity? newSeverity = null, string? newSummary = null, string? newUserSummary = null, string? newDescription = null, bool? newManuallyPromoted = null, UserId? newOwnerId = null, UserId? newNominatedById = null, bool? newAcknowledged = null, UserId? newDeclinedById = null, int? newFixChange = null, bool? newFixedSystemic = null, UserId? newResolvedById = null, List<ObjectId>? newExcludeSpanIds = null, DateTime? newLastSeenAt = null, string? newExternaIssueKey = null, UserId? newQuarantinedById = null, UserId? newForceClosedById = null, Uri? newWorkflowThreadUrl = null, CancellationToken cancellationToken = default)
 		{
 			IssueDocument issueDocument = (IssueDocument)issue;
 
@@ -1237,13 +1253,28 @@ namespace HordeServer.Issues
 			}
 			if (newFixChange != null)
 			{
-				if (newFixChange == 0)
+				if (newFixChange < 0)
+				{
+					updates.Add(Builders<IssueDocument>.Update.Unset(x => x.FixChange!).Set(x => x.FixedSystemic, true));
+				}
+				else if (newFixChange == 0)
+				{
+					updates.Add(Builders<IssueDocument>.Update.Unset(x => x.FixChange!).Unset(x => x.FixedSystemic));
+				}
+				else
+				{
+					updates.Add(Builders<IssueDocument>.Update.Set(x => x.FixChange, newFixChange).Unset(x => x.FixedSystemic));
+				}
+			}
+			else if (newFixedSystemic != null)
+			{
+				if (newFixedSystemic.Value)
 				{
-					updates.Add(Builders<IssueDocument>.Update.Unset(x => x.FixChange!));
+					updates.Add(Builders<IssueDocument>.Update.Unset(x => x.FixChange!).Set(x => x.FixedSystemic, true));
 				}
 				else
 				{
-					updates.Add(Builders<IssueDocument>.Update.Set(x => x.FixChange, newFixChange));
+					updates.Add(Builders<IssueDocument>.Update.Unset(x => x.FixedSystemic));
 				}
 			}
 			if (newResolvedById != null)
diff --git a/Engine/Source/Programs/Horde/Plugins/Build/HordeServer.Build/Issues/IssuesController.cs b/Engine/Source/Programs/Horde/Plugins/Build/HordeServer.Build/Issues/IssuesController.cs
--- a/Engine/Source/Programs/Horde/Plugins/Build/HordeServer.Build/Issues/IssuesController.cs
+++ b/Engine/Source/Programs/Horde/Plugins/Build/HordeServer.Build/Issues/IssuesController.cs
@@ -252,6 +252,7 @@ namespace HordeServer.Issues
 			}
 			response.AcknowledgedAt = issue.AcknowledgedAt;
 			response.FixChange = issue.FixChange;
+			response.FixedSystemic = issue.FixedSystemic;
 			response.ResolvedAt = issue.ResolvedAt;
 			if (resolvedBy != null)
 			{
@@ -499,6 +500,7 @@ namespace HordeServer.Issues
 			response.NominatedByInfo = details.NominatedBy?.ToThinApiResponse();
 			response.AcknowledgedAt = issue.AcknowledgedAt;
 			response.FixChange = issue.FixChange;
+			response.FixedSystemic = issue.FixedSystemic;
 			response.ResolvedAt = issue.ResolvedAt;
 			response.ResolvedBy = details.ResolvedBy?.Login;
 			response.ResolvedById = details.ResolvedBy?.Id.ToString();
diff --git a/Engine/Source/Programs/Horde/Plugins/Build/HordeServer.Build/Notifications/Sinks/SlackNotificationSink.cs b/Engine/Source/Programs/Horde/Plugins/Build/HordeServer.Build/Notifications/Sinks/SlackNotificationSink.cs
--- a/Engine/Source/Programs/Horde/Plugins/Build/HordeServer.Build/Notifications/Sinks/SlackNotificationSink.cs
+++ b/Engine/Source/Programs/Horde/Plugins/Build/HordeServer.Build/Notifications/Sinks/SlackNotificationSink.cs
@@ -1265,6 +1265,12 @@ namespace HordeServer.Notifications.Sinks
 						}
 					}
 				}
+				else if (issue.FixedSystemic)
+				{
+					string fixedEventId = $"issue_{issue.Id}_fixed_systemic";
+					string fixedMessage = $"Marked fixed as a systemic issue";
+					await PostSingleMessageToThreadAsync(triageChannel, fixedEventId, threadId, fixedMessage, cancellationToken);
+				}
 
 				// Assignment notifications
 				if (issue.OwnerId != null && issue.NominatedById != null && issue.NominatedById != issue.OwnerId)
@@ -1473,11 +1479,7 @@ namespace HordeServer.Notifications.Sinks
 				IIssueStep? fixFailedStep = issue.FindFixFailedStep(details.Spans);
 
 				string text;
-				if (issue.FixChange.Value < 0)
-				{
-					text = ":tick: Marked as a systemic issue.";
-				}
-				else if (fixFailedStep != null)
+				if (fixFailedStep != null)
 				{
 					Uri fixFailedUrl = new Uri(_serverInfo.DashboardUrl, $"job/{fixFailedStep.JobId}?step={fixFailedStep.StepId}&issue={issue.Id}");
 					text = $":cross: Marked fixed in *CL {issue.FixChange.Value}*, but seen again at *<{fixFailedUrl}|CL {fixFailedStep.Change}>*";
@@ -1488,6 +1490,10 @@ namespace HordeServer.Notifications.Sinks
 				}
 				attachment.AddSection(text);
 			}
+			else if (issue.FixedSystemic)
+			{
+				attachment.AddSection(":tick: Marked as a systemic issue.");
+			}
 			else if (userId != null && issue.OwnerId == userId)
 			{
 				if (issue.AcknowledgedAt.HasValue)
@@ -2018,6 +2024,10 @@ namespace HordeServer.Notifications.Sinks
 					status = $"Fixed in CL {issue.FixChange.Value}";
 				}
 			}
+			else if (issue.FixedSystemic)
+			{
+				status = "Closed as systemic issue";
+			}
 			else if (issue.OwnerId != null)
 			{
 				if (owner == null)
diff --git a/Engine/Source/Programs/Shared/EpicGames.Horde/Issues/IssueMessages.cs b/Engine/Source/Programs/Shared/EpicGames.Horde/Issues/IssueMessages.cs
--- a/Engine/Source/Programs/Shared/EpicGames.Horde/Issues/IssueMessages.cs
+++ b/Engine/Source/Programs/Shared/EpicGames.Horde/Issues/IssueMessages.cs
@@ -324,6 +324,11 @@ namespace EpicGames.Horde.Issues
 		/// </summary>
 		public int? FixChange { get; set; }
 
+		/// <summary>
+		/// Whether the issue is marked fixed as a systemic issue
+		/// </summary>
+		public bool FixedSystemic { get; set; }
+
 		/// <summary>
 		/// Time at which the issue was resolved
 		/// </summary>
@@ -526,6 +531,11 @@ namespace EpicGames.Horde.Issues
 		/// </summary>
 		public int? FixChange { get; set; }
 
+		/// <summary>
+		/// Whether the issue is marked fixed as a systemic issue
+		/// </summary>
+		public bool FixedSystemic { get; set; }
+
 		/// <summary>
 		/// Time at which the issue was resolved
 		/// </summary>
