Horde: Fix config preflights not reading source files required by remapped config paths.

#rnx

diff --git a/Engine/Source/Programs/Horde/HordeServer/Configuration/ConfigService.cs b/Engine/Source/Programs/Horde/HordeServer/Configuration/ConfigService.cs
--- a/Engine/Source/Programs/Horde/HordeServer/Configuration/ConfigService.cs
+++ b/Engine/Source/Programs/Horde/HordeServer/Configuration/ConfigService.cs
@@ -281,9 +281,7 @@ namespace HordeServer.Configuration
 			ConfigContext context = CreateConfigContext(overrideSources, NullLogger.Instance);
 			try
 			{
-				Uri globalConfigUri = GetGlobalConfigUri();
-
-				GlobalConfig globalConfig = await context.ReadAsync<GlobalConfig>(globalConfigUri, cancellationToken);
+				GlobalConfig globalConfig = await ReadGlobalConfigAsync(context, cancellationToken);
 				globalConfig.PostLoad(_serverSettings, _pluginCollection.LoadedPlugins, _defaultAclModifiers);
 
 				foreach (OverrideConfigFile file in overrideFiles.Values)
@@ -582,46 +580,8 @@ namespace HordeServer.Configuration
 			ConfigContext context = CreateConfigContext(_sources, _logger);
 			try
 			{
-				ConfigSnapshot snapshot = new ConfigSnapshot();
-				snapshot.ServerVersion = ServerApp.Version.ToString();
-
-				// Read the new config in
-				Uri globalConfigUri = GetGlobalConfigUri();
-
-				// Generate the schema for the preprocessor
-				ObjectConfigNode configNode = new ObjectConfigNode(typeof(GlobalConfig));
-
-				ObjectConfigNode pluginsObj = new ObjectConfigNode(false, false);
-				configNode.Properties["Plugins"] = pluginsObj;
-
-				foreach (ILoadedPlugin plugin in _pluginCollection.LoadedPlugins)
-				{
-					if (plugin.GlobalConfigType != null)
-					{
-						pluginsObj.Properties.Add(plugin.Name.ToString(), new ObjectConfigNode(plugin.GlobalConfigType));
-					}
-				}
-
-				// Copy the schema from any remapped nodes to the source node
-				foreach (KeyValuePair<string, string> remapConfigPair in s_remapConfigValues)
-				{
-					if (configNode.TryGetChildNode(remapConfigPair.Value, out ConfigNode? node))
-					{
-						configNode.AddChildNode(remapConfigPair.Key, node);
-					}
-				}
-
-				// Preprocess the file
-				JsonObject configObj = await context.PreprocessFileAsync(GetGlobalConfigUri(), configNode, cancellationToken);
-
-				// Copy the preprocessed data into the right location
-				foreach ((string source, string target) in s_remapConfigValues)
-				{
-					CopyJsonNode(configObj, source, target);
-				}
-
 				// Bind it to the target object
-				GlobalConfig globalConfig = JsonSerializer.Deserialize<GlobalConfig>(configObj, _jsonOptions)!;
+				GlobalConfig globalConfig = await ReadGlobalConfigAsync(context, cancellationToken);
 				if (globalConfig.VersionEnum < ConfigVersion.Latest)
 				{
 					List<string> message = new List<string>();
@@ -637,7 +597,9 @@ namespace HordeServer.Configuration
 					_logger.LogWarning("Global config file is using old version number ({Version}<{LatestVersion})\n\n{DeprecatedFeaturesMessage}\n", globalConfig.Version, (int)ConfigVersion.Latest, String.Join("\n", message));
 				}
 
-				// Serialize it back out to a byte array
+				// Create the snapshot data, with the config serialized back out to a byte array
+				ConfigSnapshot snapshot = new ConfigSnapshot();
+				snapshot.ServerVersion = ServerApp.Version.ToString();
 				snapshot.Data = JsonSerializer.SerializeToUtf8Bytes(globalConfig, context.JsonOptions);
 
 				// Save all the dependencies
@@ -657,6 +619,47 @@ namespace HordeServer.Configuration
 			}
 		}
 
+		async Task<GlobalConfig> ReadGlobalConfigAsync(ConfigContext context, CancellationToken cancellationToken)
+		{
+			// Read the new config in
+			Uri globalConfigUri = GetGlobalConfigUri();
+
+			// Generate the schema for the preprocessor
+			ObjectConfigNode configNode = new ObjectConfigNode(typeof(GlobalConfig));
+
+			ObjectConfigNode pluginsObj = new ObjectConfigNode(false, false);
+			configNode.Properties["Plugins"] = pluginsObj;
+
+			foreach (ILoadedPlugin plugin in _pluginCollection.LoadedPlugins)
+			{
+				if (plugin.GlobalConfigType != null)
+				{
+					pluginsObj.Properties.Add(plugin.Name.ToString(), new ObjectConfigNode(plugin.GlobalConfigType));
+				}
+			}
+
+			// Copy the schema from any remapped nodes to the source node
+			foreach (KeyValuePair<string, string> remapConfigPair in s_remapConfigValues)
+			{
+				if (configNode.TryGetChildNode(remapConfigPair.Value, out ConfigNode? node))
+				{
+					configNode.AddChildNode(remapConfigPair.Key, node);
+				}
+			}
+
+			// Preprocess the file
+			JsonObject configObj = await context.PreprocessFileAsync(globalConfigUri, configNode, cancellationToken);
+
+			// Copy the preprocessed data into the right location
+			foreach ((string source, string target) in s_remapConfigValues)
+			{
+				CopyJsonNode(configObj, source, target);
+			}
+
+			// Bind it to the target object
+			return JsonSerializer.Deserialize<GlobalConfig>(configObj, _jsonOptions)!;
+		}
+
 		static readonly KeyValuePair<string, string>[] s_remapConfigValues = new[]
 		{
 			// Analytics
