Horde: Updating schema documentation

#rnx

diff --git a/Engine/Source/Programs/Horde/Docs/Config/Schema/Globals.md b/Engine/Source/Programs/Horde/Docs/Config/Schema/Globals.md
--- a/Engine/Source/Programs/Horde/Docs/Config/Schema/Globals.md
+++ b/Engine/Source/Programs/Horde/Docs/Config/Schema/Globals.md
@@ -323,7 +323,7 @@ Name | Description
 ---- | -----------
 `type` | [PoolSizeStrategy](#poolsizestrategy-enum)<br>Strategy implementation to use
 `condition` | `string`<br>Condition if this strategy should be enabled (right now, using date/time as a distinguishing factor)
-`config` | `string`<br>Configuration for the strategy, serialized as JSON
+`config` | `object`<br>Configuration for the strategy, serialized as JSON
 `extraAgentCount` | `integer`<br>Integer to add after pool size has been calculated. Can also be negative.
 
 ## FleetManagerInfo
@@ -334,7 +334,7 @@ Name | Description
 ---- | -----------
 `type` | [FleetManagerType](#fleetmanagertype-enum)<br>Fleet manager type implementation to use
 `condition` | `string`<br>Condition if this strategy should be enabled (right now, using date/time as a distinguishing factor)
-`config` | `string`<br>Configuration for the strategy, serialized as JSON
+`config` | `object`<br>Configuration for the strategy, serialized as JSON
 
 ## FleetManagerType (Enum)
 
diff --git a/Engine/Source/Programs/Horde/Docs/Config/Schema/Projects.md b/Engine/Source/Programs/Horde/Docs/Config/Schema/Projects.md
--- a/Engine/Source/Programs/Horde/Docs/Config/Schema/Projects.md
+++ b/Engine/Source/Programs/Horde/Docs/Config/Schema/Projects.md
@@ -98,7 +98,7 @@ Name | Description
 ---- | -----------
 `type` | [PoolSizeStrategy](#poolsizestrategy-enum)<br>Strategy implementation to use
 `condition` | `string`<br>Condition if this strategy should be enabled (right now, using date/time as a distinguishing factor)
-`config` | `string`<br>Configuration for the strategy, serialized as JSON
+`config` | `object`<br>Configuration for the strategy, serialized as JSON
 `extraAgentCount` | `integer`<br>Integer to add after pool size has been calculated. Can also be negative.
 
 ## FleetManagerInfo
@@ -109,7 +109,7 @@ Name | Description
 ---- | -----------
 `type` | [FleetManagerType](#fleetmanagertype-enum)<br>Fleet manager type implementation to use
 `condition` | `string`<br>Condition if this strategy should be enabled (right now, using date/time as a distinguishing factor)
-`config` | `string`<br>Configuration for the strategy, serialized as JSON
+`config` | `object`<br>Configuration for the strategy, serialized as JSON
 
 ## FleetManagerType (Enum)
 
