Horde: Tweak the step trends graph trends line
#rnx
#skipci

diff --git a/Engine/Source/Programs/Horde/HordeDashboard/src/components/jobDetailsV2/JobDetailStepTrendsV2.tsx b/Engine/Source/Programs/Horde/HordeDashboard/src/components/jobDetailsV2/JobDetailStepTrendsV2.tsx
--- a/Engine/Source/Programs/Horde/HordeDashboard/src/components/jobDetailsV2/JobDetailStepTrendsV2.tsx
+++ b/Engine/Source/Programs/Horde/HordeDashboard/src/components/jobDetailsV2/JobDetailStepTrendsV2.tsx
@@ -402,7 +402,7 @@ class StepTrendsRenderer {
 
          let sum = 0;
          let count = 0;
-         const range = 50;
+         const range = 10;
          for (let j = idx - range; j < idx + range; j++) {
 
             if (j < 0 || j >= data.length) {
