Horde: Mark the 5.4 release in the release notes.

#rnx

diff --git a/Engine/Source/Programs/Horde/Docs/ReleaseNotes.md b/Engine/Source/Programs/Horde/Docs/ReleaseNotes.md
--- a/Engine/Source/Programs/Horde/Docs/ReleaseNotes.md
+++ b/Engine/Source/Programs/Horde/Docs/ReleaseNotes.md
@@ -201,7 +201,7 @@
 * Comment out the placeholder ticket value in the default config file. (32585713)
 * Improve ordering of nodes when writing large file trees to storage, such that the nodes which are read after each other are adjacent in the storage blobs. Nodes are read in the reverse order to which they are written, depth-first. Arranging nodes in the blob in this order prevents thrashing of the cache and improves performance. (32585474)
 
-## 2024-03-28
+## 2024-03-28 (UE 5.4 Release)
 
 * Remove /userinfo call when authenticating via JWT. If an access token is passed, it's not guaranteed it has permission to access /userinfo from OIDC. ID tokens during normal web-based login does on the other hand. (32539012)
 * Fix parsing of true/false literals in condition expressions. (32518640)
