Horde: Update release notes.

#rnx

diff --git a/Engine/Source/Programs/Horde/Docs/ReleaseNotes.md b/Engine/Source/Programs/Horde/Docs/ReleaseNotes.md
--- a/Engine/Source/Programs/Horde/Docs/ReleaseNotes.md
+++ b/Engine/Source/Programs/Horde/Docs/ReleaseNotes.md
@@ -1,5 +1,14 @@
 # Release Notes
 
+## 2024-07-25
+
+* Add a separate flag for marking an issue as fixed as a systemic issue, rather than having to pass a negative value as a fix changelist. Still supports passing/returning negative fix changelists as well for now, but will be removed in future. (35029295)
+* Serialize config for pool size strategies and fleet managers as actual json objects, rather than json objects embedded in strings. Supports reading from json objects embedded in strings for backwards compatibility. (35026676)
+* Set output assembly name for the command line tool to "Horde". (34984531)
+* Separate out log event errors and critical failures into a separate fingerprint when we exceed the 25 event limit in the Hashed Issue Handler. (34930826)
+* Add support for tagging log events with issue fingerprints. The $issue property on a log event can contain a serialized IssueFingerprint, which will be parsed by the ExternalIssueHandler handler inside Horde. This property can be added from C# code using the ILogger.BeginIssueScope() extension method. (34913930)
+* ParentView is optional per perforce documentation. Fixes Unhandled Exception (Missing 'ParentView' tag when parsing 'StreamRecord') for streams that does not have it set (34890081)
+
 ## 2024-07-16
 
 * Support for disabling conform tasks through the build config file. (34836886)
