Horde: Add a BSON serializer from JsonObject to strings to fix errors creating sessions.

#rnx

diff --git a/Engine/Source/Programs/Horde/HordeServer/Startup.cs b/Engine/Source/Programs/Horde/HordeServer/Startup.cs
--- a/Engine/Source/Programs/Horde/HordeServer/Startup.cs
+++ b/Engine/Source/Programs/Horde/HordeServer/Startup.cs
@@ -8,6 +8,8 @@ using System.Linq;
 using System.Net.Mime;
 using System.Reflection;
 using System.Runtime.InteropServices;
+using System.Text.Json;
+using System.Text.Json.Nodes;
 using System.Threading;
 using System.Threading.Tasks;
 using EpicGames.AspNet;
@@ -82,6 +84,7 @@ namespace HordeServer
 {
 	using ContentHash = EpicGames.Core.ContentHash;
 	using ILogger = Microsoft.Extensions.Logging.ILogger;
+	using JsonObject = System.Text.Json.Nodes.JsonObject;
 
 	class Startup : IServerStartup
 	{
@@ -836,6 +839,17 @@ namespace HordeServer
 			}
 		}
 
+		sealed class JsonObjectBsonSerializer : SerializerBase<JsonObject>
+		{
+			/// <inheritdoc/>
+			public override JsonObject Deserialize(BsonDeserializationContext context, BsonDeserializationArgs args)
+				=> (JsonObject)JsonObject.Parse(context.Reader.ReadString(), new JsonNodeOptions { PropertyNameCaseInsensitive = true }, new JsonDocumentOptions { AllowTrailingCommas = true, CommentHandling = JsonCommentHandling.Skip })!;
+
+			/// <inheritdoc/>
+			public override void Serialize(BsonSerializationContext context, BsonSerializationArgs args, JsonObject value)
+				=> context.Writer.WriteString(value.ToJsonString());
+		}
+
 		static int s_haveConfiguredMongoDb = 0;
 
 		public static void ConfigureMongoDbClient()
@@ -859,6 +873,7 @@ namespace HordeServer
 				BsonSerializer.RegisterSerializer(new ConditionSerializer());
 				BsonSerializer.RegisterSerializer(new SubResourceIdSerializer());
 				BsonSerializer.RegisterSerializer(new CommitTagBsonSerializer());
+				BsonSerializer.RegisterSerializer(new JsonObjectBsonSerializer());
 				BsonSerializer.RegisterSerializationProvider(new BsonSerializationProvider());
 				BsonSerializer.RegisterSerializationProvider(new StringIdBsonSerializationProvider());
 				BsonSerializer.RegisterSerializationProvider(new BinaryIdBsonSerializationProvider());
