Horde: Support deserializing empty strings into JsonObjects.

#rnx

diff --git a/Engine/Source/Programs/Horde/HordeServer/Startup.cs b/Engine/Source/Programs/Horde/HordeServer/Startup.cs
--- a/Engine/Source/Programs/Horde/HordeServer/Startup.cs
+++ b/Engine/Source/Programs/Horde/HordeServer/Startup.cs
@@ -843,7 +843,17 @@ namespace HordeServer
 		{
 			/// <inheritdoc/>
 			public override JsonObject Deserialize(BsonDeserializationContext context, BsonDeserializationArgs args)
-				=> (JsonObject)JsonObject.Parse(context.Reader.ReadString(), new JsonNodeOptions { PropertyNameCaseInsensitive = true }, new JsonDocumentOptions { AllowTrailingCommas = true, CommentHandling = JsonCommentHandling.Skip })!;
+			{
+				string str = context.Reader.ReadString();
+				if (String.IsNullOrWhiteSpace(str))
+				{
+					return new JsonObject();
+				}
+				else
+				{
+					return (JsonObject)JsonObject.Parse(str, new JsonNodeOptions { PropertyNameCaseInsensitive = true }, new JsonDocumentOptions { AllowTrailingCommas = true, CommentHandling = JsonCommentHandling.Skip })!;
+				}
+			}
 
 			/// <inheritdoc/>
 			public override void Serialize(BsonSerializationContext context, BsonSerializationArgs args, JsonObject value)
