Horde: Preliminary support for VCS commit ids that aren't an integer.

* CommitId is an abstraction for a commit id, which uses a string. Most API entry points now take these.
* CommitIdWithOrder derives from commit id and adds an integer field which can be used for ordering. In situations where the commit id itself is not numeric, this may use the commit timestamp or some similar value.
* All endpoints still support passing changelist numbers for now.

diff --git a/Engine/Source/Programs/AutomationTool/BuildGraph/Tasks/CreateArtifactTask.cs b/Engine/Source/Programs/AutomationTool/BuildGraph/Tasks/CreateArtifactTask.cs
--- a/Engine/Source/Programs/AutomationTool/BuildGraph/Tasks/CreateArtifactTask.cs
+++ b/Engine/Source/Programs/AutomationTool/BuildGraph/Tasks/CreateArtifactTask.cs
@@ -12,9 +12,11 @@ using System.Xml;
 using EpicGames.Core;
 using EpicGames.Horde;
 using EpicGames.Horde.Artifacts;
+using EpicGames.Horde.Commits;
 using EpicGames.Horde.Storage;
 using EpicGames.Horde.Storage.Clients;
 using EpicGames.Horde.Storage.Nodes;
+using EpicGames.Horde.Streams;
 using Microsoft.Extensions.DependencyInjection;
 using Microsoft.Extensions.Logging;
 
@@ -51,6 +53,12 @@ namespace AutomationTool.Tasks
 		[TaskParameter]
 		public string? BaseDir { get; set; }
 
+		/// <summary>
+		/// Stream for this artifact
+		/// </summary>
+		[TaskParameter(Optional = true)]
+		public string? StreamId { get; set; }
+
 		/// <summary>
 		/// Changelist number for this artifact
 		/// </summary>
@@ -111,8 +119,9 @@ namespace AutomationTool.Tasks
 			ArtifactType artifactType = new ArtifactType(_parameters.Type);
 
 			HordeHttpClient hordeHttpClient = serviceProvider.GetRequiredService<HordeHttpClient>();
-			int? change = (_parameters.Change == 0) ? (int?)null : _parameters.Change;
-			CreateArtifactResponse response = await hordeHttpClient.CreateArtifactAsync(artifactName, artifactType, _parameters.Description, change: change);
+			StreamId streamId = new StreamId(_parameters.StreamId ?? throw new InvalidOperationException("Missing StreamId parameter")); 
+			CommitId commitId = CommitId.FromPerforceChange(_parameters.Change ?? CommandUtils.P4Env.Changelist);
+			CreateArtifactResponse response = await hordeHttpClient.CreateArtifactAsync(artifactName, artifactType, _parameters.Description, streamId: streamId, commitId: commitId);
 			Logger.LogInformation("Creating artifact {ArtifactId} '{ArtifactName}' ({ArtifactType}) (ns: {NamespaceId}, ref: {RefName})", response.ArtifactId, artifactName, artifactType, response.NamespaceId, response.RefName);
 
 			Stopwatch timer = Stopwatch.StartNew();
diff --git a/Engine/Source/Programs/Horde/HordeCmd/Commands/Artifacts/ArtifactInfo.cs b/Engine/Source/Programs/Horde/HordeCmd/Commands/Artifacts/ArtifactInfo.cs
--- a/Engine/Source/Programs/Horde/HordeCmd/Commands/Artifacts/ArtifactInfo.cs
+++ b/Engine/Source/Programs/Horde/HordeCmd/Commands/Artifacts/ArtifactInfo.cs
@@ -31,7 +31,7 @@ namespace Horde.Commands.Artifacts
 			logger.LogInformation("Name: {Name}", artifact.Name);
 			logger.LogInformation("Type: {Type}", artifact.Type);
 			logger.LogInformation("Stream: {Stream}", artifact.StreamId);
-			logger.LogInformation("Change: {Change}", artifact.Change);
+			logger.LogInformation("Commit: {Commit}", artifact.CommitId);
 			logger.LogInformation("Description: {Desc}", artifact.Description);
 
 			if (artifact.Keys.Count > 0)
diff --git a/Engine/Source/Programs/Horde/HordeServer/Startup.cs b/Engine/Source/Programs/Horde/HordeServer/Startup.cs
--- a/Engine/Source/Programs/Horde/HordeServer/Startup.cs
+++ b/Engine/Source/Programs/Horde/HordeServer/Startup.cs
@@ -657,6 +657,28 @@ namespace HordeServer
 			LogValueFormatter.RegisterTypeAnnotation<UserId>("UserId");
 		}
 
+		public sealed class CommitIdBsonSerializer : SerializerBase<CommitId>
+		{
+			/// <inheritdoc/>
+			public override CommitId Deserialize(BsonDeserializationContext context, BsonDeserializationArgs args)
+			{
+				if (context.Reader.GetCurrentBsonType() == BsonType.Int32)
+				{
+					return CommitId.FromPerforceChange(context.Reader.ReadInt32());
+				}
+				else
+				{
+					return new CommitId(context.Reader.ReadString());
+				}
+			}
+
+			/// <inheritdoc/>
+			public override void Serialize(BsonSerializationContext context, BsonSerializationArgs args, CommitId value)
+			{
+				context.Writer.WriteString(value.Name);
+			}
+		}
+
 		public sealed class BlobLocatorBsonSerializer : SerializerBase<BlobLocator>
 		{
 			/// <inheritdoc/>
@@ -873,6 +895,7 @@ namespace HordeServer
 				ConventionRegistry.Register("Horde", conventionPack, type => true);
 
 				// Register the custom serializers
+				BsonSerializer.RegisterSerializer(new CommitIdBsonSerializer());
 				BsonSerializer.RegisterSerializer(new BlobLocatorBsonSerializer());
 				BsonSerializer.RegisterSerializer(new RefNameBsonSerializer());
 				BsonSerializer.RegisterSerializer(new IoHashBsonSerializer());
diff --git a/Engine/Source/Programs/Horde/Plugins/Build/HordeServer.Build.Tests/Agents/AgentServiceTest.cs b/Engine/Source/Programs/Horde/Plugins/Build/HordeServer.Build.Tests/Agents/AgentServiceTest.cs
--- a/Engine/Source/Programs/Horde/Plugins/Build/HordeServer.Build.Tests/Agents/AgentServiceTest.cs
+++ b/Engine/Source/Programs/Horde/Plugins/Build/HordeServer.Build.Tests/Agents/AgentServiceTest.cs
@@ -28,7 +28,7 @@ public class AgentServiceTest : BuildTestSetup
 		ActionResult<List<object>> res = await JobsController.FindJobsAsync();
 
 		List<GetJobResponse> responses = res.Value!.ConvertAll(x => (GetJobResponse)x);
-		responses.SortBy(x => x.Change);
+		responses.SortBy(x => x.CommitId);
 
 		Assert.AreEqual(2, responses.Count);
 		Assert.AreEqual("hello1", responses[0].Name);
diff --git a/Engine/Source/Programs/Horde/Plugins/Build/HordeServer.Build.Tests/Agents/Pools/PoolsControllerTest.cs b/Engine/Source/Programs/Horde/Plugins/Build/HordeServer.Build.Tests/Agents/Pools/PoolsControllerTest.cs
--- a/Engine/Source/Programs/Horde/Plugins/Build/HordeServer.Build.Tests/Agents/Pools/PoolsControllerTest.cs
+++ b/Engine/Source/Programs/Horde/Plugins/Build/HordeServer.Build.Tests/Agents/Pools/PoolsControllerTest.cs
@@ -1,6 +1,5 @@
 ﻿// Copyright Epic Games, Inc. All Rights Reserved.
 
-using EpicGames.Horde.Common;
 using HordeServer.Agents.Fleet;
 using HordeServer.Agents.Pools;
 using Microsoft.AspNetCore.Mvc;
diff --git a/Engine/Source/Programs/Horde/Plugins/Build/HordeServer.Build.Tests/Artifacts/ArtifactTests.cs b/Engine/Source/Programs/Horde/Plugins/Build/HordeServer.Build.Tests/Artifacts/ArtifactTests.cs
--- a/Engine/Source/Programs/Horde/Plugins/Build/HordeServer.Build.Tests/Artifacts/ArtifactTests.cs
+++ b/Engine/Source/Programs/Horde/Plugins/Build/HordeServer.Build.Tests/Artifacts/ArtifactTests.cs
@@ -2,9 +2,13 @@
 
 using EpicGames.Horde.Acls;
 using EpicGames.Horde.Artifacts;
+using EpicGames.Horde.Commits;
 using EpicGames.Horde.Streams;
 using HordeCommon;
 using HordeServer.Artifacts;
+using HordeServer.Projects;
+using HordeServer.Server;
+using HordeServer.Streams;
 using Microsoft.Extensions.DependencyInjection;
 
 namespace HordeServer.Tests.Artifacts
@@ -12,13 +16,28 @@ namespace HordeServer.Tests.Artifacts
 	[TestClass]
 	public class ArtifactTests : BuildTestSetup
 	{
+		public ArtifactTests()
+		{
+			void SetupFixture(GlobalConfig globalConfig)
+			{
+				ProjectConfig projectConfig = new ProjectConfig();
+				projectConfig.Streams.Add(new StreamConfig { Id = new StreamId("foo") });
+				projectConfig.Streams.Add(new StreamConfig { Id = new StreamId("bar") });
+
+				BuildConfig buildConfig = globalConfig.Plugins.GetBuildConfig();
+				buildConfig.Projects.Add(projectConfig);
+			}
+
+			UpdateConfig(SetupFixture);
+		}
+
 		[TestMethod]
 		public async Task CreateArtifactAsync()
 		{
 			StreamId streamId = new StreamId("foo");
 
 			IArtifactCollection artifactCollection = ServiceProvider.GetRequiredService<IArtifactCollection>();
-			IArtifact artifact = await artifactCollection.AddAsync(new ArtifactName("default"), ArtifactType.StepOutput, null, streamId, 1, new string[] { "test1", "test2" }, Array.Empty<string>(), AclScopeName.Root);
+			IArtifact artifact = await artifactCollection.AddAsync(new ArtifactName("default"), ArtifactType.StepOutput, null, streamId, CommitIdWithOrder.FromPerforceChange(1), new string[] { "test1", "test2" }, Array.Empty<string>(), AclScopeName.Root);
 
 			{
 				List<IArtifact> artifacts = await artifactCollection.FindAsync(streamId, keys: new[] { "test1" }).ToListAsync();
@@ -52,7 +71,7 @@ namespace HordeServer.Tests.Artifacts
 			StreamId streamId = new StreamId("foo");
 
 			IArtifactCollection artifactCollection = ServiceProvider.GetRequiredService<IArtifactCollection>();
-			IArtifact artifact = await artifactCollection.AddAsync(new ArtifactName("default"), type, null, streamId, 1, new string[] { "test1", "test2" }, Array.Empty<string>(), AclScopeName.Root);
+			IArtifact artifact = await artifactCollection.AddAsync(new ArtifactName("default"), type, null, streamId, CommitIdWithOrder.FromPerforceChange(1), new string[] { "test1", "test2" }, Array.Empty<string>(), AclScopeName.Root);
 
 			{
 				List<IArtifact> artifacts = await artifactCollection.FindAsync(streamId, keys: new[] { "test1" }).ToListAsync();
@@ -83,7 +102,7 @@ namespace HordeServer.Tests.Artifacts
 			StreamId streamId = new StreamId("foo");
 
 			IArtifactCollection artifactCollection = ServiceProvider.GetRequiredService<IArtifactCollection>();
-			IArtifact artifact = await artifactCollection.AddAsync(new ArtifactName("default"), type, null, streamId, 1, new string[] { "test1", "test2" }, Array.Empty<string>(), AclScopeName.Root);
+			IArtifact artifact = await artifactCollection.AddAsync(new ArtifactName("default"), type, null, streamId, CommitIdWithOrder.FromPerforceChange(1), new string[] { "test1", "test2" }, Array.Empty<string>(), AclScopeName.Root);
 
 			{
 				List<IArtifact> artifacts = await artifactCollection.FindAsync(streamId, keys: new[] { "test1" }).ToListAsync();
@@ -126,11 +145,11 @@ namespace HordeServer.Tests.Artifacts
 			IArtifactCollection artifactCollection = ServiceProvider.GetRequiredService<IArtifactCollection>();
 			for (int idx = 0; idx < 10; idx++)
 			{
-				await artifactCollection.AddAsync(new ArtifactName($"default-{idx}"), type, null, fooStreamId, 1, new string[] { "test1", "test2" }, Array.Empty<string>(), AclScopeName.Root);
+				await artifactCollection.AddAsync(new ArtifactName($"default-{idx}"), type, null, fooStreamId, CommitIdWithOrder.FromPerforceChange(1), new string[] { "test1", "test2" }, Array.Empty<string>(), AclScopeName.Root);
 			}
 			for (int idx = 0; idx < 10; idx++)
 			{
-				await artifactCollection.AddAsync(new ArtifactName($"default-{idx}"), type, null, barStreamId, 1, new string[] { "test1", "test2" }, Array.Empty<string>(), AclScopeName.Root);
+				await artifactCollection.AddAsync(new ArtifactName($"default-{idx}"), type, null, barStreamId, CommitIdWithOrder.FromPerforceChange(1), new string[] { "test1", "test2" }, Array.Empty<string>(), AclScopeName.Root);
 			}
 
 			await Clock.AdvanceAsync(TimeSpan.FromDays(2.0));
diff --git a/Engine/Source/Programs/Horde/Plugins/Build/HordeServer.Build.Tests/Devices/DeviceServiceTests.cs b/Engine/Source/Programs/Horde/Plugins/Build/HordeServer.Build.Tests/Devices/DeviceServiceTests.cs
--- a/Engine/Source/Programs/Horde/Plugins/Build/HordeServer.Build.Tests/Devices/DeviceServiceTests.cs
+++ b/Engine/Source/Programs/Horde/Plugins/Build/HordeServer.Build.Tests/Devices/DeviceServiceTests.cs
@@ -2,6 +2,7 @@
 
 using System.Security.Claims;
 using EpicGames.Core;
+using EpicGames.Horde.Commits;
 using EpicGames.Horde.Devices;
 using EpicGames.Horde.Jobs;
 using EpicGames.Horde.Jobs.Templates;
@@ -190,7 +191,7 @@ namespace HordeServer.Tests.Devices
 			CreateJobOptions options = new CreateJobOptions();
 			options.Arguments.Add("-Target=Run Tests");
 
-			IJob job = await JobCollection.AddAsync(JobIdUtils.GenerateNewId(), new StreamId("ue5-main"), new TemplateId("test-build"), ContentHash.SHA1("hello"), baseGraph, "Test job", 123, 123, options);
+			IJob job = await JobCollection.AddAsync(JobIdUtils.GenerateNewId(), new StreamId("ue5-main"), new TemplateId("test-build"), ContentHash.SHA1("hello"), baseGraph, "Test job", CommitIdWithOrder.FromPerforceChange(123), CommitIdWithOrder.FromPerforceChange(123), options);
 
 			job = await StartBatchAsync(job, 0);
 			job = await RunStepAsync(job, 0, 0, JobStepOutcome.Success); // Setup Build
diff --git a/Engine/Source/Programs/Horde/Plugins/Build/HordeServer.Build.Tests/Fixture.cs b/Engine/Source/Programs/Horde/Plugins/Build/HordeServer.Build.Tests/Fixture.cs
--- a/Engine/Source/Programs/Horde/Plugins/Build/HordeServer.Build.Tests/Fixture.cs
+++ b/Engine/Source/Programs/Horde/Plugins/Build/HordeServer.Build.Tests/Fixture.cs
@@ -2,6 +2,7 @@
 
 using EpicGames.Core;
 using EpicGames.Horde.Agents.Pools;
+using EpicGames.Horde.Commits;
 using EpicGames.Horde.Jobs;
 using EpicGames.Horde.Jobs.Templates;
 using EpicGames.Horde.Projects;
@@ -95,9 +96,9 @@ namespace HordeServer.Tests
 				templateHash: Template.Hash,
 				graph: Graph,
 				name: "hello1",
-				change: 1000001,
-				codeChange: 1000002,
-				new CreateJobOptions { PreflightChange = 1001 }
+				commitId: CommitIdWithOrder.FromPerforceChange(1000001),
+				codeCommitId: CommitIdWithOrder.FromPerforceChange(1000002),
+				new CreateJobOptions { PreflightCommitId = CommitId.FromPerforceChange(1001) }
 			);
 			Job1 = (await jobService.GetJobAsync(Job1.Id))!;
 
@@ -108,8 +109,8 @@ namespace HordeServer.Tests
 				templateHash: ContentHash.MD5("made-up-template-hash"),
 				graph: fg,
 				name: "hello2",
-				change: 2000001,
-				codeChange: 2000002,
+				commitId: CommitIdWithOrder.FromPerforceChange(2000001),
+				codeCommitId: CommitIdWithOrder.FromPerforceChange(2000002),
 				new CreateJobOptions()
 			);
 			Job2 = (await jobService.GetJobAsync(Job2.Id))!;
diff --git a/Engine/Source/Programs/Horde/Plugins/Build/HordeServer.Build.Tests/Fleet/JobQueueStrategyTest.cs b/Engine/Source/Programs/Horde/Plugins/Build/HordeServer.Build.Tests/Fleet/JobQueueStrategyTest.cs
--- a/Engine/Source/Programs/Horde/Plugins/Build/HordeServer.Build.Tests/Fleet/JobQueueStrategyTest.cs
+++ b/Engine/Source/Programs/Horde/Plugins/Build/HordeServer.Build.Tests/Fleet/JobQueueStrategyTest.cs
@@ -2,6 +2,7 @@
 
 using EpicGames.Core;
 using EpicGames.Horde.Agents.Pools;
+using EpicGames.Horde.Commits;
 using EpicGames.Horde.Jobs;
 using EpicGames.Horde.Jobs.Templates;
 using EpicGames.Horde.Projects;
@@ -169,7 +170,7 @@ namespace HordeServer.Tests.Fleet
 
 			IJob job = await JobCollection.AddAsync(JobIdUtils.GenerateNewId(), streamId,
 				new TemplateId("bogusTemplateRefId"), ContentHash.Empty, graph, "bogusJobName",
-				1000, 1000, options);
+				CommitIdWithOrder.FromPerforceChange(1000), CommitIdWithOrder.FromPerforceChange(1000), options);
 
 			return job;
 		}
diff --git a/Engine/Source/Programs/Horde/Plugins/Build/HordeServer.Build.Tests/Issues/IssueServiceTests.cs b/Engine/Source/Programs/Horde/Plugins/Build/HordeServer.Build.Tests/Issues/IssueServiceTests.cs
--- a/Engine/Source/Programs/Horde/Plugins/Build/HordeServer.Build.Tests/Issues/IssueServiceTests.cs
+++ b/Engine/Source/Programs/Horde/Plugins/Build/HordeServer.Build.Tests/Issues/IssueServiceTests.cs
@@ -6,6 +6,7 @@ using System.Runtime.InteropServices;
 using System.Text;
 using System.Text.Json;
 using EpicGames.Core;
+using EpicGames.Horde.Commits;
 using EpicGames.Horde.Issues;
 using EpicGames.Horde.Jobs;
 using EpicGames.Horde.Jobs.Templates;
@@ -295,7 +296,7 @@ namespace HordeServer.Tests.Issues
 			job.SetupGet(x => x.StreamId).Returns(streamId);
 			job.SetupGet(x => x.TemplateId).Returns(new TemplateId("test-template"));
 			job.SetupGet(x => x.CreateTimeUtc).Returns(utcNow);
-			job.SetupGet(x => x.Change).Returns(change);
+			job.SetupGet(x => x.CommitId).Returns(CommitIdWithOrder.FromPerforceChange(change));
 			job.SetupGet(x => x.Batches).Returns(batches);
 			job.SetupGet(x => x.ShowUgsBadges).Returns(promoteByDefault);
 			job.SetupGet(x => x.ShowUgsAlerts).Returns(promoteByDefault);
@@ -312,7 +313,7 @@ namespace HordeServer.Tests.Issues
 
 			JobStepRefId jobStepRefId = new JobStepRefId(job.Id, batch.Id, step.Id);
 			string nodeName = _graph.Groups[batch.GroupIdx].Nodes[step.NodeIdx].Name;
-			await JobStepRefCollection.InsertOrReplaceAsync(jobStepRefId, "TestJob", nodeName, job.StreamId, job.TemplateId, job.Change, step.LogId, null, null, JobStepState.Completed, outcome, job.UpdateIssues, null, null, 0.0f, 0.0f, job.CreateTimeUtc, step.StartTimeUtc!.Value, step.StartTimeUtc);
+			await JobStepRefCollection.InsertOrReplaceAsync(jobStepRefId, "TestJob", nodeName, job.StreamId, job.TemplateId, job.CommitId, step.LogId, null, null, JobStepState.Completed, outcome, job.UpdateIssues, null, null, 0.0f, 0.0f, job.CreateTimeUtc, step.StartTimeUtc!.Value, step.StartTimeUtc);
 
 			if (job.UpdateIssues)
 			{
@@ -1006,22 +1007,22 @@ namespace HordeServer.Tests.Issues
 				Assert.AreEqual("Compile errors in fog.cpp", issue.Summary);
 
 				IIssueSpan stream = spans[0];
-				Assert.AreEqual(105, stream.LastSuccess?.Change);
-				Assert.AreEqual(null, stream.NextSuccess?.Change);
+				Assert.AreEqual(105, stream.LastSuccess?.CommitId.GetPerforceChange());
+				Assert.AreEqual(null, stream.NextSuccess?.CommitId);
 
 				IReadOnlyList<IIssueSuspect> suspects = await IssueCollection.FindSuspectsAsync(issue);
 				suspects = suspects.OrderBy(x => x.Id).ToList();
 				Assert.AreEqual(3, suspects.Count);
 
-				Assert.AreEqual(75 /*115*/, suspects[1].Change);
+				Assert.AreEqual(75 /*115*/, suspects[1].CommitId.GetPerforceChange());
 				Assert.AreEqual(_jerryId, suspects[1].AuthorId);
 				//				Assert.AreEqual(75, Suspects[1].OriginatingChange);
 
-				Assert.AreEqual(110, suspects[2].Change);
+				Assert.AreEqual(110, suspects[2].CommitId.GetPerforceChange());
 				Assert.AreEqual(_bobId, suspects[2].AuthorId);
 				//			Assert.AreEqual(null, Suspects[2].OriginatingChange);
 
-				Assert.AreEqual(120, suspects[0].Change);
+				Assert.AreEqual(120, suspects[0].CommitId.GetPerforceChange());
 				Assert.AreEqual(_timId, suspects[0].AuthorId);
 				//				Assert.AreEqual(null, Suspects[0].OriginatingChange);
 
@@ -1044,16 +1045,16 @@ namespace HordeServer.Tests.Issues
 				Assert.AreEqual(spans.Count, 1);
 
 				IIssueSpan stream = spans[0];
-				Assert.AreEqual(110, stream.LastSuccess?.Change);
-				Assert.AreEqual(null, stream.NextSuccess?.Change);
+				Assert.AreEqual(110, stream.LastSuccess?.CommitId.GetPerforceChange());
+				Assert.AreEqual(null, stream.NextSuccess?.CommitId);
 
-				IReadOnlyList<IIssueSuspect> suspects = (await IssueCollection.FindSuspectsAsync(issue)).OrderByDescending(x => x.Change).ToList();
+				IReadOnlyList<IIssueSuspect> suspects = (await IssueCollection.FindSuspectsAsync(issue)).OrderByDescending(x => x.CommitId).ToList();
 				Assert.AreEqual(2, suspects.Count);
 
-				Assert.AreEqual(120, suspects[0].Change);
+				Assert.AreEqual(120, suspects[0].CommitId.GetPerforceChange());
 				Assert.AreEqual(_timId, suspects[0].AuthorId);
 
-				Assert.AreEqual(75, suspects[1].Change);
+				Assert.AreEqual(75, suspects[1].CommitId.GetPerforceChange());
 				Assert.AreEqual(_jerryId, suspects[1].AuthorId);
 
 				IReadOnlyList<IIssue> openIssues = await IssueCollection.FindIssuesAsync(resolved: false);
@@ -1075,8 +1076,8 @@ namespace HordeServer.Tests.Issues
 				Assert.AreEqual(spans.Count, 1);
 
 				IIssueSpan stream = spans[0];
-				Assert.AreEqual(110, stream.LastSuccess?.Change);
-				Assert.AreEqual(125, stream.NextSuccess?.Change);
+				Assert.AreEqual(110, stream.LastSuccess?.CommitId.GetPerforceChange());
+				Assert.AreEqual(125, stream.NextSuccess?.CommitId.GetPerforceChange());
 
 				IReadOnlyList<IIssue> openIssues = await IssueCollection.FindIssuesAsync(resolved: false);
 				Assert.AreEqual(0, openIssues.Count);
@@ -1103,8 +1104,8 @@ namespace HordeServer.Tests.Issues
 				Assert.AreEqual(1, spans.Count);
 
 				IIssueSpan span = spans[0];
-				Assert.AreEqual(110, span.LastSuccess?.Change);
-				Assert.AreEqual(125, span.NextSuccess?.Change);
+				Assert.AreEqual(110, span.LastSuccess?.CommitId.GetPerforceChange());
+				Assert.AreEqual(125, span.NextSuccess?.CommitId.GetPerforceChange());
 
 				IReadOnlyList<IIssueStep> steps = await IssueCollection.FindStepsAsync(span.Id);
 				Assert.AreEqual(2, steps.Count);
@@ -1794,8 +1795,8 @@ namespace HordeServer.Tests.Issues
 				Assert.AreEqual(issue.Fingerprints[0].Type, "Symbol");
 
 				IIssueSpan stream = spans[0];
-				Assert.AreEqual(105, stream.LastSuccess?.Change);
-				Assert.AreEqual(null, stream.NextSuccess?.Change);
+				Assert.AreEqual(105, stream.LastSuccess?.CommitId.GetPerforceChange());
+				Assert.AreEqual(null, stream.NextSuccess?.CommitId);
 
 				IReadOnlyList<IIssueSuspect> suspects = await IssueCollection.FindSuspectsAsync(issues[0]);
 
@@ -1857,12 +1858,12 @@ namespace HordeServer.Tests.Issues
 				Assert.AreEqual(issue.Fingerprints[0].Type, "Symbol");
 
 				IIssueSpan span1 = spans[0];
-				Assert.AreEqual(105, span1.LastSuccess?.Change);
-				Assert.AreEqual(null, span1.NextSuccess?.Change);
+				Assert.AreEqual(105, span1.LastSuccess?.CommitId.GetPerforceChange());
+				Assert.AreEqual(null, span1.NextSuccess?.CommitId);
 
 				IIssueSpan span2 = spans[1];
-				Assert.AreEqual(null, span2.LastSuccess?.Change);
-				Assert.AreEqual(null, span2.NextSuccess?.Change);
+				Assert.AreEqual(null, span2.LastSuccess?.CommitId);
+				Assert.AreEqual(null, span2.NextSuccess?.CommitId);
 			}
 		}
 
@@ -1904,8 +1905,8 @@ namespace HordeServer.Tests.Issues
 				Assert.AreEqual("Symbol", issue.Fingerprints[0].Type);
 
 				IIssueSpan span = spans[0];
-				Assert.AreEqual(105, span.LastSuccess?.Change);
-				Assert.AreEqual(null, span.NextSuccess?.Change);
+				Assert.AreEqual(105, span.LastSuccess?.CommitId.GetPerforceChange());
+				Assert.AreEqual(null, span.NextSuccess?.CommitId);
 			}
 
 			// #3
@@ -1931,8 +1932,8 @@ namespace HordeServer.Tests.Issues
 				Assert.AreEqual("Hashed", issue.Fingerprints[0].Type);
 
 				IIssueSpan span = spans[0];
-				Assert.AreEqual(105, span.LastSuccess?.Change);
-				Assert.AreEqual(null, span.NextSuccess?.Change);
+				Assert.AreEqual(105, span.LastSuccess?.CommitId.GetPerforceChange());
+				Assert.AreEqual(null, span.NextSuccess?.CommitId);
 			}
 		}
 
@@ -2404,7 +2405,7 @@ namespace HordeServer.Tests.Issues
 			// Scenario: Issue is marked fixed again, at a particular changelist
 			// Expected: Resolved time, owner is set
 			{
-				await IssueService.UpdateIssueAsync(issueId, resolvedById: _bobId, fixChange: 115);
+				await IssueService.UpdateIssueAsync(issueId, resolvedById: _bobId, fixCommitId: CommitId.FromPerforceChange(115));
 
 				IReadOnlyList<IIssue> openIssues = await IssueCollection.FindIssuesAsync();
 				Assert.AreEqual(0, openIssues.Count);
@@ -2437,7 +2438,7 @@ namespace HordeServer.Tests.Issues
 			// Scenario: Issue is marked fixed again, at a particular changelist
 			// Expected: Resolved time, owner is set
 			{
-				await IssueService.UpdateIssueAsync(issueId, resolvedById: _bobId, fixChange: 125);
+				await IssueService.UpdateIssueAsync(issueId, resolvedById: _bobId, fixCommitId: CommitIdWithOrder.FromPerforceChange(125));
 
 				IReadOnlyList<IIssue> openIssues = await IssueCollection.FindIssuesAsync();
 				Assert.AreEqual(0, openIssues.Count);
@@ -2943,8 +2944,8 @@ namespace HordeServer.Tests.Issues
 		static IIssueSpanSuspect MockSuspect(int change, int? originatingChange)
 		{
 			Mock<IIssueSpanSuspect> suspect = new Mock<IIssueSpanSuspect>(MockBehavior.Strict);
-			suspect.SetupGet(x => x.Change).Returns(change);
-			suspect.SetupGet(x => x.OriginatingChange).Returns(originatingChange);
+			suspect.SetupGet(x => x.CommitId).Returns(CommitIdWithOrder.FromPerforceChange(change));
+			suspect.SetupGet(x => x.SourceCommitId).Returns(CommitIdWithOrder.FromPerforceChange(originatingChange));
 			return suspect.Object;
 		}
 
diff --git a/Engine/Source/Programs/Horde/Plugins/Build/HordeServer.Build.Tests/Jobs/Bisect/BisectTests.cs b/Engine/Source/Programs/Horde/Plugins/Build/HordeServer.Build.Tests/Jobs/Bisect/BisectTests.cs
--- a/Engine/Source/Programs/Horde/Plugins/Build/HordeServer.Build.Tests/Jobs/Bisect/BisectTests.cs
+++ b/Engine/Source/Programs/Horde/Plugins/Build/HordeServer.Build.Tests/Jobs/Bisect/BisectTests.cs
@@ -2,6 +2,7 @@
 
 using System.Security.Claims;
 using EpicGames.Core;
+using EpicGames.Horde.Commits;
 using EpicGames.Horde.Jobs;
 using EpicGames.Horde.Jobs.Templates;
 using EpicGames.Horde.Projects;
@@ -39,9 +40,9 @@ namespace HordeServer.Tests.Jobs.Bisect
 			Assert.AreEqual("CompileEditor", response.NodeName);
 			Assert.AreEqual(JobStepOutcome.Failure, response.Outcome);
 			Assert.AreEqual("TestUser", response.Owner?.Name ?? "");
-			Assert.AreEqual(failedJob.Change, response.InitialChange);
+			Assert.AreEqual(failedJob.CommitId, response.InitialCommitId);
 			Assert.AreEqual(failedJob.Id.ToString(), response.InitialJobId);
-			Assert.AreEqual(failedJob.Change, response.CurrentChange);
+			Assert.AreEqual(failedJob.CommitId, response.CurrentCommitId);
 			Assert.AreEqual(failedJob.Id.ToString(), response.CurrentJobId);
 
 			IUser user = await UserCollection.FindOrAddUserByLoginAsync("TestUser");
@@ -61,10 +62,10 @@ namespace HordeServer.Tests.Jobs.Bisect
 
 			GetBisectTaskResponse bisectTask = Deref(await BisectTasksController!.GetAsync(task!.BisectTaskId));
 			Assert.AreEqual(BisectTaskState.Running, bisectTask.State);
-			Assert.AreEqual(10, bisectTask.MinChange);
+			Assert.AreEqual(10, bisectTask.MinCommitId!.GetPerforceChange());
 			jobs = await JobCollection.FindBisectTaskJobsAsync(bisectTask.Id, running: true).ToListAsync();
 			Assert.AreEqual(1, jobs.Count);
-			Assert.AreEqual(14, jobs[0].Change);
+			Assert.AreEqual(14, jobs[0].CommitId.GetPerforceChange());
 			await SetJobOutcomeAsync(jobs[0], graph, JobStepOutcome.Failure);
 
 			List<GetBisectTaskResponse> bisectTasks = Deref(await BisectTasksController!.FindBisectTasksAsync(null, bisectTask.Owner?.Id.ToString() ?? ""));
@@ -80,7 +81,7 @@ namespace HordeServer.Tests.Jobs.Bisect
 			Assert.AreEqual(BisectTaskState.Running, bisectTask.State);
 			jobs = await JobCollection.FindBisectTaskJobsAsync(bisectTask.Id, running: true).ToListAsync();
 			Assert.AreEqual(1, jobs.Count);
-			Assert.AreEqual(12, jobs[0].Change);
+			Assert.AreEqual(12, jobs[0].CommitId.GetPerforceChange());
 			await SetJobOutcomeAsync(jobs[0], graph, JobStepOutcome.Success);
 
 			await Clock.AdvanceAsync(TimeSpan.FromMinutes(30));
@@ -89,15 +90,15 @@ namespace HordeServer.Tests.Jobs.Bisect
 			Assert.AreEqual(BisectTaskState.Running, bisectTask.State);
 			jobs = await JobCollection.FindBisectTaskJobsAsync(bisectTask.Id, running: true).ToListAsync();
 			Assert.AreEqual(1, jobs.Count);
-			Assert.AreEqual(13, jobs[0].Change);
+			Assert.AreEqual(13, jobs[0].CommitId.GetPerforceChange());
 			await SetJobOutcomeAsync(jobs[0], graph, JobStepOutcome.Success);
 
 			await Clock.AdvanceAsync(TimeSpan.FromMinutes(30));
 
 			bisectTask = Deref(await BisectTasksController!.GetAsync(task!.BisectTaskId));
 			Assert.AreEqual(BisectTaskState.Succeeded, bisectTask.State);
-			Assert.AreEqual(20, bisectTask.InitialChange);
-			Assert.AreEqual(14, bisectTask.CurrentChange);
+			Assert.AreEqual(20, bisectTask.InitialCommitId.GetPerforceChange());
+			Assert.AreEqual(14, bisectTask.CurrentCommitId.GetPerforceChange());
 			jobs = await JobCollection.FindBisectTaskJobsAsync(bisectTask.Id, running: true).ToListAsync();
 			Assert.AreEqual(0, jobs.Count);
 
@@ -165,7 +166,7 @@ namespace HordeServer.Tests.Jobs.Bisect
 
 		async Task<IJob> CreateJobAsync(int change, IGraph graph, JobStepOutcome outcome, CreateJobOptions options)
 		{
-			IJob job = await JobCollection.AddAsync(JobIdUtils.GenerateNewId(), StreamId, TemplateId, ContentHash.SHA1("hello"), graph, "Test job", change, change, options);
+			IJob job = await JobCollection.AddAsync(JobIdUtils.GenerateNewId(), StreamId, TemplateId, ContentHash.SHA1("hello"), graph, "Test job", CommitIdWithOrder.FromPerforceChange(change), CommitIdWithOrder.FromPerforceChange(change), options);
 			return await SetJobOutcomeAsync(job, graph, outcome);
 		}
 
@@ -184,7 +185,7 @@ namespace HordeServer.Tests.Jobs.Bisect
 
 			INodeGroup group = graph.Groups[batch.GroupIdx];
 			INode node = group.Nodes[step.NodeIdx];
-			await JobStepRefCollection.InsertOrReplaceAsync(new JobStepRefId(job.Id, batch.Id, step.Id), job.Name, node.Name, job.StreamId, job.TemplateId, job.Change, LogIdUtils.GenerateNewId(), null, null, JobStepState.Completed, outcome, false, null, null, 0.0f, 0.0f, DateTime.MinValue, DateTime.MinValue, DateTime.MinValue);
+			await JobStepRefCollection.InsertOrReplaceAsync(new JobStepRefId(job.Id, batch.Id, step.Id), job.Name, node.Name, job.StreamId, job.TemplateId, job.CommitId, LogIdUtils.GenerateNewId(), null, null, JobStepState.Completed, outcome, false, null, null, 0.0f, 0.0f, DateTime.MinValue, DateTime.MinValue, DateTime.MinValue);
 			await BisectTaskCollection.UpdateAsync(job, batch, step, graph);
 
 			return job;
diff --git a/Engine/Source/Programs/Horde/Plugins/Build/HordeServer.Build.Tests/Jobs/JobCollectionTests.cs b/Engine/Source/Programs/Horde/Plugins/Build/HordeServer.Build.Tests/Jobs/JobCollectionTests.cs
--- a/Engine/Source/Programs/Horde/Plugins/Build/HordeServer.Build.Tests/Jobs/JobCollectionTests.cs
+++ b/Engine/Source/Programs/Horde/Plugins/Build/HordeServer.Build.Tests/Jobs/JobCollectionTests.cs
@@ -6,9 +6,11 @@ using EpicGames.Horde.Agents.Leases;
 using EpicGames.Horde.Agents.Pools;
 using EpicGames.Horde.Agents.Sessions;
 using EpicGames.Horde.Artifacts;
+using EpicGames.Horde.Commits;
 using EpicGames.Horde.Jobs;
 using EpicGames.Horde.Jobs.Templates;
 using EpicGames.Horde.Logs;
+using EpicGames.Horde.Projects;
 using EpicGames.Horde.Streams;
 using EpicGames.Horde.Users;
 using HordeServer.Agents.Sessions;
@@ -16,6 +18,8 @@ using HordeServer.Jobs;
 using HordeServer.Jobs.Graphs;
 using HordeServer.Jobs.Templates;
 using HordeServer.Logs;
+using HordeServer.Projects;
+using HordeServer.Streams;
 using HordeServer.Utilities;
 using Moq;
 
@@ -24,6 +28,18 @@ namespace HordeServer.Tests.Jobs
 	[TestClass]
 	public class JobCollectionTests : BuildTestSetup
 	{
+		public JobCollectionTests()
+		{
+			UpdateConfig(globalConfig =>
+			{
+				ProjectConfig projectConfig = new ProjectConfig();
+				projectConfig.Id = new ProjectId("ue4");
+				projectConfig.Streams.Add(new StreamConfig { Id = new StreamId("ue4-main") });
+
+				globalConfig.Plugins.GetBuildConfig().Projects.Add(projectConfig);
+			});
+		}
+
 		static NewGroup AddGroup(List<NewGroup> groups)
 		{
 			NewGroup group = new NewGroup("win64", new List<NewNode>());
@@ -79,7 +95,7 @@ namespace HordeServer.Tests.Jobs
 			options.Arguments.Add("-Target=Publish Client");
 			options.Arguments.Add("-Target=Post-Publish Client");
 
-			IJob job = await JobCollection.AddAsync(JobIdUtils.GenerateNewId(), new StreamId("ue4-main"), new TemplateId("test-build"), ContentHash.SHA1("hello"), baseGraph, "Test job", 123, 123, options);
+			IJob job = await JobCollection.AddAsync(JobIdUtils.GenerateNewId(), new StreamId("ue4-main"), new TemplateId("test-build"), ContentHash.SHA1("hello"), baseGraph, "Test job", CommitIdWithOrder.FromPerforceChange(123), CommitIdWithOrder.FromPerforceChange(123), options);
 
 			job = await RunBatchAsync(job, 0);
 			job = await RunStepAsync(job, 0, 0, JobStepOutcome.Success); // Setup Build
@@ -140,7 +156,7 @@ namespace HordeServer.Tests.Jobs
 			CreateJobOptions options = new CreateJobOptions();
 			options.Arguments.Add("-Target=Gather");
 
-			IJob job = await JobCollection.AddAsync(JobIdUtils.GenerateNewId(), new StreamId("ue4-main"), new TemplateId("test-build"), ContentHash.SHA1("hello"), baseGraph, "Test job", 123, 123, options);
+			IJob job = await JobCollection.AddAsync(JobIdUtils.GenerateNewId(), new StreamId("ue4-main"), new TemplateId("test-build"), ContentHash.SHA1("hello"), baseGraph, "Test job", CommitIdWithOrder.FromPerforceChange(123), CommitIdWithOrder.FromPerforceChange(123), options);
 
 			job = await RunBatchAsync(job, 0);
 			job = await RunStepAsync(job, 0, 0, JobStepOutcome.Success); // Setup Build
@@ -209,7 +225,7 @@ namespace HordeServer.Tests.Jobs
 			options.Arguments.Add("-Target=Gather");
 			options.Arguments.Add("-Target=Initial Node");
 
-			IJob job = await JobCollection.AddAsync(JobIdUtils.GenerateNewId(), new StreamId("ue4-main"), new TemplateId("test-build"), ContentHash.SHA1("hello"), baseGraph, "Test job", 123, 123, options);
+			IJob job = await JobCollection.AddAsync(JobIdUtils.GenerateNewId(), new StreamId("ue4-main"), new TemplateId("test-build"), ContentHash.SHA1("hello"), baseGraph, "Test job", CommitIdWithOrder.FromPerforceChange(123), CommitIdWithOrder.FromPerforceChange(123), options);
 
 			// Try a batch and fail it
 			job = await RunBatchAsync(job, 0);
@@ -272,7 +288,7 @@ namespace HordeServer.Tests.Jobs
 			options.Arguments.Add("-Target=Step 1");
 			options.Arguments.Add("-Target=Step 3");
 
-			IJob job = await JobCollection.AddAsync(JobIdUtils.GenerateNewId(), new StreamId("ue4-main"), new TemplateId("test-build"), ContentHash.SHA1("hello"), baseGraph, "Test job", 123, 123, options);
+			IJob job = await JobCollection.AddAsync(JobIdUtils.GenerateNewId(), new StreamId("ue4-main"), new TemplateId("test-build"), ContentHash.SHA1("hello"), baseGraph, "Test job", CommitIdWithOrder.FromPerforceChange(123), CommitIdWithOrder.FromPerforceChange(123), options);
 
 			job = await RunBatchAsync(job, 0);
 			job = await RunStepAsync(job, 0, 0, JobStepOutcome.Success); // Setup Build
@@ -328,7 +344,7 @@ namespace HordeServer.Tests.Jobs
 			options.Arguments.Add("-Target=Step 1");
 			options.Arguments.Add("-Target=Step 3");
 
-			IJob job = await JobCollection.AddAsync(JobIdUtils.GenerateNewId(), new StreamId("ue4-main"), new TemplateId("test-build"), ContentHash.SHA1("hello"), baseGraph, "Test job", 123, 123, options);
+			IJob job = await JobCollection.AddAsync(JobIdUtils.GenerateNewId(), new StreamId("ue4-main"), new TemplateId("test-build"), ContentHash.SHA1("hello"), baseGraph, "Test job", CommitIdWithOrder.FromPerforceChange(123), CommitIdWithOrder.FromPerforceChange(123), options);
 			Assert.AreEqual(1, job.Batches.Count);
 
 			job = await RunBatchAsync(job, 0);
@@ -379,7 +395,7 @@ namespace HordeServer.Tests.Jobs
 			options.Arguments.Add("-Target=Step 1");
 			options.Arguments.Add("-Target=Step 3");
 
-			IJob job = await JobCollection.AddAsync(JobIdUtils.GenerateNewId(), new StreamId("ue4-main"), new TemplateId("test-build"), ContentHash.SHA1("hello"), baseGraph, "Test job", 123, 123, options);
+			IJob job = await JobCollection.AddAsync(JobIdUtils.GenerateNewId(), new StreamId("ue4-main"), new TemplateId("test-build"), ContentHash.SHA1("hello"), baseGraph, "Test job", CommitIdWithOrder.FromPerforceChange(123), CommitIdWithOrder.FromPerforceChange(123), options);
 			Assert.AreEqual(1, job.Batches.Count);
 
 			// First retry
@@ -461,7 +477,7 @@ namespace HordeServer.Tests.Jobs
 			CreateJobOptions options = new CreateJobOptions();
 			options.Arguments.Add("-Target=Step 2");
 
-			IJob job = await JobCollection.AddAsync(JobIdUtils.GenerateNewId(), new StreamId("ue4-main"), new TemplateId("test-build"), ContentHash.SHA1("hello"), graph, "Test job", 123, 123, options);
+			IJob job = await JobCollection.AddAsync(JobIdUtils.GenerateNewId(), new StreamId("ue4-main"), new TemplateId("test-build"), ContentHash.SHA1("hello"), graph, "Test job", CommitIdWithOrder.FromPerforceChange(123), CommitIdWithOrder.FromPerforceChange(123), options);
 
 			// Fail the first step
 			job = await RunBatchAsync(job, 0);
@@ -563,7 +579,7 @@ namespace HordeServer.Tests.Jobs
 			CreateJobOptions options = new CreateJobOptions();
 			options.Arguments.Add("-Target=Step 2");
 
-			IJob job = await JobCollection.AddAsync(JobIdUtils.GenerateNewId(), new StreamId("ue4-main"), new TemplateId("test-build"), ContentHash.SHA1("hello"), graph, "Test job", 123, 123, options);
+			IJob job = await JobCollection.AddAsync(JobIdUtils.GenerateNewId(), new StreamId("ue4-main"), new TemplateId("test-build"), ContentHash.SHA1("hello"), graph, "Test job", CommitIdWithOrder.FromPerforceChange(123), CommitIdWithOrder.FromPerforceChange(123), options);
 
 			// Fail the first step
 			job = await RunBatchAsync(job, 0);
@@ -626,7 +642,7 @@ namespace HordeServer.Tests.Jobs
 			CreateJobOptions options = new CreateJobOptions();
 			options.Arguments.Add("-Target=Step 3");
 
-			IJob job = await JobCollection.AddAsync(JobIdUtils.GenerateNewId(), new StreamId("ue4-main"), new TemplateId("test-build"), ContentHash.SHA1("hello"), graph, "Test job", 123, 123, options);
+			IJob job = await JobCollection.AddAsync(JobIdUtils.GenerateNewId(), new StreamId("ue4-main"), new TemplateId("test-build"), ContentHash.SHA1("hello"), graph, "Test job", CommitIdWithOrder.FromPerforceChange(123), CommitIdWithOrder.FromPerforceChange(123), options);
 
 			// Fail the first step
 			job = await RunBatchAsync(job, 0);
@@ -695,7 +711,7 @@ namespace HordeServer.Tests.Jobs
 			CreateJobOptions options = new CreateJobOptions();
 			options.Arguments.Add("-Target=Step 3");
 
-			IJob job = await JobCollection.AddAsync(JobIdUtils.GenerateNewId(), new StreamId("ue4-main"), new TemplateId("test-build"), ContentHash.SHA1("hello"), graph, "Test job", 123, 123, options);
+			IJob job = await JobCollection.AddAsync(JobIdUtils.GenerateNewId(), new StreamId("ue4-main"), new TemplateId("test-build"), ContentHash.SHA1("hello"), graph, "Test job", CommitIdWithOrder.FromPerforceChange(123), CommitIdWithOrder.FromPerforceChange(123), options);
 
 			// Pass the first step
 			job = await RunBatchAsync(job, 0);
@@ -747,7 +763,7 @@ namespace HordeServer.Tests.Jobs
 			options.Arguments.Add("-Target=Step 1");
 			options.Arguments.Add("-Target=Step 3");
 
-			IJob job = await JobCollection.AddAsync(JobIdUtils.GenerateNewId(), new StreamId("ue4-main"), new TemplateId("test-build"), ContentHash.SHA1("hello"), baseGraph, "Test job", 123, 123, options);
+			IJob job = await JobCollection.AddAsync(JobIdUtils.GenerateNewId(), new StreamId("ue4-main"), new TemplateId("test-build"), ContentHash.SHA1("hello"), baseGraph, "Test job", CommitIdWithOrder.FromPerforceChange(123), CommitIdWithOrder.FromPerforceChange(123), options);
 			Assert.AreEqual(1, job.Batches.Count);
 
 			job = await RunBatchAsync(job, 0);
@@ -767,7 +783,7 @@ namespace HordeServer.Tests.Jobs
 			CreateJobOptions options = new CreateJobOptions();
 			options.Arguments.Add("-Target=Compile Editor");
 
-			IJob job = await JobCollection.AddAsync(JobIdUtils.GenerateNewId(), new StreamId("ue4-main"), new TemplateId("test-build"), ContentHash.SHA1("hello"), baseGraph, "Test job", 123, 123, options);
+			IJob job = await JobCollection.AddAsync(JobIdUtils.GenerateNewId(), new StreamId("ue4-main"), new TemplateId("test-build"), ContentHash.SHA1("hello"), baseGraph, "Test job", CommitIdWithOrder.FromPerforceChange(123), CommitIdWithOrder.FromPerforceChange(123), options);
 
 			job = await RunBatchAsync(job, 0);
 			job = await RunStepAsync(job, 0, 0, JobStepOutcome.Success); // Setup Build
diff --git a/Engine/Source/Programs/Horde/Plugins/Build/HordeServer.Build.Tests/Jobs/JobExpirationTests.cs b/Engine/Source/Programs/Horde/Plugins/Build/HordeServer.Build.Tests/Jobs/JobExpirationTests.cs
--- a/Engine/Source/Programs/Horde/Plugins/Build/HordeServer.Build.Tests/Jobs/JobExpirationTests.cs
+++ b/Engine/Source/Programs/Horde/Plugins/Build/HordeServer.Build.Tests/Jobs/JobExpirationTests.cs
@@ -1,5 +1,6 @@
 // Copyright Epic Games, Inc. All Rights Reserved.
 
+using EpicGames.Horde.Commits;
 using EpicGames.Horde.Jobs.Templates;
 using EpicGames.Horde.Projects;
 using EpicGames.Horde.Streams;
@@ -38,13 +39,13 @@ namespace HordeServer.Tests.Jobs
 			SetConfig(globalConfig);
 
 			CreateJobOptions options = new CreateJobOptions();
-			options.PreflightChange = 999;
+			options.PreflightCommitId = CommitId.FromPerforceChange(999);
 
 			ITemplate template = await TemplateCollection.GetOrAddAsync(streamConfig.Templates[0]);
 
 			IGraph graph = await GraphCollection.AddAsync(template, null);
 
-			return await JobService.CreateJobAsync(null, streamConfig, templateId, template.Hash, graph, "Hello", 1234, 1233, options);
+			return await JobService.CreateJobAsync(null, streamConfig, templateId, template.Hash, graph, "Hello", CommitIdWithOrder.FromPerforceChange(1234), CommitIdWithOrder.FromPerforceChange(1233), options);
 		}
 
 		[TestMethod]
diff --git a/Engine/Source/Programs/Horde/Plugins/Build/HordeServer.Build.Tests/Jobs/JobServiceTests.cs b/Engine/Source/Programs/Horde/Plugins/Build/HordeServer.Build.Tests/Jobs/JobServiceTests.cs
--- a/Engine/Source/Programs/Horde/Plugins/Build/HordeServer.Build.Tests/Jobs/JobServiceTests.cs
+++ b/Engine/Source/Programs/Horde/Plugins/Build/HordeServer.Build.Tests/Jobs/JobServiceTests.cs
@@ -52,14 +52,14 @@ namespace HordeServer.Tests.Jobs
 			SetConfig(globalConfig);
 
 			CreateJobOptions options = new CreateJobOptions();
-			options.PreflightChange = 999;
+			options.PreflightCommitId = CommitId.FromPerforceChange(999);
 			options.JobTriggers.AddRange(streamConfig.Templates[0].ChainedJobs!);
 
 			ITemplate template = await TemplateCollection.GetOrAddAsync(streamConfig.Templates[0]);
 
 			IGraph graph = await GraphCollection.AddAsync(template, null);
 
-			IJob job = await JobService.CreateJobAsync(null, streamConfig, templateRefId1, template.Hash, graph, "Hello", 1234, 1233, options);
+			IJob job = await JobService.CreateJobAsync(null, streamConfig, templateRefId1, template.Hash, graph, "Hello", CommitIdWithOrder.FromPerforceChange(1234), CommitIdWithOrder.FromPerforceChange(1233), options);
 			Assert.AreEqual(1, job.ChainedJobs.Count);
 
 			job = Deref(await JobService.UpdateBatchAsync(job, job.Batches[0].Id, streamConfig, LogIdUtils.GenerateNewId(), JobStepBatchState.Running));
@@ -72,9 +72,9 @@ namespace HordeServer.Tests.Jobs
 			Assert.IsNotNull(chainedJob);
 			Assert.AreEqual(chainedJob!.Id, job!.ChainedJobs[0].JobId);
 
-			Assert.AreEqual(chainedJob!.Change, job!.Change);
-			Assert.AreEqual(chainedJob!.CodeChange, job!.CodeChange);
-			Assert.AreEqual(chainedJob!.PreflightChange, job!.PreflightChange);
+			Assert.AreEqual(chainedJob!.CommitId, job!.CommitId);
+			Assert.AreEqual(chainedJob!.CodeCommitId, job!.CodeCommitId);
+			Assert.AreEqual(chainedJob!.PreflightCommitId, job!.PreflightCommitId);
 			Assert.AreEqual(chainedJob!.StartedByUserId, job!.StartedByUserId);
 		}
 
@@ -105,19 +105,19 @@ namespace HordeServer.Tests.Jobs
 			ICommitCollection commits = PerforceService.GetCommits(streamConfig);
 
 			{
-				int? change = await JobService.EvaluateChangeQueryAsync(streamId, new ChangeQueryConfig { CommitTag = CommitTag.Code }, null, commits, CancellationToken.None);
-				Assert.AreEqual(1002, change);
+				CommitIdWithOrder? change = await JobService.EvaluateChangeQueryAsync(streamId, new ChangeQueryConfig { CommitTag = CommitTag.Code }, null, commits, CancellationToken.None);
+				Assert.AreEqual(1002, change!.GetPerforceChange());
 			}
 			{
-				int? change = await JobService.EvaluateChangeQueryAsync(streamId, new ChangeQueryConfig { CommitTag = CommitTag.Content }, null, commits, CancellationToken.None);
-				Assert.AreEqual(1004, change);
+				CommitIdWithOrder? change = await JobService.EvaluateChangeQueryAsync(streamId, new ChangeQueryConfig { CommitTag = CommitTag.Content }, null, commits, CancellationToken.None);
+				Assert.AreEqual(1004, change!.GetPerforceChange());
 			}
 			{
-				int? change = await JobService.EvaluateChangeQueryAsync(streamId, new ChangeQueryConfig { CommitTag = CommitTag.Content, Condition = "tag.code == 1" }, new List<CommitTag> { CommitTag.Code }, commits, CancellationToken.None);
-				Assert.AreEqual(1004, change);
+				CommitIdWithOrder? change = await JobService.EvaluateChangeQueryAsync(streamId, new ChangeQueryConfig { CommitTag = CommitTag.Content, Condition = "tag.code == 1" }, new List<CommitTag> { CommitTag.Code }, commits, CancellationToken.None);
+				Assert.AreEqual(1004, change!.GetPerforceChange());
 			}
 			{
-				int? change = await JobService.EvaluateChangeQueryAsync(streamId, new ChangeQueryConfig { CommitTag = CommitTag.Content, Condition = "tag.content == 1" }, new List<CommitTag> { CommitTag.Code }, commits, CancellationToken.None);
+				CommitIdWithOrder? change = await JobService.EvaluateChangeQueryAsync(streamId, new ChangeQueryConfig { CommitTag = CommitTag.Content, Condition = "tag.content == 1" }, new List<CommitTag> { CommitTag.Code }, commits, CancellationToken.None);
 				Assert.IsNull(change);
 			}
 		}
@@ -187,7 +187,7 @@ namespace HordeServer.Tests.Jobs
 			IUser user = await UserCollection.FindOrAddUserByLoginAsync(startedByUserName);
 
 			CreateJobOptions options = new CreateJobOptions();
-			options.PreflightChange = preflightChange;
+			options.PreflightCommitId = CommitId.FromPerforceChange(preflightChange);
 			options.StartedByUserId = user.Id;
 			options.Arguments.AddRange(arguments);
 
@@ -198,8 +198,8 @@ namespace HordeServer.Tests.Jobs
 				templateHash: new ContentHash(Encoding.ASCII.GetBytes(templateHash)),
 				graph: fixture!.Graph,
 				name: "hello1",
-				change: 1000001,
-				codeChange: 1000002,
+				commitId: CommitIdWithOrder.FromPerforceChange(1000001),
+				codeCommitId: CommitIdWithOrder.FromPerforceChange(1000002),
 				options
 			);
 		}
@@ -268,10 +268,10 @@ namespace HordeServer.Tests.Jobs
 			graph = await GraphCollection.AppendAsync(graph, new List<NewGroup> { groupA, groupB });
 
 			CreateJobOptions options = new CreateJobOptions();
-			options.PreflightChange = 999;
+			options.PreflightCommitId = CommitId.FromPerforceChange(999);
 			options.Arguments.Add("-Target=Pak");
 
-			IJob job = await JobService.CreateJobAsync(null, streamConfig!, new TemplateId("temp"), new ContentHash(new byte[] { 1, 2, 3 }), graph, "Hello", 1234, 1233, options);
+			IJob job = await JobService.CreateJobAsync(null, streamConfig!, new TemplateId("temp"), new ContentHash(new byte[] { 1, 2, 3 }), graph, "Hello", CommitIdWithOrder.FromPerforceChange(1234), CommitIdWithOrder.FromPerforceChange(1233), options);
 
 			job = Deref(await JobService.UpdateBatchAsync(job, job.Batches[0].Id, streamConfig, LogIdUtils.GenerateNewId(), JobStepBatchState.Running));
 			job = Deref(await JobService.UpdateStepAsync(job, job.Batches[0].Id, job.Batches[0].Steps[0].Id, streamConfig, JobStepState.Running));
diff --git a/Engine/Source/Programs/Horde/Plugins/Build/HordeServer.Build.Tests/Jobs/JobTaskSourceTests.cs b/Engine/Source/Programs/Horde/Plugins/Build/HordeServer.Build.Tests/Jobs/JobTaskSourceTests.cs
--- a/Engine/Source/Programs/Horde/Plugins/Build/HordeServer.Build.Tests/Jobs/JobTaskSourceTests.cs
+++ b/Engine/Source/Programs/Horde/Plugins/Build/HordeServer.Build.Tests/Jobs/JobTaskSourceTests.cs
@@ -2,6 +2,7 @@
 
 using EpicGames.Horde.Agents;
 using EpicGames.Horde.Agents.Pools;
+using EpicGames.Horde.Commits;
 using EpicGames.Horde.Jobs;
 using EpicGames.Horde.Users;
 using HordeServer.Agents;
@@ -125,7 +126,7 @@ namespace HordeServer.Tests.Jobs
 
 			IJob job = await JobCollection.AddAsync(JobIdUtils.GenerateNewId(), stream.Id,
 				fixture.TemplateRefId1, fixture.Template.Hash, graph, "Test Paused Step Job",
-				1000, 1000, options);
+				CommitIdWithOrder.FromPerforceChange(1000), CommitIdWithOrder.FromPerforceChange(1000), options);
 
 			// validate
 			await JobTaskSource.TickAsync(CancellationToken.None);
diff --git a/Engine/Source/Programs/Horde/Plugins/Build/HordeServer.Build.Tests/Jobs/JobsControllerDbTest.cs b/Engine/Source/Programs/Horde/Plugins/Build/HordeServer.Build.Tests/Jobs/JobsControllerDbTest.cs
--- a/Engine/Source/Programs/Horde/Plugins/Build/HordeServer.Build.Tests/Jobs/JobsControllerDbTest.cs
+++ b/Engine/Source/Programs/Horde/Plugins/Build/HordeServer.Build.Tests/Jobs/JobsControllerDbTest.cs
@@ -22,7 +22,7 @@ namespace HordeServer.Tests.Jobs
 			ActionResult<List<object>> res = await JobsController.FindJobsAsync();
 
 			List<GetJobResponse> responses = res.Value!.ConvertAll(x => (GetJobResponse)x);
-			responses.SortBy(x => x.Change);
+			responses.SortBy(x => x.CommitId);
 
 			Assert.AreEqual(2, responses.Count);
 			Assert.AreEqual("hello1", responses[0].Name);
diff --git a/Engine/Source/Programs/Horde/Plugins/Build/HordeServer.Build.Tests/Jobs/Schedules/SchedulerTests.cs b/Engine/Source/Programs/Horde/Plugins/Build/HordeServer.Build.Tests/Jobs/Schedules/SchedulerTests.cs
--- a/Engine/Source/Programs/Horde/Plugins/Build/HordeServer.Build.Tests/Jobs/Schedules/SchedulerTests.cs
+++ b/Engine/Source/Programs/Horde/Plugins/Build/HordeServer.Build.Tests/Jobs/Schedules/SchedulerTests.cs
@@ -99,9 +99,9 @@ namespace HordeServer.Tests.Jobs.Schedules
 		{
 			List<IJob> jobs = await FileTestHelperAsync("....cpp");
 			Assert.AreEqual(3, jobs.Count);
-			Assert.AreEqual(100, jobs[0].Change);
-			Assert.AreEqual(103, jobs[1].Change);
-			Assert.AreEqual(104, jobs[2].Change);
+			Assert.AreEqual(100, jobs[0].CommitId.GetPerforceChange());
+			Assert.AreEqual(103, jobs[1].CommitId.GetPerforceChange());
+			Assert.AreEqual(104, jobs[2].CommitId.GetPerforceChange());
 		}
 
 		[TestMethod]
@@ -109,8 +109,8 @@ namespace HordeServer.Tests.Jobs.Schedules
 		{
 			List<IJob> jobs = await FileTestHelperAsync("/foo/...");
 			Assert.AreEqual(2, jobs.Count);
-			Assert.AreEqual(103, jobs[0].Change);
-			Assert.AreEqual(105, jobs[1].Change);
+			Assert.AreEqual(103, jobs[0].CommitId.GetPerforceChange());
+			Assert.AreEqual(105, jobs[1].CommitId.GetPerforceChange());
 		}
 
 		[TestMethod]
@@ -118,9 +118,9 @@ namespace HordeServer.Tests.Jobs.Schedules
 		{
 			List<IJob> jobs = await FileTestHelperAsync("....uasset", "-/bar/...");
 			Assert.AreEqual(3, jobs.Count);
-			Assert.AreEqual(101, jobs[0].Change);
-			Assert.AreEqual(102, jobs[1].Change);
-			Assert.AreEqual(105, jobs[2].Change);
+			Assert.AreEqual(101, jobs[0].CommitId.GetPerforceChange());
+			Assert.AreEqual(102, jobs[1].CommitId.GetPerforceChange());
+			Assert.AreEqual(105, jobs[2].CommitId.GetPerforceChange());
 		}
 
 		[TestMethod]
@@ -215,8 +215,8 @@ namespace HordeServer.Tests.Jobs.Schedules
 
 			List<IJob> jobs2 = await GetNewJobsAsync();
 			Assert.AreEqual(1, jobs2.Count);
-			Assert.AreEqual(102, jobs2[0].Change);
-			Assert.AreEqual(100, jobs2[0].CodeChange);
+			Assert.AreEqual(102, jobs2[0].CommitId.GetPerforceChange());
+			Assert.AreEqual(100, jobs2[0].CodeCommitId!.GetPerforceChange());
 		}
 
 		[TestMethod]
@@ -243,15 +243,15 @@ namespace HordeServer.Tests.Jobs.Schedules
 
 			List<IJob> jobs2 = await GetNewJobsAsync();
 			Assert.AreEqual(1, jobs2.Count);
-			Assert.AreEqual(102, jobs2[0].Change);
-			Assert.AreEqual(100, jobs2[0].CodeChange);
+			Assert.AreEqual(102, jobs2[0].CommitId.GetPerforceChange());
+			Assert.AreEqual(100, jobs2[0].CodeCommitId!.GetPerforceChange());
 
 			StreamConfig? streamConfig;
 			GlobalConfig.CurrentValue.Plugins.GetBuildConfig().TryGetStream(StreamId, out streamConfig);
 
 			IStream stream2 = (await StreamCollection.GetAsync(streamConfig!))!;
 			ITemplateSchedule schedule2 = stream2.Templates.First().Value.Schedule!;
-			Assert.AreEqual(102, schedule2.LastTriggerChange);
+			Assert.AreEqual(102, schedule2.LastTriggerCommitId!.GetPerforceChange());
 			Assert.AreEqual(Clock.UtcNow, schedule2.LastTriggerTimeUtc);
 
 			// Trigger another job
@@ -294,10 +294,10 @@ namespace HordeServer.Tests.Jobs.Schedules
 
 			List<IJob> jobs2 = await GetNewJobsAsync();
 			Assert.AreEqual(2, jobs2.Count);
-			Assert.AreEqual(104, jobs2[0].Change);
-			Assert.AreEqual(104, jobs2[0].CodeChange);
-			Assert.AreEqual(106, jobs2[1].Change);
-			Assert.AreEqual(106, jobs2[1].CodeChange);
+			Assert.AreEqual(104, jobs2[0].CommitId.GetPerforceChange());
+			Assert.AreEqual(104, jobs2[0].CodeCommitId!.GetPerforceChange());
+			Assert.AreEqual(106, jobs2[1].CommitId.GetPerforceChange());
+			Assert.AreEqual(106, jobs2[1].CodeCommitId!.GetPerforceChange());
 		}
 
 		[TestMethod]
@@ -332,10 +332,10 @@ namespace HordeServer.Tests.Jobs.Schedules
 
 			List<IJob> jobs2 = await GetNewJobsAsync();
 			Assert.AreEqual(2, jobs2.Count);
-			Assert.AreEqual(103, jobs2[0].Change);
-			Assert.AreEqual(103, jobs2[0].CodeChange);
-			Assert.AreEqual(106, jobs2[1].Change);
-			Assert.AreEqual(106, jobs2[1].CodeChange);
+			Assert.AreEqual(103, jobs2[0].CommitId.GetPerforceChange());
+			Assert.AreEqual(103, jobs2[0].CodeCommitId!.GetPerforceChange());
+			Assert.AreEqual(106, jobs2[1].CommitId.GetPerforceChange());
+			Assert.AreEqual(106, jobs2[1].CodeCommitId!.GetPerforceChange());
 		}
 
 		[TestMethod]
@@ -358,8 +358,8 @@ namespace HordeServer.Tests.Jobs.Schedules
 
 			List<IJob> jobs2 = await GetNewJobsAsync();
 			Assert.AreEqual(1, jobs2.Count);
-			Assert.AreEqual(102, jobs2[0].Change);
-			Assert.AreEqual(100, jobs2[0].CodeChange);
+			Assert.AreEqual(102, jobs2[0].CommitId.GetPerforceChange());
+			Assert.AreEqual(100, jobs2[0].CodeCommitId!.GetPerforceChange());
 
 			// Test that another job does not trigger
 			await Clock.AdvanceAsync(TimeSpan.FromHours(0.5));
@@ -377,8 +377,8 @@ namespace HordeServer.Tests.Jobs.Schedules
 
 			List<IJob> jobs4 = await GetNewJobsAsync();
 			Assert.AreEqual(1, jobs4.Count);
-			Assert.AreEqual(102, jobs4[0].Change);
-			Assert.AreEqual(100, jobs4[0].CodeChange);
+			Assert.AreEqual(102, jobs4[0].CommitId.GetPerforceChange());
+			Assert.AreEqual(100, jobs4[0].CodeCommitId!.GetPerforceChange());
 		}
 
 		[TestMethod]
@@ -399,8 +399,8 @@ namespace HordeServer.Tests.Jobs.Schedules
 
 			List<IJob> jobs2 = await GetNewJobsAsync();
 			Assert.AreEqual(1, jobs2.Count);
-			Assert.AreEqual(102, jobs2[0].Change);
-			Assert.AreEqual(100, jobs2[0].CodeChange);
+			Assert.AreEqual(102, jobs2[0].CommitId.GetPerforceChange());
+			Assert.AreEqual(100, jobs2[0].CodeCommitId!.GetPerforceChange());
 
 			// Check another job does not trigger due to the change above
 			await Clock.AdvanceAsync(TimeSpan.FromHours(1.25));
@@ -450,10 +450,10 @@ namespace HordeServer.Tests.Jobs.Schedules
 
 			// Create a job and fail it
 			CreateJobOptions options1 = new CreateJobOptions();
-			options1.PreflightChange = 999;
+			options1.PreflightCommitId = CommitId.FromPerforceChange(999);
 			options1.Arguments.Add("-Target=TriggerNext");
 
-			IJob job1 = await JobService.CreateJobAsync(null, config, newTemplateRefId1, newTemplate1.Hash, graphA, "Hello", 1234, 1233, options1);
+			IJob job1 = await JobService.CreateJobAsync(null, config, newTemplateRefId1, newTemplate1.Hash, graphA, "Hello", CommitIdWithOrder.FromPerforceChange(1234), CommitIdWithOrder.FromPerforceChange(1233), options1);
 			JobStepBatchId batchId1 = job1.Batches[0].Id;
 			JobStepId stepId1 = job1.Batches[0].Steps[0].Id;
 			job1 = Deref(await JobService.UpdateBatchAsync(job1, batchId1, config, LogIdUtils.GenerateNewId(), JobStepBatchState.Running));
@@ -470,10 +470,10 @@ namespace HordeServer.Tests.Jobs.Schedules
 
 			// Create a job and make it succeed
 			CreateJobOptions options2 = new CreateJobOptions();
-			options2.PreflightChange = 999;
+			options2.PreflightCommitId = CommitId.FromPerforceChange(999);
 			options2.Arguments.Add("-Target=TriggerNext");
 
-			IJob job2 = await JobService.CreateJobAsync(null, config, newTemplateRefId1, newTemplate1.Hash, graphA, "Hello", 1234, 1233, options2);
+			IJob job2 = await JobService.CreateJobAsync(null, config, newTemplateRefId1, newTemplate1.Hash, graphA, "Hello", CommitIdWithOrder.FromPerforceChange(1234), CommitIdWithOrder.FromPerforceChange(1233), options2);
 			JobStepBatchId batchId2 = job2.Batches[0].Id;
 			JobStepId stepId2 = job2.Batches[0].Steps[0].Id;
 			job2 = Deref(await JobService.UpdateBatchAsync(job2, batchId2, config, LogIdUtils.GenerateNewId(), JobStepBatchState.Running));
@@ -484,8 +484,8 @@ namespace HordeServer.Tests.Jobs.Schedules
 			await ScheduleService.TickForTestingAsync();
 			List<IJob> jobs4 = await GetNewJobsAsync();
 			Assert.AreEqual(1, jobs4.Count);
-			Assert.AreEqual(1234, jobs4[0].Change);
-			Assert.AreEqual(1233, jobs4[0].CodeChange);
+			Assert.AreEqual(1234, jobs4[0].CommitId.GetPerforceChange());
+			Assert.AreEqual(1233, jobs4[0].CodeCommitId!.GetPerforceChange());
 		}
 
 		[TestMethod]
@@ -550,7 +550,7 @@ namespace HordeServer.Tests.Jobs.Schedules
 				CreateJobOptions options1 = new CreateJobOptions();
 				options1.Arguments.Add("-Target=TriggerNext");
 
-				IJob job1 = await JobService.CreateJobAsync(null, config, newTemplateRefId1, _template.Hash, graphA, "Hello", change, codeChange, options1);
+				IJob job1 = await JobService.CreateJobAsync(null, config, newTemplateRefId1, _template.Hash, graphA, "Hello", CommitIdWithOrder.FromPerforceChange(change), CommitIdWithOrder.FromPerforceChange(codeChange), options1);
 				for (int batchIdx = 0; batchIdx < job1.Batches.Count; batchIdx++)
 				{
 					JobStepBatchId batchId1 = job1.Batches[batchIdx].Id;
@@ -570,8 +570,8 @@ namespace HordeServer.Tests.Jobs.Schedules
 			await ScheduleService.TriggerAsync(StreamId, newTemplateRefId2, Clock.UtcNow, default);
 			List<IJob> jobs3 = await GetNewJobsAsync();
 			Assert.AreEqual(2, jobs3.Count);
-			Assert.AreEqual(1230, jobs3[0].Change);
-			Assert.AreEqual(1233, jobs3[1].Change);
+			Assert.AreEqual(1230, jobs3[0].CommitId.GetPerforceChange());
+			Assert.AreEqual(1233, jobs3[1].CommitId.GetPerforceChange());
 		}
 
 		[TestMethod]
@@ -593,8 +593,8 @@ namespace HordeServer.Tests.Jobs.Schedules
 
 			List<IJob> jobs1 = await GetNewJobsAsync();
 			Assert.AreEqual(1, jobs1.Count);
-			Assert.AreEqual(102, jobs1[0].Change);
-			Assert.AreEqual(100, jobs1[0].CodeChange);
+			Assert.AreEqual(102, jobs1[0].CommitId.GetPerforceChange());
+			Assert.AreEqual(100, jobs1[0].CodeCommitId!.GetPerforceChange());
 
 			// Make sure the job is registered
 			IStream? stream1 = await StreamCollection.GetAsync(GlobalConfig.CurrentValue.Plugins.GetBuildConfig().Streams[0]);
@@ -638,8 +638,8 @@ namespace HordeServer.Tests.Jobs.Schedules
 
 			List<IJob> jobs3 = await GetNewJobsAsync();
 			Assert.AreEqual(1, jobs3.Count);
-			Assert.AreEqual(102, jobs3[0].Change);
-			Assert.AreEqual(100, jobs3[0].CodeChange);
+			Assert.AreEqual(102, jobs3[0].CommitId.GetPerforceChange());
+			Assert.AreEqual(100, jobs3[0].CodeCommitId!.GetPerforceChange());
 		}
 
 		[TestMethod]
@@ -722,7 +722,7 @@ namespace HordeServer.Tests.Jobs.Schedules
 			List<IJob> jobs = (await JobCollection.FindAsync()).ToList();
 			jobs.RemoveAll(x => _initialJobIds.Contains(x.Id));
 			_initialJobIds.UnionWith(jobs.Select(x => x.Id));
-			return jobs.OrderBy(x => x.Change).ToList();
+			return jobs.OrderBy(x => x.CommitId).ToList();
 		}
 	}
 }
diff --git a/Engine/Source/Programs/Horde/Plugins/Build/HordeServer.Build.Tests/Jobs/TestData/TestDataTests.cs b/Engine/Source/Programs/Horde/Plugins/Build/HordeServer.Build.Tests/Jobs/TestData/TestDataTests.cs
--- a/Engine/Source/Programs/Horde/Plugins/Build/HordeServer.Build.Tests/Jobs/TestData/TestDataTests.cs
+++ b/Engine/Source/Programs/Horde/Plugins/Build/HordeServer.Build.Tests/Jobs/TestData/TestDataTests.cs
@@ -2,6 +2,7 @@
 
 using System.Security.Claims;
 using EpicGames.Horde;
+using EpicGames.Horde.Commits;
 using EpicGames.Horde.Jobs;
 using EpicGames.Horde.Jobs.Templates;
 using EpicGames.Horde.Jobs.TestData;
@@ -146,8 +147,8 @@ namespace HordeServer.Tests.Jobs.TestData
 			job.SetupGet(x => x.Name).Returns(name);
 			job.SetupGet(x => x.StreamId).Returns(streamId);
 			job.SetupGet(x => x.TemplateId).Returns(new TemplateId("test-template"));
-			job.SetupGet(x => x.Change).Returns(change);
-			job.SetupGet(x => x.PreflightChange).Returns(0);
+			job.SetupGet(x => x.CommitId).Returns(CommitIdWithOrder.FromPerforceChange(change));
+			job.SetupGet(x => x.PreflightCommitId).Returns(default(CommitId));
 			job.SetupGet(x => x.Batches).Returns(batches);
 			job.SetupGet(x => x.ShowUgsBadges).Returns(false);
 			job.SetupGet(x => x.ShowUgsAlerts).Returns(false);
diff --git a/Engine/Source/Programs/Horde/Plugins/Build/HordeServer.Build.Tests/Notifications/NotificationServiceTests.cs b/Engine/Source/Programs/Horde/Plugins/Build/HordeServer.Build.Tests/Notifications/NotificationServiceTests.cs
--- a/Engine/Source/Programs/Horde/Plugins/Build/HordeServer.Build.Tests/Notifications/NotificationServiceTests.cs
+++ b/Engine/Source/Programs/Horde/Plugins/Build/HordeServer.Build.Tests/Notifications/NotificationServiceTests.cs
@@ -1,5 +1,6 @@
 // Copyright Epic Games, Inc. All Rights Reserved.
 
+using EpicGames.Horde.Commits;
 using EpicGames.Horde.Jobs;
 using EpicGames.Horde.Streams;
 using HordeServer.Agents;
@@ -89,7 +90,7 @@ namespace HordeServer.Tests.Notifications
 			job.SetupGet(x => x.Id).Returns(jobId);
 			job.SetupGet(x => x.Name).Returns(name);
 			job.SetupGet(x => x.StreamId).Returns(streamId);
-			job.SetupGet(x => x.Change).Returns(change);
+			job.SetupGet(x => x.CommitId).Returns(CommitIdWithOrder.FromPerforceChange(change));
 			job.SetupGet(x => x.Batches).Returns(batches);
 			return job.Object;
 		}
diff --git a/Engine/Source/Programs/Horde/Plugins/Build/HordeServer.Build.Tests/Stubs/Services/PerforceServiceStub.cs b/Engine/Source/Programs/Horde/Plugins/Build/HordeServer.Build.Tests/Stubs/Services/PerforceServiceStub.cs
--- a/Engine/Source/Programs/Horde/Plugins/Build/HordeServer.Build.Tests/Stubs/Services/PerforceServiceStub.cs
+++ b/Engine/Source/Programs/Horde/Plugins/Build/HordeServer.Build.Tests/Stubs/Services/PerforceServiceStub.cs
@@ -34,8 +34,13 @@ namespace HordeServer.Tests.Stubs.Services
 		public class Commit : ICommit
 		{
 			public StreamId StreamId { get; }
+			
 			public int Number { get; }
+			CommitIdWithOrder ICommit.Id => CommitIdWithOrder.FromPerforceChange(Number);
+	
 			public int OriginalChange { get; }
+			CommitIdWithOrder ICommit.OriginalCommitId => CommitIdWithOrder.FromPerforceChange(OriginalChange);
+
 			public UserId AuthorId { get; }
 			public UserId OwnerId { get; }
 			public string Description { get; }
@@ -123,10 +128,19 @@ namespace HordeServer.Tests.Stubs.Services
 			streamChanges.Add(number, new Commit(streamId, number, originalNumber, author.Id, owner.Id, description, files.ToList()));
 		}
 
-		public Task<List<ICommit>> GetChangesAsync(StreamConfig stream, int? minChange, int? maxChange, int? numResults, CancellationToken cancellationToken)
+		public Task<List<ICommit>> GetChangesAsync(StreamConfig stream, int? minChange, bool includeMinChange, int? maxChange, bool includeMaxChange, int? numResults, CancellationToken cancellationToken)
 		{
 			_ = cancellationToken;
 
+			if (minChange != null && !includeMinChange)
+			{
+				minChange = minChange.Value + 1;
+			}
+			if (maxChange != null && !includeMaxChange)
+			{
+				maxChange = maxChange.Value - 1;
+			}
+
 			List<ICommit> results = new List<ICommit>();
 
 			SortedDictionary<int, Commit>? streamChanges;
@@ -220,14 +234,14 @@ namespace HordeServer.Tests.Stubs.Services
 				_streamConfig = stream;
 			}
 
-			public Task<int> CreateNewAsync(string path, string description, CancellationToken cancellationToken = default)
+			public Task<CommitIdWithOrder> CreateNewAsync(string path, string description, CancellationToken cancellationToken = default)
 			{
 				throw new NotImplementedException();
 			}
 
-			public async IAsyncEnumerable<ICommit> FindAsync(int? minChange, int? maxChange, int? maxResults, IReadOnlyList<CommitTag>? tags, [EnumeratorCancellation] CancellationToken cancellationToken = default)
+			public async IAsyncEnumerable<ICommit> FindAsync(CommitId? minChange, bool includeMinChange, CommitId? maxChange, bool includeMaxChange, int? maxResults, IReadOnlyList<CommitTag>? tags, [EnumeratorCancellation] CancellationToken cancellationToken = default)
 			{
-				foreach (ICommit commit in await _owner.GetChangesAsync(_streamConfig, minChange, maxChange, null, cancellationToken))
+				foreach (ICommit commit in await _owner.GetChangesAsync(_streamConfig, minChange?.GetPerforceChange(), includeMinChange, maxChange?.GetPerforceChange(), includeMaxChange, null, cancellationToken))
 				{
 					if (tags != null && tags.Count > 0)
 					{
@@ -241,19 +255,22 @@ namespace HordeServer.Tests.Stubs.Services
 				}
 			}
 
-			public async Task<ICommit> GetAsync(int changeNumber, CancellationToken cancellationToken = default)
+			public async Task<ICommit> GetAsync(CommitId changeNumber, CancellationToken cancellationToken = default)
 			{
-				List<ICommit> commits = await _owner.GetChangeDetailsAsync(_streamConfig, new[] { changeNumber }, cancellationToken);
+				List<ICommit> commits = await _owner.GetChangeDetailsAsync(_streamConfig, new[] { changeNumber.GetPerforceChange() }, cancellationToken);
 				return commits[0];
 			}
 
-			public async Task<int> LatestNumberAsync(CancellationToken cancellationToken = default)
+			public ValueTask<CommitIdWithOrder> GetOrderedAsync(CommitId commit, CancellationToken cancellationToken = default)
+				=> ValueTask.FromResult(CommitIdWithOrder.FromPerforceChange(commit.GetPerforceChange()));
+
+			public async Task<CommitIdWithOrder> LatestNumberAsync(CancellationToken cancellationToken = default)
 			{
-				List<ICommit> commits = await _owner.GetChangesAsync(_streamConfig, null, null, 1, cancellationToken);
-				return commits[0].Number;
+				List<ICommit> commits = await _owner.GetChangesAsync(_streamConfig, null, true, null, true, 1, cancellationToken);
+				return commits[0].Id;
 			}
 
-			public IAsyncEnumerable<ICommit> SubscribeAsync(int minChange, IReadOnlyList<CommitTag>? tags = null, CancellationToken cancellationToken = default)
+			public IAsyncEnumerable<ICommit> SubscribeAsync(CommitId minChange, IReadOnlyList<CommitTag>? tags = null, CancellationToken cancellationToken = default)
 			{
 				throw new NotImplementedException();
 			}
diff --git a/Engine/Source/Programs/Horde/Plugins/Build/HordeServer.Build/Artifacts/ArtifactCollection.cs b/Engine/Source/Programs/Horde/Plugins/Build/HordeServer.Build/Artifacts/ArtifactCollection.cs
--- a/Engine/Source/Programs/Horde/Plugins/Build/HordeServer.Build/Artifacts/ArtifactCollection.cs
+++ b/Engine/Source/Programs/Horde/Plugins/Build/HordeServer.Build/Artifacts/ArtifactCollection.cs
@@ -4,8 +4,10 @@ using System.Runtime.CompilerServices;
 using EpicGames.Core;
 using EpicGames.Horde.Acls;
 using EpicGames.Horde.Artifacts;
+using EpicGames.Horde.Commits;
 using EpicGames.Horde.Storage;
 using EpicGames.Horde.Streams;
+using HordeServer.Commits;
 using HordeServer.Server;
 using HordeServer.Storage;
 using HordeServer.Utilities;
@@ -37,8 +39,18 @@ namespace HordeServer.Artifacts
 			[BsonElement("str")]
 			public StreamId StreamId { get; set; }
 
+			[BsonElement("com")]
+			public string? CommitName { get; set; }
+
 			[BsonElement("chg")]
-			public int Change { get; set; }
+			public int CommitOrder { get; set; } // Was P4 changelist number
+
+			[BsonIgnore]
+			public CommitIdWithOrder CommitId
+			{
+				get => (CommitName != null)? new CommitIdWithOrder(CommitName, CommitOrder) : CommitIdWithOrder.FromPerforceChange(CommitOrder);
+				set => (CommitName, CommitOrder) = (value.Name, value.Order);
+			}
 
 			[BsonElement("key")]
 			public List<string> Keys { get; set; } = new List<string>();
@@ -72,14 +84,14 @@ namespace HordeServer.Artifacts
 			{
 			}
 
-			public Artifact(ArtifactId id, ArtifactName name, ArtifactType type, string? description, StreamId streamId, int change, IEnumerable<string> keys, IEnumerable<string> metadata, NamespaceId namespaceId, RefName refName, DateTime createdAtUtc, AclScopeName scopeName)
+			public Artifact(ArtifactId id, ArtifactName name, ArtifactType type, string? description, StreamId streamId, CommitIdWithOrder commitId, IEnumerable<string> keys, IEnumerable<string> metadata, NamespaceId namespaceId, RefName refName, DateTime createdAtUtc, AclScopeName scopeName)
 			{
 				Id = id;
 				Name = name;
 				Type = type;
 				Description = description;
 				StreamId = streamId;
-				Change = change;
+				CommitId = commitId;
 				Keys.AddRange(keys);
 				Metadata.AddRange(metadata);
 				NamespaceId = namespaceId;
@@ -91,19 +103,21 @@ namespace HordeServer.Artifacts
 
 		readonly IMongoCollection<Artifact> _artifacts;
 		readonly IClock _clock;
+		readonly ICommitService _commitService;
 
 		/// <summary>
 		/// Constructor
 		/// </summary>
-		public ArtifactCollection(IMongoService mongoService, IClock clock)
+		public ArtifactCollection(IMongoService mongoService, IClock clock, ICommitService commitService)
 		{
 			List<MongoIndex<Artifact>> indexes = new List<MongoIndex<Artifact>>();
 			indexes.Add(keys => keys.Ascending(x => x.Keys));
 			indexes.Add(keys => keys.Ascending(x => x.Type).Descending(x => x.Id));
-			indexes.Add(keys => keys.Ascending(x => x.StreamId).Descending(x => x.Change).Ascending(x => x.Name).Descending(x => x.Id));
+			indexes.Add(keys => keys.Ascending(x => x.StreamId).Descending(x => x.CommitOrder).Ascending(x => x.Name).Descending(x => x.Id));
 			_artifacts = mongoService.GetCollection<Artifact>("ArtifactsV2", indexes);
 
 			_clock = clock;
+			_commitService = commitService;
 		}
 
 #pragma warning disable CA1308 // Expect ansi-only keys here
@@ -117,7 +131,7 @@ namespace HordeServer.Artifacts
 		public static string GetArtifactPath(StreamId streamId, ArtifactName name, ArtifactType type) => $"{streamId}/{name}/{type}";
 
 		/// <inheritdoc/>
-		public async Task<IArtifact> AddAsync(ArtifactName name, ArtifactType type, string? description, StreamId streamId, int change, IEnumerable<string> keys, IEnumerable<string> metadata, AclScopeName scopeName, CancellationToken cancellationToken)
+		public async Task<IArtifact> AddAsync(ArtifactName name, ArtifactType type, string? description, StreamId streamId, CommitId commitId, IEnumerable<string> keys, IEnumerable<string> metadata, AclScopeName scopeName, CancellationToken cancellationToken)
 		{
 			if (name.Id.IsEmpty)
 			{
@@ -131,9 +145,11 @@ namespace HordeServer.Artifacts
 			ArtifactId id = new ArtifactId(BinaryIdUtils.CreateNew());
 
 			NamespaceId namespaceId = Namespace.Artifacts;
-			RefName refName = new RefName($"{GetArtifactPath(streamId, name, type)}/{change}/{id}");
+			RefName refName = new RefName($"{GetArtifactPath(streamId, name, type)}/{commitId}/{id}");
 
-			Artifact artifact = new Artifact(id, name, type, description, streamId, change, keys.Select(x => NormalizeKey(x)), metadata, namespaceId, refName, _clock.UtcNow, scopeName);
+			CommitIdWithOrder commitIdWithOrder = await _commitService.GetOrderedAsync(streamId, commitId, cancellationToken);
+
+			Artifact artifact = new Artifact(id, name, type, description, streamId, commitIdWithOrder, keys.Select(x => NormalizeKey(x)), metadata, namespaceId, refName, _clock.UtcNow, scopeName);
 			await _artifacts.InsertOneAsync(artifact, null, cancellationToken);
 			return artifact;
 		}
@@ -146,20 +162,22 @@ namespace HordeServer.Artifacts
 		}
 
 		/// <inheritdoc/>
-		public async IAsyncEnumerable<IArtifact> FindAsync(StreamId? streamId = null, int? minChange = null, int? maxChange = null, ArtifactName? name = null, ArtifactType? type = null, IEnumerable<string>? keys = null, int maxResults = 100, [EnumeratorCancellation] CancellationToken cancellationToken = default)
+		public async IAsyncEnumerable<IArtifact> FindAsync(StreamId? streamId = null, CommitId? minCommitId = null, CommitId? maxCommitId = null, ArtifactName? name = null, ArtifactType? type = null, IEnumerable<string>? keys = null, int maxResults = 100, [EnumeratorCancellation] CancellationToken cancellationToken = default)
 		{
 			FilterDefinition<Artifact> filter = FilterDefinition<Artifact>.Empty;
 			if (streamId != null)
 			{
 				filter &= Builders<Artifact>.Filter.Eq(x => x.StreamId, streamId.Value);
-			}
-			if (minChange != null)
-			{
-				filter &= Builders<Artifact>.Filter.Gte(x => x.Change, minChange.Value);
-			}
-			if (maxChange != null)
-			{
-				filter &= Builders<Artifact>.Filter.Lte(x => x.Change, maxChange.Value);
+				if (minCommitId != null)
+				{
+					CommitIdWithOrder minCommitIdWithOrder = await _commitService.GetOrderedAsync(streamId.Value, minCommitId, cancellationToken);
+					filter &= Builders<Artifact>.Filter.Gte(x => x.CommitOrder, minCommitIdWithOrder.Order);
+				}
+				if (maxCommitId != null)
+				{
+					CommitIdWithOrder maxCommitIdWithOrder = await _commitService.GetOrderedAsync(streamId.Value, maxCommitId, cancellationToken);
+					filter &= Builders<Artifact>.Filter.Lte(x => x.CommitOrder, maxCommitIdWithOrder.Order);
+				}
 			}
 			if (name != null)
 			{
@@ -174,7 +192,7 @@ namespace HordeServer.Artifacts
 				filter &= Builders<Artifact>.Filter.All(x => x.Keys, keys.Select(x => NormalizeKey(x)));
 			}
 
-			using (IAsyncCursor<Artifact> cursor = await _artifacts.Find(filter).SortByDescending(x => x.Change).ThenByDescending(x => x.Id).Limit(maxResults).ToCursorAsync(cancellationToken))
+			using (IAsyncCursor<Artifact> cursor = await _artifacts.Find(filter).SortByDescending(x => x.CommitOrder).ThenByDescending(x => x.Id).Limit(maxResults).ToCursorAsync(cancellationToken))
 			{
 				while (await cursor.MoveNextAsync(cancellationToken))
 				{
diff --git a/Engine/Source/Programs/Horde/Plugins/Build/HordeServer.Build/Artifacts/ArtifactsController.cs b/Engine/Source/Programs/Horde/Plugins/Build/HordeServer.Build/Artifacts/ArtifactsController.cs
--- a/Engine/Source/Programs/Horde/Plugins/Build/HordeServer.Build/Artifacts/ArtifactsController.cs
+++ b/Engine/Source/Programs/Horde/Plugins/Build/HordeServer.Build/Artifacts/ArtifactsController.cs
@@ -10,6 +10,7 @@ using EpicGames.Core;
 using EpicGames.Horde.Acls;
 using EpicGames.Horde.Agents.Leases;
 using EpicGames.Horde.Artifacts;
+using EpicGames.Horde.Commits;
 using EpicGames.Horde.Jobs;
 using EpicGames.Horde.Storage;
 using EpicGames.Horde.Storage.Bundles;
@@ -80,11 +81,11 @@ namespace HordeServer.Artifacts
 		public async Task<ActionResult<CreateArtifactResponse>> CreateArtifactAsync([FromBody] CreateArtifactRequest request, CancellationToken cancellationToken = default)
 		{
 			StreamConfig? streamConfig;
-			if (request.StreamId != null && request.Change != null && _buildConfig.TryGetStream(request.StreamId.Value, out streamConfig))
+			if (request.StreamId != null && request.CommitId != null && _buildConfig.TryGetStream(request.StreamId.Value, out streamConfig))
 			{
 				if (streamConfig.Authorize(ArtifactAclAction.WriteArtifact, User))
 				{
-					return await CreateArtifactInternalAsync(request.Name, request.Type, request.Description, request.StreamId.Value, request.Change.Value, request.Keys, request.Metadata, AclScopeName.Root, cancellationToken);
+					return await CreateArtifactInternalAsync(request.Name, request.Type, request.Description, request.StreamId.Value, request.CommitId, request.Keys, request.Metadata, AclScopeName.Root, cancellationToken);
 				}
 			}
 
@@ -136,15 +137,15 @@ namespace HordeServer.Artifacts
 					scopeName = templateRefConfig.Acl.ScopeName;
 				}
 
-				return await CreateArtifactInternalAsync(request.Name, request.Type, request.Description, streamId, request.Change ?? job.Change, keys, request.Metadata, scopeName, cancellationToken);
+				return await CreateArtifactInternalAsync(request.Name, request.Type, request.Description, streamId, request.CommitId ?? job.CommitId, keys, request.Metadata, scopeName, cancellationToken);
 			}
 
 			return Forbid(ArtifactAclAction.WriteArtifact);
 		}
 
-		async Task<ActionResult<CreateArtifactResponse>> CreateArtifactInternalAsync(ArtifactName name, ArtifactType type, string? description, StreamId streamId, int change, List<string> keys, List<string> metadata, AclScopeName scopeName, CancellationToken cancellationToken)
+		async Task<ActionResult<CreateArtifactResponse>> CreateArtifactInternalAsync(ArtifactName name, ArtifactType type, string? description, StreamId streamId, CommitId commitId, List<string> keys, List<string> metadata, AclScopeName scopeName, CancellationToken cancellationToken)
 		{
-			IArtifact artifact = await _artifactCollection.AddAsync(name, type, description, streamId, change, keys, metadata, scopeName, cancellationToken);
+			IArtifact artifact = await _artifactCollection.AddAsync(name, type, description, streamId, commitId, keys, metadata, scopeName, cancellationToken);
 			RefName? prevRefName = await GetPrevRefNameForArtifactAsync(artifact, cancellationToken);
 
 			List<AclClaimConfig> claims = new List<AclClaimConfig>();
@@ -158,8 +159,13 @@ namespace HordeServer.Artifacts
 		async Task<RefName?> GetPrevRefNameForArtifactAsync(IArtifact artifact, CancellationToken cancellationToken)
 		{
 			IStorageBackend storageBackend = _storageService.CreateBackend(artifact.NamespaceId);
-			await foreach (IArtifact prevArtifact in _artifactCollection.FindAsync(artifact.StreamId, maxChange: artifact.Change - 1, name: artifact.Name, type: artifact.Type, cancellationToken: cancellationToken))
+			await foreach (IArtifact prevArtifact in _artifactCollection.FindAsync(artifact.StreamId, maxCommitId: artifact.CommitId, name: artifact.Name, type: artifact.Type, cancellationToken: cancellationToken))
 			{
+				if (prevArtifact.CommitId == artifact.CommitId)
+				{
+					continue;
+				}
+
 				if (prevArtifact.NamespaceId != artifact.NamespaceId)
 				{
 					break;
@@ -189,14 +195,14 @@ namespace HordeServer.Artifacts
 		[HttpGet]
 		[Route("/api/v2/artifacts")]
 		[ProducesResponseType(typeof(FindArtifactsResponse), 200)]
-		public async Task<ActionResult<object>> FindArtifactsAsync([FromQuery] StreamId? streamId = null, [FromQuery] int? minChange = null, [FromQuery] int? maxChange = null, [FromQuery(Name = "name")] ArtifactName? name = null, [FromQuery(Name = "type")] ArtifactType? type = null, [FromQuery(Name = "key")] IEnumerable<string>? keys = null, [FromQuery] int maxResults = 100, [FromQuery] PropertyFilter? filter = null)
+		public async Task<ActionResult<object>> FindArtifactsAsync([FromQuery] StreamId? streamId = null, [FromQuery] CommitId? minChange = null, [FromQuery] CommitId? maxChange = null, [FromQuery(Name = "name")] ArtifactName? name = null, [FromQuery(Name = "type")] ArtifactType? type = null, [FromQuery(Name = "key")] IEnumerable<string>? keys = null, [FromQuery] int maxResults = 100, [FromQuery] PropertyFilter? filter = null)
 		{
 			FindArtifactsResponse response = new FindArtifactsResponse();
 			await foreach (IArtifact artifact in _artifactCollection.FindAsync(streamId, minChange, maxChange, name, type, keys, maxResults, HttpContext.RequestAborted))
 			{
 				if (_buildConfig.AuthorizeArtifact(artifact.Type, artifact.StreamId, ArtifactAclAction.ReadArtifact, User))
 				{
-					response.Artifacts.Add(new GetArtifactResponse(artifact.Id, artifact.Name, artifact.Type, artifact.Description, artifact.StreamId, artifact.Change, artifact.Keys, artifact.Metadata, artifact.CreatedAtUtc));
+					response.Artifacts.Add(new GetArtifactResponse(artifact.Id, artifact.Name, artifact.Type, artifact.Description, artifact.StreamId, artifact.Keys, artifact.Metadata, artifact.CreatedAtUtc) { CommitId = artifact.CommitId });
 				}
 			}
 
@@ -224,7 +230,7 @@ namespace HordeServer.Artifacts
 				return Forbid(ArtifactAclAction.ReadArtifact, artifact.StreamId);
 			}
 
-			return PropertyFilter.Apply(new GetArtifactResponse(artifact.Id, artifact.Name, artifact.Type, artifact.Description, artifact.StreamId, artifact.Change, artifact.Keys, artifact.Metadata, artifact.CreatedAtUtc), filter);
+			return PropertyFilter.Apply(new GetArtifactResponse(artifact.Id, artifact.Name, artifact.Type, artifact.Description, artifact.StreamId, artifact.Keys, artifact.Metadata, artifact.CreatedAtUtc) { CommitId = artifact.CommitId }, filter);
 		}
 
 		/// <summary>
diff --git a/Engine/Source/Programs/Horde/Plugins/Build/HordeServer.Build/Artifacts/IArtifact.cs b/Engine/Source/Programs/Horde/Plugins/Build/HordeServer.Build/Artifacts/IArtifact.cs
--- a/Engine/Source/Programs/Horde/Plugins/Build/HordeServer.Build/Artifacts/IArtifact.cs
+++ b/Engine/Source/Programs/Horde/Plugins/Build/HordeServer.Build/Artifacts/IArtifact.cs
@@ -1,6 +1,7 @@
 // Copyright Epic Games, Inc. All Rights Reserved.
 
 using EpicGames.Horde.Artifacts;
+using EpicGames.Horde.Commits;
 using EpicGames.Horde.Storage;
 using EpicGames.Horde.Streams;
 
@@ -39,7 +40,7 @@ namespace HordeServer.Artifacts
 		/// <summary>
 		/// Change that the artifact corresponds to
 		/// </summary>
-		public int Change { get; }
+		public CommitIdWithOrder CommitId { get; }
 
 		/// <summary>
 		/// Keys used to collate artifacts
diff --git a/Engine/Source/Programs/Horde/Plugins/Build/HordeServer.Build/Artifacts/IArtifactCollection.cs b/Engine/Source/Programs/Horde/Plugins/Build/HordeServer.Build/Artifacts/IArtifactCollection.cs
--- a/Engine/Source/Programs/Horde/Plugins/Build/HordeServer.Build/Artifacts/IArtifactCollection.cs
+++ b/Engine/Source/Programs/Horde/Plugins/Build/HordeServer.Build/Artifacts/IArtifactCollection.cs
@@ -2,6 +2,7 @@
 
 using EpicGames.Horde.Acls;
 using EpicGames.Horde.Artifacts;
+using EpicGames.Horde.Commits;
 using EpicGames.Horde.Streams;
 
 namespace HordeServer.Artifacts
@@ -18,13 +19,13 @@ namespace HordeServer.Artifacts
 		/// <param name="type">Type identifier for the artifact</param>
 		/// <param name="description">Description for the artifact</param>
 		/// <param name="streamId">Stream that the artifact was built from</param>
-		/// <param name="change">Change number that the artifact was built from</param>
+		/// <param name="commitId">Commit that the artifact was built from</param>
 		/// <param name="keys">Keys for the artifact</param>
 		/// <param name="metadata">Metadata for the artifact</param>
 		/// <param name="scopeName">Inherited scope used for permissions</param>
 		/// <param name="cancellationToken">Cancellation token for the operation</param>
 		/// <returns>The new log file document</returns>
-		Task<IArtifact> AddAsync(ArtifactName name, ArtifactType type, string? description, StreamId streamId, int change, IEnumerable<string> keys, IEnumerable<string> metadata, AclScopeName scopeName, CancellationToken cancellationToken = default);
+		Task<IArtifact> AddAsync(ArtifactName name, ArtifactType type, string? description, StreamId streamId, CommitId commitId, IEnumerable<string> keys, IEnumerable<string> metadata, AclScopeName scopeName, CancellationToken cancellationToken = default);
 
 		/// <summary>
 		/// Deletes artifacts
@@ -47,15 +48,15 @@ namespace HordeServer.Artifacts
 		/// Finds artifacts with the given keys.
 		/// </summary>
 		/// <param name="streamId">Stream to find artifacts for</param>
-		/// <param name="minChange">Minimum changelist number for the artifacts</param>
-		/// <param name="maxChange">Maximum changelist number for the artifacts</param>
+		/// <param name="minCommitId">Minimum commit for the artifacts (inclusive)</param>
+		/// <param name="maxCommitId">Maximum commit for the artifacts (inclusive)</param>
 		/// <param name="name">Name of the artifact to search for</param>
 		/// <param name="type">The artifact type</param>
 		/// <param name="keys">Set of keys, all of which must all be present on any returned artifacts</param>
 		/// <param name="maxResults">Maximum number of results to return</param>
 		/// <param name="cancellationToken">Cancellation token for the operation</param>
 		/// <returns>Sequence of artifacts. Ordered by descending CL order, then by descending order in which they were created.</returns>
-		IAsyncEnumerable<IArtifact> FindAsync(StreamId? streamId = null, int? minChange = null, int? maxChange = null, ArtifactName? name = null, ArtifactType? type = null, IEnumerable<string>? keys = null, int maxResults = 100, CancellationToken cancellationToken = default);
+		IAsyncEnumerable<IArtifact> FindAsync(StreamId? streamId = null, CommitId? minCommitId = null, CommitId? maxCommitId = null, ArtifactName? name = null, ArtifactType? type = null, IEnumerable<string>? keys = null, int maxResults = 100, CancellationToken cancellationToken = default);
 
 		/// <summary>
 		/// Gets an artifact by ID
diff --git a/Engine/Source/Programs/Horde/Plugins/Build/HordeServer.Build/Commits/CommitService.cs b/Engine/Source/Programs/Horde/Plugins/Build/HordeServer.Build/Commits/CommitService.cs
--- a/Engine/Source/Programs/Horde/Plugins/Build/HordeServer.Build/Commits/CommitService.cs
+++ b/Engine/Source/Programs/Horde/Plugins/Build/HordeServer.Build/Commits/CommitService.cs
@@ -1,7 +1,10 @@
 // Copyright Epic Games, Inc. All Rights Reserved.
 
+using EpicGames.Horde.Commits;
+using EpicGames.Horde.Streams;
 using HordeServer.Perforce;
 using HordeServer.Streams;
+using Microsoft.Extensions.Options;
 
 namespace HordeServer.Commits
 {
@@ -11,16 +14,30 @@ namespace HordeServer.Commits
 	class CommitService : ICommitService
 	{
 		readonly IPerforceService _perforceService;
+		readonly IOptionsMonitor<BuildConfig> _buildConfig;
 
 		/// <summary>
 		/// Constructor
 		/// </summary>
-		public CommitService(IPerforceService perforceService)
+		public CommitService(IPerforceService perforceService, IOptionsMonitor<BuildConfig> buildConfig)
 		{
 			_perforceService = perforceService;
+			_buildConfig = buildConfig;
 		}
 
 		/// <inheritdoc/>
 		public ICommitCollection GetCollection(StreamConfig streamConfig) => _perforceService.GetCommits(streamConfig);
+
+		/// <inheritdoc/>
+		public async ValueTask<CommitIdWithOrder> GetOrderedAsync(StreamId streamId, CommitId commitId, CancellationToken cancellationToken = default)
+		{
+			if (!_buildConfig.CurrentValue.TryGetStream(streamId, out StreamConfig? streamConfig))
+			{
+				throw new StreamNotFoundException(streamId);
+			}
+
+			ICommitCollection commitCollection = GetCollection(streamConfig);
+			return await commitCollection.GetOrderedAsync(commitId, cancellationToken);
+		}
 	}
 }
diff --git a/Engine/Source/Programs/Horde/Plugins/Build/HordeServer.Build/Commits/ICommit.cs b/Engine/Source/Programs/Horde/Plugins/Build/HordeServer.Build/Commits/ICommit.cs
--- a/Engine/Source/Programs/Horde/Plugins/Build/HordeServer.Build/Commits/ICommit.cs
+++ b/Engine/Source/Programs/Horde/Plugins/Build/HordeServer.Build/Commits/ICommit.cs
@@ -13,19 +13,19 @@ namespace HordeServer.Commits
 	public interface ICommit
 	{
 		/// <summary>
-		/// Stream containing the commit
+		/// Id for this commit
 		/// </summary>
-		StreamId StreamId { get; }
+		CommitIdWithOrder Id { get; }
 
 		/// <summary>
-		/// The changelist number
+		/// Stream containing the commit
 		/// </summary>
-		int Number { get; }
+		StreamId StreamId { get; }
 
 		/// <summary>
 		/// The change that this commit originates from
 		/// </summary>
-		int OriginalChange { get; }
+		CommitIdWithOrder OriginalCommitId { get; }
 
 		/// <summary>
 		/// The author user id
diff --git a/Engine/Source/Programs/Horde/Plugins/Build/HordeServer.Build/Commits/ICommitCollection.cs b/Engine/Source/Programs/Horde/Plugins/Build/HordeServer.Build/Commits/ICommitCollection.cs
--- a/Engine/Source/Programs/Horde/Plugins/Build/HordeServer.Build/Commits/ICommitCollection.cs
+++ b/Engine/Source/Programs/Horde/Plugins/Build/HordeServer.Build/Commits/ICommitCollection.cs
@@ -18,35 +18,45 @@ namespace HordeServer.Commits
 		/// <param name="description">Description of the change</param>
 		/// <param name="cancellationToken">Cancellation token for the operation</param>
 		/// <returns>New commit information</returns>
-		Task<int> CreateNewAsync(string path, string description, CancellationToken cancellationToken = default);
+		Task<CommitIdWithOrder> CreateNewAsync(string path, string description, CancellationToken cancellationToken = default);
 
 		/// <summary>
-		/// Gets the latest change for a particular stream
+		/// Gets a commit by id
 		/// </summary>
-		/// <param name="changeNumber">Change numbers to query</param>
+		/// <param name="commitId">Commit to query</param>
 		/// <param name="cancellationToken">Cancellation token for the operation</param>
 		/// <returns>Commit details</returns>
-		Task<ICommit> GetAsync(int changeNumber, CancellationToken cancellationToken = default);
+		Task<ICommit> GetAsync(CommitId commitId, CancellationToken cancellationToken = default);
+
+		/// <summary>
+		/// Gets an ordered commit id
+		/// </summary>
+		/// <param name="commitId">The commit to query</param>
+		/// <param name="cancellationToken">Cancellation token for the operation</param>
+		/// <returns>Numbered commit id</returns>
+		ValueTask<CommitIdWithOrder> GetOrderedAsync(CommitId commitId, CancellationToken cancellationToken = default);
 
 		/// <summary>
 		/// Finds changes submitted to a stream, in reverse order.
 		/// </summary>
-		/// <param name="minChange">The minimum changelist number</param>
-		/// <param name="maxChange">The maximum changelist number</param>
+		/// <param name="minCommitId">The minimum changelist number</param>
+		/// <param name="includeMinCommit">Whether to include the minimum changelist in the range of enumerated responses</param>
+		/// <param name="maxCommitId">The maximum changelist number</param>
+		/// <param name="includeMaxCommit">Whether to include the maximum changelist in the range of enumerated responses</param>
 		/// <param name="maxResults">Maximum number of results to return</param>
 		/// <param name="tags">Tags for the commits to return</param>
 		/// <param name="cancellationToken">Cancellation token for the operation</param>
 		/// <returns>Changelist information</returns>
-		IAsyncEnumerable<ICommit> FindAsync(int? minChange = null, int? maxChange = null, int? maxResults = null, IReadOnlyList<CommitTag>? tags = null, CancellationToken cancellationToken = default);
+		IAsyncEnumerable<ICommit> FindAsync(CommitId? minCommitId = null, bool includeMinCommit = true, CommitId? maxCommitId = null, bool includeMaxCommit = true, int? maxResults = null, IReadOnlyList<CommitTag>? tags = null, CancellationToken cancellationToken = default);
 
 		/// <summary>
 		/// Subscribes to changes from this commit source
 		/// </summary>
-		/// <param name="minChange">Minimum changelist number (exclusive)</param>
+		/// <param name="minCommitId">Minimum changelist number (exclusive)</param>
 		/// <param name="tags">Tags for the commit to return</param>
 		/// <param name="cancellationToken">Cancellation token for the operation</param>
 		/// <returns>New change information</returns>
-		IAsyncEnumerable<ICommit> SubscribeAsync(int minChange, IReadOnlyList<CommitTag>? tags = null, CancellationToken cancellationToken = default);
+		IAsyncEnumerable<ICommit> SubscribeAsync(CommitId minCommitId, IReadOnlyList<CommitTag>? tags = null, CancellationToken cancellationToken = default);
 	}
 
 	/// <summary>
@@ -57,37 +67,50 @@ namespace HordeServer.Commits
 		/// <summary>
 		/// Creates a new change for a template
 		/// </summary>
-		/// <param name="perforce">The Perforce service instance</param>
+		/// <param name="commitCollection">The Perforce service instance</param>
 		/// <param name="template">The template being built</param>
 		/// <param name="cancellationToken">Cancellation token for the operation</param>
 		/// <returns>New changelist number</returns>
-		public static Task<int> CreateNewAsync(this ICommitCollection perforce, ITemplate template, CancellationToken cancellationToken)
+		public static Task<CommitIdWithOrder> CreateNewAsync(this ICommitCollection commitCollection, ITemplate template, CancellationToken cancellationToken)
 		{
 			string description = (template.SubmitDescription ?? "[Horde] New change for $(TemplateName)").Replace("$(TemplateName)", template.Name, StringComparison.OrdinalIgnoreCase);
-			return perforce.CreateNewAsync(template.SubmitNewChange!, description, cancellationToken);
+			return commitCollection.CreateNewAsync(template.SubmitNewChange!, description, cancellationToken);
 		}
 
+		/// <summary>
+		/// Finds changes submitted to a stream, in reverse order.
+		/// </summary>
+		/// <param name="commitCollection">Collection to operate on</param>
+		/// <param name="minCommitId">The minimum changelist number</param>
+		/// <param name="maxCommitId">The maximum changelist number</param>
+		/// <param name="maxResults">Maximum number of results to return</param>
+		/// <param name="tags">Tags for the commits to return</param>
+		/// <param name="cancellationToken">Cancellation token for the operation</param>
+		/// <returns>Changelist information</returns>
+		public static IAsyncEnumerable<ICommit> FindAsync(this ICommitCollection commitCollection, CommitId? minCommitId, CommitId? maxCommitId, int? maxResults = null, IReadOnlyList<CommitTag>? tags = null, CancellationToken cancellationToken = default)
+			=> commitCollection.FindAsync(minCommitId: minCommitId, maxCommitId: maxCommitId, maxResults: maxResults, tags: tags, cancellationToken: cancellationToken);
+
 		/// <summary>
 		/// Gets the last code code equal or before the given change number
 		/// </summary>
-		/// <param name="source">The commit source to query</param>
-		/// <param name="maxChange">Maximum code change to query</param>
+		/// <param name="commitCollection">The commit source to query</param>
+		/// <param name="maxCommitId">Maximum code change to query</param>
 		/// <param name="cancellationToken">Cancellation token for the operation</param>
 		/// <returns>The last code change</returns>
-		public static async ValueTask<ICommit?> GetLastCodeChangeAsync(this ICommitCollection source, int? maxChange, CancellationToken cancellationToken = default)
+		public static async ValueTask<ICommit?> GetLastCodeChangeAsync(this ICommitCollection commitCollection, CommitId? maxCommitId, CancellationToken cancellationToken = default)
 		{
-			return await source.FindAsync(null, maxChange, 1, new[] { CommitTag.Code }, cancellationToken).FirstOrDefaultAsync(cancellationToken);
+			return await commitCollection.FindAsync(minCommitId: null, maxCommitId: maxCommitId, maxResults: 1, tags: new[] { CommitTag.Code }, cancellationToken: cancellationToken).FirstOrDefaultAsync(cancellationToken);
 		}
 
 		/// <summary>
 		/// Finds the latest commit from a source
 		/// </summary>
-		/// <param name="source">The commit source to query</param>
+		/// <param name="commitCollection">The commit source to query</param>
 		/// <param name="cancellationToken">Cancellation token for the operation</param>
 		/// <returns>The latest commit</returns>
-		public static async Task<ICommit> GetLatestAsync(this ICommitCollection source, CancellationToken cancellationToken = default)
+		public static async Task<ICommit> GetLatestAsync(this ICommitCollection commitCollection, CancellationToken cancellationToken = default)
 		{
-			ICommit? commit = await source.FindAsync(null, null, 1, null, cancellationToken).FirstOrDefaultAsync(cancellationToken);
+			ICommit? commit = await commitCollection.FindAsync(null, null, maxResults: 1, cancellationToken: cancellationToken).FirstOrDefaultAsync(cancellationToken);
 			if (commit == null)
 			{
 				throw new PerforceException("No changes found for stream.");
@@ -98,13 +121,13 @@ namespace HordeServer.Commits
 		/// <summary>
 		/// Finds the latest commit from a source
 		/// </summary>
-		/// <param name="source">The commit source to query</param>
+		/// <param name="commitCollection">The commit source to query</param>
 		/// <param name="cancellationToken">Cancellation token for the operation</param>
 		/// <returns>The latest commit</returns>
-		public static async Task<int> GetLatestNumberAsync(this ICommitCollection source, CancellationToken cancellationToken = default)
+		public static async Task<CommitIdWithOrder> GetLastCommitIdAsync(this ICommitCollection commitCollection, CancellationToken cancellationToken = default)
 		{
-			ICommit commit = await GetLatestAsync(source, cancellationToken);
-			return commit.Number;
+			ICommit commit = await GetLatestAsync(commitCollection, cancellationToken);
+			return commit.Id;
 		}
 	}
 }
diff --git a/Engine/Source/Programs/Horde/Plugins/Build/HordeServer.Build/Commits/ICommitService.cs b/Engine/Source/Programs/Horde/Plugins/Build/HordeServer.Build/Commits/ICommitService.cs
--- a/Engine/Source/Programs/Horde/Plugins/Build/HordeServer.Build/Commits/ICommitService.cs
+++ b/Engine/Source/Programs/Horde/Plugins/Build/HordeServer.Build/Commits/ICommitService.cs
@@ -1,5 +1,7 @@
 // Copyright Epic Games, Inc. All Rights Reserved.
 
+using EpicGames.Horde.Commits;
+using EpicGames.Horde.Streams;
 using HordeServer.Streams;
 
 namespace HordeServer.Commits
@@ -15,5 +17,31 @@ namespace HordeServer.Commits
 		/// <param name="streamConfig">Stream to get commits for</param>
 		/// <returns>Collection object</returns>
 		ICommitCollection GetCollection(StreamConfig streamConfig);
+
+		/// <summary>
+		/// Gets a numbered commit id
+		/// </summary>
+		/// <param name="streamId">Stream containing the commit</param>
+		/// <param name="commitId">The commit to query</param>
+		/// <param name="cancellationToken">Cancellation token for the operation</param>
+		/// <returns>Numbered commit id</returns>
+		ValueTask<CommitIdWithOrder> GetOrderedAsync(StreamId streamId, CommitId commitId, CancellationToken cancellationToken = default);
+	}
+
+	/// <summary>
+	/// Exception thrown when a stream can't be found
+	/// </summary>
+	public sealed class StreamNotFoundException : Exception
+	{
+		/// <summary>
+		/// The stream identifier
+		/// </summary>
+		public StreamId StreamId { get; }
+
+		/// <summary>
+		/// Constructor
+		/// </summary>
+		public StreamNotFoundException(StreamId streamId) : base($"Stream {streamId} not found")
+			=> StreamId = streamId;
 	}
 }
