Horde: Preliminary support for VCS commit ids that aren't an integer.

diff --git a/Engine/Source/Programs/Horde/Plugins/Build/HordeServer.Build/Issues/IIssue.cs b/Engine/Source/Programs/Horde/Plugins/Build/HordeServer.Build/Issues/IIssue.cs
--- a/Engine/Source/Programs/Horde/Plugins/Build/HordeServer.Build/Issues/IIssue.cs
+++ b/Engine/Source/Programs/Horde/Plugins/Build/HordeServer.Build/Issues/IIssue.cs
@@ -1,5 +1,6 @@
 // Copyright Epic Games, Inc. All Rights Reserved.
 
+using EpicGames.Horde.Commits;
 using EpicGames.Horde.Issues;
 using EpicGames.Horde.Streams;
 using EpicGames.Horde.Users;
@@ -103,14 +104,14 @@ namespace HordeServer.Issues
 		public DateTime LastSeenAt { get; }
 
 		/// <summary>
-		/// Fix changelist for this issue
+		/// Fix commit for this issue
 		/// </summary>
-		public int? FixChange { get; }
+		public CommitId? FixCommitId { get; }
 
 		/// <summary>
 		/// Whether the issue is marked fixed as systemic
 		/// </summary>
-		public bool FixedSystemic { get; }
+		public bool FixSystemic { get; }
 
 		//		/// <summary>
 		//		/// The first stream that encountered the error. The fix will be considered failed if an error after FixChange occurs in this stream. 
@@ -205,9 +206,9 @@ namespace HordeServer.Issues
 		UserId AuthorId { get; }
 
 		/// <summary>
-		/// The change suspected of causing this issue (in the origin stream)
+		/// The commit suspected of causing this issue (in the origin stream)
 		/// </summary>
-		int Change { get; }
+		CommitIdWithOrder CommitId { get; }
 
 		/// <summary>
 		/// Time at which the author declined the issue
@@ -220,6 +221,12 @@ namespace HordeServer.Issues
 	/// </summary>
 	static class IssueExtensions
 	{
+		/// <summary>
+		/// Test whether an issue is marked as fixed
+		/// </summary>
+		public static bool IsMarkedFixed(this IIssue issue)
+			=> issue.FixCommitId != null || issue.FixSystemic;
+
 		/// <summary>
 		/// Creates a lookup from stream id to whether it's fixed
 		/// </summary>
@@ -237,29 +244,5 @@ namespace HordeServer.Issues
 			}
 			return fixStreamIds;
 		}
-
-		/// <summary>
-		/// Find the first fix-failed step for this issue
-		/// </summary>
-		/// <param name="issue"></param>
-		/// <param name="spans"></param>
-		/// <returns></returns>
-		public static IIssueStep? FindFixFailedStep(this IIssue issue, IEnumerable<IIssueSpan> spans)
-		{
-			IIssueStep? fixFailedStep = null;
-			if (issue.FixChange != null && issue.FixChange.Value >= 0)
-			{
-				foreach (IIssueSpan span in spans)
-				{
-					IIssueStream? stream = issue.Streams.FirstOrDefault(x => x.StreamId == span.StreamId);
-					if (stream != null && (stream.ContainsFix ?? false) && span.LastFailure.Change >= issue.FixChange.Value)
-					{
-						fixFailedStep = span.LastFailure;
-						break;
-					}
-				}
-			}
-			return fixFailedStep;
-		}
 	}
 }
\ No newline at end of file
diff --git a/Engine/Source/Programs/Horde/Plugins/Build/HordeServer.Build/Issues/IIssueCollection.cs b/Engine/Source/Programs/Horde/Plugins/Build/HordeServer.Build/Issues/IIssueCollection.cs
--- a/Engine/Source/Programs/Horde/Plugins/Build/HordeServer.Build/Issues/IIssueCollection.cs
+++ b/Engine/Source/Programs/Horde/Plugins/Build/HordeServer.Build/Issues/IIssueCollection.cs
@@ -1,5 +1,6 @@
 // Copyright Epic Games, Inc. All Rights Reserved.
 
+using EpicGames.Horde.Commits;
 using EpicGames.Horde.Issues;
 using EpicGames.Horde.Jobs;
 using EpicGames.Horde.Jobs.Templates;
@@ -257,8 +258,8 @@ namespace HordeServer.Issues
 	/// </summary>
 	public class NewIssueStepData
 	{
-		/// <inheritdoc cref="IIssueStep.Change"/>
-		public int Change { get; set; }
+		/// <inheritdoc cref="IIssueStep.CommitId"/>
+		public CommitIdWithOrder CommitId { get; set; }
 
 		/// <inheritdoc cref="IIssueStep.Severity"/>
 		public IssueSeverity Severity { get; set; }
@@ -290,7 +291,7 @@ namespace HordeServer.Issues
 		/// <summary>
 		/// Constructor
 		/// </summary>
-		/// <param name="change">The changelist number for this job</param>
+		/// <param name="commitId">The commit for this job</param>
 		/// <param name="severity">Severity of the issue in this step</param>
 		/// <param name="jobName">The job name</param>
 		/// <param name="jobId">The unique job id</param>
@@ -300,9 +301,9 @@ namespace HordeServer.Issues
 		/// <param name="logId">Unique id of the log file for this step</param>
 		/// <param name="annotations">Annotations for this step</param>
 		/// <param name="promoted">Whether this step is promoted</param>
-		public NewIssueStepData(int change, IssueSeverity severity, string jobName, JobId jobId, JobStepBatchId batchId, JobStepId stepId, DateTime stepTime, LogId? logId, IReadOnlyNodeAnnotations? annotations, bool promoted)
+		public NewIssueStepData(CommitIdWithOrder commitId, IssueSeverity severity, string jobName, JobId jobId, JobStepBatchId batchId, JobStepId stepId, DateTime stepTime, LogId? logId, IReadOnlyNodeAnnotations? annotations, bool promoted)
 		{
-			Change = change;
+			CommitId = commitId;
 			Severity = severity;
 			JobName = jobName;
 			JobId = jobId;
@@ -327,7 +328,7 @@ namespace HordeServer.Issues
 		/// <param name="annotations">Annotations for this step</param>
 		/// <param name="promoted">Whether this step is promoted</param>
 		public NewIssueStepData(IJob job, IJobStepBatch batch, IJobStep step, IssueSeverity severity, IReadOnlyNodeAnnotations? annotations, bool promoted)
-			: this(job.Change, severity, job.Name, job.Id, batch.Id, step.Id, step.StartTimeUtc ?? default, step.LogId, annotations, promoted)
+			: this(job.CommitId, severity, job.Name, job.Id, batch.Id, step.Id, step.StartTimeUtc ?? default, step.LogId, annotations, promoted)
 		{
 		}
 
@@ -336,7 +337,7 @@ namespace HordeServer.Issues
 		/// </summary>
 		/// <param name="jobStepRef">The jobstep reference</param>
 		public NewIssueStepData(IJobStepRef jobStepRef)
-			: this(jobStepRef.Change, IssueSeverity.Unspecified, jobStepRef.JobName, jobStepRef.Id.JobId, jobStepRef.Id.BatchId, jobStepRef.Id.StepId, jobStepRef.StartTimeUtc, jobStepRef.LogId, null, false)
+			: this(jobStepRef.CommitId, IssueSeverity.Unspecified, jobStepRef.JobName, jobStepRef.Id.JobId, jobStepRef.Id.BatchId, jobStepRef.Id.StepId, jobStepRef.StartTimeUtc, jobStepRef.LogId, null, false)
 		{
 		}
 	}
@@ -349,18 +350,18 @@ namespace HordeServer.Issues
 		/// <inheritdoc cref="IIssueSuspect.AuthorId"/>
 		public UserId AuthorId { get; set; }
 
-		/// <inheritdoc cref="IIssueSuspect.Change"/>
-		public int Change { get; set; }
+		/// <inheritdoc cref="IIssueSuspect.CommitId"/>
+		public CommitIdWithOrder CommitId { get; set; }
 
 		/// <summary>
 		/// Constructor
 		/// </summary>
 		/// <param name="authorId">Author of the change</param>
-		/// <param name="change">The changelist number</param>
-		public NewIssueSuspectData(UserId authorId, int change)
+		/// <param name="commitId">The changelist number</param>
+		public NewIssueSuspectData(UserId authorId, CommitIdWithOrder commitId)
 		{
 			AuthorId = authorId;
-			Change = change;
+			CommitId = commitId;
 		}
 	}
 
@@ -369,23 +370,23 @@ namespace HordeServer.Issues
 	/// </summary>
 	public class NewIssueSpanSuspectData
 	{
-		/// <inheritdoc cref="IIssueSpanSuspect.Change"/>
-		public int Change { get; set; }
+		/// <inheritdoc cref="IIssueSpanSuspect.CommitId"/>
+		public CommitIdWithOrder CommitId { get; set; }
 
 		/// <inheritdoc cref="IIssueSpanSuspect.AuthorId"/>
 		public UserId AuthorId { get; set; }
 
-		/// <inheritdoc cref="IIssueSpanSuspect.OriginatingChange"/>
-		public int? OriginatingChange { get; set; }
+		/// <inheritdoc cref="IIssueSpanSuspect.SourceCommitId"/>
+		public CommitIdWithOrder? SourceCommitId { get; set; }
 
 		/// <summary>
 		/// Constructor
 		/// </summary>
 		/// <param name="change">The changelist number</param>
 		/// <param name="authorId">Author of the change</param>
-		public NewIssueSpanSuspectData(int change, UserId authorId)
+		public NewIssueSpanSuspectData(CommitIdWithOrder change, UserId authorId)
 		{
-			Change = change;
+			CommitId = change;
 			AuthorId = authorId;
 		}
 	}
@@ -460,23 +461,23 @@ namespace HordeServer.Issues
 		/// <param name="ids">Set of issue ids to find</param>
 		/// <param name="ownerId">The user to find issues for</param>
 		/// <param name="streamId">The stream affected by the issue</param>
-		/// <param name="minChange">Minimum changelist affected by the issue</param>
-		/// <param name="maxChange">Maximum changelist affected by the issue</param>
+		/// <param name="minCommitId">Minimum changelist affected by the issue</param>
+		/// <param name="maxCommitId">Maximum changelist affected by the issue</param>
 		/// <param name="resolved">Include issues that are now resolved</param>
 		/// <param name="promoted">Include only promoted issues</param>
 		/// <param name="index">Index within the results to return</param>
 		/// <param name="count">Number of results</param>
 		/// <param name="cancellationToken">Cancellation token for the operation</param>
 		/// <returns>List of streams open in the given stream at the given changelist</returns>
-		Task<IReadOnlyList<IIssue>> FindIssuesAsync(IEnumerable<int>? ids = null, UserId? ownerId = null, StreamId? streamId = null, int? minChange = null, int? maxChange = null, bool? resolved = null, bool? promoted = null, int? index = null, int? count = null, CancellationToken cancellationToken = default);
+		Task<IReadOnlyList<IIssue>> FindIssuesAsync(IEnumerable<int>? ids = null, UserId? ownerId = null, StreamId? streamId = null, CommitId? minCommitId = null, CommitId? maxCommitId = null, bool? resolved = null, bool? promoted = null, int? index = null, int? count = null, CancellationToken cancellationToken = default);
 
 		/// <summary>
 		/// Searches for open issues
 		/// </summary>
-		/// <param name="changes">List of suspect changes</param>
+		/// <param name="commits">List of suspect commits</param>
 		/// <param name="cancellationToken">Cancellation token for the operation</param>
 		/// <returns>List of issues that are affected by the given changes</returns>
-		Task<IReadOnlyList<IIssue>> FindIssuesForChangesAsync(List<int> changes, CancellationToken cancellationToken = default);
+		Task<IReadOnlyList<IIssue>> FindIssuesForChangesAsync(List<CommitIdWithOrder> commits, CancellationToken cancellationToken = default);
 
 		/// <summary>
 		/// Try to update the state of an issue
@@ -492,8 +493,8 @@ namespace HordeServer.Issues
 		/// <param name="newNominatedById">Person that nominated the new owner</param>
 		/// <param name="newAcknowledged">Whether the issue has been acknowledged</param>
 		/// <param name="newDeclinedById">Name of a user that has declined the issue</param>
-		/// <param name="newFixChange">Fix changelist for the issue. Pass 0 to clear the fix changelist, -1 for systemic issue.</param>
-		/// <param name="newFixedSystemic">Whether the issue should be marked fixed as a systemic issue</param>
+		/// <param name="newFixCommitId">Fix changelist for the issue. Pass a default value to clear the fix changelist.</param>
+		/// <param name="newFixSystemic">Whether the issue should be marked fixed as a systemic issue</param>
 		/// <param name="newResolvedById">User that resolved the issue (may be ObjectId.Empty to clear)</param>
 		/// <param name="newExcludeSpanIds">List of span ids to exclude from this issue</param>
 		/// <param name="newLastSeenAt"></param>
@@ -503,7 +504,7 @@ namespace HordeServer.Issues
 		/// <param name="newWorkflowThreadUrl">The workflow thread url associated with the issue</param>
 		/// <param name="cancellationToken">Cancellation token for the operation</param>
 		/// <returns>True if the issue was updated</returns>
-		Task<IIssue?> TryUpdateIssueAsync(IIssue issue, UserId? initiatedByUserId, IssueSeverity? newSeverity = null, string? newSummary = null, string? newUserSummary = null, string? newDescription = null, bool? newPromoted = null, UserId? newOwnerId = null, UserId? newNominatedById = null, bool? newAcknowledged = null, UserId? newDeclinedById = null, int? newFixChange = null, bool? newFixedSystemic = null, UserId? newResolvedById = null, List<ObjectId>? newExcludeSpanIds = null, DateTime? newLastSeenAt = null, string? newExternalIssueKey = null, UserId? newQuarantinedById = null, UserId? newForceClosedById = null, Uri? newWorkflowThreadUrl = null, CancellationToken cancellationToken = default);
+		Task<IIssue?> TryUpdateIssueAsync(IIssue issue, UserId? initiatedByUserId, IssueSeverity? newSeverity = null, string? newSummary = null, string? newUserSummary = null, string? newDescription = null, bool? newPromoted = null, UserId? newOwnerId = null, UserId? newNominatedById = null, bool? newAcknowledged = null, UserId? newDeclinedById = null, CommitId? newFixCommitId = null, bool? newFixSystemic = null, UserId? newResolvedById = null, List<ObjectId>? newExcludeSpanIds = null, DateTime? newLastSeenAt = null, string? newExternalIssueKey = null, UserId? newQuarantinedById = null, UserId? newForceClosedById = null, Uri? newWorkflowThreadUrl = null, CancellationToken cancellationToken = default);
 
 		/// <summary>
 		/// Updates derived data for an issue (ie. data computed from the spans attached to it). Also clears the issue's 'modified' state.
@@ -577,10 +578,10 @@ namespace HordeServer.Issues
 		/// <param name="streamId">The stream id</param>
 		/// <param name="templateId">The template id</param>
 		/// <param name="name">Name of the node</param>
-		/// <param name="change">Changelist number to query</param>
+		/// <param name="commitId">Changelist number to query</param>
 		/// <param name="cancellationToken">Cancellation token for the operation</param>
 		/// <returns>List of open issues</returns>
-		Task<IReadOnlyList<IIssueSpan>> FindOpenSpansAsync(StreamId streamId, TemplateId templateId, string name, int change, CancellationToken cancellationToken = default);
+		Task<IReadOnlyList<IIssueSpan>> FindOpenSpansAsync(StreamId streamId, TemplateId templateId, string name, CommitId commitId, CancellationToken cancellationToken = default);
 
 		/// <summary>
 		/// Searches for open issues
@@ -588,14 +589,14 @@ namespace HordeServer.Issues
 		/// <param name="ids">Set of issue ids to find</param>
 		/// <param name="issueIds">The issue ids to retrieve spans for</param>
 		/// <param name="streamId">The stream affected by the issue</param>
-		/// <param name="minChange">Minimum changelist affected by the issue</param>
-		/// <param name="maxChange">Maximum changelist affected by the issue</param>
+		/// <param name="minCommitId">Minimum commit affected by the issue</param>
+		/// <param name="maxCommitId">Maximum commit affected by the issue</param>
 		/// <param name="resolved">Include issues that are now resolved</param>
 		/// <param name="index">Index within the results to return</param>
 		/// <param name="count">Number of results</param>
 		/// <param name="cancellationToken">Cancellation token for the operation</param>
 		/// <returns>List of streams open in the given stream at the given changelist</returns>
-		Task<IReadOnlyList<IIssueSpan>> FindSpansAsync(IEnumerable<ObjectId>? ids = null, IEnumerable<int>? issueIds = null, StreamId? streamId = null, int? minChange = null, int? maxChange = null, bool? resolved = null, int? index = null, int? count = null, CancellationToken cancellationToken = default);
+		Task<IReadOnlyList<IIssueSpan>> FindSpansAsync(IEnumerable<ObjectId>? ids = null, IEnumerable<int>? issueIds = null, StreamId? streamId = null, CommitId? minCommitId = null, CommitId? maxCommitId = null, bool? resolved = null, int? index = null, int? count = null, CancellationToken cancellationToken = default);
 
 		#endregion
 
diff --git a/Engine/Source/Programs/Horde/Plugins/Build/HordeServer.Build/Issues/IIssueSpan.cs b/Engine/Source/Programs/Horde/Plugins/Build/HordeServer.Build/Issues/IIssueSpan.cs
--- a/Engine/Source/Programs/Horde/Plugins/Build/HordeServer.Build/Issues/IIssueSpan.cs
+++ b/Engine/Source/Programs/Horde/Plugins/Build/HordeServer.Build/Issues/IIssueSpan.cs
@@ -1,5 +1,6 @@
 // Copyright Epic Games, Inc. All Rights Reserved.
 
+using EpicGames.Horde.Commits;
 using EpicGames.Horde.Issues;
 using EpicGames.Horde.Jobs.Templates;
 using EpicGames.Horde.Streams;
@@ -87,7 +88,7 @@ namespace HordeServer.Issues
 		/// <summary>
 		/// The submitted changelist
 		/// </summary>
-		public int Change { get; }
+		public CommitIdWithOrder CommitId { get; }
 
 		/// <summary>
 		/// Id of the changelist's author
@@ -97,6 +98,6 @@ namespace HordeServer.Issues
 		/// <summary>
 		/// The original changelist number, if merged from another branch. For changes merged between several branches, this is the originally submitted change.
 		/// </summary>
-		public int? OriginatingChange { get; }
+		public CommitIdWithOrder? SourceCommitId { get; }
 	}
 }
diff --git a/Engine/Source/Programs/Horde/Plugins/Build/HordeServer.Build/Issues/IIssueStep.cs b/Engine/Source/Programs/Horde/Plugins/Build/HordeServer.Build/Issues/IIssueStep.cs
--- a/Engine/Source/Programs/Horde/Plugins/Build/HordeServer.Build/Issues/IIssueStep.cs
+++ b/Engine/Source/Programs/Horde/Plugins/Build/HordeServer.Build/Issues/IIssueStep.cs
@@ -1,5 +1,6 @@
 // Copyright Epic Games, Inc. All Rights Reserved.
 
+using EpicGames.Horde.Commits;
 using EpicGames.Horde.Issues;
 using EpicGames.Horde.Jobs;
 using EpicGames.Horde.Logs;
@@ -9,7 +10,7 @@ using MongoDB.Bson;
 namespace HordeServer.Issues
 {
 	/// <summary>
-	/// Identifies a particular changelist and job that contributes to a span
+	/// Identifies a particular commit and job that contributes to a span
 	/// </summary>
 	public interface IIssueStep
 	{
@@ -19,9 +20,9 @@ namespace HordeServer.Issues
 		public ObjectId SpanId { get; }
 
 		/// <summary>
-		/// The changelist number
+		/// The commit
 		/// </summary>
-		public int Change { get; }
+		public CommitIdWithOrder CommitId { get; }
 
 		/// <summary>
 		/// Severity of the issue in this step
diff --git a/Engine/Source/Programs/Horde/Plugins/Build/HordeServer.Build/Issues/IssueCollection.cs b/Engine/Source/Programs/Horde/Plugins/Build/HordeServer.Build/Issues/IssueCollection.cs
--- a/Engine/Source/Programs/Horde/Plugins/Build/HordeServer.Build/Issues/IssueCollection.cs
+++ b/Engine/Source/Programs/Horde/Plugins/Build/HordeServer.Build/Issues/IssueCollection.cs
@@ -4,6 +4,7 @@ using System.Data;
 using System.Diagnostics;
 using System.Diagnostics.CodeAnalysis;
 using EpicGames.Core;
+using EpicGames.Horde.Commits;
 using EpicGames.Horde.Issues;
 using EpicGames.Horde.Jobs;
 using EpicGames.Horde.Jobs.Templates;
@@ -13,6 +14,7 @@ using EpicGames.Horde.Telemetry;
 using EpicGames.Horde.Users;
 using EpicGames.Redis.Utility;
 using HordeServer.Auditing;
+using HordeServer.Commits;
 using HordeServer.Jobs;
 using HordeServer.Jobs.Graphs;
 using HordeServer.Logs;
@@ -97,18 +99,24 @@ namespace HordeServer.Issues
 
 			public DateTime LastSeenAt { get; set; }
 
-			[BsonIgnoreIfNull]
+			[BsonIgnore]
+			public CommitId? FixCommitId
+			{
+				get => (FixCommitName != null)? new CommitId(FixCommitName) : (FixChange != null)? CommitId.FromPerforceChange(FixChange.Value) : null;
+				set => (FixCommitName, FixChange) = (value?.Name, null);
+			}
+
+			[BsonElement("FixCommit"), BsonIgnoreIfNull]
+			public string? FixCommitName { get; set; }
+
+			[BsonElement("FixChange"), BsonIgnoreIfNull]
 			public int? FixChange { get; set; }
-			int? IIssue.FixChange => FixChange ?? (FixedSystemic ? -1 : null);
 
 			[BsonIgnoreIfDefault, BsonDefaultValue(false)]
-			public bool FixedSystemic { get; set; }
+			public bool FixSystemic { get; set; }
 
 			public List<IssueStreamDocument> Streams { get; set; } = new List<IssueStreamDocument>();
 
-			public int MinSuspectChange { get; set; }
-			public int MaxSuspectChange { get; set; }
-
 			[BsonElement("Suspects"), BsonIgnoreIfNull]
 			public List<IssueSuspectDocument>? SuspectsDeprecated { get; set; }
 
@@ -207,7 +215,20 @@ namespace HordeServer.Issues
 			public ObjectId Id { get; set; }
 			public int IssueId { get; set; }
 			public UserId AuthorId { get; set; }
-			public int Change { get; set; }
+
+			[BsonIgnore]
+			public CommitIdWithOrder CommitId
+			{
+				get => (CommitName != null) ? new CommitIdWithOrder(CommitName, CommitOrder) : CommitIdWithOrder.FromPerforceChange(CommitOrder);
+				set => (CommitName, CommitOrder) = (value.Name, value.Order);
+			}
+
+			[BsonElement("Commit")]
+			public string? CommitName { get; set; }
+
+			[BsonElement("Change")]
+			public int CommitOrder { get; set; }
+
 			public DateTime? DeclinedAt { get; set; }
 			public DateTime? ResolvedAt { get; set; } // Degenerate
 
@@ -220,21 +241,21 @@ namespace HordeServer.Issues
 				Id = ObjectId.GenerateNewId();
 				IssueId = issueId;
 				AuthorId = newSuspect.AuthorId;
-				Change = newSuspect.Change;
+				CommitId = newSuspect.CommitId;
 				ResolvedAt = resolvedAt;
 			}
 
 			public IssueSuspectDocument(int issueId, IIssueSpanSuspect suspect)
-				: this(issueId, suspect.AuthorId, suspect.OriginatingChange ?? suspect.Change, null, null)
+				: this(issueId, suspect.AuthorId, suspect.SourceCommitId ?? suspect.CommitId, null, null)
 			{
 			}
 
-			public IssueSuspectDocument(int issueId, UserId authorId, int change, DateTime? declinedAt, DateTime? resolvedAt)
+			public IssueSuspectDocument(int issueId, UserId authorId, CommitIdWithOrder commit, DateTime? declinedAt, DateTime? resolvedAt)
 			{
 				Id = ObjectId.GenerateNewId();
 				IssueId = issueId;
 				AuthorId = authorId;
-				Change = change;
+				CommitId = commit;
 				DeclinedAt = declinedAt;
 				ResolvedAt = resolvedAt;
 			}
@@ -425,8 +446,11 @@ namespace HordeServer.Issues
 			[BsonRequired]
 			public IssueFingerprintDocument Fingerprint { get; set; }
 
-			public int MinChange { get; set; }
-			public int MaxChange { get; set; } = Int32.MaxValue;
+			[BsonElement("MinChange")]
+			public int MinCommitOrder { get; set; }
+
+			[BsonElement("MaxChange")]
+			public int MaxCommitOrder { get; set; } = Int32.MaxValue;
 
 			public IssueStepDocument? LastSuccess { get; set; }
 
@@ -476,14 +500,14 @@ namespace HordeServer.Issues
 				Fingerprint = new IssueFingerprintDocument(newSpan.Fingerprint);
 				if (newSpan.LastSuccess != null)
 				{
-					MinChange = newSpan.LastSuccess.Change;
+					MinCommitOrder = newSpan.LastSuccess.CommitId.Order;
 					LastSuccess = new IssueStepDocument(Id, newSpan.LastSuccess);
 				}
 				FirstFailure = new IssueStepDocument(Id, newSpan.FirstFailure);
 				LastFailure = new IssueStepDocument(Id, newSpan.FirstFailure);
 				if (newSpan.NextSuccess != null)
 				{
-					MaxChange = newSpan.NextSuccess.Change;
+					MaxCommitOrder = newSpan.NextSuccess.CommitId.Order;
 					NextSuccess = new IssueStepDocument(Id, newSpan.NextSuccess);
 				}
 				PromoteByDefault = newSpan.FirstFailure.PromoteByDefault;
@@ -494,9 +518,33 @@ namespace HordeServer.Issues
 
 		class IssueSpanSuspectDocument : IIssueSpanSuspect
 		{
-			public int Change { get; set; }
+			[BsonIgnore]
+			public CommitIdWithOrder CommitId
+			{
+				get => (CommitName != null)? new CommitIdWithOrder(CommitName, CommitOrder) : CommitIdWithOrder.FromPerforceChange(CommitOrder);
+				set => (CommitName, CommitOrder) = (value.Name, value.Order);
+			}
+
+			[BsonElement("Commit")]
+			private string? CommitName { get; set; }
+
+			[BsonElement("Change")]
+			private int CommitOrder { get; set; }
+
 			public UserId AuthorId { get; set; }
-			public int? OriginatingChange { get; set; }
+
+			[BsonIgnore]
+			public CommitIdWithOrder? SourceCommitId
+			{
+				get => (SourceCommitName != null)? new CommitIdWithOrder(SourceCommitName, SourceCommitOrder ?? 0) : (SourceCommitOrder != null)? CommitIdWithOrder.FromPerforceChange(SourceCommitOrder.Value) : null;
+				set => (SourceCommitName, SourceCommitOrder) = (value?.Name, value?.Order);
+			}
+
+			[BsonElement("SourceCommit")]
+			public string? SourceCommitName { get; set; }
+
+			[BsonElement("OriginatingChange")]
+			public int? SourceCommitOrder { get; set; }
 
 			[BsonConstructor]
 			private IssueSpanSuspectDocument()
@@ -505,9 +553,9 @@ namespace HordeServer.Issues
 
 			public IssueSpanSuspectDocument(NewIssueSpanSuspectData newSuspectData)
 			{
-				Change = newSuspectData.Change;
+				CommitId = newSuspectData.CommitId;
 				AuthorId = newSuspectData.AuthorId;
-				OriginatingChange = newSuspectData.OriginatingChange;
+				SourceCommitId = newSuspectData.SourceCommitId;
 			}
 		}
 
@@ -516,7 +564,19 @@ namespace HordeServer.Issues
 			public ObjectId Id { get; set; }
 			public ObjectId SpanId { get; set; }
 
-			public int Change { get; set; }
+			[BsonIgnore]
+			public CommitIdWithOrder CommitId
+			{
+				get => (CommitName != null)? new CommitIdWithOrder(CommitName, CommitOrder) : CommitIdWithOrder.FromPerforceChange(CommitOrder);
+				set => (CommitName, CommitOrder) = (value.Name, value.Order);
+			}
+
+			[BsonElement("Commit")]
+			public string? CommitName { get; set; }
+
+			[BsonElement("Change")]
+			public int CommitOrder { get; set; }
+
 			public IssueSeverity Severity { get; set; }
 
 			[BsonRequired]
@@ -556,7 +616,7 @@ namespace HordeServer.Issues
 			{
 				Id = ObjectId.GenerateNewId();
 				SpanId = spanId;
-				Change = stepData.Change;
+				CommitId = stepData.CommitId;
 				Severity = stepData.Severity;
 				JobName = stepData.JobName;
 				JobId = stepData.JobId;
@@ -598,6 +658,7 @@ namespace HordeServer.Issues
 
 		readonly IRedisService _redisService;
 		readonly IUserCollection _userCollection;
+		readonly ICommitService _commitService;
 		readonly ISingletonDocument<IssueLedger> _ledgerSingleton;
 		readonly IMongoCollection<IssueDocument> _issues;
 		readonly IMongoCollection<IssueSpanDocument> _issueSpans;
@@ -628,10 +689,11 @@ namespace HordeServer.Issues
 			});
 		}
 
-		public IssueCollection(IMongoService mongoService, IRedisService redisService, IUserCollection userCollection, IAuditLogFactory<int> auditLogFactory, ITelemetryWriter telemetryWriter, IOptionsMonitor<BuildConfig> buildConfig, Tracer tracer, ILogger<IssueCollection> logger)
+		public IssueCollection(IMongoService mongoService, IRedisService redisService, IUserCollection userCollection, ICommitService commitService, IAuditLogFactory<int> auditLogFactory, ITelemetryWriter telemetryWriter, IOptionsMonitor<BuildConfig> buildConfig, Tracer tracer, ILogger<IssueCollection> logger)
 		{
 			_redisService = redisService;
 			_userCollection = userCollection;
+			_commitService = commitService;
 			_telemetryWriter = telemetryWriter;
 			_buildConfig = buildConfig;
 			_tracer = tracer;
@@ -646,8 +708,8 @@ namespace HordeServer.Issues
 
 			List<MongoIndex<IssueSpanDocument>> issueSpanIndexes = new List<MongoIndex<IssueSpanDocument>>();
 			issueSpanIndexes.Add(keys => keys.Ascending(x => x.IssueId));
-			issueSpanIndexes.Add(keys => keys.Ascending(x => x.StreamId).Ascending(x => x.MinChange).Ascending(x => x.MaxChange));
-			issueSpanIndexes.Add("StreamChanges", keys => keys.Ascending(x => x.StreamId).Ascending(x => x.TemplateRefId).Ascending(x => x.NodeName).Ascending(x => x.MinChange).Ascending(x => x.MaxChange));
+			issueSpanIndexes.Add(keys => keys.Ascending(x => x.StreamId).Ascending(x => x.MinCommitOrder).Ascending(x => x.MaxCommitOrder));
+			issueSpanIndexes.Add("StreamChanges", keys => keys.Ascending(x => x.StreamId).Ascending(x => x.TemplateRefId).Ascending(x => x.NodeName).Ascending(x => x.MinCommitOrder).Ascending(x => x.MaxCommitOrder));
 			_issueSpans = mongoService.GetCollection<IssueSpanDocument>("IssuesV2.Spans", issueSpanIndexes);
 
 			List<MongoIndex<IssueStepDocument>> issueStepIndexes = new List<MongoIndex<IssueStepDocument>>();
@@ -656,9 +718,9 @@ namespace HordeServer.Issues
 			_issueSteps = mongoService.GetCollection<IssueStepDocument>("IssuesV2.Steps", issueStepIndexes);
 
 			List<MongoIndex<IssueSuspectDocument>> issueSuspectIndexes = new List<MongoIndex<IssueSuspectDocument>>();
-			issueSuspectIndexes.Add(keys => keys.Ascending(x => x.Change));
+			issueSuspectIndexes.Add(keys => keys.Ascending(x => x.CommitOrder));
 			issueSuspectIndexes.Add(keys => keys.Ascending(x => x.AuthorId).Ascending(x => x.ResolvedAt));
-			issueSuspectIndexes.Add(keys => keys.Ascending(x => x.IssueId).Ascending(x => x.Change), unique: true);
+			issueSuspectIndexes.Add(keys => keys.Ascending(x => x.IssueId).Ascending(x => x.CommitOrder), unique: true);
 			_issueSuspects = mongoService.GetCollection<IssueSuspectDocument>("IssuesV2.Suspects", issueSuspectIndexes);
 
 			_auditLog = auditLogFactory.Create("IssuesV2.History", "IssueId");
@@ -757,8 +819,8 @@ namespace HordeServer.Issues
 					AcknowledgedAt = issue.AcknowledgedAt,
 					CreatedAt = issue.CreatedAt,
 					OwnerId = issue.OwnerId,
-					FixChange = issue.FixChange,
-					FixedSystemic = issue.FixedSystemic,
+					FixChange = issue.FixCommitId,
+					FixedSystemic = issue.FixSystemic,
 					LastSeenAt = issue.LastSeenAt,
 					NominatedAt = issue.NominatedAt,
 					NominatedById = issue.NominatedById,
@@ -842,26 +904,26 @@ namespace HordeServer.Issues
 					issueLogger.LogInformation("Issue was acknowledged by {UserName} ({UserId})", await GetUserNameAsync(newIssue.OwnerId, cancellationToken), newIssue.OwnerId);
 				}
 			}
-			if (newIssue.FixChange != oldIssue.FixChange)
+			if (newIssue.FixCommitId != oldIssue.FixCommitId)
 			{
-				if (newIssue.FixChange == 0)
+				if (newIssue.FixCommitId == null)
 				{
 					issueLogger.LogInformation("Issue was marked as not fixed");
 				}
 				else
 				{
-					issueLogger.LogInformation("Issue was marked as fixed in {Change}", newIssue.FixChange);
+					issueLogger.LogInformation("Issue was marked as fixed in {Commit}", newIssue.FixCommitId);
 				}
 			}
-			if (newIssue.FixedSystemic != oldIssue.FixedSystemic)
+			if (newIssue.FixSystemic != oldIssue.FixSystemic)
 			{
-				if (newIssue.FixChange == 0)
+				if (newIssue.FixSystemic)
 				{
-					issueLogger.LogInformation("Issue was marked as not fixed");
+					issueLogger.LogInformation("Issue was marked fixed as a systemic issue");
 				}
 				else
 				{
-					issueLogger.LogInformation("Issue was marked as fixed in {Change}", newIssue.FixChange);
+					issueLogger.LogInformation("Issue was marked as not fixed as a systemic issue");
 				}
 			}
 			if (newIssue.ResolvedById != oldIssue.ResolvedById)
@@ -957,15 +1019,15 @@ namespace HordeServer.Issues
 
 		async Task LogIssueSuspectChangesAsync(ILogger issueLogger, IReadOnlyList<IssueSuspectDocument> oldIssueSuspects, List<IssueSuspectDocument> newIssueSuspects, CancellationToken cancellationToken)
 		{
-			HashSet<(UserId, int)> oldSuspects = new HashSet<(UserId, int)>(oldIssueSuspects.Select(x => (x.AuthorId, x.Change)));
-			HashSet<(UserId, int)> newSuspects = new HashSet<(UserId, int)>(newIssueSuspects.Select(x => (x.AuthorId, x.Change)));
-			foreach ((UserId userId, int change) in newSuspects.Where(x => !oldSuspects.Contains(x)))
+			HashSet<(UserId, CommitId)> oldSuspects = new HashSet<(UserId, CommitId)>(oldIssueSuspects.Select(x => (x.AuthorId, (CommitId)x.CommitId)));
+			HashSet<(UserId, CommitId)> newSuspects = new HashSet<(UserId, CommitId)>(newIssueSuspects.Select(x => (x.AuthorId, (CommitId)x.CommitId)));
+			foreach ((UserId userId, CommitId commit) in newSuspects.Where(x => !oldSuspects.Contains(x)))
 			{
-				issueLogger.LogInformation("Added suspect {UserName} ({UserId}) for change {Change}", await GetUserNameAsync(userId, cancellationToken), userId, change);
+				issueLogger.LogInformation("Added suspect {UserName} ({UserId}) for commit {Commit}", await GetUserNameAsync(userId, cancellationToken), userId, commit);
 			}
-			foreach ((UserId userId, int change) in oldSuspects.Where(x => !newSuspects.Contains(x)))
+			foreach ((UserId userId, CommitId commit) in oldSuspects.Where(x => !newSuspects.Contains(x)))
 			{
-				issueLogger.LogInformation("Removed suspect {UserName} ({UserId}) for change {Change}", await GetUserNameAsync(userId, cancellationToken), userId, change);
+				issueLogger.LogInformation("Removed suspect {UserName} ({UserId}) for commit {Commit}", await GetUserNameAsync(userId, cancellationToken), userId, commit);
 			}
 
 			HashSet<UserId> oldDeclinedBy = new HashSet<UserId>(oldIssueSuspects.Where(x => x.DeclinedAt != null).Select(x => x.AuthorId));
@@ -1000,13 +1062,13 @@ namespace HordeServer.Issues
 		}
 
 		/// <inheritdoc/>
-		public async Task<IReadOnlyList<IIssue>> FindIssuesAsync(IEnumerable<int>? ids = null, UserId? ownerId = null, StreamId? streamId = null, int? minChange = null, int? maxChange = null, bool? resolved = null, bool? promoted = null, int? index = null, int? count = null, CancellationToken cancellationToken = default)
+		public async Task<IReadOnlyList<IIssue>> FindIssuesAsync(IEnumerable<int>? ids = null, UserId? ownerId = null, StreamId? streamId = null, CommitId? minCommitId = null, CommitId? maxCommitId = null, bool? resolved = null, bool? promoted = null, int? index = null, int? count = null, CancellationToken cancellationToken = default)
 		{
 			IReadOnlyList<IIssue> results;
 
 			if (ownerId == null)
 			{
-				results = await FilterIssuesByStreamIdAsync(ids, streamId, minChange, maxChange, resolved ?? false, promoted, index ?? 0, count, cancellationToken);
+				results = await FilterIssuesByStreamIdAsync(ids, streamId, minCommitId, maxCommitId, resolved ?? false, promoted, index ?? 0, count, cancellationToken);
 			}
 			else
 			{
@@ -1016,11 +1078,16 @@ namespace HordeServer.Issues
 			return results;
 		}
 
-		async Task<IReadOnlyList<IIssue>> FilterIssuesByStreamIdAsync(IEnumerable<int>? ids, StreamId? streamId, int? minChange, int? maxChange, bool? resolved, bool? promoted, int index, int? count, CancellationToken cancellationToken)
+		async Task<IReadOnlyList<IIssue>> FilterIssuesByStreamIdAsync(IEnumerable<int>? ids, StreamId? streamId, CommitId? minCommitId, CommitId? maxCommitId, bool? resolved, bool? promoted, int index, int? count, CancellationToken cancellationToken)
 		{
 			if (streamId == null)
 			{
-				return await FilterIssuesByOtherFieldsAsync(ids, minChange, maxChange, resolved, promoted, index, count, cancellationToken);
+				if (minCommitId != null || maxCommitId != null)
+				{
+					throw new ArgumentException("Cannot search commit range without stream parameter");
+				}
+
+				return await FilterIssuesByOtherFieldsAsync(ids, resolved, promoted, index, count, cancellationToken);
 			}
 			else
 			{
@@ -1034,13 +1101,15 @@ namespace HordeServer.Issues
 					filter &= Builders<IssueSpanDocument>.Filter.Exists(x => x.IssueId);
 				}
 
-				if (minChange != null)
+				if (minCommitId != null)
 				{
-					filter &= Builders<IssueSpanDocument>.Filter.Not(Builders<IssueSpanDocument>.Filter.Lt(x => x.MaxChange, minChange.Value));
+					CommitIdWithOrder minCommitWithOrder = await _commitService.GetOrderedAsync(streamId.Value, minCommitId, cancellationToken);
+					filter &= Builders<IssueSpanDocument>.Filter.Not(Builders<IssueSpanDocument>.Filter.Lt(x => x.MaxCommitOrder, minCommitWithOrder.Order));
 				}
-				if (maxChange != null)
+				if (maxCommitId != null)
 				{
-					filter &= Builders<IssueSpanDocument>.Filter.Not(Builders<IssueSpanDocument>.Filter.Gt(x => x.MinChange, maxChange.Value));
+					CommitIdWithOrder maxCommitWithOrder = await _commitService.GetOrderedAsync(streamId.Value, maxCommitId, cancellationToken);
+					filter &= Builders<IssueSpanDocument>.Filter.Not(Builders<IssueSpanDocument>.Filter.Gt(x => x.MinCommitOrder, maxCommitWithOrder.Order));
 				}
 
 				if (resolved != null)
@@ -1057,7 +1126,7 @@ namespace HordeServer.Issues
 
 				using (IAsyncCursor<ProjectedIssueId> cursor = await _issueSpans.Aggregate().Match(filter).Group(x => x.IssueId, x => new ProjectedIssueId { _id = x.Key }).SortByDescending(x => x._id).ToCursorAsync(cancellationToken))
 				{
-					List<IssueDocument> results = await PaginatedJoinAsync(cursor, (nextIds, nextIndex, nextCount) => FilterIssuesByOtherFieldsAsync(nextIds, null, null, null, promoted, nextIndex, nextCount, cancellationToken), index, count, cancellationToken);
+					List<IssueDocument> results = await PaginatedJoinAsync(cursor, (nextIds, nextIndex, nextCount) => FilterIssuesByOtherFieldsAsync(nextIds, null, promoted, nextIndex, nextCount, cancellationToken), index, count, cancellationToken);
 					if (resolved != null)
 					{
 						for (int idx = results.Count - 1; idx >= 0; idx--)
@@ -1075,7 +1144,7 @@ namespace HordeServer.Issues
 			}
 		}
 
-		async Task<List<IssueDocument>> FilterIssuesByOtherFieldsAsync(IEnumerable<int>? ids, int? minChange, int? maxChange, bool? resolved, bool? promoted, int index, int? count, CancellationToken cancellationToken)
+		async Task<List<IssueDocument>> FilterIssuesByOtherFieldsAsync(IEnumerable<int>? ids, bool? resolved, bool? promoted, int index, int? count, CancellationToken cancellationToken)
 		{
 			FilterDefinition<IssueDocument> filter = FilterDefinition<IssueDocument>.Empty;
 			if (ids != null)
@@ -1093,14 +1162,6 @@ namespace HordeServer.Issues
 					filter &= Builders<IssueDocument>.Filter.Eq(x => x.ResolvedAt, null);
 				}
 			}
-			if (minChange != null)
-			{
-				filter &= Builders<IssueDocument>.Filter.Not(Builders<IssueDocument>.Filter.Lt(x => x.MaxSuspectChange, minChange.Value));
-			}
-			if (maxChange != null)
-			{
-				filter &= Builders<IssueDocument>.Filter.Not(Builders<IssueDocument>.Filter.Gt(x => x.MinSuspectChange, maxChange.Value));
-			}
 			if (promoted != null)
 			{
 				if (promoted.Value)
@@ -1147,14 +1208,15 @@ namespace HordeServer.Issues
 		}
 
 		/// <inheritdoc/>
-		public async Task<IReadOnlyList<IIssue>> FindIssuesForChangesAsync(List<int> changes, CancellationToken cancellationToken)
+		public async Task<IReadOnlyList<IIssue>> FindIssuesForChangesAsync(List<CommitIdWithOrder> commits, CancellationToken cancellationToken)
 		{
-			List<int> issueIds = await (await _issueSuspects.DistinctAsync(x => x.IssueId, Builders<IssueSuspectDocument>.Filter.In(x => x.Change, changes), cancellationToken: cancellationToken)).ToListAsync(cancellationToken);
+			List<int> commitOrders = commits.ConvertAll(x => x.Order);
+			List<int> issueIds = await (await _issueSuspects.DistinctAsync(x => x.IssueId, Builders<IssueSuspectDocument>.Filter.In(x => x.CommitOrder, commitOrders), cancellationToken: cancellationToken)).ToListAsync(cancellationToken);
 			return await _issues.Find(Builders<IssueDocument>.Filter.In(x => x.Id, issueIds)).ToListAsync(cancellationToken);
 		}
 
 		/// <inheritdoc/>
-		public async Task<IIssue?> TryUpdateIssueAsync(IIssue issue, UserId? initiatedByUserId, IssueSeverity? newSeverity = null, string? newSummary = null, string? newUserSummary = null, string? newDescription = null, bool? newManuallyPromoted = null, UserId? newOwnerId = null, UserId? newNominatedById = null, bool? newAcknowledged = null, UserId? newDeclinedById = null, int? newFixChange = null, bool? newFixedSystemic = null, UserId? newResolvedById = null, List<ObjectId>? newExcludeSpanIds = null, DateTime? newLastSeenAt = null, string? newExternaIssueKey = null, UserId? newQuarantinedById = null, UserId? newForceClosedById = null, Uri? newWorkflowThreadUrl = null, CancellationToken cancellationToken = default)
+		public async Task<IIssue?> TryUpdateIssueAsync(IIssue issue, UserId? initiatedByUserId, IssueSeverity? newSeverity = null, string? newSummary = null, string? newUserSummary = null, string? newDescription = null, bool? newManuallyPromoted = null, UserId? newOwnerId = null, UserId? newNominatedById = null, bool? newAcknowledged = null, UserId? newDeclinedById = null, CommitId? newFixCommit = null, bool? newFixSystemic = null, UserId? newResolvedById = null, List<ObjectId>? newExcludeSpanIds = null, DateTime? newLastSeenAt = null, string? newExternaIssueKey = null, UserId? newQuarantinedById = null, UserId? newForceClosedById = null, Uri? newWorkflowThreadUrl = null, CancellationToken cancellationToken = default)
 		{
 			IssueDocument issueDocument = (IssueDocument)issue;
 
@@ -1251,30 +1313,26 @@ namespace HordeServer.Issues
 					}
 				}
 			}
-			if (newFixChange != null)
+			if (newFixCommit != null)
 			{
-				if (newFixChange < 0)
-				{
-					updates.Add(Builders<IssueDocument>.Update.Unset(x => x.FixChange!).Set(x => x.FixedSystemic, true));
-				}
-				else if (newFixChange == 0)
+				if (String.IsNullOrEmpty(newFixCommit.Name))
 				{
-					updates.Add(Builders<IssueDocument>.Update.Unset(x => x.FixChange!).Unset(x => x.FixedSystemic));
+					updates.Add(Builders<IssueDocument>.Update.Unset(x => x.FixCommitName).Unset(x => x.FixChange).Unset(x => x.FixSystemic));
 				}
 				else
 				{
-					updates.Add(Builders<IssueDocument>.Update.Set(x => x.FixChange, newFixChange).Unset(x => x.FixedSystemic));
+					updates.Add(Builders<IssueDocument>.Update.Set(x => x.FixCommitName, newFixCommit.Name).Unset(x => x.FixChange).Unset(x => x.FixSystemic));
 				}
 			}
-			else if (newFixedSystemic != null)
+			else if (newFixSystemic != null)
 			{
-				if (newFixedSystemic.Value)
+				if (newFixSystemic.Value)
 				{
-					updates.Add(Builders<IssueDocument>.Update.Unset(x => x.FixChange!).Set(x => x.FixedSystemic, true));
+					updates.Add(Builders<IssueDocument>.Update.Unset(x => x.FixCommitName).Unset(x => x.FixChange).Set(x => x.FixSystemic, true));
 				}
 				else
 				{
-					updates.Add(Builders<IssueDocument>.Update.Unset(x => x.FixedSystemic));
+					updates.Add(Builders<IssueDocument>.Update.Unset(x => x.FixCommitName).Unset(x => x.FixChange).Unset(x => x.FixSystemic));
 				}
 			}
 			if (newResolvedById != null)
@@ -1435,10 +1493,6 @@ namespace HordeServer.Issues
 				}
 			}
 
-			// Get the range of suspect changes
-			int newMinSuspectChange = (newSuspects.Count > 0) ? newSuspects.Min(x => x.Change) : 0;
-			int newMaxSuspectChange = (newSuspects.Count > 0) ? newSuspects.Min(x => x.Change) : 0;
-
 			// Perform the actual update with this data
 			List<UpdateDefinition<IssueDocument>> updates = new List<UpdateDefinition<IssueDocument>>();
 			if (!String.Equals(issue.Summary, newSummary, StringComparison.Ordinal))
@@ -1461,14 +1515,6 @@ namespace HordeServer.Issues
 			{
 				updates.Add(Builders<IssueDocument>.Update.Set(x => x.Streams, newStreams.Select(x => new IssueStreamDocument(x))));
 			}
-			if (issueImpl.MinSuspectChange != newMinSuspectChange)
-			{
-				updates.Add(Builders<IssueDocument>.Update.Set(x => x.MinSuspectChange, newMinSuspectChange));
-			}
-			if (issueImpl.MaxSuspectChange != newMaxSuspectChange)
-			{
-				updates.Add(Builders<IssueDocument>.Update.Set(x => x.MaxSuspectChange, newMaxSuspectChange));
-			}
 			if (issueImpl.DefaultOwnerId != newDefaultOwnerId)
 			{
 				updates.Add(Builders<IssueDocument>.Update.Set(x => x.DefaultOwnerId, newDefaultOwnerId));
@@ -1505,11 +1551,11 @@ namespace HordeServer.Issues
 			List<IssueSuspectDocument> newSuspectImpls = new List<IssueSuspectDocument>(oldSuspectImpls);
 
 			// Find the current list of suspects
-			HashSet<(UserId, int)> curSuspectKeys = new HashSet<(UserId, int)>(oldSuspectImpls.Select(x => (x.AuthorId, x.Change)));
-			List<IssueSuspectDocument> createSuspects = newSuspects.Where(x => !curSuspectKeys.Contains((x.AuthorId, x.Change))).Select(x => new IssueSuspectDocument(issueId, x, resolvedAt)).ToList();
+			HashSet<(UserId, CommitId)> curSuspectKeys = new HashSet<(UserId, CommitId)>(oldSuspectImpls.Select(x => (x.AuthorId, (CommitId)x.CommitId)));
+			List<IssueSuspectDocument> createSuspects = newSuspects.Where(x => !curSuspectKeys.Contains((x.AuthorId, x.CommitId))).Select(x => new IssueSuspectDocument(issueId, x, resolvedAt)).ToList();
 
-			HashSet<(UserId, int)> newSuspectKeys = new HashSet<(UserId, int)>(newSuspects.Select(x => (x.AuthorId, x.Change)));
-			List<IssueSuspectDocument> deleteSuspects = oldSuspectImpls.Where(x => !newSuspectKeys.Contains((x.AuthorId, x.Change))).ToList();
+			HashSet<(UserId, CommitId)> newSuspectKeys = new HashSet<(UserId, CommitId)>(newSuspects.Select(x => (x.AuthorId, (CommitId)x.CommitId)));
+			List<IssueSuspectDocument> deleteSuspects = oldSuspectImpls.Where(x => !newSuspectKeys.Contains((x.AuthorId, x.CommitId))).ToList();
 
 			// Apply the suspect changes
 			if (createSuspects.Count > 0)
@@ -1520,7 +1566,7 @@ namespace HordeServer.Issues
 			if (deleteSuspects.Count > 0)
 			{
 				await _issueSuspects.DeleteManyAsync(Builders<IssueSuspectDocument>.Filter.In(x => x.Id, deleteSuspects.Select(y => y.Id)), cancellationToken);
-				newSuspectImpls.RemoveAll(x => !newSuspectKeys.Contains((x.AuthorId, x.Change)));
+				newSuspectImpls.RemoveAll(x => !newSuspectKeys.Contains((x.AuthorId, x.CommitId)));
 			}
 
 			// Make sure all the remaining suspects have the correct resolved time
@@ -1554,8 +1600,8 @@ namespace HordeServer.Issues
 					Id = issueSpan.Id,
 					IssueId = issueSpan.IssueId,
 					Fingerprint = new { Type = issueSpan.Fingerprint.Type, Keys = issueSpan.Fingerprint.Keys },
-					FirstFailure = new { JobId = issueSpan.FirstFailure.JobId, JobName = issueSpan.FirstFailure.JobName, Change = issueSpan.FirstFailure.Change, StepId = issueSpan.FirstFailure.SpanId },
-					LastFailure = (issueSpan.LastFailure != null) ? new { JobId = issueSpan.LastFailure.JobId, JobName = issueSpan.LastFailure.JobName, Change = issueSpan.LastFailure.Change, StepId = issueSpan.LastFailure.SpanId } : null,
+					FirstFailure = new { JobId = issueSpan.FirstFailure.JobId, JobName = issueSpan.FirstFailure.JobName, Change = issueSpan.FirstFailure.CommitId, StepId = issueSpan.FirstFailure.SpanId },
+					LastFailure = (issueSpan.LastFailure != null) ? new { JobId = issueSpan.LastFailure.JobId, JobName = issueSpan.LastFailure.JobName, Change = issueSpan.LastFailure.CommitId, StepId = issueSpan.LastFailure.SpanId } : null,
 					StreamId = issueSpan.StreamId,
 					StreamName = issueSpan.StreamName,
 					TemplateRefId = issueSpan.TemplateRefId
@@ -1575,27 +1621,27 @@ namespace HordeServer.Issues
 			List<UpdateDefinition<IssueSpanDocument>> updates = new List<UpdateDefinition<IssueSpanDocument>>();
 			if (newLastSuccess != null)
 			{
-				updates.Add(Builders<IssueSpanDocument>.Update.Set(x => x.MinChange, newLastSuccess.Change));
+				updates.Add(Builders<IssueSpanDocument>.Update.Set(x => x.MinCommitOrder, newLastSuccess.CommitId.Order));
 				updates.Add(Builders<IssueSpanDocument>.Update.Set(x => x.LastSuccess, new IssueStepDocument(span.Id, newLastSuccess)));
 			}
 			if (newFailure != null)
 			{
-				if (newFailure.Change < span.FirstFailure.Change)
+				if (newFailure.CommitId < span.FirstFailure.CommitId)
 				{
 					updates.Add(Builders<IssueSpanDocument>.Update.Set(x => x.FirstFailure, new IssueStepDocument(span.Id, newFailure)));
 				}
-				if (newFailure.Change >= span.LastFailure.Change)
+				if (newFailure.CommitId >= span.LastFailure.CommitId)
 				{
 					updates.Add(Builders<IssueSpanDocument>.Update.Set(x => x.LastFailure, new IssueStepDocument(span.Id, newFailure)));
 				}
-				if (newFailure.PromoteByDefault != span.PromoteByDefault && newFailure.Change >= span.LastFailure.Change)
+				if (newFailure.PromoteByDefault != span.PromoteByDefault && newFailure.CommitId >= span.LastFailure.CommitId)
 				{
 					updates.Add(Builders<IssueSpanDocument>.Update.Set(x => x.PromoteByDefault, newFailure.PromoteByDefault));
 				}
 			}
 			if (newNextSuccess != null)
 			{
-				updates.Add(Builders<IssueSpanDocument>.Update.Set(x => x.MaxChange, newNextSuccess.Change));
+				updates.Add(Builders<IssueSpanDocument>.Update.Set(x => x.MaxCommitOrder, newNextSuccess.CommitId.Order));
 				updates.Add(Builders<IssueSpanDocument>.Update.Set(x => x.NextSuccess, new IssueStepDocument(span.Id, newNextSuccess)));
 			}
 			if (newSuspects != null)
@@ -1618,15 +1664,15 @@ namespace HordeServer.Issues
 				ILogger logger = GetLogger(newSpan.IssueId);
 				if (newLastSuccess != null)
 				{
-					logger.LogInformation("Set last success for span {SpanId} to job {JobId} at CL {Change}", newSpan.Id, newLastSuccess.JobId, newLastSuccess.Change);
+					logger.LogInformation("Set last success for span {SpanId} to job {JobId} at commit {Commit}", newSpan.Id, newLastSuccess.JobId, newLastSuccess.CommitId);
 				}
 				if (newNextSuccess != null)
 				{
-					logger.LogInformation("Set next success for span {SpanId} to job {JobId} at CL {Change}", newSpan.Id, newNextSuccess.JobId, newNextSuccess.Change);
+					logger.LogInformation("Set next success for span {SpanId} to job {JobId} at commit {Commit}", newSpan.Id, newNextSuccess.JobId, newNextSuccess.CommitId);
 				}
 				if (newFailure != null)
 				{
-					logger.LogInformation("Added failure for span {SpanId} in job {JobId} at CL {Change}", newSpan.Id, newFailure.JobId, newFailure.Change);
+					logger.LogInformation("Added failure for span {SpanId} in job {JobId} at commit {Commit}", newSpan.Id, newFailure.JobId, newFailure.CommitId);
 				}
 				SendTelemetry(newSpan);
 			}
@@ -1646,13 +1692,14 @@ namespace HordeServer.Issues
 		}
 
 		/// <inheritdoc/>
-		public async Task<IReadOnlyList<IIssueSpan>> FindOpenSpansAsync(StreamId streamId, TemplateId templateId, string nodeName, int change, CancellationToken cancellationToken)
+		public async Task<IReadOnlyList<IIssueSpan>> FindOpenSpansAsync(StreamId streamId, TemplateId templateId, string nodeName, CommitId commitId, CancellationToken cancellationToken)
 		{
-			return await _issueSpans.Find(x => x.StreamId == streamId && x.TemplateRefId == templateId && x.NodeName == nodeName && change >= x.MinChange && change <= x.MaxChange).ToListAsync(cancellationToken);
+			CommitIdWithOrder commitIdWithOrder = await _commitService.GetOrderedAsync(streamId, commitId, cancellationToken);
+			return await _issueSpans.Find(x => x.StreamId == streamId && x.TemplateRefId == templateId && x.NodeName == nodeName && commitIdWithOrder.Order >= x.MinCommitOrder && commitIdWithOrder.Order <= x.MaxCommitOrder).ToListAsync(cancellationToken);
 		}
 
 		/// <inheritdoc/>
-		public async Task<IReadOnlyList<IIssueSpan>> FindSpansAsync(IEnumerable<ObjectId>? spanIds, IEnumerable<int>? issueIds, StreamId? streamId, int? minChange, int? maxChange, bool? resolved, int? index, int? count, CancellationToken cancellationToken)
+		public async Task<IReadOnlyList<IIssueSpan>> FindSpansAsync(IEnumerable<ObjectId>? spanIds, IEnumerable<int>? issueIds, StreamId? streamId, CommitId? minCommitId, CommitId? maxCommitId, bool? resolved, int? index, int? count, CancellationToken cancellationToken)
 		{
 			FilterDefinition<IssueSpanDocument> filter = FilterDefinition<IssueSpanDocument>.Empty;
 
@@ -1664,20 +1711,22 @@ namespace HordeServer.Issues
 			if (streamId != null)
 			{
 				filter &= Builders<IssueSpanDocument>.Filter.Eq(x => x.StreamId, streamId);
+				if (minCommitId != null)
+				{
+					CommitIdWithOrder minCommitIdWithOrder = await _commitService.GetOrderedAsync(streamId.Value, minCommitId, cancellationToken);
+					filter &= Builders<IssueSpanDocument>.Filter.Not(Builders<IssueSpanDocument>.Filter.Lt(x => x.MaxCommitOrder, minCommitIdWithOrder.Order));
+				}
+				if (maxCommitId != null)
+				{
+					CommitIdWithOrder maxCommitIdWithOrder = await _commitService.GetOrderedAsync(streamId.Value, maxCommitId, cancellationToken);
+					filter &= Builders<IssueSpanDocument>.Filter.Not(Builders<IssueSpanDocument>.Filter.Gt(x => x.MinCommitOrder, maxCommitIdWithOrder.Order));
+				}
 			}
 
 			if (issueIds != null)
 			{
 				filter &= Builders<IssueSpanDocument>.Filter.In(x => x.IssueId, issueIds.Select<int, int?>(x => x));
 			}
-			if (minChange != null)
-			{
-				filter &= Builders<IssueSpanDocument>.Filter.Not(Builders<IssueSpanDocument>.Filter.Lt(x => x.MaxChange, minChange.Value));
-			}
-			if (maxChange != null)
-			{
-				filter &= Builders<IssueSpanDocument>.Filter.Not(Builders<IssueSpanDocument>.Filter.Gt(x => x.MinChange, maxChange.Value));
-			}
 			if (resolved != null)
 			{
 				if (resolved.Value)
diff --git a/Engine/Source/Programs/Horde/Plugins/Build/HordeServer.Build/Issues/IssueService.cs b/Engine/Source/Programs/Horde/Plugins/Build/HordeServer.Build/Issues/IssueService.cs
--- a/Engine/Source/Programs/Horde/Plugins/Build/HordeServer.Build/Issues/IssueService.cs
+++ b/Engine/Source/Programs/Horde/Plugins/Build/HordeServer.Build/Issues/IssueService.cs
@@ -5,6 +5,7 @@ using System.Reflection;
 using System.Text;
 using System.Text.RegularExpressions;
 using EpicGames.Core;
+using EpicGames.Horde.Commits;
 using EpicGames.Horde.Issues;
 using EpicGames.Horde.Jobs;
 using EpicGames.Horde.Jobs.Templates;
@@ -443,7 +444,8 @@ namespace HordeServer.Issues
 		/// <param name="nominatedById">Person that nominated the new owner</param>
 		/// <param name="acknowledged">Whether the issue has been acknowledged</param>
 		/// <param name="declinedById">Name of the user that has declined this issue</param>
-		/// <param name="fixChange">Fix changelist for the issue</param>
+		/// <param name="fixCommitId">Fix changelist for the issue</param>
+		/// <param name="fixSystemic">There was a systemic fix for this issue</param>
 		/// <param name="resolvedById">Whether the issue has been resolved</param>
 		/// <param name="addSpanIds">Add spans to this issue</param>
 		/// <param name="removeSpanIds">Remove spans from this issue</param>
@@ -454,7 +456,7 @@ namespace HordeServer.Issues
 		/// <param name="workflowThreadUrl">The workfloe thread created for the issue</param>
 		/// <param name="cancellationToken">Cancellation token for the operation</param>
 		/// <returns>True if the issue was updated</returns>
-		public async Task<bool> UpdateIssueAsync(int id, string? summary = null, string? description = null, bool? promoted = null, UserId? ownerId = null, UserId? nominatedById = null, bool? acknowledged = null, UserId? declinedById = null, int? fixChange = null, UserId? resolvedById = null, List<ObjectId>? addSpanIds = null, List<ObjectId>? removeSpanIds = null, string? externalIssueKey = null, UserId? quarantinedById = null, UserId? forceClosedById = null, UserId? initiatedById = null, Uri? workflowThreadUrl = null, CancellationToken cancellationToken = default)
+		public async Task<bool> UpdateIssueAsync(int id, string? summary = null, string? description = null, bool? promoted = null, UserId? ownerId = null, UserId? nominatedById = null, bool? acknowledged = null, UserId? declinedById = null, CommitId? fixCommitId = null, bool? fixSystemic = null, UserId? resolvedById = null, List<ObjectId>? addSpanIds = null, List<ObjectId>? removeSpanIds = null, string? externalIssueKey = null, UserId? quarantinedById = null, UserId? forceClosedById = null, UserId? initiatedById = null, Uri? workflowThreadUrl = null, CancellationToken cancellationToken = default)
 		{
 			IIssue? issue;
 			for (; ; )
@@ -465,7 +467,7 @@ namespace HordeServer.Issues
 					return false;
 				}
 
-				issue = await _issueCollection.TryUpdateIssueAsync(issue, initiatedById, newUserSummary: summary, newDescription: description, newPromoted: promoted, newOwnerId: ownerId ?? resolvedById, newNominatedById: nominatedById, newDeclinedById: declinedById, newAcknowledged: acknowledged, newFixChange: fixChange, newResolvedById: resolvedById, newExcludeSpanIds: removeSpanIds, newExternalIssueKey: externalIssueKey, newQuarantinedById: quarantinedById, newForceClosedById: forceClosedById, newWorkflowThreadUrl: workflowThreadUrl, cancellationToken: cancellationToken);
+				issue = await _issueCollection.TryUpdateIssueAsync(issue, initiatedById, newUserSummary: summary, newDescription: description, newPromoted: promoted, newOwnerId: ownerId ?? resolvedById, newNominatedById: nominatedById, newDeclinedById: declinedById, newAcknowledged: acknowledged, newFixCommitId: fixCommitId, newFixSystemic: fixSystemic, newResolvedById: resolvedById, newExcludeSpanIds: removeSpanIds, newExternalIssueKey: externalIssueKey, newQuarantinedById: quarantinedById, newForceClosedById: forceClosedById, newWorkflowThreadUrl: workflowThreadUrl, cancellationToken: cancellationToken);
 				if (issue != null)
 				{
 					break;
@@ -495,7 +497,7 @@ namespace HordeServer.Issues
 						continue;
 					}
 
-					await _issueCollection.TryUpdateSpanAsync(span, newNextSuccess: new NewIssueStepData(step.Change, step.Severity, step.JobName, step.JobId, step.BatchId, step.StepId, step.StepTime, step.LogId, step.Annotations, step.PromoteByDefault), cancellationToken: cancellationToken);
+					await _issueCollection.TryUpdateSpanAsync(span, newNextSuccess: new NewIssueStepData(step.CommitId, step.Severity, step.JobName, step.JobId, step.BatchId, step.StepId, step.StepTime, step.LogId, step.Annotations, step.PromoteByDefault), cancellationToken: cancellationToken);
 				}
 			}
 
@@ -619,8 +621,8 @@ namespace HordeServer.Issues
 				for (; ; )
 				{
 					// Get the spans that are currently open
-					List<IIssueSpan> openSpans = (await _issueCollection.FindOpenSpansAsync(job.StreamId, job.TemplateId, node.Name, job.Change, cancellationToken)).ToList();
-					_logger.LogDebug("{NumSpans} spans are open in {StreamId} at CL {Change} for template {TemplateId}, node {Node}", openSpans.Count, job.StreamId, job.Change, job.TemplateId, node.Name);
+					List<IIssueSpan> openSpans = (await _issueCollection.FindOpenSpansAsync(job.StreamId, job.TemplateId, node.Name, job.CommitId, cancellationToken)).ToList();
+					_logger.LogDebug("{NumSpans} spans are open in {StreamId} at CL {Change} for template {TemplateId}, node {Node}", openSpans.Count, job.StreamId, job.CommitId, job.TemplateId, node.Name);
 
 					// Add the events to existing issues, and create new issues for everything else
 					if (eventGroups.Count > 0)
@@ -786,7 +788,7 @@ namespace HordeServer.Issues
 
 						// Update the span if this changes the current range
 						IIssueSpan? newSpan = openSpan;
-						if (newFailure.Change <= openSpan.FirstFailure.Change || newFailure.Change >= openSpan.LastFailure.Change)
+						if (newFailure.CommitId <= openSpan.FirstFailure.CommitId || newFailure.CommitId >= openSpan.LastFailure.CommitId)
 						{
 							newSpan = await _issueCollection.TryUpdateSpanAsync(openSpan, newFailure: newFailure, cancellationToken: cancellationToken);
 							if (newSpan == null)
@@ -880,14 +882,14 @@ namespace HordeServer.Issues
 				// Get the span data
 				NewIssueSpanData spanData = new NewIssueSpanData(streamConfig.Id, streamConfig.Name, job.TemplateId, node.Name, eventGroup.Fingerprint, stepData);
 
-				IJobStepRef? prevJob = await _jobStepRefs.GetPrevStepForNodeAsync(job.StreamId, job.TemplateId, node.Name, job.Change, JobStepOutcome.Success, true);
+				IJobStepRef? prevJob = await _jobStepRefs.GetPrevStepForNodeAsync(job.StreamId, job.TemplateId, node.Name, job.CommitId, JobStepOutcome.Success, true);
 				if (prevJob != null)
 				{
 					spanData.LastSuccess = new NewIssueStepData(prevJob);
-					spanData.Suspects = await FindSuspectsForSpanAsync(streamConfig, spanData.Fingerprint, spanData.LastSuccess.Change + 1, spanData.FirstFailure.Change, cancellationToken);
+					spanData.Suspects = await FindSuspectsForSpanAsync(streamConfig, spanData.Fingerprint, spanData.LastSuccess.CommitId, spanData.FirstFailure.CommitId, cancellationToken);
 				}
 
-				IJobStepRef? nextJob = await _jobStepRefs.GetNextStepForNodeAsync(job.StreamId, job.TemplateId, node.Name, job.Change, JobStepOutcome.Success, true);
+				IJobStepRef? nextJob = await _jobStepRefs.GetNextStepForNodeAsync(job.StreamId, job.TemplateId, node.Name, job.CommitId, JobStepOutcome.Success, true);
 				if (nextJob != null)
 				{
 					spanData.NextSuccess = new NewIssueStepData(nextJob);
@@ -920,22 +922,22 @@ namespace HordeServer.Issues
 			return issue;
 		}
 
-		static void GetChangeMap(IIssueSpan span, Dictionary<int, int> originChangeToSuspectChange)
+		static void GetChangeMap(IIssueSpan span, Dictionary<CommitId, CommitIdWithOrder> originChangeToSuspectChange)
 		{
 			foreach (IIssueSpanSuspect suspect in span.Suspects)
 			{
-				originChangeToSuspectChange[suspect.OriginatingChange ?? suspect.Change] = suspect.Change;
+				originChangeToSuspectChange[suspect.SourceCommitId ?? suspect.CommitId] = suspect.CommitId;
 			}
 		}
 
-		static int CompareSpans(Dictionary<int, int> originChangeToSuspectChange, IIssueSpan span)
+		static int CompareSpans(Dictionary<CommitId, CommitIdWithOrder> originChangeToSuspectChange, IIssueSpan span)
 		{
 			foreach (IIssueSpanSuspect suspect in span.Suspects)
 			{
-				int originChange = suspect.OriginatingChange ?? suspect.Change;
-				if (originChangeToSuspectChange.TryGetValue(originChange, out int suspectChange) && suspectChange != suspect.Change)
+				CommitId originChange = suspect.SourceCommitId ?? suspect.CommitId;
+				if (originChangeToSuspectChange.TryGetValue(originChange, out CommitIdWithOrder? suspectChange) && suspectChange != suspect.CommitId)
 				{
-					return suspectChange - suspect.Change;
+					return suspectChange.Order - suspect.CommitId.Order;
 				}
 			}
 			return 0;
@@ -950,7 +952,7 @@ namespace HordeServer.Issues
 			{
 				originSpans.Add(spans[0]);
 
-				Dictionary<int, int> originChangeToSuspectChange = new Dictionary<int, int>();
+				Dictionary<CommitId, CommitIdWithOrder> originChangeToSuspectChange = new Dictionary<CommitId, CommitIdWithOrder>();
 				GetChangeMap(spans[0], originChangeToSuspectChange);
 
 				for (int idx = 1; idx < spans.Count; idx++)
@@ -974,7 +976,7 @@ namespace HordeServer.Issues
 
 		async Task<IIssue?> UpdateIssueDerivedDataAsync(IIssue issue, CancellationToken cancellationToken)
 		{
-			Dictionary<(StreamId, int), bool> fixChangeCache = new Dictionary<(StreamId, int), bool>();
+			Dictionary<(StreamId, CommitId), bool> fixChangeCache = new Dictionary<(StreamId, CommitId), bool>();
 			for (; ; )
 			{
 				BuildConfig buildConfig = _buildConfig.CurrentValue;
@@ -998,12 +1000,12 @@ namespace HordeServer.Issues
 					List<IIssueSpan> spansWithSuspects = spans.Where(x => x.LastSuccess != null).ToList();
 					if (spansWithSuspects.Count > 0)
 					{
-						HashSet<int> suspectChanges = new HashSet<int>(spansWithSuspects[0].Suspects.Select(x => x.OriginatingChange ?? x.Change));
+						HashSet<CommitId> suspectChanges = new HashSet<CommitId>(spansWithSuspects[0].Suspects.Select(x => (CommitId)(x.SourceCommitId ?? x.CommitId)));
 						for (int spanIdx = 1; spanIdx < spansWithSuspects.Count; spanIdx++)
 						{
-							suspectChanges.IntersectWith(spansWithSuspects[spanIdx].Suspects.Select(x => x.OriginatingChange ?? x.Change));
+							suspectChanges.IntersectWith(spansWithSuspects[spanIdx].Suspects.Select(x => (CommitId)(x.SourceCommitId ?? x.CommitId)));
 						}
-						newSuspects = spansWithSuspects[0].Suspects.Where(x => suspectChanges.Contains(x.OriginatingChange ?? x.Change)).Select(x => new NewIssueSuspectData(x.AuthorId, x.OriginatingChange ?? x.Change)).ToList();
+						newSuspects = spansWithSuspects[0].Suspects.Where(x => suspectChanges.Contains(x.SourceCommitId ?? x.CommitId)).Select(x => new NewIssueSuspectData(x.AuthorId, x.SourceCommitId ?? x.CommitId)).ToList();
 					}
 				}
 
@@ -1092,7 +1094,7 @@ namespace HordeServer.Issues
 
 				// Get the new list of streams
 				List<NewIssueStream> newStreams = spans.Select(x => x.StreamId).OrderBy(x => x.ToString(), StringComparer.Ordinal).Distinct().Select(x => new NewIssueStream(x)).ToList();
-				if (issue.FixChange != null)
+				if (issue.FixCommitId != null)
 				{
 					foreach (NewIssueStream newStream in newStreams)
 					{
@@ -1102,9 +1104,10 @@ namespace HordeServer.Issues
 							newStream.ContainsFix = stream.ContainsFix;
 						}
 
-						newStream.ContainsFix ??= await ContainsFixChangeAsync(buildConfig, newStream.StreamId, issue.FixChange.Value, fixChangeCache, cancellationToken);
+						newStream.ContainsFix ??= await ContainsFixChangeAsync(buildConfig, newStream.StreamId, issue.FixCommitId, fixChangeCache, cancellationToken);
 
-						if (spans.Any(x => x.StreamId == newStream.StreamId && x.LastFailure.Change > issue.FixChange.Value))
+						CommitIdWithOrder commitIdWithOrder = await _commitService.GetOrderedAsync(newStream.StreamId, issue.FixCommitId, cancellationToken);
+						if (spans.Any(x => x.StreamId == newStream.StreamId && x.LastFailure.CommitId > commitIdWithOrder))
 						{
 							newStream.FixFailed = true;
 						}
@@ -1125,7 +1128,7 @@ namespace HordeServer.Issues
 				// If a fix changelist has been specified and it's not valid, clear it out
 				if (newVerifiedAt == null && newResolvedAt != null)
 				{
-					if (spans.Any(x => HasFixFailed(x, issue.FixChange, newResolvedAt.Value, newStreams)))
+					if(await HasFixFailedAsync(spans, issue.FixCommitId, newResolvedAt.Value, newStreams, cancellationToken))
 					{
 						newStreams.ForEach(x => x.ContainsFix = null);
 						newResolvedAt = null;
@@ -1153,21 +1156,22 @@ namespace HordeServer.Issues
 		/// <summary>
 		/// Determine if an issue should be marked as fix failed
 		/// </summary>
-		static bool HasFixFailed(IIssueSpan span, int? fixChange, DateTime resolvedAt, IReadOnlyList<NewIssueStream> streams)
+		async Task<bool> HasFixFailedAsync(IEnumerable<IIssueSpan> spans, CommitId? fixCommit, DateTime resolvedAt, IReadOnlyList<NewIssueStream> streams, CancellationToken cancellationToken)
 		{
-			if (span.NextSuccess == null)
+			foreach (IGrouping<StreamId, IIssueSpan> group in spans.Where(x => x.NextSuccess == null).GroupBy(x => x.StreamId))
 			{
-				NewIssueStream? stream = streams.FirstOrDefault(x => x.StreamId == span.StreamId);
-				if (stream != null && fixChange != null && (stream.ContainsFix ?? false))
+				NewIssueStream? stream = streams.FirstOrDefault(x => x.StreamId == group.Key);
+				if (stream != null && fixCommit != null && (stream.ContainsFix ?? false))
 				{
-					if (span.LastFailure.Change > fixChange.Value)
+					CommitIdWithOrder commitIdWithOrder = await _commitService.GetOrderedAsync(stream.StreamId, fixCommit, cancellationToken);
+					if (group.Any(x => x.LastFailure.CommitId > commitIdWithOrder))
 					{
 						return true;
 					}
 				}
 				else
 				{
-					if (span.LastFailure.StepTime > resolvedAt + TimeSpan.FromHours(24.0))
+					if (group.Any(x => x.LastFailure.StepTime > resolvedAt + TimeSpan.FromHours(24.0)))
 					{
 						return true;
 					}
@@ -1179,10 +1183,10 @@ namespace HordeServer.Issues
 		/// <summary>
 		/// Figure out if a stream contains a fix changelist
 		/// </summary>
-		async ValueTask<bool> ContainsFixChangeAsync(BuildConfig buildConfig, StreamId streamId, int fixChange, Dictionary<(StreamId, int), bool> cachedContainsFixChange, CancellationToken cancellationToken)
+		async ValueTask<bool> ContainsFixChangeAsync(BuildConfig buildConfig, StreamId streamId, CommitId fixChange, Dictionary<(StreamId, CommitId), bool> cachedContainsFixChange, CancellationToken cancellationToken)
 		{
 			bool containsFixChange;
-			if (!cachedContainsFixChange.TryGetValue((streamId, fixChange), out containsFixChange) && fixChange > 0)
+			if (!cachedContainsFixChange.TryGetValue((streamId, fixChange), out containsFixChange))
 			{
 				StreamConfig? streamConfig;
 				if (buildConfig.TryGetStream(streamId, out streamConfig))
@@ -1206,19 +1210,19 @@ namespace HordeServer.Issues
 		/// </summary>
 		/// <param name="streamConfig">The stream name</param>
 		/// <param name="fingerprint">The fingerprint for the span</param>
-		/// <param name="minChange">Minimum changelist to consider</param>
-		/// <param name="maxChange">Maximum changelist to consider</param>
+		/// <param name="lastSuccess">The last successful commit</param>
+		/// <param name="firstFailure">The first failing commit</param>
 		/// <param name="cancellationToken">Cancellation token for the operation</param>
 		/// <returns>List of suspects</returns>
-		async Task<List<NewIssueSpanSuspectData>> FindSuspectsForSpanAsync(StreamConfig streamConfig, IIssueFingerprint fingerprint, int minChange, int maxChange, CancellationToken cancellationToken)
+		async Task<List<NewIssueSpanSuspectData>> FindSuspectsForSpanAsync(StreamConfig streamConfig, IIssueFingerprint fingerprint, CommitId lastSuccess, CommitId firstFailure, CancellationToken cancellationToken)
 		{
 			List<NewIssueSpanSuspectData> suspects = new List<NewIssueSpanSuspectData>();
 			if (!String.IsNullOrEmpty(fingerprint.ChangeFilter))
 			{
-				_logger.LogDebug("Querying for changes in {StreamName} between {MinChange} and {MaxChange}", streamConfig.Name, minChange, maxChange);
+				_logger.LogDebug("Querying for changes in {StreamName} between {LastSuccess} and {FirstFailure}", streamConfig.Name, lastSuccess, firstFailure);
 
 				// Get the submitted changes before this job
-				List<ICommit> changes = await _commitService.GetCollection(streamConfig).FindAsync(minChange, maxChange, MaxChanges, cancellationToken: cancellationToken).ToListAsync(cancellationToken);
+				List<ICommit> changes = await _commitService.GetCollection(streamConfig).FindAsync(lastSuccess, includeMinCommit: false, firstFailure, includeMaxCommit: true, MaxChanges, cancellationToken: cancellationToken).ToListAsync(cancellationToken);
 				_logger.LogDebug("Found {NumResults} changes", changes.Count);
 
 				// Get all the parameters used to rank suspects
@@ -1267,8 +1271,8 @@ namespace HordeServer.Issues
 
 					if (rank == maxRank)
 					{
-						NewIssueSpanSuspectData suspect = new NewIssueSpanSuspectData(commit.Number, commit.OwnerId);
-						suspect.OriginatingChange = commit.OriginalChange;
+						NewIssueSpanSuspectData suspect = new NewIssueSpanSuspectData(commit.Id, commit.OwnerId);
+						suspect.SourceCommitId = commit.OriginalCommitId;
 						suspects.Add(suspect);
 					}
 				}
@@ -1325,11 +1329,11 @@ namespace HordeServer.Issues
 			IReadOnlyList<IIssue> existingIssues;
 			if (span.LastSuccess == null)
 			{
-				existingIssues = await _issueCollection.FindIssuesAsync(streamId: span.StreamId, minChange: span.FirstFailure.Change, maxChange: span.FirstFailure.Change, resolved: false, cancellationToken: cancellationToken);
+				existingIssues = await _issueCollection.FindIssuesAsync(streamId: span.StreamId, minCommitId: span.FirstFailure.CommitId, maxCommitId: span.FirstFailure.CommitId, resolved: false, cancellationToken: cancellationToken);
 			}
 			else
 			{
-				existingIssues = await _issueCollection.FindIssuesForChangesAsync(span.Suspects.ConvertAll(x => x.OriginatingChange ?? x.Change), cancellationToken);
+				existingIssues = await _issueCollection.FindIssuesForChangesAsync(span.Suspects.ConvertAll(x => x.SourceCommitId ?? x.CommitId), cancellationToken);
 			}
 
 			IIssue? issue = existingIssues.FirstOrDefault(x => x.Fingerprints.Any(y => y.IsMatch(span.Fingerprint)));
@@ -1352,13 +1356,13 @@ namespace HordeServer.Issues
 			IReadOnlyList<IIssue> existingIssues;
 			if (span.LastSuccess == null)
 			{
-				existingIssues = await _issueCollection.FindIssuesAsync(streamId: span.StreamId, minChange: span.FirstFailure.Change, maxChange: span.FirstFailure.Change, resolved: false, cancellationToken: cancellationToken);
-				_logger.LogDebug("Found {NumIssues} open issues at {Change} in {StreamId}", existingIssues.Count, span.FirstFailure.Change, span.StreamId);
+				existingIssues = await _issueCollection.FindIssuesAsync(streamId: span.StreamId, minCommitId: span.FirstFailure.CommitId, maxCommitId: span.FirstFailure.CommitId, resolved: false, cancellationToken: cancellationToken);
+				_logger.LogDebug("Found {NumIssues} open issues at {Change} in {StreamId}", existingIssues.Count, span.FirstFailure.CommitId, span.StreamId);
 			}
 			else
 			{
-				existingIssues = await _issueCollection.FindIssuesForChangesAsync(span.Suspects.ConvertAll(x => x.OriginatingChange ?? x.Change), cancellationToken);
-				_logger.LogDebug("Found {NumIssues} open issues in {StreamId} from [{ChangeList}]", existingIssues.Count, span.StreamId, String.Join(", ", span.Suspects.ConvertAll(x => (x.OriginatingChange ?? x.Change).ToString())));
+				existingIssues = await _issueCollection.FindIssuesForChangesAsync(span.Suspects.ConvertAll(x => x.SourceCommitId ?? x.CommitId), cancellationToken);
+				_logger.LogDebug("Found {NumIssues} open issues in {StreamId} from [{ChangeList}]", existingIssues.Count, span.StreamId, String.Join(", ", span.Suspects.ConvertAll(x => (x.SourceCommitId ?? x.CommitId).ToString())));
 			}
 
 			IIssue? issue = existingIssues.FirstOrDefault(x => x.VerifiedAt == null && x.Fingerprints.Any(y => y.IsMatch(span.Fingerprint)));
@@ -1504,32 +1508,32 @@ namespace HordeServer.Issues
 		{
 			foreach (IIssueSpan span in spans)
 			{
-				if (job.Change < span.FirstFailure.Change && (span.LastSuccess == null || job.Change > span.LastSuccess.Change))
+				if (job.CommitId < span.FirstFailure.CommitId && (span.LastSuccess == null || job.CommitId > span.LastSuccess.CommitId))
 				{
-					NewIssueStepData newLastSuccess = new NewIssueStepData(job.Change, IssueSeverity.Unspecified, job.Name, job.Id, batch.Id, step.Id, step.StartTimeUtc ?? default, step.LogId, null, false);
-					List<NewIssueSpanSuspectData> newSuspects = await FindSuspectsForSpanAsync(streamConfig, span.Fingerprint, job.Change + 1, span.FirstFailure.Change, cancellationToken);
+					NewIssueStepData newLastSuccess = new NewIssueStepData(job.CommitId, IssueSeverity.Unspecified, job.Name, job.Id, batch.Id, step.Id, step.StartTimeUtc ?? default, step.LogId, null, false);
+					List<NewIssueSpanSuspectData> newSuspects = await FindSuspectsForSpanAsync(streamConfig, span.Fingerprint, job.CommitId, span.FirstFailure.CommitId, cancellationToken);
 
 					if (await _issueCollection.TryUpdateSpanAsync(span, newLastSuccess: newLastSuccess, newSuspects: newSuspects, cancellationToken: cancellationToken) == null)
 					{
 						return false;
 					}
 
-					_logger.LogInformation("Set last success for issue {IssueId}, template {TemplateId}, node {Node} as job {JobId}, cl {Change}", span.IssueId, job.TemplateId, span.NodeName, job.Id, job.Change);
+					_logger.LogInformation("Set last success for issue {IssueId}, template {TemplateId}, node {Node} as job {JobId}, cl {Change}", span.IssueId, job.TemplateId, span.NodeName, job.Id, job.CommitId);
 					await UpdateIssueDerivedDataAsync(span.IssueId, cancellationToken);
 				}
-				else if (job.Change > span.LastFailure.Change && (span.NextSuccess == null || job.Change < span.NextSuccess.Change))
+				else if (job.CommitId > span.LastFailure.CommitId && (span.NextSuccess == null || job.CommitId < span.NextSuccess.CommitId))
 				{
 					IIssue? issue = await _issueCollection.GetIssueAsync(span.IssueId, cancellationToken);
 					if (issue == null || issue.QuarantinedByUserId == null)
 					{
-						NewIssueStepData newNextSucccess = new NewIssueStepData(job.Change, IssueSeverity.Unspecified, job.Name, job.Id, batch.Id, step.Id, step.StartTimeUtc ?? default, step.LogId, null, false);
+						NewIssueStepData newNextSucccess = new NewIssueStepData(job.CommitId, IssueSeverity.Unspecified, job.Name, job.Id, batch.Id, step.Id, step.StartTimeUtc ?? default, step.LogId, null, false);
 
 						if (await _issueCollection.TryUpdateSpanAsync(span, newNextSuccess: newNextSucccess, cancellationToken: cancellationToken) == null)
 						{
 							return false;
 						}
 
-						_logger.LogInformation("Set next success for issue {IssueId}, template {TemplateId}, node {Node} as job {JobId}, cl {Change}", span.IssueId, job.TemplateId, span.NodeName, job.Id, job.Change);
+						_logger.LogInformation("Set next success for issue {IssueId}, template {TemplateId}, node {Node} as job {JobId}, cl {Change}", span.IssueId, job.TemplateId, span.NodeName, job.Id, job.CommitId);
 					}
 					else
 					{
@@ -1540,7 +1544,7 @@ namespace HordeServer.Issues
 						{
 							NewIssueStepData newStep = new NewIssueStepData(job, batch, step, IssueSeverity.Unspecified, null, false);
 							await _issueCollection.AddStepAsync(span.Id, newStep, cancellationToken);
-							_logger.LogInformation("Adding step to quarantined issue {IssueId}, template {TemplateId}, node {Node} job {JobId} batch {BatchId} step {StepId} cl {Change}", span.IssueId, job.TemplateId, span.NodeName, job.Id, batch.Id, step.Id, job.Change);
+							_logger.LogInformation("Adding step to quarantined issue {IssueId}, template {TemplateId}, node {Node} job {JobId} batch {BatchId} step {StepId} cl {Change}", span.IssueId, job.TemplateId, span.NodeName, job.Id, batch.Id, step.Id, job.CommitId);
 						}
 					}
 
diff --git a/Engine/Source/Programs/Horde/Plugins/Build/HordeServer.Build/Issues/IssueTagService.cs b/Engine/Source/Programs/Horde/Plugins/Build/HordeServer.Build/Issues/IssueTagService.cs
--- a/Engine/Source/Programs/Horde/Plugins/Build/HordeServer.Build/Issues/IssueTagService.cs
+++ b/Engine/Source/Programs/Horde/Plugins/Build/HordeServer.Build/Issues/IssueTagService.cs
@@ -2,6 +2,7 @@
 
 using System.Text.RegularExpressions;
 using EpicGames.Core;
+using EpicGames.Horde.Commits;
 using EpicGames.Horde.Streams;
 using HordeServer.Commits;
 using HordeServer.Server;
@@ -24,7 +25,7 @@ namespace HordeServer.Issues
 		class State : SingletonBase
 		{
 			[BsonElement("streams"), BsonDictionaryOptions(DictionaryRepresentation.ArrayOfDocuments)]
-			public Dictionary<StreamId, int> Streams { get; set; } = new Dictionary<StreamId, int>();
+			public Dictionary<StreamId, CommitId> Streams { get; set; } = new Dictionary<StreamId, CommitId>();
 		}
 
 		readonly ISingletonDocument<State> _state;
@@ -108,15 +109,15 @@ namespace HordeServer.Issues
 		{
 			ICommitCollection commits = _commitService.GetCollection(streamConfig);
 
-			int minChange;
-			if (!initialState.Streams.TryGetValue(streamConfig.Id, out minChange))
+			CommitId? minCommitId;
+			if (!initialState.Streams.TryGetValue(streamConfig.Id, out minCommitId))
 			{
-				minChange = await commits.GetLatestNumberAsync(cancellationToken);
+				minCommitId = await commits.GetLastCommitIdAsync(cancellationToken);
 			}
 
-			await foreach (ICommit commit in commits.SubscribeAsync(minChange + 1, null, cancellationToken))
+			await foreach (ICommit commit in commits.SubscribeAsync(minCommitId, null, cancellationToken))
 			{
-				_logger.LogDebug("Checking commit {Change} in {StreamId}", commit.Number, streamConfig.Id);
+				_logger.LogDebug("Checking commit {Change} in {StreamId}", commit.Id, streamConfig.Id);
 				foreach (int issueId in ParseTags(_buildConfig.CurrentValue.IssueFixedTag, commit.Description))
 				{
 					for (; ; )
@@ -124,19 +125,19 @@ namespace HordeServer.Issues
 						IIssue? issue = await _issueCollection.GetIssueAsync(issueId, cancellationToken);
 						if (issue == null)
 						{
-							_logger.LogInformation("Commit {Change} by {Author} in {StreamId} has invalid issue id {IssueId}", commit.Number, commit.AuthorId, streamConfig.Id, issueId);
+							_logger.LogInformation("Commit {Change} by {Author} in {StreamId} has invalid issue id {IssueId}", commit.Id, commit.AuthorId, streamConfig.Id, issueId);
 							break;
 						}
 
-						issue = await _issueCollection.TryUpdateIssueAsync(issue, commit.AuthorId, newFixChange: commit.Number, newResolvedById: commit.AuthorId, cancellationToken: cancellationToken);
+						issue = await _issueCollection.TryUpdateIssueAsync(issue, commit.AuthorId, newFixCommitId: commit.Id, newResolvedById: commit.AuthorId, cancellationToken: cancellationToken);
 						if (issue != null)
 						{
-							_logger.LogInformation("Commit {Change} by {Author} in {StreamId} fixes issue id {IssueId}", commit.Number, commit.AuthorId, streamConfig.Id, issueId);
+							_logger.LogInformation("Commit {Change} by {Author} in {StreamId} fixes issue id {IssueId}", commit.Id, commit.AuthorId, streamConfig.Id, issueId);
 							break;
 						}
 					}
 				}
-				await _state.UpdateAsync(x => x.Streams[streamConfig.Id] = commit.Number, cancellationToken);
+				await _state.UpdateAsync(x => x.Streams[streamConfig.Id] = commit.Id, cancellationToken);
 			}
 		}
 
diff --git a/Engine/Source/Programs/Horde/Plugins/Build/HordeServer.Build/Issues/IssuesController.cs b/Engine/Source/Programs/Horde/Plugins/Build/HordeServer.Build/Issues/IssuesController.cs
--- a/Engine/Source/Programs/Horde/Plugins/Build/HordeServer.Build/Issues/IssuesController.cs
+++ b/Engine/Source/Programs/Horde/Plugins/Build/HordeServer.Build/Issues/IssuesController.cs
@@ -1,6 +1,7 @@
 // Copyright Epic Games, Inc. All Rights Reserved.
 
 using EpicGames.Core;
+using EpicGames.Horde.Commits;
 using EpicGames.Horde.Issues;
 using EpicGames.Horde.Jobs;
 using EpicGames.Horde.Jobs.Templates;
@@ -8,6 +9,7 @@ using EpicGames.Horde.Logs;
 using EpicGames.Horde.Streams;
 using EpicGames.Horde.Users;
 using HordeServer.Auditing;
+using HordeServer.Commits;
 using HordeServer.Issues.External;
 using HordeServer.Jobs;
 using HordeServer.Jobs.Graphs;
@@ -71,7 +73,7 @@ namespace HordeServer.Issues
 		[HttpGet]
 		[Route("/api/v2/issues")]
 		[ProducesResponseType(typeof(List<FindIssueResponse>), 200)]
-		public async Task<ActionResult<object>> FindIssuesV2Async([FromQuery(Name = "Id")] int[]? ids = null, [FromQuery] StreamId? streamId = null, [FromQuery] int? minChange = null, [FromQuery] int? maxChange = null, [FromQuery] bool? resolved = null, [FromQuery] int index = 0, [FromQuery] int count = 10, [FromQuery] PropertyFilter? filter = null, CancellationToken cancellationToken = default)
+		public async Task<ActionResult<object>> FindIssuesV2Async([FromQuery(Name = "Id")] int[]? ids = null, [FromQuery] StreamId? streamId = null, [FromQuery] CommitId? minChange = null, [FromQuery] CommitId? maxChange = null, [FromQuery] bool? resolved = null, [FromQuery] int index = 0, [FromQuery] int count = 10, [FromQuery] PropertyFilter? filter = null, CancellationToken cancellationToken = default)
 		{
 			if (ids != null && ids.Length == 0)
 			{
@@ -209,7 +211,7 @@ namespace HordeServer.Issues
 		static GetIssueStepResponse NewGetIssueStepResponse(IIssueStep issueStep)
 		{
 			GetIssueStepResponse response = new GetIssueStepResponse();
-			response.Change = issueStep.Change;
+			response.CommitId = issueStep.CommitId;
 			response.Severity = issueStep.Severity;
 			response.JobName = issueStep.JobName;
 			response.JobId = issueStep.JobId;
@@ -251,8 +253,8 @@ namespace HordeServer.Issues
 				response.NominatedBy = nominatedBy.ToThinApiResponse();
 			}
 			response.AcknowledgedAt = issue.AcknowledgedAt;
-			response.FixChange = issue.FixChange;
-			response.FixedSystemic = issue.FixedSystemic;
+			response.FixCommitId = issue.FixCommitId;
+			response.FixSystemic = issue.FixSystemic;
 			response.ResolvedAt = issue.ResolvedAt;
 			if (resolvedBy != null)
 			{
@@ -295,7 +297,7 @@ namespace HordeServer.Issues
 		[HttpGet]
 		[Route("/api/v1/issues")]
 		[ProducesResponseType(typeof(List<GetIssueResponse>), 200)]
-		public async Task<ActionResult<object>> FindIssuesAsync([FromQuery(Name = "Id")] int[]? ids = null, [FromQuery] string? streamId = null, [FromQuery] int? change = null, [FromQuery] int? minChange = null, [FromQuery] int? maxChange = null, [FromQuery] JobId? jobId = null, [FromQuery] JobStepBatchId? batchId = null, [FromQuery] JobStepId? stepId = null, [FromQuery(Name = "label")] int? labelIdx = null, [FromQuery] string? ownerId = null, [FromQuery] bool? resolved = null, [FromQuery] bool? promoted = null, [FromQuery] int index = 0, [FromQuery] int count = 10, [FromQuery] PropertyFilter? filter = null, CancellationToken cancellationToken = default)
+		public async Task<ActionResult<object>> FindIssuesAsync([FromQuery(Name = "Id")] int[]? ids = null, [FromQuery] string? streamId = null, [FromQuery] CommitId? change = null, [FromQuery] CommitId? minChange = null, [FromQuery] CommitId? maxChange = null, [FromQuery] JobId? jobId = null, [FromQuery] JobStepBatchId? batchId = null, [FromQuery] JobStepId? stepId = null, [FromQuery(Name = "label")] int? labelIdx = null, [FromQuery] string? ownerId = null, [FromQuery] bool? resolved = null, [FromQuery] bool? promoted = null, [FromQuery] int index = 0, [FromQuery] int count = 10, [FromQuery] PropertyFilter? filter = null, CancellationToken cancellationToken = default)
 		{
 			if (ids != null && ids.Length == 0)
 			{
@@ -499,8 +501,8 @@ namespace HordeServer.Issues
 			response.NominatedBy = details.NominatedBy?.Login;
 			response.NominatedByInfo = details.NominatedBy?.ToThinApiResponse();
 			response.AcknowledgedAt = issue.AcknowledgedAt;
-			response.FixChange = issue.FixChange;
-			response.FixedSystemic = issue.FixedSystemic;
+			response.FixCommitId = issue.FixCommitId;
+			response.FixSystemic = issue.FixSystemic;
 			response.ResolvedAt = issue.ResolvedAt;
 			response.ResolvedBy = details.ResolvedBy?.Login;
 			response.ResolvedById = details.ResolvedBy?.Id.ToString();
@@ -585,13 +587,13 @@ namespace HordeServer.Issues
 
 			foreach (IIssueSpan span in spans)
 			{
-				if (span.LastSuccess != null && (response.MinChange == null || span.LastSuccess.Change < response.MinChange.Value))
+				if (span.LastSuccess != null && (response.MinCommitId == null || span.LastSuccess.CommitId < response.MinCommitId))
 				{
-					response.MinChange = span.LastSuccess.Change;
+					response.MinCommitId = span.LastSuccess.CommitId;
 				}
-				if (span.NextSuccess != null && (response.MaxChange == null || span.NextSuccess.Change > response.MaxChange.Value))
+				if (span.NextSuccess != null && (response.MaxCommitId == null || span.NextSuccess.CommitId > response.MaxCommitId))
 				{
-					response.MaxChange = span.NextSuccess.Change;
+					response.MaxCommitId = span.NextSuccess.CommitId;
 				}
 				response.Nodes.Add(NewGetIssueSpanResponse(span, steps.Where(y => y.SpanId == span.Id).ToList()));
 			}
@@ -836,7 +838,7 @@ namespace HordeServer.Issues
 				removeSpans = request.RemoveSpans.ConvertAll(x => ObjectId.Parse(x));
 			}
 
-			if (!await _issueService.UpdateIssueAsync(issueId, request.Summary, request.Description, request.Promoted, newOwnerId, newNominatedById, request.Acknowledged, newDeclinedById, request.FixChange, newResolvedById, addSpans, removeSpans, request.ExternalIssueKey, newQuarantinedById, newForceClosedById, initiatedById: User.GetUserId()))
+			if (!await _issueService.UpdateIssueAsync(issueId, request.Summary, request.Description, request.Promoted, newOwnerId, newNominatedById, request.Acknowledged, newDeclinedById, request.FixCommitId, request.FixSystemic, newResolvedById, addSpans, removeSpans, request.ExternalIssueKey, newQuarantinedById, newForceClosedById, initiatedById: User.GetUserId()))
 			{
 				return NotFound();
 			}
diff --git a/Engine/Source/Programs/Horde/Plugins/Build/HordeServer.Build/Jobs/Bisect/BisectController.cs b/Engine/Source/Programs/Horde/Plugins/Build/HordeServer.Build/Jobs/Bisect/BisectController.cs
--- a/Engine/Source/Programs/Horde/Plugins/Build/HordeServer.Build/Jobs/Bisect/BisectController.cs
+++ b/Engine/Source/Programs/Horde/Plugins/Build/HordeServer.Build/Jobs/Bisect/BisectController.cs
@@ -6,6 +6,7 @@ using EpicGames.Horde.Jobs.Bisect;
 using EpicGames.Horde.Jobs.Templates;
 using EpicGames.Horde.Streams;
 using EpicGames.Horde.Users;
+using HordeServer.Commits;
 using HordeServer.Jobs.Graphs;
 using HordeServer.Streams;
 using HordeServer.Users;
@@ -43,7 +44,23 @@ namespace HordeServer.Jobs.Bisect
 		/// <summary>
 		/// Set of changes to ignore. Can be modified later through <see cref="UpdateBisectTaskRequest"/>.
 		/// </summary>
-		public List<int>? IgnoreChanges { get; set; }
+		[Obsolete("Use IgnoreCommitIds instead")]
+		public List<int>? IgnoreChanges
+		{
+			get => _ignoreChanges ?? _ignoreCommitIds?.ConvertAll(x => x.GetPerforceChangeOrMinusOne());
+			set => _ignoreChanges = value;
+		}
+		List<int>? _ignoreChanges;
+
+		/// <summary>
+		/// Set of changes to ignore. Can be modified later through <see cref="UpdateBisectTaskRequest"/>.
+		/// </summary>
+		public List<CommitId>? IgnoreCommitIds
+		{
+			get => _ignoreCommitIds ?? _ignoreChanges?.ConvertAll(x => CommitId.FromPerforceChange(x));
+			set => _ignoreCommitIds = value;
+		}
+		List<CommitId>? _ignoreCommitIds;
 
 		/// <summary>
 		/// Set of jobs to ignore. Can be modified later through <see cref="UpdateBisectTaskRequest"/>.
@@ -101,8 +118,12 @@ namespace HordeServer.Jobs.Bisect
 		/// Initial Job Step Id
 		public string InitialStepId => _bisectTask.InitialJobStep.StepId.ToString();
 
-		/// <inheritdoc cref="IBisectTask.InitialChange"/>
-		public int InitialChange => _bisectTask.InitialChange;
+		/// <inheritdoc cref="IBisectTask.InitialCommitId"/>
+		[Obsolete("Use InitialCommitId instead")]
+		public int InitialChange => InitialCommitId.GetPerforceChangeOrMinusOne();
+
+		/// <inheritdoc cref="IBisectTask.InitialCommitId"/>
+		public CommitIdWithOrder InitialCommitId => _bisectTask.InitialCommitId;
 
 		/// Min Job Id
 		public string? MinJobId => _bisectTask.MinJobStep?.JobId.ToString();
@@ -113,8 +134,12 @@ namespace HordeServer.Jobs.Bisect
 		/// Min Job Step Id
 		public string? MinStepId => _bisectTask.MinJobStep?.StepId.ToString();
 
-		/// <inheritdoc cref="IBisectTask.MinChange"/>
-		public int? MinChange => _bisectTask.MinChange;
+		/// <inheritdoc cref="IBisectTask.MinCommitId"/>
+		[Obsolete("Use MinCommitId instead")]
+		public int? MinChange => MinCommitId?.GetPerforceChangeOrMinusOne();
+
+		/// <inheritdoc cref="IBisectTask.MinCommitId"/>
+		public CommitIdWithOrder? MinCommitId => _bisectTask.MinCommitId;
 
 		/// Current Job Id
 		public string CurrentJobId => _bisectTask.CurrentJobStep.JobId.ToString();
@@ -125,8 +150,12 @@ namespace HordeServer.Jobs.Bisect
 		/// Current Job Step Id
 		public string CurrentStepId => _bisectTask.CurrentJobStep.StepId.ToString();
 
-		/// <inheritdoc cref="IBisectTask.CurrentChange"/>
-		public int CurrentChange => _bisectTask.CurrentChange;
+		/// <inheritdoc cref="IBisectTask.CurrentCommitId"/>
+		[Obsolete("Use CurrentCommit instead")]
+		public int CurrentChange => CurrentCommitId.GetPerforceChangeOrMinusOne();
+
+		/// <inheritdoc cref="IBisectTask.CurrentCommitId"/>
+		public CommitIdWithOrder CurrentCommitId => _bisectTask.CurrentCommitId;
 
 		/// <summary>
 		/// The next job id for a running bisection task
@@ -136,7 +165,13 @@ namespace HordeServer.Jobs.Bisect
 		/// <summary>
 		/// The next job change
 		/// </summary>
-		public int? NextJobChange { get; }
+		[Obsolete("Use NextJobCommit instead")]
+		public int? NextJobChange => NextJobCommit?.TryGetPerforceChange();
+
+		/// <summary>
+		/// The next job change
+		/// </summary>
+		public CommitIdWithOrder? NextJobCommit { get; }
 
 		/// <summary>
 		/// The steps involved in the bisection
@@ -149,7 +184,7 @@ namespace HordeServer.Jobs.Bisect
 			Owner = owner;
 			Steps = steps.Select(s => StreamsController.CreateGetJobStepRefResponse(s)).ToList();
 			NextJobId = nextJob?.Id;
-			NextJobChange = nextJob?.Change;
+			NextJobCommit = nextJob?.CommitId;
 		}
 	}
 
@@ -166,12 +201,44 @@ namespace HordeServer.Jobs.Bisect
 		/// <summary>
 		/// List of change numbers to include in the search. 
 		/// </summary>
-		public List<int> IncludeChanges { get; set; } = new List<int>();
+		[Obsolete("Use IncludeCommits instead")]
+		public List<int> IncludeChanges
+		{
+			get => _includeChanges ?? _includeCommitIds?.ConvertAll(x => x.GetPerforceChangeOrMinusOne()) ?? new List<int>();
+			set => _includeChanges = value;
+		}
+		List<int>? _includeChanges;
+
+		/// <summary>
+		/// List of commits to include in the search. 
+		/// </summary>
+		public List<CommitId> IncludeCommitIds
+		{
+			get => _includeCommitIds ?? _includeChanges?.ConvertAll(x => CommitId.FromPerforceChange(x)) ?? new List<CommitId>();
+			set => _includeCommitIds = value;
+		}
+		List<CommitId>? _includeCommitIds;
 
 		/// <summary>
 		/// List of change numbers to exclude from the search.
 		/// </summary>
-		public List<int> ExcludeChanges { get; set; } = new List<int>();
+		[Obsolete("Use ExcludeCommitIds instead")]
+		public List<int> ExcludeChanges
+		{
+			get => _excludeChanges ?? _excludeCommitIds?.Select(x => x.TryGetPerforceChange() ?? 0).Where(x => x != 0).ToList() ?? new List<int>();
+			set => _excludeChanges = value;
+		}
+		List<int>? _excludeChanges;
+
+		/// <summary>
+		/// List of commits to exclude from the search
+		/// </summary>
+		public List<CommitId> ExcludeCommitIds
+		{
+			get => _excludeCommitIds ?? _excludeChanges?.ConvertAll(x => CommitId.FromPerforceChange(x)) ?? new List<CommitId>();
+			set => _excludeCommitIds = value;
+		}
+		List<CommitId>? _excludeCommitIds;
 
 		/// <summary>
 		/// List of jobs to include in the search.
@@ -195,6 +262,7 @@ namespace HordeServer.Jobs.Bisect
 		readonly IBisectTaskCollection _bisectTaskCollection;
 		readonly IJobCollection _jobCollection;
 		readonly JobService _jobService;
+		readonly ICommitService _commitService;
 		readonly IJobStepRefCollection _jobStepRefs;
 		readonly IGraphCollection _graphCollection;
 		readonly IUserCollection _userCollection;
@@ -205,10 +273,11 @@ namespace HordeServer.Jobs.Bisect
 		/// <summary>
 		/// Constructor
 		/// </summary>
-		public BisectTasksController(IBisectTaskCollection bisectTaskCollection, JobService jobService, IJobCollection jobCollection, IJobStepRefCollection jobStepRefs, IGraphCollection graphCollection, IUserCollection userCollection, Tracer tracer, ILogger<BisectTasksController> logger, IOptionsSnapshot<BuildConfig> buildConfig)
+		public BisectTasksController(IBisectTaskCollection bisectTaskCollection, JobService jobService, ICommitService commitService, IJobCollection jobCollection, IJobStepRefCollection jobStepRefs, IGraphCollection graphCollection, IUserCollection userCollection, Tracer tracer, ILogger<BisectTasksController> logger, IOptionsSnapshot<BuildConfig> buildConfig)
 		{
 			_bisectTaskCollection = bisectTaskCollection;
 			_jobService = jobService;
+			_commitService = commitService;
 			_jobCollection = jobCollection;
 			_jobStepRefs = jobStepRefs;
 			_graphCollection = graphCollection;
@@ -287,8 +356,8 @@ namespace HordeServer.Jobs.Bisect
 
 			CreateBisectTaskOptions options = new CreateBisectTaskOptions();
 			options.CommitTags = create.CommitTags;
-			options.IgnoreChanges = create.IgnoreChanges;
-			options.IgnoreJobs = create.IgnoreJobs;
+			options.IgnoreCommitIds = create.IgnoreCommitIds;
+			options.IgnoreJobIds = create.IgnoreJobs;
 
 			UserId? userId = User.GetUserId();
 
@@ -362,8 +431,8 @@ namespace HordeServer.Jobs.Bisect
 
 				UpdateBisectTaskOptions options = new UpdateBisectTaskOptions();
 				options.State = (!request.Cancel.HasValue) ? null : request.Cancel.Value ? BisectTaskState.Cancelled : BisectTaskState.Running;
-				options.IncludeChanges = request.IncludeChanges;
-				options.ExcludeChanges = request.ExcludeChanges;
+				options.IncludeCommitIds = request.IncludeCommitIds;
+				options.ExcludeCommitIds = request.ExcludeCommitIds;
 				options.IncludeJobs = request.IncludeJobs;
 				options.ExcludeJobs = request.ExcludeJobs;
 
