Horde: Preliminary support for VCS commit ids that aren't an integer.

diff --git a/Engine/Source/Programs/Horde/Plugins/Build/HordeServer.Build/Jobs/Bisect/BisectService.cs b/Engine/Source/Programs/Horde/Plugins/Build/HordeServer.Build/Jobs/Bisect/BisectService.cs
--- a/Engine/Source/Programs/Horde/Plugins/Build/HordeServer.Build/Jobs/Bisect/BisectService.cs
+++ b/Engine/Source/Programs/Horde/Plugins/Build/HordeServer.Build/Jobs/Bisect/BisectService.cs
@@ -92,8 +92,8 @@ namespace HordeServer.Jobs.Bisect
 			// Check if the task needs to be updated with the current bisect state
 			if (currentJobStepRef != null && bisectTask.CurrentJobStep.JobId != currentJobStepRef.Id.JobId)
 			{
-				_logger.LogInformation("Bisect task {BisectTaskId} ({StreamId}:{TemplateId}:{NodeName}): Changing current job from {PrevJobId} ({PrevChange}) -> {NextJobId} ({NextChange}).", bisectTask.Id, bisectTask.StreamId, bisectTask.TemplateId, bisectTask.NodeName, bisectTask.CurrentJobStep.JobId, bisectTask.CurrentChange, currentJobStepRef.Id.JobId, currentJobStepRef.Change);
-				await _bisectTaskCollection.TryUpdateAsync(bisectTask, new UpdateBisectTaskOptions { CurrentJobStep = (currentJobStepRef.Id, currentJobStepRef.Change) }, cancellationToken);
+				_logger.LogInformation("Bisect task {BisectTaskId} ({StreamId}:{TemplateId}:{NodeName}): Changing current job from {PrevJobId} ({PrevChange}) -> {NextJobId} ({NextChange}).", bisectTask.Id, bisectTask.StreamId, bisectTask.TemplateId, bisectTask.NodeName, bisectTask.CurrentJobStep.JobId, bisectTask.CurrentCommitId, currentJobStepRef.Id.JobId, currentJobStepRef.CommitId);
+				await _bisectTaskCollection.TryUpdateAsync(bisectTask, new UpdateBisectTaskOptions { CurrentJobStep = (currentJobStepRef.Id, currentJobStepRef.CommitId) }, cancellationToken);
 				return null;
 			}
 
@@ -114,7 +114,7 @@ namespace HordeServer.Jobs.Bisect
 			// Update bisection lower bound
 			if (bisectTask.MinJobStep == null)
 			{
-				await _bisectTaskCollection.TryUpdateAsync(bisectTask, new UpdateBisectTaskOptions { MinJobStep = (previousJobStepRef.Id, previousJobStepRef.Change) }, cancellationToken);
+				await _bisectTaskCollection.TryUpdateAsync(bisectTask, new UpdateBisectTaskOptions { MinJobStep = (previousJobStepRef.Id, previousJobStepRef.CommitId) }, cancellationToken);
 			}
 
 			// Find the next commit to test
@@ -126,17 +126,17 @@ namespace HordeServer.Jobs.Bisect
 			}
 
 			ICommitCollection commitCollection = _commitService.GetCollection(streamConfig);
-			List<ICommit> commits = await commitCollection.FindAsync(previousJobStepRef.Change + 1, bisectTask.CurrentChange - 1, 100, bisectTask.CommitTags, cancellationToken).ToListAsync(cancellationToken);
-			commits.RemoveAll(x => bisectTask.IgnoreChanges.Contains(x.Number));
+			List<ICommit> commits = await commitCollection.FindAsync(previousJobStepRef.CommitId, includeMinCommit: false, bisectTask.CurrentCommitId, includeMaxCommit: false, 100, bisectTask.CommitTags, cancellationToken).ToListAsync(cancellationToken);
+			commits.RemoveAll(x => bisectTask.IgnoreCommitIds.Contains(x.Id));
 
 			if (commits.Count == 0)
 			{
-				_logger.LogInformation("Bisect task {BisectTaskId} ({StreamId}:{TemplateId}:{NodeName}): Success - first failure was CL {Change}.", bisectTask.Id, bisectTask.StreamId, bisectTask.TemplateId, bisectTask.NodeName, bisectTask.CurrentChange);
+				_logger.LogInformation("Bisect task {BisectTaskId} ({StreamId}:{TemplateId}:{NodeName}): Success - first failure was CL {Change}.", bisectTask.Id, bisectTask.StreamId, bisectTask.TemplateId, bisectTask.NodeName, bisectTask.CurrentCommitId);
 				return BisectTaskState.Succeeded;
 			}
 
 			ICommit nextCommit = commits[commits.Count / 2];
-			ICommit nextCodeCommit = await commitCollection.GetLastCodeChangeAsync(nextCommit.Number, cancellationToken) ?? nextCommit;
+			ICommit nextCodeCommit = await commitCollection.GetLastCodeChangeAsync(nextCommit.Id, cancellationToken) ?? nextCommit;
 
 			// Get the initial job
 			IJob? job = await _jobCollection.GetAsync(bisectTask.InitialJobStep.JobId, cancellationToken);
@@ -164,8 +164,8 @@ namespace HordeServer.Jobs.Bisect
 			options.Arguments.AddRange(job.Arguments.Where(x => !x.StartsWith(IJob.TargetArgumentPrefix, StringComparison.OrdinalIgnoreCase)));
 			options.Arguments.Add($"{IJob.TargetArgumentPrefix}{bisectTask.NodeName}");
 
-			IJob nextJob = await _jobCollection.AddAsync(JobIdUtils.GenerateNewId(), bisectTask.StreamId, bisectTask.TemplateId, template.Hash, graph, $"{template.Name} (Bisect)", nextCommit.Number, nextCodeCommit.Number, options, cancellationToken);
-			_logger.LogInformation("Bisect task {BisectTaskId} ({StreamId}:{TemplateId}:{NodeName}): {NumCommits} possible commits ({MinChange}..{MaxChange}). Started new job {JobId} at CL {Change}.", bisectTask.Id, bisectTask.StreamId, bisectTask.TemplateId, bisectTask.NodeName, commits.Count, commits[^1].Number, commits[0].Number, nextJob.Id, nextCommit.Number);
+			IJob nextJob = await _jobCollection.AddAsync(JobIdUtils.GenerateNewId(), bisectTask.StreamId, bisectTask.TemplateId, template.Hash, graph, $"{template.Name} (Bisect)", nextCommit.Id, nextCodeCommit.Id, options, cancellationToken);
+			_logger.LogInformation("Bisect task {BisectTaskId} ({StreamId}:{TemplateId}:{NodeName}): {NumCommits} possible commits ({MinChange}..{MaxChange}). Started new job {JobId} at CL {Change}.", bisectTask.Id, bisectTask.StreamId, bisectTask.TemplateId, bisectTask.NodeName, commits.Count, commits[^1].Id, commits[0].Id, nextJob.Id, nextCommit.Id);
 
 			return BisectTaskState.Running;
 		}
@@ -175,12 +175,12 @@ namespace HordeServer.Jobs.Bisect
 			// Query jobs before the start position
 			for (int maxJobCount = 20; ; maxJobCount += 20)
 			{
-				List<IJobStepRef> jobStepRefs = await _jobStepRefCollection.GetStepsForNodeAsync(bisectTask.StreamId, bisectTask.TemplateId, bisectTask.NodeName, bisectTask.InitialChange, true, maxJobCount, cancellationToken);
+				List<IJobStepRef> jobStepRefs = await _jobStepRefCollection.GetStepsForNodeAsync(bisectTask.StreamId, bisectTask.TemplateId, bisectTask.NodeName, bisectTask.InitialCommitId, true, maxJobCount, cancellationToken);
 
 				IJobStepRef? currentJobStepRef = null;
-				foreach (IJobStepRef jobStepRef in jobStepRefs.OrderByDescending(x => x.Change).ThenBy(x => x.Id))
+				foreach (IJobStepRef jobStepRef in jobStepRefs.OrderByDescending(x => x.CommitId).ThenBy(x => x.Id))
 				{
-					if (!bisectTask.IgnoreJobs.Contains(jobStepRef.Id.JobId))
+					if (!bisectTask.IgnoreJobIds.Contains(jobStepRef.Id.JobId))
 					{
 						if (jobStepRef.Outcome == bisectTask.Outcome)
 						{
diff --git a/Engine/Source/Programs/Horde/Plugins/Build/HordeServer.Build/Jobs/Bisect/BisectTaskCollection.cs b/Engine/Source/Programs/Horde/Plugins/Build/HordeServer.Build/Jobs/Bisect/BisectTaskCollection.cs
--- a/Engine/Source/Programs/Horde/Plugins/Build/HordeServer.Build/Jobs/Bisect/BisectTaskCollection.cs
+++ b/Engine/Source/Programs/Horde/Plugins/Build/HordeServer.Build/Jobs/Bisect/BisectTaskCollection.cs
@@ -1,6 +1,7 @@
 // Copyright Epic Games, Inc. All Rights Reserved.
 
 using System.Runtime.CompilerServices;
+using EpicGames.Core;
 using EpicGames.Horde.Commits;
 using EpicGames.Horde.Jobs;
 using EpicGames.Horde.Jobs.Bisect;
@@ -49,21 +50,50 @@ namespace HordeServer.Jobs.Bisect
 			[BsonElement("ijob")]
 			public JobStepRefId InitialJobStep { get; set; }
 
+			[BsonIgnore]
+			public CommitIdWithOrder InitialCommitId
+			{
+				get => (InitialCommitName != null) ? new CommitIdWithOrder(InitialCommitName, InitialCommitOrder) : CommitIdWithOrder.FromPerforceChange(InitialCommitOrder);
+				set => (InitialCommitName, InitialCommitOrder) = (value.Name, value.Order);
+			}
+
+			[BsonElement("icom")]
+			public string? InitialCommitName { get; set; }
+
 			[BsonElement("ichg")]
-			public int InitialChange { get; set; }
+			public int InitialCommitOrder { get; set; }
 
 			[BsonElement("curjob")]
 			public JobStepRefId CurrentJobStep { get; set; }
 
+			[BsonIgnore]
+			public CommitIdWithOrder CurrentCommitId
+			{
+				get => (CurrentCommitName != null) ? new CommitIdWithOrder(CurrentCommitName, CurrentCommitOrder) : CommitIdWithOrder.FromPerforceChange(CurrentCommitOrder);
+				set => (CurrentCommitName, CurrentCommitOrder) = (value.Name, value.Order);
+			}
+
+			[BsonElement("curCom")]
+			public string? CurrentCommitName { get; set; }
+
 			[BsonElement("curChg")]
-			public int CurrentChange { get; set; }
+			public int CurrentCommitOrder { get; set; }
 
 			// Lower bounds of bisection
 			[BsonElement("minJob"), BsonIgnoreIfNull]
 			public JobStepRefId? MinJobStep { get; set; }
 
+			public CommitIdWithOrder? MinCommitId
+			{
+				get => (MinCommitName != null) ? new CommitIdWithOrder(MinCommitName, MinCommitOrder ?? 0) : CommitIdWithOrder.FromPerforceChange(MinCommitOrder);
+				set => (MinCommitName, MinCommitOrder) = (value?.Name, value?.Order);
+			}
+
+			[BsonElement("minCom"), BsonIgnoreIfNull]
+			public string? MinCommitName { get; set; }
+
 			[BsonElement("minChg"), BsonIgnoreIfNull]
-			public int? MinChange { get; set; }
+			public int? MinCommitOrder { get; set; }
 
 			[BsonElement("steps")]
 			public List<JobStepRefId> Steps { get; set; } = new List<JobStepRefId>();
@@ -77,12 +107,12 @@ namespace HordeServer.Jobs.Bisect
 			IReadOnlyList<CommitTag>? IBisectTask.CommitTags => CommitTags;
 
 			[BsonElement("nochg")]
-			public HashSet<int> IgnoreChanges { get; set; } = new HashSet<int>();
-			IReadOnlySet<int> IBisectTask.IgnoreChanges => IgnoreChanges;
+			public HashSet<CommitId> IgnoreCommitIds { get; set; } = new HashSet<CommitId>();
+			IReadOnlySet<CommitId> IBisectTask.IgnoreCommitIds => IgnoreCommitIds;
 
 			[BsonElement("nojob")]
 			public HashSet<JobId> IgnoreJobs { get; set; } = new HashSet<JobId>();
-			IReadOnlySet<JobId> IBisectTask.IgnoreJobs => IgnoreJobs;
+			IReadOnlySet<JobId> IBisectTask.IgnoreJobIds => IgnoreJobs;
 		}
 
 		readonly Tracer _tracer;
@@ -114,9 +144,9 @@ namespace HordeServer.Jobs.Bisect
 			bisectTaskDoc.Outcome = outcome;
 			bisectTaskDoc.InitialJobId = job.Id;
 			bisectTaskDoc.InitialJobStep = new JobStepRefId(job.Id, batchId, stepId);
-			bisectTaskDoc.InitialChange = job.Change;
+			bisectTaskDoc.InitialCommitId = job.CommitId;
 			bisectTaskDoc.CurrentJobStep = new JobStepRefId(job.Id, batchId, stepId);
-			bisectTaskDoc.CurrentChange = job.Change;
+			bisectTaskDoc.CurrentCommitId = job.CommitId;
 
 			if (options != null)
 			{
@@ -124,13 +154,13 @@ namespace HordeServer.Jobs.Bisect
 				{
 					bisectTaskDoc.CommitTags = new List<CommitTag>(options.CommitTags);
 				}
-				if (options.IgnoreChanges != null)
+				if (options.IgnoreCommitIds != null)
 				{
-					bisectTaskDoc.IgnoreChanges.UnionWith(options.IgnoreChanges);
+					bisectTaskDoc.IgnoreCommitIds.UnionWith(options.IgnoreCommitIds);
 				}
-				if (options.IgnoreJobs != null)
+				if (options.IgnoreJobIds != null)
 				{
-					bisectTaskDoc.IgnoreJobs.UnionWith(options.IgnoreJobs);
+					bisectTaskDoc.IgnoreJobs.UnionWith(options.IgnoreJobIds);
 				}
 			}
 
@@ -203,7 +233,7 @@ namespace HordeServer.Jobs.Bisect
 				filter &= filterBuilder.Lte(x => x.Id!, maxTime);
 			}
 
-			List<BisectTaskDoc> steps = await _bisectTasks.Find(filter).SortByDescending(x => x.InitialChange).Range(index, count).ToListAsync(cancellationToken);
+			List<BisectTaskDoc> steps = await _bisectTasks.Find(filter).SortByDescending(x => x.InitialCommitOrder).Range(index, count).ToListAsync(cancellationToken);
 			return steps.ConvertAll<IBisectTask>(x => x);
 		}
 
@@ -215,12 +245,18 @@ namespace HordeServer.Jobs.Bisect
 			UpdateDefinition<BisectTaskDoc> update = Builders<BisectTaskDoc>.Update.Inc(x => x.UpdateIdx, 1);
 			if (options.CurrentJobStep != null)
 			{
-				update = update.Set(x => x.CurrentJobStep, options.CurrentJobStep.Value.Step).Set(x => x.CurrentChange, options.CurrentJobStep.Value.Change);
+				update = update
+					.Set(x => x.CurrentJobStep, options.CurrentJobStep.Value.Step)
+					.Set(x => x.CurrentCommitName, options.CurrentJobStep.Value.CommitId.Name)
+					.Set(x => x.CurrentCommitOrder, options.CurrentJobStep.Value.CommitId.Order);
 			}
 
 			if (options.MinJobStep != null)
 			{
-				update = update.Set(x => x.MinJobStep, options.MinJobStep.Value.Step).Set(x => x.MinChange, options.MinJobStep.Value.Change);
+				update = update
+					.Set(x => x.MinJobStep, options.MinJobStep.Value.Step)
+					.Set(x => x.MinCommitName, options.MinJobStep.Value.CommitId.Name)
+					.Set(x => x.MinCommitOrder, options.MinJobStep.Value.CommitId.Order);
 			}
 
 			if (options.State != null)
@@ -241,13 +277,13 @@ namespace HordeServer.Jobs.Bisect
 				update = update.AddToSet(x => x.Steps, options.NewJobStep.Value);
 			}
 
-			if (options.IncludeChanges != null && options.IncludeChanges.Count > 0)
+			if (options.IncludeCommitIds != null && options.IncludeCommitIds.Count > 0)
 			{
-				update = update.PullAll(x => x.IgnoreChanges, options.IncludeChanges);
+				update = update.PullAll(x => x.IgnoreCommitIds, options.IncludeCommitIds);
 			}
-			else if (options.ExcludeChanges != null && options.ExcludeChanges.Count > 0)
+			else if (options.ExcludeCommitIds != null && options.ExcludeCommitIds.Count > 0)
 			{
-				update = update.AddToSetEach(x => x.IgnoreChanges, options.ExcludeChanges);
+				update = update.AddToSetEach(x => x.IgnoreCommitIds, options.ExcludeCommitIds);
 			}
 
 			if (options.IncludeJobs != null && options.IncludeJobs.Count > 0)
diff --git a/Engine/Source/Programs/Horde/Plugins/Build/HordeServer.Build/Jobs/Bisect/IBisectTask.cs b/Engine/Source/Programs/Horde/Plugins/Build/HordeServer.Build/Jobs/Bisect/IBisectTask.cs
--- a/Engine/Source/Programs/Horde/Plugins/Build/HordeServer.Build/Jobs/Bisect/IBisectTask.cs
+++ b/Engine/Source/Programs/Horde/Plugins/Build/HordeServer.Build/Jobs/Bisect/IBisectTask.cs
@@ -98,7 +98,7 @@ namespace HordeServer.Jobs.Bisect
 		/// <summary>
 		/// Starting change for the bisection
 		/// </summary>
-		public int InitialChange { get; }
+		public CommitIdWithOrder InitialCommitId { get; }
 
 		/// <summary>
 		/// First known job id that is broken
@@ -108,7 +108,7 @@ namespace HordeServer.Jobs.Bisect
 		/// <summary>
 		/// Changelist number of the first broken job id
 		/// </summary>
-		public int CurrentChange { get; }
+		public CommitIdWithOrder CurrentCommitId { get; }
 
 		/// <summary>
 		/// Lower bound of bisection job Id
@@ -118,7 +118,7 @@ namespace HordeServer.Jobs.Bisect
 		/// <summary>
 		/// Lower bound of bisection changelist number
 		/// </summary>
-		public int? MinChange { get; }
+		public CommitIdWithOrder? MinCommitId { get; }
 
 		/// <summary>
 		/// The steps involved in the bisection
@@ -133,11 +133,11 @@ namespace HordeServer.Jobs.Bisect
 		/// <summary>
 		/// Set of changes to ignore while doing the bisection
 		/// </summary>
-		public IReadOnlySet<int> IgnoreChanges { get; }
+		public IReadOnlySet<CommitId> IgnoreCommitIds { get; }
 
 		/// <summary>
 		/// Set of jobs to ignore while doing the bisection
 		/// </summary>
-		public IReadOnlySet<JobId> IgnoreJobs { get; }
+		public IReadOnlySet<JobId> IgnoreJobIds { get; }
 	}
 }
\ No newline at end of file
diff --git a/Engine/Source/Programs/Horde/Plugins/Build/HordeServer.Build/Jobs/Bisect/IBisectTaskCollection.cs b/Engine/Source/Programs/Horde/Plugins/Build/HordeServer.Build/Jobs/Bisect/IBisectTaskCollection.cs
--- a/Engine/Source/Programs/Horde/Plugins/Build/HordeServer.Build/Jobs/Bisect/IBisectTaskCollection.cs
+++ b/Engine/Source/Programs/Horde/Plugins/Build/HordeServer.Build/Jobs/Bisect/IBisectTaskCollection.cs
@@ -78,12 +78,12 @@ namespace HordeServer.Jobs.Bisect
 		/// <summary>
 		/// Changes to exclude from the bisection
 		/// </summary>
-		public IReadOnlyList<int>? IgnoreChanges { get; set; }
+		public IReadOnlyList<CommitId>? IgnoreCommitIds { get; set; }
 
 		/// <summary>
 		/// Jobs to exclude from the bisection
 		/// </summary>
-		public IReadOnlyList<JobId>? IgnoreJobs { get; set; }
+		public IReadOnlyList<JobId>? IgnoreJobIds { get; set; }
 	}
 
 	/// <summary>
@@ -94,12 +94,12 @@ namespace HordeServer.Jobs.Bisect
 		/// <summary>
 		/// First job to bisect from
 		/// </summary>
-		public (JobStepRefId Step, int Change)? CurrentJobStep { get; set; }
+		public (JobStepRefId Step, CommitIdWithOrder CommitId)? CurrentJobStep { get; set; }
 
 		/// <summary>
 		/// The lower bounds of the bisection task
 		/// </summary>
-		public (JobStepRefId Step, int Change)? MinJobStep { get; set; }
+		public (JobStepRefId Step, CommitIdWithOrder CommitId)? MinJobStep { get; set; }
 
 		/// <summary>
 		/// New state for the task
@@ -114,12 +114,12 @@ namespace HordeServer.Jobs.Bisect
 		/// <summary>
 		/// Changes to include in the bisection
 		/// </summary>
-		public IReadOnlyList<int>? IncludeChanges { get; set; }
+		public IReadOnlyList<CommitId>? IncludeCommitIds { get; set; }
 
 		/// <summary>
 		/// Changes to exclude from the bisection
 		/// </summary>
-		public IReadOnlyList<int>? ExcludeChanges { get; set; }
+		public IReadOnlyList<CommitId>? ExcludeCommitIds { get; set; }
 
 		/// <summary>
 		/// Jobs to include in the bisection
diff --git a/Engine/Source/Programs/Horde/Plugins/Build/HordeServer.Build/Jobs/IJob.cs b/Engine/Source/Programs/Horde/Plugins/Build/HordeServer.Build/Jobs/IJob.cs
--- a/Engine/Source/Programs/Horde/Plugins/Build/HordeServer.Build/Jobs/IJob.cs
+++ b/Engine/Source/Programs/Horde/Plugins/Build/HordeServer.Build/Jobs/IJob.cs
@@ -8,6 +8,7 @@ using EpicGames.Horde.Agents;
 using EpicGames.Horde.Agents.Leases;
 using EpicGames.Horde.Agents.Pools;
 using EpicGames.Horde.Agents.Sessions;
+using EpicGames.Horde.Commits;
 using EpicGames.Horde.Jobs;
 using EpicGames.Horde.Jobs.Bisect;
 using EpicGames.Horde.Jobs.Templates;
@@ -99,24 +100,24 @@ namespace HordeServer.Jobs
 		public string Name { get; }
 
 		/// <summary>
-		/// The changelist number to build
+		/// The commit to build
 		/// </summary>
-		public int Change { get; }
+		public CommitIdWithOrder CommitId { get; }
 
 		/// <summary>
-		/// The code changelist number for this build
+		/// The code commit for this build
 		/// </summary>
-		public int CodeChange { get; }
+		public CommitIdWithOrder? CodeCommitId { get; }
 
 		/// <summary>
 		/// The preflight changelist number
 		/// </summary>
-		public int PreflightChange { get; }
+		public CommitId? PreflightCommitId { get; }
 
 		/// <summary>
 		/// The cloned preflight changelist number (if the prefight change is duplicated via p4 reshelve)
 		/// </summary>
-		public int ClonedPreflightChange { get; }
+		public CommitId? ClonedPreflightCommitId { get; }
 
 		/// <summary>
 		/// Description for the shelved change if running a preflight
@@ -1178,10 +1179,10 @@ namespace HordeServer.Jobs
 		{
 			RpcGetJobResponse response = new RpcGetJobResponse();
 			response.StreamId = job.StreamId.ToString();
-			response.Change = job.Change;
-			response.CodeChange = job.CodeChange;
-			response.PreflightChange = job.PreflightChange;
-			response.ClonedPreflightChange = job.ClonedPreflightChange;
+			response.Change = job.CommitId.GetPerforceChange();
+			response.CodeChange = job.CodeCommitId?.GetPerforceChange() ?? 0;
+			response.PreflightChange = job.PreflightCommitId?.GetPerforceChange() ?? 0;
+			response.ClonedPreflightChange = job.ClonedPreflightCommitId?.GetPerforceChange() ?? 0;
 			response.Arguments.Add(job.Arguments);
 			return response;
 		}
diff --git a/Engine/Source/Programs/Horde/Plugins/Build/HordeServer.Build/Jobs/IJobCollection.cs b/Engine/Source/Programs/Horde/Plugins/Build/HordeServer.Build/Jobs/IJobCollection.cs
--- a/Engine/Source/Programs/Horde/Plugins/Build/HordeServer.Build/Jobs/IJobCollection.cs
+++ b/Engine/Source/Programs/Horde/Plugins/Build/HordeServer.Build/Jobs/IJobCollection.cs
@@ -1,6 +1,7 @@
 // Copyright Epic Games, Inc. All Rights Reserved.
 
 using EpicGames.Core;
+using EpicGames.Horde.Commits;
 using EpicGames.Horde.Jobs;
 using EpicGames.Horde.Jobs.Bisect;
 using EpicGames.Horde.Jobs.Templates;
@@ -19,11 +20,11 @@ namespace HordeServer.Jobs
 	/// </summary>
 	public class CreateJobOptions
 	{
-		/// <inheritdoc cref="IJob.PreflightChange"/>
-		public int? PreflightChange { get; set; }
+		/// <inheritdoc cref="IJob.PreflightCommitId"/>
+		public CommitId? PreflightCommitId { get; set; }
 
-		/// <inheritdoc cref="IJob.ClonedPreflightChange"/>
-		public int? ClonedPreflightChange { get; set; }
+		/// <inheritdoc cref="IJob.ClonedPreflightCommitId"/>
+		public CommitId? ClonedPreflightCommitId { get; set; }
 
 		/// <inheritdoc cref="IJob.PreflightDescription"/>
 		public string? PreflightDescription { get; set; }
@@ -123,12 +124,12 @@ namespace HordeServer.Jobs
 		/// <param name="templateHash">Template for this job</param>
 		/// <param name="graph">The graph for the new job</param>
 		/// <param name="name">Name of the job</param>
-		/// <param name="change">The change to build</param>
-		/// <param name="codeChange">The corresponding code changelist number</param>
+		/// <param name="commitId">The commit to build</param>
+		/// <param name="codeCommitId">The corresponding code changelist number</param>
 		/// <param name="options">Additional options for the new job</param>
 		/// <param name="cancellationToken">Cancellation token for the operation</param>
 		/// <returns>The new job document</returns>
-		Task<IJob> AddAsync(JobId jobId, StreamId streamId, TemplateId templateRefId, ContentHash templateHash, IGraph graph, string name, int change, int codeChange, CreateJobOptions options, CancellationToken cancellationToken = default);
+		Task<IJob> AddAsync(JobId jobId, StreamId streamId, TemplateId templateRefId, ContentHash templateHash, IGraph graph, string name, CommitId commitId, CommitId? codeCommitId, CreateJobOptions options, CancellationToken cancellationToken = default);
 
 		/// <summary>
 		/// Gets a job with the given unique id
@@ -153,10 +154,11 @@ namespace HordeServer.Jobs
 		/// <param name="streamId">The stream containing the job</param>
 		/// <param name="name">Name of the job</param>
 		/// <param name="templates">Templates to look for</param>
-		/// <param name="minChange">The minimum changelist number</param>
-		/// <param name="maxChange">The maximum changelist number</param>
-		/// <param name="preflightChange">Preflight change to find</param>
+		/// <param name="minCommitId">The minimum commit</param>
+		/// <param name="maxCommitId">The maximum commit</param>
+		/// <param name="preflightCommitId">Preflight change to find</param>
 		/// <param name="preflightOnly">Whether to only include preflights</param>
+		/// <param name="includePreflight">Whether to include preflights in the results</param>
 		/// <param name="startedByUser">User id for which to include jobs</param>
 		/// <param name="preflightStartedByUser">User for which to include preflight jobs</param>
 		/// <param name="minCreateTime">The minimum creation time</param>
@@ -171,7 +173,7 @@ namespace HordeServer.Jobs
 		/// <param name="excludeUserJobs">Whether to exclude user jobs from the find</param>
 		/// <param name="cancellationToken">Cancellation token for the operation</param>
 		/// <returns>List of jobs matching the given criteria</returns>
-		Task<IReadOnlyList<IJob>> FindAsync(JobId[]? jobIds = null, StreamId? streamId = null, string? name = null, TemplateId[]? templates = null, int? minChange = null, int? maxChange = null, int? preflightChange = null, bool? preflightOnly = null, UserId? preflightStartedByUser = null, UserId? startedByUser = null, DateTimeOffset? minCreateTime = null, DateTimeOffset? maxCreateTime = null, DateTimeOffset? modifiedBefore = null, DateTimeOffset? modifiedAfter = null, JobStepBatchState? batchState = null, int? index = null, int? count = null, bool consistentRead = true, string? indexHint = null, bool? excludeUserJobs = null, CancellationToken cancellationToken = default);
+		Task<IReadOnlyList<IJob>> FindAsync(JobId[]? jobIds = null, StreamId? streamId = null, string? name = null, TemplateId[]? templates = null, CommitId? minCommitId = null, CommitId? maxCommitId = null, CommitId? preflightCommitId = null, bool? preflightOnly = null, bool? includePreflight = null, UserId? preflightStartedByUser = null, UserId? startedByUser = null, DateTimeOffset? minCreateTime = null, DateTimeOffset? maxCreateTime = null, DateTimeOffset? modifiedBefore = null, DateTimeOffset? modifiedAfter = null, JobStepBatchState? batchState = null, int? index = null, int? count = null, bool consistentRead = true, string? indexHint = null, bool? excludeUserJobs = null, CancellationToken cancellationToken = default);
 
 		/// <summary>
 		/// Searches for jobs matching the given criteria
diff --git a/Engine/Source/Programs/Horde/Plugins/Build/HordeServer.Build/Jobs/IJobStepRef.cs b/Engine/Source/Programs/Horde/Plugins/Build/HordeServer.Build/Jobs/IJobStepRef.cs
--- a/Engine/Source/Programs/Horde/Plugins/Build/HordeServer.Build/Jobs/IJobStepRef.cs
+++ b/Engine/Source/Programs/Horde/Plugins/Build/HordeServer.Build/Jobs/IJobStepRef.cs
@@ -3,6 +3,7 @@
 using System.Diagnostics.CodeAnalysis;
 using EpicGames.Horde.Agents;
 using EpicGames.Horde.Agents.Pools;
+using EpicGames.Horde.Commits;
 using EpicGames.Horde.Jobs;
 using EpicGames.Horde.Jobs.Templates;
 using EpicGames.Horde.Logs;
@@ -181,7 +182,7 @@ namespace HordeServer.Jobs
 		/// <summary>
 		/// The change number being built
 		/// </summary>
-		public int Change { get; }
+		public CommitIdWithOrder CommitId { get; }
 
 		/// <summary>
 		/// Log for this step
@@ -216,12 +217,12 @@ namespace HordeServer.Jobs
 		/// <summary>
 		/// The last change that succeeded. Note that this is only set when the ref is updated; it is not necessarily consistent with steps run later.
 		/// </summary>
-		public int? LastSuccess { get; }
+		public CommitIdWithOrder? LastSuccess { get; }
 
 		/// <summary>
 		/// The last change that succeeded, or completed a warning. See <see cref="LastSuccess"/>.
 		/// </summary>
-		public int? LastWarning { get; }
+		public CommitIdWithOrder? LastWarning { get; }
 
 		/// <summary>
 		/// Time taken for the batch containing this batch to start after it became ready
diff --git a/Engine/Source/Programs/Horde/Plugins/Build/HordeServer.Build/Jobs/IJobStepRefCollection.cs b/Engine/Source/Programs/Horde/Plugins/Build/HordeServer.Build/Jobs/IJobStepRefCollection.cs
--- a/Engine/Source/Programs/Horde/Plugins/Build/HordeServer.Build/Jobs/IJobStepRefCollection.cs
+++ b/Engine/Source/Programs/Horde/Plugins/Build/HordeServer.Build/Jobs/IJobStepRefCollection.cs
@@ -2,6 +2,7 @@
 
 using EpicGames.Horde.Agents;
 using EpicGames.Horde.Agents.Pools;
+using EpicGames.Horde.Commits;
 using EpicGames.Horde.Jobs;
 using EpicGames.Horde.Jobs.Templates;
 using EpicGames.Horde.Logs;
@@ -24,7 +25,7 @@ namespace HordeServer.Jobs
 		/// <param name="stepName">Name of the step</param>
 		/// <param name="streamId">Unique id for the stream containing the job</param>
 		/// <param name="templateId"></param>
-		/// <param name="change">The change number being built</param>
+		/// <param name="commitId">The commit being built</param>
 		/// <param name="logId">The log file id</param>
 		/// <param name="poolId">The pool id</param>
 		/// <param name="agentId">The agent id</param>
@@ -38,7 +39,7 @@ namespace HordeServer.Jobs
 		/// <param name="jobStartTimeUtc">Start time of the job</param>
 		/// <param name="startTimeUtc">Start time</param>
 		/// <param name="finishTimeUtc">Finish time for the step, if known</param>		
-		Task<IJobStepRef> InsertOrReplaceAsync(JobStepRefId id, string jobName, string stepName, StreamId streamId, TemplateId templateId, int change, LogId? logId, PoolId? poolId, AgentId? agentId, JobStepState? state, JobStepOutcome? outcome, bool updateIssues, int? lastSuccess, int? lastWarning, float waitTime, float initTime, DateTime jobStartTimeUtc, DateTime startTimeUtc, DateTime? finishTimeUtc);
+		Task<IJobStepRef> InsertOrReplaceAsync(JobStepRefId id, string jobName, string stepName, StreamId streamId, TemplateId templateId, CommitIdWithOrder commitId, LogId? logId, PoolId? poolId, AgentId? agentId, JobStepState? state, JobStepOutcome? outcome, bool updateIssues, CommitIdWithOrder? lastSuccess, CommitIdWithOrder? lastWarning, float waitTime, float initTime, DateTime jobStartTimeUtc, DateTime startTimeUtc, DateTime? finishTimeUtc);
 
 		/// <summary>
 		/// Updates a job step ref 
@@ -73,12 +74,12 @@ namespace HordeServer.Jobs
 		/// <param name="streamId">Unique id for a stream</param>
 		/// <param name="templateId"></param>
 		/// <param name="nodeName">Name of the node</param>
-		/// <param name="change">The current change</param>
+		/// <param name="commitId">The current change</param>
 		/// <param name="includeFailed">Whether to include failed nodes</param>
 		/// <param name="maxCount">Number of results to return</param>
 		/// <param name="cancellationToken">Cancellation token for the operation</param>
 		/// <returns>List of step references</returns>
-		Task<List<IJobStepRef>> GetStepsForNodeAsync(StreamId streamId, TemplateId templateId, string nodeName, int? change, bool includeFailed, int maxCount, CancellationToken cancellationToken = default);
+		Task<List<IJobStepRef>> GetStepsForNodeAsync(StreamId streamId, TemplateId templateId, string nodeName, CommitIdWithOrder? commitId, bool includeFailed, int maxCount, CancellationToken cancellationToken = default);
 
 		/// <summary>
 		/// Gets the previous job that ran a given step
@@ -86,12 +87,12 @@ namespace HordeServer.Jobs
 		/// <param name="streamId">Id of the stream to search</param>
 		/// <param name="templateId">The template id</param>
 		/// <param name="nodeName">Name of the step to find</param>
-		/// <param name="change">The current changelist number</param>
+		/// <param name="commitId">The current changelist number</param>
 		/// <param name="outcome">The outcome to filter by or include all outcomes if null</param>
 		/// <param name="updateIssues">If true, constrain to steps which update issues</param>		 
 		/// <param name="excludeJobIds">Jobs to exclude from the search</param>
 		/// <returns>The previous job, or null.</returns>
-		Task<IJobStepRef?> GetPrevStepForNodeAsync(StreamId streamId, TemplateId templateId, string nodeName, int change, JobStepOutcome? outcome = null, bool? updateIssues = null, IEnumerable<JobId>? excludeJobIds = null);
+		Task<IJobStepRef?> GetPrevStepForNodeAsync(StreamId streamId, TemplateId templateId, string nodeName, CommitIdWithOrder commitId, JobStepOutcome? outcome = null, bool? updateIssues = null, IEnumerable<JobId>? excludeJobIds = null);
 
 		/// <summary>
 		/// Gets the next job that ran a given step
@@ -99,11 +100,11 @@ namespace HordeServer.Jobs
 		/// <param name="streamId">Id of the stream to search</param>
 		/// <param name="templateId">The template id</param>
 		/// <param name="nodeName">Name of the step to find</param>
-		/// <param name="change">The current changelist number</param>
+		/// <param name="commitId">The current changelist number</param>
 		/// <param name="outcome">The outcome to filter by or include all outcomes if null</param>
 		/// <param name="updateIssues">If true, constrain to steps which update issues</param>
 		/// <returns>The previous job, or null.</returns>
-		Task<IJobStepRef?> GetNextStepForNodeAsync(StreamId streamId, TemplateId templateId, string nodeName, int change, JobStepOutcome? outcome = null, bool? updateIssues = null);
+		Task<IJobStepRef?> GetNextStepForNodeAsync(StreamId streamId, TemplateId templateId, string nodeName, CommitIdWithOrder commitId, JobStepOutcome? outcome = null, bool? updateIssues = null);
 
 	}
 
@@ -111,7 +112,7 @@ namespace HordeServer.Jobs
 	{
 		public static async Task UpdateAsync(this IJobStepRefCollection jobStepRefs, IJob job, IJobStepBatch batch, IJobStep step, IGraph graph, ILogger? logger = null)
 		{
-			if (job.PreflightChange == 0)
+			if (job.PreflightCommitId == null)
 			{
 				float waitTime = (float)(batch.GetWaitTime() ?? TimeSpan.Zero).TotalSeconds;
 				float initTime = (float)(batch.GetInitTime() ?? TimeSpan.Zero).TotalSeconds;
@@ -119,11 +120,11 @@ namespace HordeServer.Jobs
 				string nodeName = graph.Groups[batch.GroupIdx].Nodes[step.NodeIdx].Name;
 				JobStepOutcome? outcome = step.IsPending() ? (JobStepOutcome?)null : step.Outcome;
 
-				int? lastSuccess = null;
-				int? lastWarning = null;
+				CommitIdWithOrder? lastSuccess = null;
+				CommitIdWithOrder? lastWarning = null;
 				if (outcome != JobStepOutcome.Success)
 				{
-					IJobStepRef? prevStep = await jobStepRefs.GetPrevStepForNodeAsync(job.StreamId, job.TemplateId, nodeName, job.Change);
+					IJobStepRef? prevStep = await jobStepRefs.GetPrevStepForNodeAsync(job.StreamId, job.TemplateId, nodeName, job.CommitId);
 					if (prevStep != null)
 					{
 						lastSuccess = prevStep.LastSuccess;
@@ -149,7 +150,7 @@ namespace HordeServer.Jobs
 					logger.LogInformation("Updating step reference {StepId} for job {JobId}, batch {BatchId}, with outcome {JobStepOutcome}", step.Id, job.Id, batch.Id, outcome);
 				}
 
-				await jobStepRefs.InsertOrReplaceAsync(new JobStepRefId(job.Id, batch.Id, step.Id), job.Name, nodeName, job.StreamId, job.TemplateId, job.Change, step.LogId, batch.PoolId, batch.AgentId, step.State, outcome, job.UpdateIssues, lastSuccess, lastWarning, waitTime, initTime, job.CreateTimeUtc, step.StartTimeUtc ?? DateTime.UtcNow, step.FinishTimeUtc);
+				await jobStepRefs.InsertOrReplaceAsync(new JobStepRefId(job.Id, batch.Id, step.Id), job.Name, nodeName, job.StreamId, job.TemplateId, job.CommitId, step.LogId, batch.PoolId, batch.AgentId, step.State, outcome, job.UpdateIssues, lastSuccess, lastWarning, waitTime, initTime, job.CreateTimeUtc, step.StartTimeUtc ?? DateTime.UtcNow, step.FinishTimeUtc);
 			}
 		}
 	}
diff --git a/Engine/Source/Programs/Horde/Plugins/Build/HordeServer.Build/Jobs/JobCollection.cs b/Engine/Source/Programs/Horde/Plugins/Build/HordeServer.Build/Jobs/JobCollection.cs
--- a/Engine/Source/Programs/Horde/Plugins/Build/HordeServer.Build/Jobs/JobCollection.cs
+++ b/Engine/Source/Programs/Horde/Plugins/Build/HordeServer.Build/Jobs/JobCollection.cs
@@ -9,6 +9,7 @@ using EpicGames.Horde.Agents;
 using EpicGames.Horde.Agents.Leases;
 using EpicGames.Horde.Agents.Pools;
 using EpicGames.Horde.Agents.Sessions;
+using EpicGames.Horde.Commits;
 using EpicGames.Horde.Jobs;
 using EpicGames.Horde.Jobs.Bisect;
 using EpicGames.Horde.Jobs.Templates;
@@ -16,6 +17,7 @@ using EpicGames.Horde.Logs;
 using EpicGames.Horde.Streams;
 using EpicGames.Horde.Users;
 using HordeServer.Acls;
+using HordeServer.Commits;
 using HordeServer.Jobs.Graphs;
 using HordeServer.Logs;
 using HordeServer.Server;
@@ -61,10 +63,10 @@ namespace HordeServer.Jobs
 
 			BisectTaskId? IJob.StartedByBisectTaskId => _document.StartedByBisectTaskId;
 			string IJob.Name => _document.Name;
-			int IJob.Change => _document.Change;
-			int IJob.CodeChange => _document.CodeChange;
-			int IJob.PreflightChange => _document.PreflightChange;
-			int IJob.ClonedPreflightChange => _document.ClonedPreflightChange;
+			CommitIdWithOrder IJob.CommitId => _document.CommitId;
+			CommitIdWithOrder? IJob.CodeCommitId => _document.CodeCommitId;
+			CommitId? IJob.PreflightCommitId => _document.PreflightCommitId;
+			CommitId? IJob.ClonedPreflightCommitId => _document.ClonedPreflightCommitId;
 			string? IJob.PreflightDescription => _document.PreflightDescription;
 			Priority IJob.Priority => _document.Priority;
 			bool IJob.AutoSubmit => _document.AutoSubmit;
@@ -415,10 +417,48 @@ namespace HordeServer.Jobs
 			[BsonRequired]
 			public string Name { get; set; }
 
-			public int Change { get; set; }
-			public int CodeChange { get; set; }
+			[BsonIgnore]
+			public CommitIdWithOrder CommitId
+			{
+				get => (CommitName != null)? new CommitIdWithOrder(CommitName, CommitOrder) : CommitIdWithOrder.FromPerforceChange(CommitOrder);
+				set => (CommitName, CommitOrder) = (value.Name, value.Order);
+			}
+
+			[BsonElement("Commit")]
+			public string? CommitName { get; set; }
+
+			[BsonElement("Change")]
+			public int CommitOrder { get; set; }
+
+			[BsonIgnore]
+			public CommitIdWithOrder? CodeCommitId
+			{
+				get => (CodeCommitName != null) ? new CommitIdWithOrder(CodeCommitName, CodeCommitOrder) : (CodeCommitOrder != 0) ? CommitIdWithOrder.FromPerforceChange(CodeCommitOrder) : null;
+				set => (CodeCommitName, CodeCommitOrder) = (value?.Name, value?.Order ?? 0);
+			}
+
+			[BsonElement("CodeCommit")]
+			public string? CodeCommitName { get; set; }
+
+			[BsonElement("CodeChange")]
+			public int CodeCommitOrder { get; set; }
+
+			[BsonIgnore]
+			public CommitId? PreflightCommitId
+			{
+				get => (PreflightCommitName != null) ? new CommitId(PreflightCommitName) : (PreflightChange != 0) ? EpicGames.Horde.Commits.CommitId.FromPerforceChange(PreflightChange) : null;
+				set => (PreflightCommitName, PreflightChange) = (value?.Name, (value == null) ? 0 : value.TryGetPerforceChange() ?? -1);
+			}
+
+			[BsonElement("PreflightCommit")]
+			public string? PreflightCommitName { get; set; }
+
+			// -1 for non-P4 preflights
 			public int PreflightChange { get; set; }
-			public int ClonedPreflightChange { get; set; }
+
+			[BsonElement("ClonedPreflightChange"), BsonIgnoreIfNull]
+			public CommitId? ClonedPreflightCommitId { get; set; }
+
 			public string? PreflightDescription { get; set; }
 			public Priority Priority { get; set; }
 
@@ -491,7 +531,7 @@ namespace HordeServer.Jobs
 				GraphHash = null!;
 			}
 
-			public JobDocument(JobId id, StreamId streamId, TemplateId templateId, ContentHash templateHash, ContentHash graphHash, string name, int change, int codeChange, CreateJobOptions options, DateTime createTimeUtc)
+			public JobDocument(JobId id, StreamId streamId, TemplateId templateId, ContentHash templateHash, ContentHash graphHash, string name, CommitIdWithOrder commitId, CommitIdWithOrder? codeCommitId, CreateJobOptions options, DateTime createTimeUtc)
 			{
 				Id = id;
 				StreamId = streamId;
@@ -499,16 +539,16 @@ namespace HordeServer.Jobs
 				TemplateHash = templateHash;
 				GraphHash = graphHash;
 				Name = name;
-				Change = change;
-				CodeChange = codeChange;
-				PreflightChange = options.PreflightChange ?? 0;
-				ClonedPreflightChange = options.ClonedPreflightChange ?? 0;
+				CommitId = commitId;
+				CodeCommitId = codeCommitId;
+				PreflightCommitId = options.PreflightCommitId;
+				ClonedPreflightCommitId = options.ClonedPreflightCommitId;
 				PreflightDescription = options.PreflightDescription;
 				StartedByUserId = options.StartedByUserId;
 				StartedByBisectTaskId = options.StartedByBisectTaskId;
 				Priority = options.Priority ?? Priority.Normal;
 				AutoSubmit = options.AutoSubmit ?? false;
-				UpdateIssues = options.UpdateIssues ?? (options.StartedByUserId == null && (options.PreflightChange == 0 || options.PreflightChange == null));
+				UpdateIssues = options.UpdateIssues ?? (options.StartedByUserId == null && options.PreflightCommitId == null);
 				PromoteIssuesByDefault = options.PromoteIssuesByDefault ?? false;
 				Claims = options.Claims;
 				JobOptions = options.JobOptions;
@@ -725,6 +765,7 @@ namespace HordeServer.Jobs
 		readonly IClock _clock;
 		readonly IGraphCollection _graphCollection;
 		readonly ILogCollection _logCollection;
+		readonly ICommitService _commitService;
 		readonly Tracer _tracer;
 		readonly IOptionsMonitor<BuildConfig> _buildConfig;
 		readonly ILogger<JobCollection> _logger;
@@ -732,11 +773,12 @@ namespace HordeServer.Jobs
 		/// <summary>
 		/// Constructor
 		/// </summary>
-		public JobCollection(IMongoService mongoService, IClock clock, IGraphCollection graphCollection, ILogCollection logCollection, ITelemetryWriter telemetryWriter, IOptionsMonitor<BuildConfig> buildConfig, Tracer tracer, ILogger<JobCollection> logger)
+		public JobCollection(IMongoService mongoService, IClock clock, IGraphCollection graphCollection, ILogCollection logCollection, ICommitService commitService, ITelemetryWriter telemetryWriter, IOptionsMonitor<BuildConfig> buildConfig, Tracer tracer, ILogger<JobCollection> logger)
 		{
 			_clock = clock;
 			_graphCollection = graphCollection;
 			_logCollection = logCollection;
+			_commitService = commitService;
 			_telemetryWriter = telemetryWriter;
 			_buildConfig = buildConfig;
 			_tracer = tracer;
@@ -745,8 +787,9 @@ namespace HordeServer.Jobs
 			List<MongoIndex<JobDocument>> indexes = new List<MongoIndex<JobDocument>>();
 			indexes.Add(keys => keys.Ascending(x => x.StreamId).Descending(x => x.CreateTimeUtc));
 			indexes.Add(_streamThenTemplateThenCreationTimeIndex = MongoIndex.Create<JobDocument>(keys => keys.Ascending(x => x.StreamId).Ascending(x => x.TemplateId).Descending(x => x.CreateTimeUtc)));
-			indexes.Add(MongoIndex.Create<JobDocument>(keys => keys.Ascending(x => x.StreamId).Ascending(x => x.TemplateId).Descending(x => x.Change)));
-			indexes.Add(keys => keys.Ascending(x => x.Change));
+			indexes.Add(MongoIndex.Create<JobDocument>(keys => keys.Ascending(x => x.StreamId).Ascending(x => x.TemplateId).Descending(x => x.CommitOrder)));
+			indexes.Add(keys => keys.Ascending(x => x.CommitOrder));
+			indexes.Add(keys => keys.Ascending(x => x.PreflightCommitName));
 			indexes.Add(keys => keys.Ascending(x => x.PreflightChange));
 			indexes.Add(_createTimeIndex = MongoIndex.Create<JobDocument>(keys => keys.Descending(x => x.CreateTimeUtc)));
 			indexes.Add(keys => keys.Ascending(x => x.StartedByUserId));
@@ -777,9 +820,17 @@ namespace HordeServer.Jobs
 		}
 
 		/// <inheritdoc/>
-		public async Task<IJob> AddAsync(JobId jobId, StreamId streamId, TemplateId templateRefId, ContentHash templateHash, IGraph graph, string name, int change, int codeChange, CreateJobOptions options, CancellationToken cancellationToken)
+		public async Task<IJob> AddAsync(JobId jobId, StreamId streamId, TemplateId templateRefId, ContentHash templateHash, IGraph graph, string name, CommitId commitId, CommitId? codeCommitId, CreateJobOptions options, CancellationToken cancellationToken)
 		{
-			JobDocument newJob = new JobDocument(jobId, streamId, templateRefId, templateHash, graph.Id, name, change, codeChange, options, DateTime.UtcNow);
+			CommitIdWithOrder commitIdWithOrder = await _commitService.GetOrderedAsync(streamId, commitId, cancellationToken);
+
+			CommitIdWithOrder? codeCommitIdWithOrder = null;
+			if (codeCommitId != null)
+			{
+				codeCommitIdWithOrder = await _commitService.GetOrderedAsync(streamId, codeCommitId, cancellationToken);
+			}
+
+			JobDocument newJob = new JobDocument(jobId, streamId, templateRefId, templateHash, graph.Id, name, commitIdWithOrder, codeCommitIdWithOrder, options, DateTime.UtcNow);
 			CreateBatches(newJob, graph, _logger);
 
 			await _jobs.InsertOneAsync(newJob, null, cancellationToken);
@@ -793,12 +844,12 @@ namespace HordeServer.Jobs
 					StreamId = newJob.StreamId,
 					Arguments = newJob.Arguments,
 					AutoSubmit = newJob.AutoSubmit,
-					Change = newJob.Change,
-					CodeChange = newJob.CodeChange,
+					Change = newJob.CommitId.Name,
+					CodeChange = newJob.CodeCommitId?.Name,
 					CreateTimeUtc = newJob.CreateTimeUtc,
 					GraphHash = newJob.GraphHash,
 					Name = newJob.Name,
-					PreflightChange = newJob.PreflightChange,
+					PreflightChange = newJob.PreflightCommitId,
 					PreflightDescription = newJob.PreflightDescription,
 					Priority = newJob.Priority,
 					StartedByUserId = newJob.StartedByUserId,
@@ -847,7 +898,7 @@ namespace HordeServer.Jobs
 		}
 
 		/// <inheritdoc/>
-		public async Task<IReadOnlyList<IJob>> FindAsync(JobId[]? jobIds, StreamId? streamId, string? name, TemplateId[]? templates, int? minChange, int? maxChange, int? preflightChange, bool? preflightOnly, UserId? preflightStartedByUser, UserId? startedByUser, DateTimeOffset? minCreateTime, DateTimeOffset? maxCreateTime, DateTimeOffset? modifiedBefore, DateTimeOffset? modifiedAfter, JobStepBatchState? batchState, int? index, int? count, bool consistentRead, string? indexHint, bool? excludeUserJobs, CancellationToken cancellationToken)
+		public async Task<IReadOnlyList<IJob>> FindAsync(JobId[]? jobIds, StreamId? streamId, string? name, TemplateId[]? templates, CommitId? minCommitId, CommitId? maxCommitId, CommitId? preflightCommitId, bool? preflightOnly, bool? includePreflight, UserId? preflightStartedByUser, UserId? startedByUser, DateTimeOffset? minCreateTime, DateTimeOffset? maxCreateTime, DateTimeOffset? modifiedBefore, DateTimeOffset? modifiedAfter, JobStepBatchState? batchState, int? index, int? count, bool consistentRead, string? indexHint, bool? excludeUserJobs, CancellationToken cancellationToken)
 		{
 			FilterDefinitionBuilder<JobDocument> filterBuilder = Builders<JobDocument>.Filter;
 
@@ -859,6 +910,16 @@ namespace HordeServer.Jobs
 			if (streamId != null)
 			{
 				filter &= filterBuilder.Eq(x => x.StreamId, streamId.Value);
+				if (minCommitId != null)
+				{
+					CommitIdWithOrder minCommitIdWithOrder = await _commitService.GetOrderedAsync(streamId.Value, minCommitId, cancellationToken);
+					filter &= filterBuilder.Gte(x => x.CommitOrder, minCommitIdWithOrder.Order);
+				}
+				if (maxCommitId != null)
+				{
+					CommitIdWithOrder maxCommitIdWithOrder = await _commitService.GetOrderedAsync(streamId.Value, maxCommitId, cancellationToken);
+					filter &= filterBuilder.Lte(x => x.CommitOrder, maxCommitIdWithOrder.Order);
+				}
 			}
 			if (name != null)
 			{
@@ -883,17 +944,21 @@ namespace HordeServer.Jobs
 					filter &= filterBuilder.In(x => x.TemplateId, templates);
 				}
 			}
-			if (minChange != null)
+			if (preflightCommitId != null)
 			{
-				filter &= filterBuilder.Gte(x => x.Change, minChange);
-			}
-			if (maxChange != null)
-			{
-				filter &= filterBuilder.Lte(x => x.Change, maxChange);
+				int? preflightChange = preflightCommitId.TryGetPerforceChange();
+				if (preflightChange == null)
+				{
+					filter &= filterBuilder.Eq(x => x.PreflightCommitName, preflightCommitId.Name);
+				}
+				else
+				{
+					filter &= filterBuilder.Eq(x => x.PreflightChange, preflightChange);
+				}
 			}
-			if (preflightChange != null)
+			if (includePreflight != null && !includePreflight.Value)
 			{
-				filter &= filterBuilder.Eq(x => x.PreflightChange, preflightChange);
+				filter &= filterBuilder.Eq(x => x.PreflightChange, 0);
 			}
 			if (preflightOnly != null && preflightOnly.Value)
 			{
diff --git a/Engine/Source/Programs/Horde/Plugins/Build/HordeServer.Build/Jobs/JobRpcService.cs b/Engine/Source/Programs/Horde/Plugins/Build/HordeServer.Build/Jobs/JobRpcService.cs
--- a/Engine/Source/Programs/Horde/Plugins/Build/HordeServer.Build/Jobs/JobRpcService.cs
+++ b/Engine/Source/Programs/Horde/Plugins/Build/HordeServer.Build/Jobs/JobRpcService.cs
@@ -1,6 +1,5 @@
 // Copyright Epic Games, Inc. All Rights Reserved.
 
-using System.Globalization;
 using System.Security.Claims;
 using System.Text;
 using System.Text.Json;
@@ -10,6 +9,7 @@ using EpicGames.Horde;
 using EpicGames.Horde.Agents;
 using EpicGames.Horde.Agents.Sessions;
 using EpicGames.Horde.Artifacts;
+using EpicGames.Horde.Commits;
 using EpicGames.Horde.Jobs;
 using EpicGames.Horde.Logs;
 using Google.Protobuf.WellKnownTypes;
@@ -197,10 +197,10 @@ namespace HordeServer.Jobs
 			response.LogId = batch.LogId.ToString();
 			response.AgentType = graph.Groups[batch.GroupIdx].AgentType;
 			response.StreamName = streamConfig.Name;
-			response.Change = job.Change;
-			response.CodeChange = job.CodeChange;
-			response.PreflightChange = job.PreflightChange;
-			response.ClonedPreflightChange = job.ClonedPreflightChange;
+			response.Change = job.CommitId.GetPerforceChange();
+			response.CodeChange = job.CodeCommitId?.GetPerforceChange() ?? 0;
+			response.PreflightChange = job.PreflightCommitId?.GetPerforceChange() ?? 0;
+			response.ClonedPreflightChange = job.ClonedPreflightCommitId?.GetPerforceChange() ?? 0;
 			response.Arguments.AddRange(job.Arguments);
 			if (agentConfig.TempStorageDir != null)
 			{
@@ -374,60 +374,60 @@ namespace HordeServer.Jobs
 					response.EnvVars.Add(job.Environment);
 				}
 
-				IJobStepRef? lastStep = await _jobStepRefCollection.GetPrevStepForNodeAsync(job.StreamId, job.TemplateId, node.Name, job.Change);
+				IJobStepRef? lastStep = await _jobStepRefCollection.GetPrevStepForNodeAsync(job.StreamId, job.TemplateId, node.Name, job.CommitId);
 				if (lastStep != null)
 				{
-					response.EnvVars.Add("UE_HORDE_LAST_CL", lastStep.Change.ToString(CultureInfo.InvariantCulture));
+					response.EnvVars.Add("UE_HORDE_LAST_CL", lastStep.CommitId.ToString());
 
-					int? lastSuccessChange = null;
+					CommitIdWithOrder? lastSuccessCommitId = null;
 					if (lastStep.Outcome == JobStepOutcome.Success)
 					{
-						lastSuccessChange = lastStep.Change;
+						lastSuccessCommitId = lastStep.CommitId;
 					}
 					else if (lastStep.LastSuccess != null)
 					{
-						lastSuccessChange = lastStep.LastSuccess.Value;
+						lastSuccessCommitId = lastStep.LastSuccess;
 					}
 					else
 					{
 						// Previous job hasn't finished yet; need to search for *current* last success step
-						IJobStepRef? lastSuccess = await _jobStepRefCollection.GetPrevStepForNodeAsync(job.StreamId, job.TemplateId, node.Name, job.Change, outcome: JobStepOutcome.Success);
+						IJobStepRef? lastSuccess = await _jobStepRefCollection.GetPrevStepForNodeAsync(job.StreamId, job.TemplateId, node.Name, job.CommitId, outcome: JobStepOutcome.Success);
 						if (lastSuccess != null)
 						{
-							lastSuccessChange = lastSuccess.Change;
+							lastSuccessCommitId = lastSuccess.CommitId;
 						}
 					}
 
-					int? lastWarningChange = null;
+					CommitIdWithOrder? lastWarningCommitId = null;
 					if (lastStep.Outcome == JobStepOutcome.Success || lastStep.Outcome == JobStepOutcome.Warnings)
 					{
-						lastWarningChange = lastStep.Change;
+						lastWarningCommitId = lastStep.CommitId;
 					}
 					else if (lastStep.LastWarning != null)
 					{
-						lastWarningChange = lastStep.LastWarning.Value;
+						lastWarningCommitId = lastStep.LastWarning;
 					}
 					else
 					{
 						// Previous job hasn't finished yet; need to search for *current* last warning step
-						IJobStepRef? lastWarnings = await _jobStepRefCollection.GetPrevStepForNodeAsync(job.StreamId, job.TemplateId, node.Name, job.Change, outcome: JobStepOutcome.Warnings);
-						if (lastWarnings != null && (lastSuccessChange == null || lastWarnings.Change > lastSuccessChange.Value))
+						IJobStepRef? lastWarnings = await _jobStepRefCollection.GetPrevStepForNodeAsync(job.StreamId, job.TemplateId, node.Name, job.CommitId, outcome: JobStepOutcome.Warnings);
+						if (lastWarnings != null && (lastSuccessCommitId == null || lastWarnings.CommitId > lastSuccessCommitId))
 						{
-							lastWarningChange = lastWarnings.Change;
+							lastWarningCommitId = lastWarnings.CommitId;
 						}
 						else
 						{
-							lastWarningChange = lastSuccessChange;
+							lastWarningCommitId = lastSuccessCommitId;
 						}
 					}
 
-					if (lastSuccessChange != null)
+					if (lastSuccessCommitId != null)
 					{
-						response.EnvVars.Add("UE_HORDE_LAST_SUCCESS_CL", lastSuccessChange.Value.ToString(CultureInfo.InvariantCulture));
+						response.EnvVars.Add("UE_HORDE_LAST_SUCCESS_CL", lastSuccessCommitId.ToString());
 					}
-					if (lastWarningChange != null)
+					if (lastWarningCommitId != null)
 					{
-						response.EnvVars.Add("UE_HORDE_LAST_WARNING_CL", lastWarningChange.Value.ToString(CultureInfo.InvariantCulture));
+						response.EnvVars.Add("UE_HORDE_LAST_WARNING_CL", lastWarningCommitId.ToString());
 					}
 				}
 
@@ -819,7 +819,7 @@ namespace HordeServer.Jobs
 				description = request.Name;
 			}
 
-			IArtifact artifact = await _artifactCollection.AddAsync(name, type, description, job.StreamId, job.Change, keys, request.Metadata, templateConfig.Acl.ScopeName, context.CancellationToken);
+			IArtifact artifact = await _artifactCollection.AddAsync(name, type, description, job.StreamId, job.CommitId, keys, request.Metadata, templateConfig.Acl.ScopeName, context.CancellationToken);
 
 			List<AclClaimConfig> claims = new List<AclClaimConfig>();
 			claims.Add(new AclClaimConfig(HordeClaimTypes.WriteNamespace, $"{artifact.NamespaceId}:{artifact.RefName}"));
diff --git a/Engine/Source/Programs/Horde/Plugins/Build/HordeServer.Build/Jobs/JobService.cs b/Engine/Source/Programs/Horde/Plugins/Build/HordeServer.Build/Jobs/JobService.cs
--- a/Engine/Source/Programs/Horde/Plugins/Build/HordeServer.Build/Jobs/JobService.cs
+++ b/Engine/Source/Programs/Horde/Plugins/Build/HordeServer.Build/Jobs/JobService.cs
@@ -1,6 +1,5 @@
 // Copyright Epic Games, Inc. All Rights Reserved.
 
-using System.Globalization;
 using System.Security.Claims;
 using System.Text.RegularExpressions;
 using EpicGames.Core;
@@ -159,12 +158,12 @@ namespace HordeServer.Jobs
 		/// <param name="templateHash">Template for this job</param>
 		/// <param name="graph">The graph for the new job</param>
 		/// <param name="name">Name of the job</param>
-		/// <param name="change">The change to build</param>
-		/// <param name="codeChange">The corresponding code changelist</param>
+		/// <param name="commitId">The change to build</param>
+		/// <param name="codeCommitId">The corresponding code changelist</param>
 		/// <param name="options">Options for the new job</param>
 		/// <param name="cancellationToken">Cancellation token for the operation</param>
 		/// <returns>Unique id representing the job</returns>
-		public async Task<IJob> CreateJobAsync(JobId? jobId, StreamConfig streamConfig, TemplateId templateRefId, ContentHash templateHash, IGraph graph, string name, int change, int codeChange, CreateJobOptions options, CancellationToken cancellationToken = default)
+		public async Task<IJob> CreateJobAsync(JobId? jobId, StreamConfig streamConfig, TemplateId templateRefId, ContentHash templateHash, IGraph graph, string name, CommitId commitId, CommitId? codeCommitId, CreateJobOptions options, CancellationToken cancellationToken = default)
 		{
 			using TelemetrySpan span = _tracer.StartActiveSpan($"{nameof(JobService)}.{nameof(CreateJobAsync)}");
 			span.SetAttribute("JobId", jobId?.ToString());
@@ -173,10 +172,10 @@ namespace HordeServer.Jobs
 			span.SetAttribute("TemplateHash", templateHash);
 			span.SetAttribute("GraphId", graph.Id);
 			span.SetAttribute("Name", name);
-			span.SetAttribute("Change", change);
-			span.SetAttribute("CodeChange", codeChange);
-			span.SetAttribute("PreflightChange", options.PreflightChange);
-			span.SetAttribute("ClonedPreflightChange", options.ClonedPreflightChange);
+			span.SetAttribute("Change", commitId.ToString());
+			span.SetAttribute("CodeChange", codeCommitId?.ToString());
+			span.SetAttribute("PreflightChange", options.PreflightCommitId?.GetPerforceChange());
+			span.SetAttribute("ClonedPreflightChange", options.ClonedPreflightCommitId?.GetPerforceChange());
 			span.SetAttribute("StartedByUserId", options.StartedByUserId.ToString());
 			span.SetAttribute("Priority", options.Priority.ToString());
 			span.SetAttribute("ShowUgsBadges", options.ShowUgsBadges);
@@ -200,18 +199,18 @@ namespace HordeServer.Jobs
 			JobId jobIdValue = jobId ?? JobIdUtils.GenerateNewId();
 			using IDisposable? scope = _logger.BeginScope("CreateJobAsync({JobId})", jobIdValue);
 
-			if (options.PreflightChange != null && ShouldClonePreflightChange(streamConfig.Id))
+			if (options.PreflightCommitId != null && ShouldClonePreflightChange(streamConfig.Id))
 			{
-				options.ClonedPreflightChange = await CloneShelvedChangeAsync(streamConfig.ClusterName, options.ClonedPreflightChange ?? options.PreflightChange.Value, cancellationToken);
+				options.ClonedPreflightCommitId = await CloneShelvedChangeAsync(streamConfig.ClusterName, options.ClonedPreflightCommitId ?? options.PreflightCommitId, cancellationToken);
 			}
 
-			_logger.LogInformation("Creating job at CL {Change}, code CL {CodeChange}, preflight CL {PreflightChange}, cloned CL {ClonedPreflightChange}", change, codeChange, options.PreflightChange, options.ClonedPreflightChange);
+			_logger.LogInformation("Creating job at CL {Change}, code CL {CodeChange}, preflight CL {PreflightChange}, cloned CL {ClonedPreflightChange}", commitId, codeCommitId, options.PreflightCommitId, options.ClonedPreflightCommitId);
 
 			Dictionary<string, string> properties = new Dictionary<string, string>();
-			properties["Change"] = change.ToString(CultureInfo.InvariantCulture);
-			properties["CodeChange"] = codeChange.ToString(CultureInfo.InvariantCulture);
-			properties["PreflightChange"] = options.PreflightChange?.ToString(CultureInfo.InvariantCulture) ?? String.Empty;
-			properties["ClonedPreflightChange"] = options.ClonedPreflightChange?.ToString(CultureInfo.InvariantCulture) ?? String.Empty;
+			properties["Change"] = commitId.ToString();
+			properties["CodeChange"] = codeCommitId?.ToString() ?? String.Empty;
+			properties["PreflightChange"] = options.PreflightCommitId?.ToString() ?? String.Empty;
+			properties["ClonedPreflightChange"] = options.ClonedPreflightCommitId?.ToString() ?? String.Empty;
 			properties["StreamId"] = streamConfig.Id.ToString();
 			properties["TemplateId"] = templateRefId.ToString();
 			properties["JobId"] = jobIdValue.ToString();
@@ -223,7 +222,7 @@ namespace HordeServer.Jobs
 
 			name = StringUtils.ExpandProperties(name, properties);
 
-			IJob newJob = await _jobs.AddAsync(jobIdValue, streamConfig.Id, templateRefId, templateHash, graph, name, change, codeChange, options, cancellationToken);
+			IJob newJob = await _jobs.AddAsync(jobIdValue, streamConfig.Id, templateRefId, templateHash, graph, name, commitId, codeCommitId, options, cancellationToken);
 			_jobTaskSource.UpdateQueuedJob(newJob, graph);
 
 			await _jobTaskSource.UpdateUgsBadgesAsync(newJob, graph, new List<(LabelState, LabelOutcome)>(), cancellationToken);
@@ -245,9 +244,9 @@ namespace HordeServer.Jobs
 			span.SetAttribute("JobName", newJob.Name);
 
 			IReadOnlyList<IJob> jobsToAbort = new List<IJob>();
-			if (newJob.PreflightChange > 0)
+			if (newJob.PreflightCommitId != null)
 			{
-				jobsToAbort = await _jobs.FindAsync(preflightChange: newJob.PreflightChange, cancellationToken: cancellationToken);
+				jobsToAbort = await _jobs.FindAsync(preflightCommitId: newJob.PreflightCommitId, cancellationToken: cancellationToken);
 			}
 
 			foreach (IJob job in jobsToAbort)
@@ -397,14 +396,14 @@ namespace HordeServer.Jobs
 		/// Evaluate a change query to determine which CL to run a job at
 		/// </summary>
 		/// <returns></returns>
-		public async Task<int?> EvaluateChangeQueriesAsync(StreamId streamId, List<ChangeQueryConfig> queries, List<CommitTag>? commitTags, ICommitCollection commits, CancellationToken cancellationToken)
+		public async Task<CommitIdWithOrder?> EvaluateChangeQueriesAsync(StreamId streamId, List<ChangeQueryConfig> queries, List<CommitTag>? commitTags, ICommitCollection commits, CancellationToken cancellationToken)
 		{
 			foreach (ChangeQueryConfig query in queries)
 			{
-				int? change = await EvaluateChangeQueryAsync(streamId, query, commitTags, commits, cancellationToken);
-				if (change != null)
+				CommitIdWithOrder? commitId = await EvaluateChangeQueryAsync(streamId, query, commitTags, commits, cancellationToken);
+				if (commitId != null)
 				{
-					return change;
+					return commitId;
 				}
 			}
 			return null;
@@ -414,7 +413,7 @@ namespace HordeServer.Jobs
 		/// Evaluate a change query to determine which CL to run a job at
 		/// </summary>
 		/// <returns></returns>
-		public async ValueTask<int?> EvaluateChangeQueryAsync(StreamId streamId, ChangeQueryConfig query, List<CommitTag>? commitTags, ICommitCollection commits, CancellationToken cancellationToken)
+		public async ValueTask<CommitIdWithOrder?> EvaluateChangeQueryAsync(StreamId streamId, ChangeQueryConfig query, List<CommitTag>? commitTags, ICommitCollection commits, CancellationToken cancellationToken)
 		{
 			if (query.Condition == null || query.Condition.Evaluate(propertyName => GetTagPropertyValues(propertyName, commitTags)))
 			{
@@ -424,8 +423,8 @@ namespace HordeServer.Jobs
 					ICommit? taggedCommit = await commits.FindAsync(null, null, 1, new[] { query.CommitTag.Value }, cancellationToken).FirstOrDefaultAsync(cancellationToken);
 					if (taggedCommit != null)
 					{
-						_logger.LogInformation("Last commit with tag '{Tag}' was {Change}", query.CommitTag.Value, taggedCommit.Number);
-						return taggedCommit.Number;
+						_logger.LogInformation("Last commit with tag '{Tag}' was {Change}", query.CommitTag.Value, taggedCommit.Id);
+						return taggedCommit.Id;
 					}
 				}
 
@@ -437,8 +436,8 @@ namespace HordeServer.Jobs
 					IReadOnlyList<IJob> jobs = await FindJobsAsync(streamId: streamId, templates: new[] { query.TemplateId.Value }, target: query.Target, state: new[] { JobStepState.Completed }, outcome: outcomes.ToArray(), count: 1, excludeUserJobs: true, excludeCancelled: true, cancellationToken: cancellationToken);
 					if (jobs.Count > 0)
 					{
-						_logger.LogInformation("Last successful build of {TemplateId} target {Target} was job {JobId} at change {Change}", query.TemplateId, query.Target, jobs[0].Id, jobs[0].Change);
-						return jobs[0].Change;
+						_logger.LogInformation("Last successful build of {TemplateId} target {Target} was job {JobId} at change {Change}", query.TemplateId, query.Target, jobs[0].Id, jobs[0].CommitId);
+						return jobs[0].CommitId;
 					}
 				}
 			}
@@ -536,10 +535,11 @@ namespace HordeServer.Jobs
 		/// <param name="streamId">The stream containing the job</param>
 		/// <param name="name">Name of the job</param>
 		/// <param name="templates">Templates to look for</param>
-		/// <param name="minChange">The minimum changelist number</param>
-		/// <param name="maxChange">The maximum changelist number</param>		
-		/// <param name="preflightChange">The preflight change to look for</param>
+		/// <param name="minCommitId">The minimum commit</param>
+		/// <param name="maxCommitId">The maximum commit</param>		
+		/// <param name="preflightCommitId">The preflight change to look for</param>
 		/// <param name="preflightOnly">Whether to only include preflights</param>
+		/// <param name="includePreflights">Whether to include preflight jobs</param>
 		/// <param name="preflightStartedByUser">User for which to include preflight jobs</param>		
 		/// <param name="startedByUser">User for which to include jobs</param>
 		/// <param name="minCreateTime">The minimum creation time</param>
@@ -557,16 +557,16 @@ namespace HordeServer.Jobs
 		/// <param name="excludeCancelled">Whether to exclude cancelled jobs</param>
 		/// <param name="cancellationToken">Cancellation token for the operation</param>
 		/// <returns>List of jobs matching the given criteria</returns>
-		public async Task<IReadOnlyList<IJob>> FindJobsAsync(JobId[]? jobIds = null, StreamId? streamId = null, string? name = null, TemplateId[]? templates = null, int? minChange = null, int? maxChange = null, int? preflightChange = null, bool? preflightOnly = null, UserId? preflightStartedByUser = null, UserId? startedByUser = null, DateTimeOffset? minCreateTime = null, DateTimeOffset? maxCreateTime = null, string? target = null, JobStepBatchState? batchState = null, JobStepState[]? state = null, JobStepOutcome[]? outcome = null, DateTimeOffset? modifiedBefore = null, DateTimeOffset? modifiedAfter = null, int? index = null, int? count = null, bool consistentRead = true, bool? excludeUserJobs = null, bool? excludeCancelled = null, CancellationToken cancellationToken = default)
+		public async Task<IReadOnlyList<IJob>> FindJobsAsync(JobId[]? jobIds = null, StreamId? streamId = null, string? name = null, TemplateId[]? templates = null, CommitId? minCommitId = null, CommitId? maxCommitId = null, CommitId? preflightCommitId = null, bool? preflightOnly = null, bool? includePreflights = null, UserId? preflightStartedByUser = null, UserId? startedByUser = null, DateTimeOffset? minCreateTime = null, DateTimeOffset? maxCreateTime = null, string? target = null, JobStepBatchState? batchState = null, JobStepState[]? state = null, JobStepOutcome[]? outcome = null, DateTimeOffset? modifiedBefore = null, DateTimeOffset? modifiedAfter = null, int? index = null, int? count = null, bool consistentRead = true, bool? excludeUserJobs = null, bool? excludeCancelled = null, CancellationToken cancellationToken = default)
 		{
 			using TelemetrySpan span = _tracer.StartActiveSpan($"{nameof(JobService)}.{nameof(FindJobsAsync)}");
 			span.SetAttribute("JobIds", (jobIds == null) ? null : String.Join(',', jobIds));
 			span.SetAttribute("StreamId", streamId);
 			span.SetAttribute("Name", name);
 			span.SetAttribute("Templates", templates);
-			span.SetAttribute("MinChange", minChange);
-			span.SetAttribute("MaxChange", maxChange);
-			span.SetAttribute("PreflightChange", preflightChange);
+			span.SetAttribute("MinCommitId", minCommitId?.ToString());
+			span.SetAttribute("MaxCommitId", maxCommitId?.ToString());
+			span.SetAttribute("PreflightCommitId", preflightCommitId?.ToString());
 			span.SetAttribute("PreflightStartedByUser", preflightStartedByUser?.ToString());
 			span.SetAttribute("StartedByUser", startedByUser?.ToString());
 			span.SetAttribute("MinCreateTime", minCreateTime);
@@ -581,24 +581,31 @@ namespace HordeServer.Jobs
 
 			if (target == null && (state == null || state.Length == 0) && (outcome == null || outcome.Length == 0))
 			{
-				return await _jobs.FindAsync(jobIds, streamId, name, templates, minChange, maxChange, preflightChange, preflightOnly, preflightStartedByUser, startedByUser, minCreateTime, maxCreateTime, modifiedBefore, modifiedAfter, batchState, index, count, consistentRead, null, excludeUserJobs, cancellationToken);
+				return await _jobs.FindAsync(jobIds, streamId, name, templates, minCommitId, maxCommitId, preflightCommitId, preflightOnly, includePreflights, preflightStartedByUser, startedByUser, minCreateTime, maxCreateTime, modifiedBefore, modifiedAfter, batchState, index, count, consistentRead, null, excludeUserJobs, cancellationToken);
 			}
 			else
 			{
 				List<IJob> results = new List<IJob>();
 				_logger.LogInformation("Performing scan for job with ");
 
+				bool excludeMaxCommitId = false;
+
 				int maxCount = (count ?? 1);
 				while (results.Count < maxCount)
 				{
-					IReadOnlyList<IJob> scanJobs = await _jobs.FindAsync(jobIds, streamId, name, templates, minChange, maxChange, preflightChange, preflightOnly, preflightStartedByUser, startedByUser, minCreateTime, maxCreateTime, modifiedBefore, modifiedAfter, batchState, 0, 5, consistentRead, null, excludeUserJobs, cancellationToken);
+					IReadOnlyList<IJob> scanJobs = await _jobs.FindAsync(jobIds, streamId, name, templates, minCommitId, maxCommitId, preflightCommitId, preflightOnly, includePreflights, preflightStartedByUser, startedByUser, minCreateTime, maxCreateTime, modifiedBefore, modifiedAfter, batchState, 0, 5, consistentRead, null, excludeUserJobs, cancellationToken);
 					if (scanJobs.Count == 0)
 					{
 						break;
 					}
 
-					foreach (IJob job in scanJobs.OrderByDescending(x => x.Change))
+					foreach (IJob job in scanJobs.OrderByDescending(x => x.CommitId))
 					{
+						if (excludeMaxCommitId && job.CommitId == maxCommitId)
+						{
+							continue;
+						}
+
 						if (excludeCancelled != null && excludeCancelled.Value && WasCancelled(job))
 						{
 							continue;
@@ -628,7 +635,8 @@ namespace HordeServer.Jobs
 						}
 					}
 
-					maxChange = scanJobs.Min(x => x.Change) - 1;
+					maxCommitId = scanJobs.Min(x => x.CommitId);
+					excludeMaxCommitId = true;
 				}
 
 				return results;
@@ -731,7 +739,7 @@ namespace HordeServer.Jobs
 							JobStepTimingData? stepTimingData;
 							if (!cachedNewSteps.TryGetValue(node.Name, out stepTimingData))
 							{
-								stepTimingData = await GetStepTimingInfoAsync(job.StreamId, job.TemplateId, node.Name, job.Change, cancellationToken);
+								stepTimingData = await GetStepTimingInfoAsync(job.StreamId, job.TemplateId, node.Name, job.CommitId, cancellationToken);
 							}
 							newSteps.Add(stepTimingData);
 						}
@@ -774,19 +782,19 @@ namespace HordeServer.Jobs
 		/// <param name="streamId">The stream to search</param>
 		/// <param name="templateId">The template id</param>
 		/// <param name="nodeName">Name of the node</param>
-		/// <param name="change">Maximum changelist to consider</param>
+		/// <param name="commitId">Maximum changelist to consider</param>
 		/// <param name="cancellationToken">Cancellation token for the operation</param>
 		/// <returns>Expected duration for the given step</returns>
-		async Task<JobStepTimingData> GetStepTimingInfoAsync(StreamId streamId, TemplateId templateId, string nodeName, int? change, CancellationToken cancellationToken)
+		async Task<JobStepTimingData> GetStepTimingInfoAsync(StreamId streamId, TemplateId templateId, string nodeName, CommitIdWithOrder? commitId, CancellationToken cancellationToken)
 		{
 			using TelemetrySpan span = _tracer.StartActiveSpan($"{nameof(JobService)}.{nameof(GetStepTimingInfoAsync)}");
 			span.SetAttribute("StreamId", streamId);
 			span.SetAttribute("TemplateId", templateId);
 			span.SetAttribute("NodeName", nodeName);
-			span.SetAttribute("Change", change);
+			span.SetAttribute("Change", commitId?.ToString());
 
 			// Find all the steps matching the given criteria
-			List<IJobStepRef> steps = await _jobStepRefs.GetStepsForNodeAsync(streamId, templateId, nodeName, change, false, 10, cancellationToken);
+			List<IJobStepRef> steps = await _jobStepRefs.GetStepsForNodeAsync(streamId, templateId, nodeName, commitId, false, 10, cancellationToken);
 
 			// Sum up all the durations and wait times
 			int count = 0;
@@ -922,9 +930,9 @@ namespace HordeServer.Jobs
 		{
 			try
 			{
-				if (job.PreflightChange != 0)
+				if (job.PreflightCommitId != null)
 				{
-					(CheckShelfResult result, _) = await _perforceService.CheckShelfAsync(streamConfig, job.PreflightChange, cancellationToken);
+					(CheckShelfResult result, _) = await _perforceService.CheckShelfAsync(streamConfig, job.PreflightCommitId.GetPerforceChange(), cancellationToken);
 					if (result != CheckShelfResult.Ok)
 					{
 						_logger.LogWarning("Job {JobId} is no longer valid - check shelf returned {Result}", job.Id, result);
@@ -1080,7 +1088,7 @@ namespace HordeServer.Jobs
 					}
 
 					// Submit the change if auto-submit is enabled
-					if (job.PreflightChange != 0 && newState == JobStepState.Completed)
+					if (job.PreflightCommitId != null && newState == JobStepState.Completed)
 					{
 						(JobStepState state, JobStepOutcome outcome) = job.GetTargetState();
 						if (state == JobStepState.Completed)
@@ -1089,18 +1097,18 @@ namespace HordeServer.Jobs
 							{
 								job = await AutoSubmitChangeAsync(streamConfig, job, graph, cancellationToken);
 							}
-							else if (job.ClonedPreflightChange == 0 && job.StartedByUserId.HasValue && outcome == JobStepOutcome.Success && job.AbortedByUserId == null)
+							else if (job.ClonedPreflightCommitId == null && job.StartedByUserId.HasValue && outcome == JobStepOutcome.Success && job.AbortedByUserId == null)
 							{
 								IUserSettings settings = await _userCollection.GetSettingsAsync(job.StartedByUserId.Value, cancellationToken);
 								if (settings.AlwaysTagPreflightCL)
 								{
-									_logger.LogInformation("Updating description for {PreflightChange} for {UserId} user settings", job.PreflightChange, job.StartedByUserId.Value);
-									await _perforceService.UpdateChangelistDescriptionAsync(streamConfig.ClusterName, job.PreflightChange, x => x.TrimEnd() + $"\n#preflight {job.Id}", cancellationToken);
+									_logger.LogInformation("Updating description for {PreflightChange} for {UserId} user settings", job.PreflightCommitId, job.StartedByUserId.Value);
+									await _perforceService.UpdateChangelistDescriptionAsync(streamConfig.ClusterName, job.PreflightCommitId.GetPerforceChange(), x => x.TrimEnd() + $"\n#preflight {job.Id}", cancellationToken);
 								}
 							}
-							else if (job.ClonedPreflightChange != 0)
+							else if (job.ClonedPreflightCommitId != null)
 							{
-								await DeleteShelvedChangeAsync(streamConfig.ClusterName, job.ClonedPreflightChange);
+								await DeleteShelvedChangeAsync(streamConfig.ClusterName, job.ClonedPreflightCommitId);
 							}
 						}
 					}
@@ -1186,6 +1194,11 @@ namespace HordeServer.Jobs
 		/// <returns></returns>
 		private async Task<IJob> AutoSubmitChangeAsync(StreamConfig streamConfig, IJob job, IGraph graph, CancellationToken cancellationToken)
 		{
+			if (job.PreflightCommitId == null)
+			{
+				return job;
+			}
+
 			using TelemetrySpan span = _tracer.StartActiveSpan($"{nameof(JobService)}.{nameof(AutoSubmitChangeAsync)}");
 			span.SetAttribute("Job", job.Id.ToString());
 			span.SetAttribute("Graph", graph.Id);
@@ -1194,27 +1207,27 @@ namespace HordeServer.Jobs
 			string message;
 			try
 			{
-				int clonedPreflightChange = job.ClonedPreflightChange;
-				if (clonedPreflightChange == 0)
+				CommitId? clonedPreflightChange = job.ClonedPreflightCommitId;
+				if (clonedPreflightChange == null)
 				{
 					if (ShouldClonePreflightChange(job.StreamId))
 					{
-						clonedPreflightChange = await CloneShelvedChangeAsync(streamConfig.ClusterName, job.PreflightChange, cancellationToken);
+						clonedPreflightChange = await CloneShelvedChangeAsync(streamConfig.ClusterName, job.PreflightCommitId, cancellationToken);
 					}
 					else
 					{
-						clonedPreflightChange = job.PreflightChange;
+						clonedPreflightChange = job.PreflightCommitId;
 					}
 				}
 
 				_logger.LogInformation("Updating description for {ClonedPreflightChange}", clonedPreflightChange);
 
-				await _perforceService.UpdateChangelistDescriptionAsync(streamConfig.ClusterName, clonedPreflightChange, x => x.TrimEnd() + $"\n#preflight {job.Id}", cancellationToken);
+				await _perforceService.UpdateChangelistDescriptionAsync(streamConfig.ClusterName, clonedPreflightChange.GetPerforceChange(), x => x.TrimEnd() + $"\n#preflight {job.Id}", cancellationToken);
 
-				_logger.LogInformation("Submitting change {Change} (through {ChangeCopy}) after successful completion of {JobId}", job.PreflightChange, clonedPreflightChange, job.Id);
-				(change, message) = await _perforceService.SubmitShelvedChangeAsync(streamConfig, clonedPreflightChange, job.PreflightChange, cancellationToken);
+				_logger.LogInformation("Submitting change {Change} (through {ChangeCopy}) after successful completion of {JobId}", job.PreflightCommitId, clonedPreflightChange, job.Id);
+				(change, message) = await _perforceService.SubmitShelvedChangeAsync(streamConfig, clonedPreflightChange.GetPerforceChange(), job.PreflightCommitId.GetPerforceChange(), cancellationToken);
 
-				_logger.LogInformation("Attempt to submit {Change} (through {ChangeCopy}): {Message}", job.PreflightChange, clonedPreflightChange, message);
+				_logger.LogInformation("Attempt to submit {Change} (through {ChangeCopy}): {Message}", job.PreflightCommitId, clonedPreflightChange, message);
 
 				if (!String.IsNullOrEmpty(message))
 				{
@@ -1223,16 +1236,16 @@ namespace HordeServer.Jobs
 
 				if (ShouldClonePreflightChange(job.StreamId))
 				{
-					if (change != null && job.ClonedPreflightChange != 0)
+					if (change != null && job.ClonedPreflightCommitId != null)
 					{
-						await DeleteShelvedChangeAsync(streamConfig.ClusterName, job.PreflightChange);
+						await DeleteShelvedChangeAsync(streamConfig.ClusterName, job.PreflightCommitId);
 					}
 				}
 				else
 				{
-					if (change != null && job.PreflightChange != 0)
+					if (change != null)
 					{
-						await DeleteShelvedChangeAsync(streamConfig.ClusterName, job.PreflightChange);
+						await DeleteShelvedChangeAsync(streamConfig.ClusterName, job.PreflightCommitId);
 					}
 				}
 			}
@@ -1267,12 +1280,12 @@ namespace HordeServer.Jobs
 		/// <param name="change">The changelist to clone</param>
 		/// <param name="cancellationToken">Cancellation token for the operation</param>
 		/// <returns></returns>
-		private async Task<int> CloneShelvedChangeAsync(string clusterName, int change, CancellationToken cancellationToken = default)
+		private async Task<CommitId> CloneShelvedChangeAsync(string clusterName, CommitId change, CancellationToken cancellationToken = default)
 		{
 			int clonedChange;
 			try
 			{
-				clonedChange = await _perforceService.DuplicateShelvedChangeAsync(clusterName, change, cancellationToken);
+				clonedChange = await _perforceService.DuplicateShelvedChangeAsync(clusterName, change.GetPerforceChange(), cancellationToken);
 				_logger.LogInformation("CL {Change} was duplicated into {ClonedChange}", change, clonedChange);
 			}
 			catch (Exception ex)
@@ -1280,7 +1293,7 @@ namespace HordeServer.Jobs
 				_logger.LogError(ex, "Unable to CL {Change} for preflight: {Message}", change, ex.Message);
 				throw;
 			}
-			return clonedChange;
+			return CommitId.FromPerforceChange(clonedChange);
 		}
 
 		/// <summary>
@@ -1289,7 +1302,7 @@ namespace HordeServer.Jobs
 		/// <param name="clusterName"></param>
 		/// <param name="change">The changelist to delete</param>
 		/// <returns>True if the change was deleted successfully, false otherwise</returns>
-		private Task<bool> DeleteShelvedChangeAsync(string clusterName, int change)
+		private Task<bool> DeleteShelvedChangeAsync(string clusterName, CommitId change)
 		{
 			_ = clusterName;
 			_logger.LogInformation("Leaving shelved change {Change}", change);
@@ -1346,39 +1359,39 @@ namespace HordeServer.Jobs
 
 					IGraph triggerGraph = await _graphs.AddAsync(template, streamConfig.InitialAgentType);
 
-					int? change = null;
-					int codeChange;
+					CommitIdWithOrder? commitId = null;
+					CommitIdWithOrder? codeCommitId;
 					if (jobTrigger.UseDefaultChangeForTemplate)
 					{
 						ICommitCollection commits = _perforceService.GetCommits(streamConfig);
 
 						if (templateRefConfig.DefaultChange != null)
 						{
-							change = await EvaluateChangeQueriesAsync(streamConfig.Id, templateRefConfig.DefaultChange, null, commits, cancellationToken);
+							commitId = await EvaluateChangeQueriesAsync(streamConfig.Id, templateRefConfig.DefaultChange, null, commits, cancellationToken);
 						}
-						if (change == null)
+						if (commitId == null)
 						{
-							change = await commits.GetLatestNumberAsync(cancellationToken);
+							commitId = await commits.GetLastCommitIdAsync(cancellationToken);
 						}
 
-						ICommit? commit = await commits.GetLastCodeChangeAsync(change, cancellationToken);
-						codeChange = commit?.Number ?? change.Value;
+						ICommit? codeCommit = await commits.GetLastCodeChangeAsync(commitId, cancellationToken);
+						codeCommitId = codeCommit?.Id;
 					}
 					else
 					{
-						change = job.Change;
-						codeChange = job.CodeChange;
+						commitId = job.CommitId;
+						codeCommitId = job.CodeCommitId;
 					}
 
-					_logger.LogInformation("Creating downstream job {ChainedJobId} from job {JobId} at change {Change}", chainedJobId, newJob.Id, change);
+					_logger.LogInformation("Creating downstream job {ChainedJobId} from job {JobId} at change {Change}", chainedJobId, newJob.Id, commitId);
 
 					CreateJobOptions options = new CreateJobOptions(templateRefConfig);
-					options.PreflightChange = newJob.PreflightChange;
+					options.PreflightCommitId = newJob.PreflightCommitId;
 					options.PreflightDescription = newJob.PreflightDescription;
 					template.GetDefaultParameters(options.Parameters, true);
 					template.GetArgumentsForParameters(options.Parameters, options.Arguments);
 
-					await CreateJobAsync(chainedJobId, streamConfig, jobTrigger.TemplateRefId, template.Hash, triggerGraph, templateRefConfig.Name, change.Value, codeChange, options);
+					await CreateJobAsync(chainedJobId, streamConfig, jobTrigger.TemplateRefId, template.Hash, triggerGraph, templateRefConfig.Name, commitId, codeCommitId, options);
 					return newJob;
 				}
 
diff --git a/Engine/Source/Programs/Horde/Plugins/Build/HordeServer.Build/Jobs/JobStepRefCollection.cs b/Engine/Source/Programs/Horde/Plugins/Build/HordeServer.Build/Jobs/JobStepRefCollection.cs
--- a/Engine/Source/Programs/Horde/Plugins/Build/HordeServer.Build/Jobs/JobStepRefCollection.cs
+++ b/Engine/Source/Programs/Horde/Plugins/Build/HordeServer.Build/Jobs/JobStepRefCollection.cs
@@ -2,6 +2,7 @@
 
 using EpicGames.Horde.Agents;
 using EpicGames.Horde.Agents.Pools;
+using EpicGames.Horde.Commits;
 using EpicGames.Horde.Jobs;
 using EpicGames.Horde.Jobs.Templates;
 using EpicGames.Horde.Logs;
@@ -32,18 +33,51 @@ namespace HordeServer.Jobs
 
 			public StreamId StreamId { get; set; }
 			public TemplateId TemplateId { get; set; }
-			public int Change { get; set; }
+
+			[BsonIgnore]
+			public CommitIdWithOrder CommitId
+			{
+				get => (CommitName != null) ? new CommitIdWithOrder(CommitName, CommitOrder) : CommitIdWithOrder.FromPerforceChange(CommitOrder);
+				set => (CommitName, CommitOrder) = (value.Name, value.Order);
+			}
+
+			[BsonElement("Commit")]
+			public string? CommitName { get; set; }
+
+			[BsonElement("Change")]
+			public int CommitOrder { get; set; }
+
 			public LogId? LogId { get; set; }
 			public PoolId? PoolId { get; set; }
 			public AgentId? AgentId { get; set; }
 			public JobStepState? State { get; set; }
 			public JobStepOutcome? Outcome { get; set; }
 
+			[BsonIgnore]
+			public CommitIdWithOrder? LastSuccess
+			{
+				get => (LastSuccessName != null) ? new CommitIdWithOrder(LastSuccessName, LastSuccessOrder ?? 0) : (LastSuccessOrder != null) ? CommitIdWithOrder.FromPerforceChange(LastSuccessOrder.Value) : null;
+				set => (LastSuccessName, LastSuccessOrder) = (value?.Name, value?.Order);
+			}
+
 			[BsonIgnoreIfNull]
-			public int? LastSuccess { get; }
+			public string? LastSuccessName { get; set; }
+
+			[BsonElement("LastSuccess"), BsonIgnoreIfNull]
+			public int? LastSuccessOrder { get; set; }
+
+			[BsonIgnore]
+			public CommitIdWithOrder? LastWarning
+			{
+				get => (LastWarningName != null) ? new CommitIdWithOrder(LastWarningName, LastWarningOrder ?? 0) : (LastWarningOrder != null) ? CommitIdWithOrder.FromPerforceChange(LastWarningOrder.Value) : null;
+				set => (LastWarningName, LastWarningOrder) = (value?.Name, value?.Order);
+			}
 
 			[BsonIgnoreIfNull]
-			public int? LastWarning { get; }
+			public string? LastWarningName { get; set; }
+
+			[BsonElement("LastWarning"), BsonIgnoreIfNull]
+			public int? LastWarningOrder { get; set; }
 
 			public float BatchWaitTime { get; set; }
 			public float BatchInitTime { get; set; }
@@ -74,14 +108,14 @@ namespace HordeServer.Jobs
 			bool IJobStepRef.UpdateIssues => UpdateIssues ?? false;
 			IReadOnlyList<int>? IJobStepRef.IssueIds => IssueIds;
 
-			public JobStepRef(JobStepRefId id, string jobName, string nodeName, StreamId streamId, TemplateId templateId, int change, LogId? logId, PoolId? poolId, AgentId? agentId, JobStepState? state, JobStepOutcome? outcome, bool updateIssues, int? lastSuccess, int? lastWarning, float batchWaitTime, float batchInitTime, DateTime jobStartTimeUtc, DateTime startTimeUtc, DateTime? finishTimeUtc)
+			public JobStepRef(JobStepRefId id, string jobName, string nodeName, StreamId streamId, TemplateId templateId, CommitIdWithOrder commitId, LogId? logId, PoolId? poolId, AgentId? agentId, JobStepState? state, JobStepOutcome? outcome, bool updateIssues, CommitIdWithOrder? lastSuccess, CommitIdWithOrder? lastWarning, float batchWaitTime, float batchInitTime, DateTime jobStartTimeUtc, DateTime startTimeUtc, DateTime? finishTimeUtc)
 			{
 				Id = id;
 				JobName = jobName;
 				Name = nodeName;
 				StreamId = streamId;
 				TemplateId = templateId;
-				Change = change;
+				CommitId = commitId;
 				LogId = logId;
 				PoolId = poolId;
 				AgentId = agentId;
@@ -108,7 +142,7 @@ namespace HordeServer.Jobs
 		public JobStepRefCollection(IMongoService mongoService, ITelemetryWriter telemetryWriter, IOptionsMonitor<BuildConfig> buildConfig)
 		{
 			List<MongoIndex<JobStepRef>> indexes = new List<MongoIndex<JobStepRef>>();
-			indexes.Add(keys => keys.Ascending(x => x.StreamId).Ascending(x => x.TemplateId).Ascending(x => x.Name).Descending(x => x.Change));
+			indexes.Add(keys => keys.Ascending(x => x.StreamId).Ascending(x => x.TemplateId).Ascending(x => x.Name).Descending(x => x.CommitOrder));
 
 			_jobStepRefs = mongoService.GetCollection<JobStepRef>("JobStepRefs", indexes);
 			_telemetryWriter = telemetryWriter;
@@ -116,9 +150,9 @@ namespace HordeServer.Jobs
 		}
 
 		/// <inheritdoc/>
-		public async Task<IJobStepRef> InsertOrReplaceAsync(JobStepRefId id, string jobName, string stepName, StreamId streamId, TemplateId templateId, int change, LogId? logId, PoolId? poolId, AgentId? agentId, JobStepState? state, JobStepOutcome? outcome, bool updateIssues, int? lastSuccess, int? lastWarning, float waitTime, float initTime, DateTime jobStartTimeUtc, DateTime startTimeUtc, DateTime? finishTimeUtc)
+		public async Task<IJobStepRef> InsertOrReplaceAsync(JobStepRefId id, string jobName, string stepName, StreamId streamId, TemplateId templateId, CommitIdWithOrder commitId, LogId? logId, PoolId? poolId, AgentId? agentId, JobStepState? state, JobStepOutcome? outcome, bool updateIssues, CommitIdWithOrder? lastSuccess, CommitIdWithOrder? lastWarning, float waitTime, float initTime, DateTime jobStartTimeUtc, DateTime startTimeUtc, DateTime? finishTimeUtc)
 		{
-			JobStepRef newJobStepRef = new JobStepRef(id, jobName, stepName, streamId, templateId, change, logId, poolId, agentId, state, outcome, updateIssues, lastSuccess, lastWarning, waitTime, initTime, jobStartTimeUtc, startTimeUtc, finishTimeUtc);
+			JobStepRef newJobStepRef = new JobStepRef(id, jobName, stepName, streamId, templateId, commitId, logId, poolId, agentId, state, outcome, updateIssues, lastSuccess, lastWarning, waitTime, initTime, jobStartTimeUtc, startTimeUtc, finishTimeUtc);
 			await _jobStepRefs.ReplaceOneAsync(Builders<JobStepRef>.Filter.Eq(x => x.Id, newJobStepRef.Id), newJobStepRef, new ReplaceOptions { IsUpsert = true });
 
 			if (_buildConfig.CurrentValue.TryGetStream(streamId, out StreamConfig? streamConfig) && !streamConfig.TelemetryStoreId.IsEmpty)
@@ -133,7 +167,7 @@ namespace HordeServer.Jobs
 					AgentId = agentId,
 					BatchInitTime = initTime,
 					BatchWaitTime = waitTime,
-					Change = change,
+					Change = commitId,
 					FinishTime = finishTimeUtc,
 					JobName = jobName,
 					JobStartTime = jobStartTimeUtc,
@@ -186,12 +220,12 @@ namespace HordeServer.Jobs
 			FilterDefinition<JobStepRef> filter = FilterDefinition<JobStepRef>.Empty;
 			filter &= filterBuilder.In(x => x.Id, ids);
 
-			List<JobStepRef> steps = await _jobStepRefs.Find(filter).SortByDescending(x => x.Change).ToListAsync(cancellationToken);
+			List<JobStepRef> steps = await _jobStepRefs.Find(filter).SortByDescending(x => x.CommitOrder).ToListAsync(cancellationToken);
 			return steps.ConvertAll<IJobStepRef>(x => x);
 		}
 
 		/// <inheritdoc/>
-		public async Task<List<IJobStepRef>> GetStepsForNodeAsync(StreamId streamId, TemplateId templateId, string nodeName, int? change, bool includeFailed, int maxCount, CancellationToken cancellationToken)
+		public async Task<List<IJobStepRef>> GetStepsForNodeAsync(StreamId streamId, TemplateId templateId, string nodeName, CommitIdWithOrder? commitId, bool includeFailed, int maxCount, CancellationToken cancellationToken)
 		{
 			// Find all the steps matching the given criteria
 			FilterDefinitionBuilder<JobStepRef> filterBuilder = Builders<JobStepRef>.Filter;
@@ -200,21 +234,21 @@ namespace HordeServer.Jobs
 			filter &= filterBuilder.Eq(x => x.StreamId, streamId);
 			filter &= filterBuilder.Eq(x => x.TemplateId, templateId);
 			filter &= filterBuilder.Eq(x => x.Name, nodeName);
-			if (change != null)
+			if (commitId != null)
 			{
-				filter &= filterBuilder.Lte(x => x.Change, change.Value);
+				filter &= filterBuilder.Lte(x => x.CommitOrder, commitId.Order);
 			}
 			if (!includeFailed)
 			{
 				filter &= filterBuilder.Ne(x => x.Outcome, JobStepOutcome.Failure);
 			}
 
-			List<JobStepRef> steps = await _jobStepRefs.Find(filter).SortByDescending(x => x.Change).ThenByDescending(x => x.StartTimeUtc).Limit(maxCount).ToListAsync(cancellationToken);
+			List<JobStepRef> steps = await _jobStepRefs.Find(filter).SortByDescending(x => x.CommitOrder).ThenByDescending(x => x.StartTimeUtc).Limit(maxCount).ToListAsync(cancellationToken);
 			return steps.ConvertAll<IJobStepRef>(x => x);
 		}
 
 		/// <inheritdoc/>
-		public async Task<IJobStepRef?> GetPrevStepForNodeAsync(StreamId streamId, TemplateId templateId, string nodeName, int change, JobStepOutcome? outcome = null, bool? updateIssues = null, IEnumerable<JobId>? excludeJobIds = null)
+		public async Task<IJobStepRef?> GetPrevStepForNodeAsync(StreamId streamId, TemplateId templateId, string nodeName, CommitIdWithOrder commitId, JobStepOutcome? outcome = null, bool? updateIssues = null, IEnumerable<JobId>? excludeJobIds = null)
 		{
 			FilterDefinitionBuilder<JobStepRef> filterBuilder = Builders<JobStepRef>.Filter;
 
@@ -222,7 +256,7 @@ namespace HordeServer.Jobs
 			filter &= filterBuilder.Eq(x => x.StreamId, streamId);
 			filter &= filterBuilder.Eq(x => x.TemplateId, templateId);
 			filter &= filterBuilder.Eq(x => x.Name, nodeName);
-			filter &= filterBuilder.Lt(x => x.Change, change);
+			filter &= filterBuilder.Lt(x => x.CommitOrder, commitId.Order);
 
 			if (outcome != null)
 			{
@@ -243,11 +277,11 @@ namespace HordeServer.Jobs
 				filter &= filterBuilder.Nin(x => x.Id.JobId, excludeJobIds);
 			}
 
-			return await _jobStepRefs.Find(filter).SortByDescending(x => x.Change).FirstOrDefaultAsync();
+			return await _jobStepRefs.Find(filter).SortByDescending(x => x.CommitOrder).FirstOrDefaultAsync();
 		}
 
 		/// <inheritdoc/>
-		public async Task<IJobStepRef?> GetNextStepForNodeAsync(StreamId streamId, TemplateId templateId, string nodeName, int change, JobStepOutcome? outcome = null, bool? updateIssues = null)
+		public async Task<IJobStepRef?> GetNextStepForNodeAsync(StreamId streamId, TemplateId templateId, string nodeName, CommitIdWithOrder commitId, JobStepOutcome? outcome = null, bool? updateIssues = null)
 		{
 			FilterDefinitionBuilder<JobStepRef> filterBuilder = Builders<JobStepRef>.Filter;
 
@@ -255,7 +289,7 @@ namespace HordeServer.Jobs
 			filter &= filterBuilder.Eq(x => x.StreamId, streamId);
 			filter &= filterBuilder.Eq(x => x.TemplateId, templateId);
 			filter &= filterBuilder.Eq(x => x.Name, nodeName);
-			filter &= filterBuilder.Gt(x => x.Change, change);
+			filter &= filterBuilder.Gt(x => x.CommitOrder, commitId.Order);
 
 			if (outcome != null)
 			{
@@ -271,7 +305,7 @@ namespace HordeServer.Jobs
 				filter &= filterBuilder.Ne(x => x.UpdateIssues, false);
 			}
 
-			return await _jobStepRefs.Find(filter).SortBy(x => x.Change).FirstOrDefaultAsync();
+			return await _jobStepRefs.Find(filter).SortBy(x => x.CommitOrder).FirstOrDefaultAsync();
 		}
 	}
 }
diff --git a/Engine/Source/Programs/Horde/Plugins/Build/HordeServer.Build/Jobs/JobTaskSource.cs b/Engine/Source/Programs/Horde/Plugins/Build/HordeServer.Build/Jobs/JobTaskSource.cs
--- a/Engine/Source/Programs/Horde/Plugins/Build/HordeServer.Build/Jobs/JobTaskSource.cs
+++ b/Engine/Source/Programs/Horde/Plugins/Build/HordeServer.Build/Jobs/JobTaskSource.cs
@@ -798,13 +798,13 @@ namespace HordeServer.Jobs
 
 				// Get the lease name
 				StringBuilder leaseName = new StringBuilder($"{streamConfig.Name} - ");
-				if (job.PreflightChange > 0)
+				if (job.PreflightCommitId != null)
 				{
-					leaseName.Append((job.Change > 0) ? $"Preflight CL {job.PreflightChange} against CL {job.Change}" : $"Preflight CL {job.PreflightChange} against latest");
+					leaseName.Append($"Preflight CL {job.PreflightCommitId} against CL {job.CommitId}");
 				}
 				else
 				{
-					leaseName.Append((job.Change > 0) ? $"CL {job.Change}" : "Latest CL");
+					leaseName.Append($"CL {job.CommitId}");
 				}
 				leaseName.Append(CultureInfo.InvariantCulture, $" - {job.Name}");
 
@@ -875,13 +875,13 @@ namespace HordeServer.Jobs
 		{
 			// Get the lease name
 			StringBuilder leaseName = new StringBuilder($"{streamConfig.Name} - ");
-			if (job.PreflightChange > 0)
+			if (job.PreflightCommitId != null)
 			{
-				leaseName.Append((job.Change > 0) ? $"Preflight CL {job.PreflightChange} against CL {job.Change}" : $"Preflight CL {job.PreflightChange} against latest");
+				leaseName.Append($"Preflight CL {job.PreflightCommitId} against CL {job.CommitId}");
 			}
 			else
 			{
-				leaseName.Append((job.Change > 0) ? $"CL {job.Change}" : "Latest CL");
+				leaseName.Append($"CL {job.CommitId}");
 			}
 			leaseName.Append(CultureInfo.InvariantCulture, $" - {job.Name}");
 
@@ -894,7 +894,7 @@ namespace HordeServer.Jobs
 			claims.Add(HordeClaims.AgentRoleClaim);
 			claims.Add(new AclClaimConfig(HordeClaimTypes.Lease, leaseId.ToString()));
 
-			string storagePrefix = $"{job.StreamId}/{job.Change}-{job.Id}";
+			string storagePrefix = $"{job.StreamId}/{job.CommitId}-{job.Id}";
 			claims.Add(new AclClaimConfig(HordeClaimTypes.ReadNamespace, $"{namespaceId}:{storagePrefix}"));
 			claims.Add(new AclClaimConfig(HordeClaimTypes.WriteNamespace, $"{namespaceId}:{storagePrefix}"));
 
@@ -990,7 +990,7 @@ namespace HordeServer.Jobs
 		/// <returns>Async task</returns>
 		public async Task UpdateUgsBadgesAsync(IJob job, IGraph graph, IReadOnlyList<(LabelState, LabelOutcome)> oldLabelStates, IReadOnlyList<(LabelState, LabelOutcome)> newLabelStates, CancellationToken cancellationToken)
 		{
-			if (!job.ShowUgsBadges || job.PreflightChange != 0)
+			if (!job.ShowUgsBadges || job.PreflightCommitId != null)
 			{
 				return;
 			}
@@ -1050,11 +1050,11 @@ namespace HordeServer.Jobs
 				int change;
 				if (label.Change == RpcLabelChange.Code)
 				{
-					change = job.CodeChange;
+					change = job.CodeCommitId?.GetPerforceChange() ?? 0;
 				}
 				else
 				{
-					change = job.Change;
+					change = job.CommitId.GetPerforceChange();
 				}
 
 				// Get the current metadata state
