Horde: Preliminary support for VCS commit ids that aren't an integer.

diff --git a/Engine/Source/Programs/Horde/Plugins/Build/HordeServer.Build/Jobs/JobsController.cs b/Engine/Source/Programs/Horde/Plugins/Build/HordeServer.Build/Jobs/JobsController.cs
--- a/Engine/Source/Programs/Horde/Plugins/Build/HordeServer.Build/Jobs/JobsController.cs
+++ b/Engine/Source/Programs/Horde/Plugins/Build/HordeServer.Build/Jobs/JobsController.cs
@@ -2,6 +2,7 @@
 
 using EpicGames.Core;
 using EpicGames.Horde.Artifacts;
+using EpicGames.Horde.Commits;
 using EpicGames.Horde.Jobs;
 using EpicGames.Horde.Jobs.Graphs;
 using EpicGames.Horde.Jobs.Templates;
@@ -106,7 +107,7 @@ namespace HordeServer.Jobs
 			{
 				return BadRequest("Missing template referenced by {TemplateId}", create.TemplateId);
 			}
-			if (!template.AllowPreflights && create.PreflightChange > 0)
+			if (!template.AllowPreflights && create.PreflightCommitId != null)
 			{
 				return BadRequest("Template {TemplateId} does not allow preflights", create.TemplateId);
 			}
@@ -131,23 +132,25 @@ namespace HordeServer.Jobs
 
 			// Check the preflight change is valid
 			ShelfInfo? shelfInfo = null;
-			if (create.PreflightChange != null)
+			if (create.PreflightCommitId != null)
 			{
-				(CheckShelfResult result, shelfInfo) = await _perforce.CheckShelfAsync(streamConfig, create.PreflightChange.Value, HttpContext.RequestAborted);
+				int preflightChange = create.PreflightCommitId.GetPerforceChange();
+
+				(CheckShelfResult result, shelfInfo) = await _perforce.CheckShelfAsync(streamConfig, preflightChange, HttpContext.RequestAborted);
 				switch (result)
 				{
 					case CheckShelfResult.Ok:
 						break;
 					case CheckShelfResult.NoChange:
-						return BadRequest(KnownLogEvents.Horde_InvalidPreflight, "CL {Change} does not exist", create.PreflightChange);
+						return BadRequest(KnownLogEvents.Horde_InvalidPreflight, "CL {Change} does not exist", preflightChange);
 					case CheckShelfResult.NoShelvedFiles:
-						return BadRequest(KnownLogEvents.Horde_InvalidPreflight, "CL {Change} does not contain any shelved files", create.PreflightChange);
+						return BadRequest(KnownLogEvents.Horde_InvalidPreflight, "CL {Change} does not contain any shelved files", preflightChange);
 					case CheckShelfResult.WrongStream:
-						return BadRequest(KnownLogEvents.Horde_InvalidPreflight, "CL {Change} does not contain files in {Stream}", create.PreflightChange, streamConfig.Name);
+						return BadRequest(KnownLogEvents.Horde_InvalidPreflight, "CL {Change} does not contain files in {Stream}", preflightChange, streamConfig.Name);
 					case CheckShelfResult.MixedStream:
-						return BadRequest(KnownLogEvents.Horde_InvalidPreflight, "CL {Change} contains files from multiple streams", create.PreflightChange);
+						return BadRequest(KnownLogEvents.Horde_InvalidPreflight, "CL {Change} contains files from multiple streams", preflightChange);
 					default:
-						return BadRequest(KnownLogEvents.Horde_InvalidPreflight, "CL {Change} cannot be preflighted ({Result})", create.PreflightChange, result);
+						return BadRequest(KnownLogEvents.Horde_InvalidPreflight, "CL {Change} cannot be preflighted ({Result})", preflightChange, result);
 				}
 
 				if (shelfInfo!.Tags != null)
@@ -167,11 +170,11 @@ namespace HordeServer.Jobs
 			ICommitCollection commits = _commitService.GetCollection(streamConfig);
 
 			// Get the change to build
-			int change = await GetChangeToBuildAsync(create, streamConfig.Id, template, shelfInfo, commits, HttpContext.RequestAborted);
+			CommitIdWithOrder commitId = await GetChangeToBuildAsync(create, streamConfig.Id, template, shelfInfo, commits, HttpContext.RequestAborted);
 
 			// And get the matching code changelist
-			ICommit? lastCodeCommit = await commits.GetLastCodeChangeAsync(change, HttpContext.RequestAborted);
-			int codeChange = lastCodeCommit?.Number ?? change;
+			ICommit? lastCodeCommit = await commits.GetLastCodeChangeAsync(commitId, HttpContext.RequestAborted);
+			CommitIdWithOrder? codeCommitId = lastCodeCommit?.Id;
 
 			// New properties for the job
 			bool? updateIssues = null;
@@ -186,7 +189,7 @@ namespace HordeServer.Jobs
 
 			// Create options for the new job
 			CreateJobOptions options = new CreateJobOptions(templateRefConfig);
-			options.PreflightChange = create.PreflightChange;
+			options.PreflightCommitId = create.PreflightCommitId;
 			options.PreflightDescription = shelfInfo?.Description;
 			options.StartedByUserId = User.GetUserId();
 			options.Priority = priority;
@@ -247,37 +250,37 @@ namespace HordeServer.Jobs
 			}
 
 			// Create the job
-			IJob job = await _jobService.CreateJobAsync(null, streamConfig, templateRefId, template.Hash, graph, name, change, codeChange, options, cancellationToken);
+			IJob job = await _jobService.CreateJobAsync(null, streamConfig, templateRefId, template.Hash, graph, name, commitId, codeCommitId, options, cancellationToken);
 			await UpdateNotificationsAsync(job.Id, new UpdateNotificationsRequest { Slack = true }, cancellationToken);
 			return new CreateJobResponse(job.Id.ToString());
 		}
 
-		async ValueTask<int> GetChangeToBuildAsync(CreateJobRequest create, StreamId streamId, ITemplate template, ShelfInfo? shelfInfo, ICommitCollection commits, CancellationToken cancellationToken)
+		async ValueTask<CommitIdWithOrder> GetChangeToBuildAsync(CreateJobRequest create, StreamId streamId, ITemplate template, ShelfInfo? shelfInfo, ICommitCollection commits, CancellationToken cancellationToken)
 		{
 			// If there's an explicit change specified, use that
-			if (create.Change.HasValue && create.Change.Value != -1)
+			if (create.CommitId != null)
 			{
-				return create.Change.Value;
+				return await _commitService.GetOrderedAsync(streamId, create.CommitId, cancellationToken);
 			}
 
 			// Evaluate the change queries
 			if (create.ChangeQueries != null && create.ChangeQueries.Count > 0)
 			{
-				int? change = await _jobService.EvaluateChangeQueriesAsync(streamId, create.ChangeQueries, shelfInfo?.Tags, commits, cancellationToken);
-				if (change != null)
+				CommitIdWithOrder? commitId = await _jobService.EvaluateChangeQueriesAsync(streamId, create.ChangeQueries, shelfInfo?.Tags, commits, cancellationToken);
+				if (commitId != null)
 				{
-					return change.Value;
+					return commitId;
 				}
 			}
 
 			// If we need to submit a new change, do that
-			if (create.PreflightChange == null && template.SubmitNewChange != null)
+			if (create.PreflightCommitId == null && template.SubmitNewChange != null)
 			{
 				return await commits.CreateNewAsync(template, cancellationToken);
 			}
 
 			// Otherwise return the latest change
-			return await commits.GetLatestNumberAsync(cancellationToken);
+			return await commits.GetLastCommitIdAsync(cancellationToken);
 		}
 
 		/// <summary>
@@ -589,10 +592,10 @@ namespace HordeServer.Jobs
 		static GetJobResponse CreateGetJobResponse(IJob job, GetThinUserInfoResponse? startedByUserInfo, GetThinUserInfoResponse? abortedByUserInfo)
 		{
 			GetJobResponse response = new GetJobResponse(job.Id, job.StreamId, job.TemplateId, job.Name);
-			response.Change = job.Change;
-			response.CodeChange = (job.CodeChange != 0) ? (int?)job.CodeChange : null;
-			response.PreflightChange = (job.PreflightChange != 0) ? (int?)job.PreflightChange : null;
-			response.ClonedPreflightChange = (job.ClonedPreflightChange != 0) ? (int?)job.ClonedPreflightChange : null;
+			response.CommitId = job.CommitId;
+			response.CodeCommitId = job.CodeCommitId;
+			response.PreflightCommitId = job.PreflightCommitId;
+			response.ClonedPreflightCommitId = job.ClonedPreflightCommitId;
 			response.PreflightDescription = job.PreflightDescription;
 			response.TemplateHash = job.TemplateHash?.ToString() ?? String.Empty;
 			response.GraphHash = job.GraphHash.ToString();
@@ -1054,11 +1057,11 @@ namespace HordeServer.Jobs
 			[FromQuery] string? streamId = null,
 			[FromQuery] string? name = null,
 			[FromQuery(Name = "template")] string[]? templates = null,
-			[FromQuery] int? minChange = null,
-			[FromQuery] int? maxChange = null,
+			[FromQuery] CommitId? minChange = null,
+			[FromQuery] CommitId? maxChange = null,
 			[FromQuery] bool includePreflight = true,
 			[FromQuery] bool? preflightOnly = null,
-			[FromQuery] int? preflightChange = null,
+			[FromQuery] CommitId? preflightChange = null,
 			[FromQuery] string? preflightStartedByUserId = null,
 			[FromQuery] string? startedByUserId = null,
 			[FromQuery] DateTimeOffset? minCreateTime = null,
@@ -1077,11 +1080,6 @@ namespace HordeServer.Jobs
 
 			TemplateId[]? templateRefIds = (templates != null && templates.Length > 0) ? templates.Select(x => new TemplateId(x)).ToArray() : null;
 
-			if (includePreflight == false)
-			{
-				preflightChange = 0;
-			}
-
 			UserId? preflightStartedByUserIdValue = null;
 
 			if (preflightStartedByUserId != null)
@@ -1098,7 +1096,7 @@ namespace HordeServer.Jobs
 
 			IReadOnlyList<IJob> jobs;
 			jobs = await _jobService.FindJobsAsync(jobIdValues, streamIdValue, name, templateRefIds, minChange,
-				maxChange, preflightChange, preflightOnly, preflightStartedByUserIdValue, startedByUserIdValue, minCreateTime?.UtcDateTime, maxCreateTime?.UtcDateTime, target, null, state, outcome,
+				maxChange, preflightChange, preflightOnly, includePreflight, preflightStartedByUserIdValue, startedByUserIdValue, minCreateTime?.UtcDateTime, maxCreateTime?.UtcDateTime, target, null, state, outcome,
 				modifiedBefore, modifiedAfter, index, count, false);
 
 			return await CreateAuthorizedJobResponsesAsync(jobs, filter);
diff --git a/Engine/Source/Programs/Horde/Plugins/Build/HordeServer.Build/Jobs/Schedules/ScheduleService.cs b/Engine/Source/Programs/Horde/Plugins/Build/HordeServer.Build/Jobs/Schedules/ScheduleService.cs
--- a/Engine/Source/Programs/Horde/Plugins/Build/HordeServer.Build/Jobs/Schedules/ScheduleService.cs
+++ b/Engine/Source/Programs/Horde/Plugins/Build/HordeServer.Build/Jobs/Schedules/ScheduleService.cs
@@ -325,14 +325,11 @@ namespace HordeServer.Jobs.Schedules
 			}
 
 			// Minimum changelist number, inclusive
-			int minChangeNumber = schedule.LastTriggerChange;
-			if (minChangeNumber > 0 && !schedule.Config.RequireSubmittedChange)
-			{
-				minChangeNumber--;
-			}
+			CommitIdWithOrder? minCommitId = schedule.LastTriggerCommitId;
+			bool includeMinCommitId = !schedule.Config.RequireSubmittedChange;
 
 			// Maximum changelist number, exclusive
-			int? maxChangeNumber = null;
+			CommitIdWithOrder? maxCommitId = null;
 
 			// Get the maximum number of changes to trigger
 			int maxNewChanges = 1;
@@ -357,65 +354,68 @@ namespace HordeServer.Jobs.Schedules
 
 			// Cache the Perforce history as we're iterating through changes to improve query performance
 			ICommitCollection commits = _commitService.GetCollection(stream.Config);
-			IAsyncEnumerable<ICommit> commitEnumerable = commits.FindAsync(minChangeNumber, null, null, schedule.Config.GetAllCommitTags(), cancellationToken);
+			IAsyncEnumerable<ICommit> commitEnumerable = commits.FindAsync(minCommitId, includeMinCommitId, tags: schedule.Config.GetAllCommitTags(), cancellationToken: cancellationToken);
 			await using IAsyncEnumerator<ICommit> commitEnumerator = commitEnumerable.GetAsyncEnumerator(cancellationToken);
 
 			// Start as many jobs as possible
-			List<(int Change, int CodeChange)> triggerChanges = new List<(int, int)>();
+			List<(CommitIdWithOrder CommitId, CommitIdWithOrder CodeCommitId)> triggerChanges = new List<(CommitIdWithOrder, CommitIdWithOrder)>();
 			while (triggerChanges.Count < maxNewChanges)
 			{
 				cancellationToken.ThrowIfCancellationRequested();
 
 				// Get the next valid change
-				int change;
+				CommitIdWithOrder? commitId;
 				ICommit? commit;
 
 				if (schedule.Config.Gate != null)
 				{
-					change = await GetNextChangeForGateAsync(stream.Id, templateId, schedule.Config.Gate, minChangeNumber, maxChangeNumber, cancellationToken);
-					commit = await commits.FindAsync(change, change, 1, null, cancellationToken).FirstOrDefaultAsync(cancellationToken); // May be a change in a different stream
+					commitId = await GetNextChangeForGateAsync(stream.Id, templateId, schedule.Config.Gate, minCommitId, maxCommitId, cancellationToken);
+					commit = (commitId != null)? await commits.FindAsync(commitId, commitId, 1, null, cancellationToken).FirstOrDefaultAsync(cancellationToken) : null; // May be a change in a different stream
 				}
 				else if (await commitEnumerator.MoveNextAsync(cancellationToken))
 				{
 					commit = commitEnumerator.Current;
-					change = commit.Number;
+					commitId = commit.Id;
 				}
 				else
 				{
 					commit = null;
-					change = 0;
+					commitId = null;
 				}
 
 				// Quit if we didn't find anything
-				if (change <= 0)
+				if (commitId == null)
 				{
 					break;
 				}
-				if (change < minChangeNumber)
+				if (minCommitId != null)
 				{
-					break;
-				}
-				if (change == minChangeNumber && (schedule.Config.RequireSubmittedChange || triggerChanges.Count > 0))
-				{
-					break;
+					if (commitId < minCommitId)
+					{
+						break;
+					}
+					if (commitId == minCommitId && (schedule.Config.RequireSubmittedChange || triggerChanges.Count > 0))
+					{
+						break;
+					}
 				}
 
 				// Adjust the changelist for the desired filter
 				if (commit == null || await ShouldBuildChangeAsync(commit, schedule.Config.GetAllCommitTags(), fileFilter, cancellationToken))
 				{
-					int codeChange = change;
+					CommitIdWithOrder codeCommitId = commitId;
 
-					ICommit? lastCodeCommit = await commits.GetLastCodeChangeAsync(change, cancellationToken);
+					ICommit? lastCodeCommit = await commits.GetLastCodeChangeAsync(commitId, cancellationToken);
 					if (lastCodeCommit != null)
 					{
-						codeChange = lastCodeCommit.Number;
+						codeCommitId = lastCodeCommit.Id;
 					}
 					else
 					{
-						_logger.LogWarning("Unable to find code change for CL {Change}", change);
+						_logger.LogWarning("Unable to find code change for CL {Change}", commitId);
 					}
 
-					triggerChanges.Add((change, codeChange));
+					triggerChanges.Add((commitId, codeCommitId));
 				}
 
 				// Check we haven't exceeded the time limit
@@ -426,8 +426,8 @@ namespace HordeServer.Jobs.Schedules
 				}
 
 				// Update the remaining range of changes to check for
-				maxChangeNumber = change - 1;
-				if (maxChangeNumber < minChangeNumber)
+				maxCommitId = commitId;
+				if (minCommitId != null && maxCommitId < minCommitId)
 				{
 					break;
 				}
@@ -436,7 +436,7 @@ namespace HordeServer.Jobs.Schedules
 			// Early out if there's nothing to do
 			if (triggerChanges.Count == 0)
 			{
-				_logger.LogInformation("Skipping trigger of {StreamName} template {TemplateId} - no candidate changes after CL {LastTriggerChange}", stream.Id, templateId, schedule.LastTriggerChange);
+				_logger.LogInformation("Skipping trigger of {StreamName} template {TemplateId} - no candidate changes after CL {LastTriggerChange}", stream.Id, templateId, schedule.LastTriggerCommitId);
 				return;
 			}
 
@@ -449,14 +449,14 @@ namespace HordeServer.Jobs.Schedules
 			// We may need to submit a new change for any new jobs. This only makes sense if there's one change.
 			if (template.SubmitNewChange != null)
 			{
-				int newChange = await commits.CreateNewAsync(template, cancellationToken);
-				ICommit? newCodeChange = await commits.GetLastCodeChangeAsync(newChange, cancellationToken);
-				triggerChanges = new List<(int, int)> { (newChange, newCodeChange?.Number ?? newChange) };
+				CommitIdWithOrder newCommitId = await commits.CreateNewAsync(template, cancellationToken);
+				ICommit? newCodeCommit = await commits.GetLastCodeChangeAsync(newCommitId, cancellationToken);
+				triggerChanges = new List<(CommitIdWithOrder, CommitIdWithOrder)> { (newCommitId, newCodeCommit?.Id ?? newCommitId) };
 			}
 
 			// Try to start all the new jobs
 			_logger.LogInformation("Starting {NumJobs} new jobs for {StreamId} template {TemplateId} (active: {NumActive}, max new: {MaxNewJobs})", triggerChanges.Count, stream.Id, templateId, numActiveJobs, maxNewChanges);
-			foreach ((int change, int codeChange) in triggerChanges.OrderBy(x => x.Change))
+			foreach ((CommitIdWithOrder commitId, CommitIdWithOrder codeCommitId) in triggerChanges.OrderBy(x => x.CommitId))
 			{
 				cancellationToken.ThrowIfCancellationRequested();
 
@@ -465,9 +465,9 @@ namespace HordeServer.Jobs.Schedules
 				template.GetDefaultParameters(options.Parameters, true);
 				template.GetArgumentsForParameters(options.Parameters, options.Arguments);
 
-				IJob newJob = await _jobService.CreateJobAsync(null, stream.Config, templateId, template.Hash, graph, template.Name, change, codeChange, options, cancellationToken);
-				_logger.LogInformation("Started new job for {StreamId} template {TemplateId} at CL {Change} (Code CL {CodeChange}): {JobId}", stream.Id, templateId, change, codeChange, newJob.Id);
-				await _streamCollection.UpdateScheduleTriggerAsync(stream, templateId, utcNow, change, new List<JobId> { newJob.Id }, new List<JobId>(), cancellationToken);
+				IJob newJob = await _jobService.CreateJobAsync(null, stream.Config, templateId, template.Hash, graph, template.Name, commitId, codeCommitId, options, cancellationToken);
+				_logger.LogInformation("Started new job for {StreamId} template {TemplateId} at CL {Change} (Code CL {CodeChange}): {JobId}", stream.Id, templateId, commitId, codeCommitId, newJob.Id);
+				await _streamCollection.UpdateScheduleTriggerAsync(stream, templateId, utcNow, commitId, new List<JobId> { newJob.Id }, new List<JobId>(), cancellationToken);
 			}
 		}
 
@@ -490,7 +490,7 @@ namespace HordeServer.Jobs.Schedules
 				IReadOnlyList<CommitTag> commitTags = await commit.GetTagsAsync(cancellationToken);
 				if (!commitTags.Any(x => filterTags.Contains(x)))
 				{
-					_logger.LogDebug("Not building change {Change} ({ChangeTags}) due to filter tags ({FilterTags})", commit.Number, String.Join(", ", commitTags.Select(x => x.ToString())), String.Join(", ", filterTags.Select(x => x.ToString())));
+					_logger.LogDebug("Not building change {Change} ({ChangeTags}) due to filter tags ({FilterTags})", commit.Id, String.Join(", ", commitTags.Select(x => x.ToString())), String.Join(", ", filterTags.Select(x => x.ToString())));
 					return false;
 				}
 			}
@@ -498,7 +498,7 @@ namespace HordeServer.Jobs.Schedules
 			{
 				if (!await commit.MatchesFilterAsync(fileFilter, cancellationToken))
 				{
-					_logger.LogDebug("Not building change {Change} due to file filter", commit.Number);
+					_logger.LogDebug("Not building change {Change} due to file filter", commit.Id);
 					return false;
 				}
 			}
@@ -508,21 +508,20 @@ namespace HordeServer.Jobs.Schedules
 		/// <summary>
 		/// Gets the next change to build for a schedule on a gate
 		/// </summary>
-		/// <returns></returns>
-		private async Task<int> GetNextChangeForGateAsync(StreamId streamId, TemplateId templateRefId, ScheduleGateConfig gate, int? minChange, int? maxChange, CancellationToken cancellationToken)
+		private async Task<CommitIdWithOrder?> GetNextChangeForGateAsync(StreamId streamId, TemplateId templateRefId, ScheduleGateConfig gate, CommitIdWithOrder? minCommitId, CommitIdWithOrder? maxCommitId, CancellationToken cancellationToken)
 		{
 			for (; ; )
 			{
 				cancellationToken.ThrowIfCancellationRequested();
 
-				IReadOnlyList<IJob> jobs = await _jobCollection.FindAsync(streamId: streamId, templates: new[] { gate.TemplateId }, minChange: minChange, maxChange: maxChange, count: 1, cancellationToken: cancellationToken);
-				if (jobs.Count == 0)
+				IReadOnlyList<IJob> jobs = await _jobCollection.FindAsync(streamId: streamId, templates: new[] { gate.TemplateId }, minCommitId: minCommitId, maxCommitId: maxCommitId, count: 2, cancellationToken: cancellationToken);
+
+				IJob? job = jobs.FirstOrDefault(x => maxCommitId == null || x.CommitId < maxCommitId);
+				if (job == null)
 				{
-					return 0;
+					return null;
 				}
 
-				IJob job = jobs[0];
-
 				IGraph? graph = await _graphs.GetAsync(job.GraphHash, cancellationToken);
 				if (graph != null)
 				{
@@ -532,13 +531,13 @@ namespace HordeServer.Jobs.Schedules
 						JobStepOutcome outcome = state.Value.Item2;
 						if (outcome == JobStepOutcome.Success || outcome == JobStepOutcome.Warnings)
 						{
-							return job.Change;
+							return job.CommitId;
 						}
 						_logger.LogInformation("Skipping trigger of {StreamName} template {TemplateId} - last {OtherTemplateRefId} job ({JobId}) ended with errors", streamId, templateRefId, gate.TemplateId, job.Id);
 					}
 				}
 
-				maxChange = job.Change - 1;
+				maxCommitId = job.CommitId;
 			}
 		}
 	}
diff --git a/Engine/Source/Programs/Horde/Plugins/Build/HordeServer.Build/Jobs/Schedules/Schedules.cs b/Engine/Source/Programs/Horde/Plugins/Build/HordeServer.Build/Jobs/Schedules/Schedules.cs
--- a/Engine/Source/Programs/Horde/Plugins/Build/HordeServer.Build/Jobs/Schedules/Schedules.cs
+++ b/Engine/Source/Programs/Horde/Plugins/Build/HordeServer.Build/Jobs/Schedules/Schedules.cs
@@ -60,7 +60,23 @@ namespace HordeServer.Jobs.Schedules
 		/// <summary>
 		/// Last changelist number that this was triggered for
 		/// </summary>
-		public int LastTriggerChange { get; set; }
+		[Obsolete("Use LastTriggerCommitId instead")]
+		public int LastTriggerChange
+		{
+			get => _lastTriggerChange ?? _lastTriggerCommitId?.TryGetPerforceChange() ?? 0;
+			set => _lastTriggerChange = value;
+		}
+		int? _lastTriggerChange;
+
+		/// <summary>
+		/// Last changelist number that this was triggered for
+		/// </summary>
+		public CommitIdWithOrder? LastTriggerCommitId
+		{
+			get => _lastTriggerCommitId ?? CommitIdWithOrder.FromPerforceChange(_lastTriggerChange);
+			set => _lastTriggerCommitId = value;
+		}
+		CommitIdWithOrder? _lastTriggerCommitId;
 
 		/// <summary>
 		/// Last time that the schedule was triggered
@@ -95,7 +111,7 @@ namespace HordeServer.Jobs.Schedules
 #pragma warning restore CS0618 // Type or member is obsolete
 			TemplateParameters = schedule.Config.TemplateParameters;
 			Patterns = schedule.Config.Patterns;
-			LastTriggerChange = schedule.LastTriggerChange;
+			LastTriggerCommitId = schedule.LastTriggerCommitId;
 			LastTriggerTime = schedule.LastTriggerTimeUtc;
 			ActiveJobs = new List<JobId>(schedule.ActiveJobs);
 
diff --git a/Engine/Source/Programs/Horde/Plugins/Build/HordeServer.Build/Jobs/TestData/ITestData.cs b/Engine/Source/Programs/Horde/Plugins/Build/HordeServer.Build/Jobs/TestData/ITestData.cs
--- a/Engine/Source/Programs/Horde/Plugins/Build/HordeServer.Build/Jobs/TestData/ITestData.cs
+++ b/Engine/Source/Programs/Horde/Plugins/Build/HordeServer.Build/Jobs/TestData/ITestData.cs
@@ -1,5 +1,6 @@
 // Copyright Epic Games, Inc. All Rights Reserved.
 
+using EpicGames.Horde.Commits;
 using EpicGames.Horde.Jobs;
 using EpicGames.Horde.Jobs.Templates;
 using EpicGames.Horde.Jobs.TestData;
@@ -17,6 +18,7 @@ namespace HordeServer.Jobs.TestData
 		/// The test meta id
 		/// </summary>
 		TestMetaId Id { get; }
+
 		/// <summary>
 		/// The name of the test platform
 		/// </summary>
@@ -146,7 +148,9 @@ namespace HordeServer.Jobs.TestData
 	/// </summary>
 	public interface ITestDataRef
 	{
+		/// <summary>
 		/// The test ref id
+		/// </summary>
 		TestRefId Id { get; }
 
 		/// <summary>
@@ -172,7 +176,7 @@ namespace HordeServer.Jobs.TestData
 		/// <summary>
 		/// The build changelist upon which the test ran, may not correspond to the job changelist
 		/// </summary>
-		int BuildChangeList { get; }
+		CommitIdWithOrder BuildCommitId { get; }
 
 		/// <summary>
 		/// The environment the test ran in
@@ -220,7 +224,9 @@ namespace HordeServer.Jobs.TestData
 	/// </summary>
 	public interface ITestDataDetails
 	{
+		/// <summary>
 		/// The corresponding test ref		
+		/// </summary>
 		TestRefId Id { get; }
 
 		/// <summary>
@@ -286,9 +292,9 @@ namespace HordeServer.Jobs.TestData
 		JobStepId StepId { get; }
 
 		/// <summary>
-		/// The changelist number that contained the data
+		/// The commit that contained the data
 		/// </summary>
-		int Change { get; }
+		CommitIdWithOrder CommitId { get; }
 
 		/// <summary>
 		/// Key used to identify the particular data
diff --git a/Engine/Source/Programs/Horde/Plugins/Build/HordeServer.Build/Jobs/TestData/ITestDataCollection.cs b/Engine/Source/Programs/Horde/Plugins/Build/HordeServer.Build/Jobs/TestData/ITestDataCollection.cs
--- a/Engine/Source/Programs/Horde/Plugins/Build/HordeServer.Build/Jobs/TestData/ITestDataCollection.cs
+++ b/Engine/Source/Programs/Horde/Plugins/Build/HordeServer.Build/Jobs/TestData/ITestDataCollection.cs
@@ -1,5 +1,6 @@
 // Copyright Epic Games, Inc. All Rights Reserved.
 
+using EpicGames.Horde.Commits;
 using EpicGames.Horde.Jobs;
 using EpicGames.Horde.Streams;
 using MongoDB.Bson;
@@ -28,11 +29,11 @@ namespace HordeServer.Jobs.TestData
 		/// <param name="suiteIds"></param>
 		/// <param name="minCreateTime"></param>
 		/// <param name="maxCreateTime"></param>
-		/// <param name="minChange"></param>
-		/// <param name="maxChange"></param>
+		/// <param name="minCommitId"></param>
+		/// <param name="maxCommitId"></param>
 		/// <param name="cancellationToken">Cancellation token for the operation</param>
 		/// <returns></returns>
-		Task<IReadOnlyList<ITestDataRef>> FindTestRefsAsync(StreamId[] streamIds, TestMetaId[]? metaIds = null, TestId[]? testIds = null, TestSuiteId[]? suiteIds = null, DateTime? minCreateTime = null, DateTime? maxCreateTime = null, int? minChange = null, int? maxChange = null, CancellationToken cancellationToken = default);
+		Task<IReadOnlyList<ITestDataRef>> FindTestRefsAsync(StreamId[] streamIds, TestMetaId[]? metaIds = null, TestId[]? testIds = null, TestSuiteId[]? suiteIds = null, DateTime? minCreateTime = null, DateTime? maxCreateTime = null, CommitId? minCommitId = null, CommitId? maxCommitId = null, CancellationToken cancellationToken = default);
 
 		/// <summary>
 		/// Find test details
@@ -80,8 +81,8 @@ namespace HordeServer.Jobs.TestData
 		/// Searches for test data that matches a set of criteria
 		/// </summary>
 		/// <param name="streamId">The stream id</param>
-		/// <param name="minChange">The minimum changelist number to return (inclusive)</param>
-		/// <param name="maxChange">The maximum changelist number to return (inclusive)</param>
+		/// <param name="minCommitId">The minimum commit to return (inclusive)</param>
+		/// <param name="maxCommitId">The maximum commit to return (inclusive)</param>
 		/// <param name="jobId">The job id</param>
 		/// <param name="stepId">The unique step id</param>
 		/// <param name="key">Key identifying the result to return</param>
@@ -89,7 +90,7 @@ namespace HordeServer.Jobs.TestData
 		/// <param name="count">Number of results to return</param>
 		/// <param name="cancellationToken">Cancellation token for the operation</param>
 		/// <returns>The stream document</returns>
-		Task<IReadOnlyList<ITestData>> FindAsync(StreamId? streamId, int? minChange, int? maxChange, JobId? jobId, JobStepId? stepId, string? key = null, int index = 0, int count = 10, CancellationToken cancellationToken = default);
+		Task<IReadOnlyList<ITestData>> FindAsync(StreamId? streamId, CommitId? minCommitId, CommitId? maxCommitId, JobId? jobId, JobStepId? stepId, string? key = null, int index = 0, int count = 10, CancellationToken cancellationToken = default);
 
 		/// <summary>
 		/// Find test meta data
diff --git a/Engine/Source/Programs/Horde/Plugins/Build/HordeServer.Build/Jobs/TestData/TestDataCollection.cs b/Engine/Source/Programs/Horde/Plugins/Build/HordeServer.Build/Jobs/TestData/TestDataCollection.cs
--- a/Engine/Source/Programs/Horde/Plugins/Build/HordeServer.Build/Jobs/TestData/TestDataCollection.cs
+++ b/Engine/Source/Programs/Horde/Plugins/Build/HordeServer.Build/Jobs/TestData/TestDataCollection.cs
@@ -1,9 +1,11 @@
 // Copyright Epic Games, Inc. All Rights Reserved.
 
+using EpicGames.Horde.Commits;
 using EpicGames.Horde.Jobs;
 using EpicGames.Horde.Jobs.Templates;
 using EpicGames.Horde.Jobs.TestData;
 using EpicGames.Horde.Streams;
+using HordeServer.Commits;
 using HordeServer.Server;
 using HordeServer.Utilities;
 using Microsoft.Extensions.Logging;
@@ -243,8 +245,17 @@ namespace HordeServer.Jobs.TestData
 			[BsonRequired, BsonElement("m")]
 			public TestMetaId Metadata { get; set; }
 
+			public CommitIdWithOrder BuildCommitId
+			{
+				get => (BuildCommitName != null) ? new CommitIdWithOrder(BuildCommitName, BuildCommitOrder) : CommitIdWithOrder.FromPerforceChange(BuildCommitOrder);
+				set => (BuildCommitName, BuildCommitOrder) = (value.Name, value.Order);
+			}
+
+			[BsonIgnoreIfNull, BsonElement("bcn")]
+			public string? BuildCommitName { get; set; }
+
 			[BsonRequired, BsonElement("bcl")]
-			public int BuildChangeList { get; set; }
+			public int BuildCommitOrder { get; set; }
 
 			[BsonRequired, BsonElement("d")]
 			public TimeSpan Duration { get; set; }
@@ -335,7 +346,19 @@ namespace HordeServer.Jobs.TestData
 			public TemplateId TemplateRefId { get; set; }
 			public JobId JobId { get; set; }
 			public JobStepId StepId { get; set; }
-			public int Change { get; set; }
+
+			[BsonIgnore]
+			public CommitIdWithOrder CommitId
+			{
+				get => (CommitName != null) ? new CommitIdWithOrder(CommitName, CommitOrder) : CommitIdWithOrder.FromPerforceChange(CommitOrder);
+				set => (CommitName, CommitOrder) = (value.Name, value.Order);
+			}
+
+			public string? CommitName { get; set; }
+			
+			[BsonElement("Change")]
+			public int CommitOrder { get; set; }
+
 			public string Key { get; set; }
 			public BsonDocument Data { get; set; }
 
@@ -352,59 +375,34 @@ namespace HordeServer.Jobs.TestData
 				TemplateRefId = job.TemplateId;
 				JobId = job.Id;
 				StepId = jobStep.Id;
-				Change = job.Change;
+				CommitId = job.CommitId;
 				Key = key;
 				Data = value;
 			}
 		}
 
-		/// <summary>
-		/// The detailed test data collection
-		/// </summary>
 		readonly IMongoCollection<TestDataDocument> _testDataDocuments;
-
-		/// <summary>
-		/// Test meta collection
-		/// </summary>
 		readonly IMongoCollection<TestMetaDocument> _testMeta;
-
-		/// <summary>
-		/// Test collection
-		/// </summary>
 		readonly IMongoCollection<TestDocument> _tests;
-
-		/// <summary>
-		/// Test suite collection
-		/// </summary>
 		readonly IMongoCollection<TestSuiteDocument> _testSuites;
-
-		/// <summary>
-		/// Test data refs collection
-		/// </summary>
 		readonly IMongoCollection<TestDataRefDocument> _testRefs;
-
-		/// <summary>
-		/// Test data refs collection
-		/// </summary>
 		readonly IMongoCollection<TestDataDetailsDocument> _testDetails;
-		/// <summary>
-		/// Test streams collection
-		/// </summary>
 		readonly IMongoCollection<TestStreamDocument> _testStreams;
-
+		readonly ICommitService _commitService;
 		readonly Tracer _tracer;
 		readonly ILogger _logger;
 
 		/// <summary>
 		/// Constructor
 		/// </summary>
-		public TestDataCollection(IMongoService mongoService, Tracer tracer, ILogger<TestDataCollection> logger)
+		public TestDataCollection(IMongoService mongoService, ICommitService commitService, Tracer tracer, ILogger<TestDataCollection> logger)
 		{
+			_commitService = commitService;
 			_tracer = tracer;
 			_logger = logger;
 
 			List<MongoIndex<TestDataDocument>> indexes = new List<MongoIndex<TestDataDocument>>();
-			indexes.Add(keys => keys.Ascending(x => x.StreamId).Ascending(x => x.Change).Ascending(x => x.Key));
+			indexes.Add(keys => keys.Ascending(x => x.StreamId).Ascending(x => x.CommitOrder).Ascending(x => x.Key));
 			indexes.Add(keys => keys.Ascending(x => x.JobId).Ascending(x => x.StepId).Ascending(x => x.Key), unique: true);
 			_testDataDocuments = mongoService.GetCollection<TestDataDocument>("TestData", indexes);
 
@@ -421,7 +419,7 @@ namespace HordeServer.Jobs.TestData
 			_testSuites = mongoService.GetCollection<TestSuiteDocument>("TestData.TestSuitesV2", testSuiteIndexes);
 
 			List<MongoIndex<TestDataRefDocument>> testRefIndexes = new List<MongoIndex<TestDataRefDocument>>();
-			testRefIndexes.Add(keys => keys.Ascending(x => x.StreamId).Ascending(x => x.Metadata).Descending(x => x.BuildChangeList).Ascending(x => x.TestId).Ascending(x => x.SuiteId));
+			testRefIndexes.Add(keys => keys.Ascending(x => x.StreamId).Ascending(x => x.Metadata).Descending(x => x.BuildCommitOrder).Ascending(x => x.TestId).Ascending(x => x.SuiteId));
 			_testRefs = mongoService.GetCollection<TestDataRefDocument>("TestData.TestRefsV2", testRefIndexes);
 
 			_testDetails = mongoService.GetCollection<TestDataDetailsDocument>("TestData.TestDetailsV2");
@@ -433,7 +431,7 @@ namespace HordeServer.Jobs.TestData
 		}
 
 		/// <inheritdoc/>
-		public async Task<IReadOnlyList<ITestDataRef>> FindTestRefsAsync(StreamId[] streamIds, TestMetaId[]? metaIds = null, TestId[]? testIds = null, TestSuiteId[]? suiteIds = null, DateTime? minCreateTime = null, DateTime? maxCreateTime = null, int? minChange = null, int? maxChange = null, CancellationToken cancellationToken = default)
+		public async Task<IReadOnlyList<ITestDataRef>> FindTestRefsAsync(StreamId[] streamIds, TestMetaId[]? metaIds = null, TestId[]? testIds = null, TestSuiteId[]? suiteIds = null, DateTime? minCreateTime = null, DateTime? maxCreateTime = null, CommitId? minCommitId = null, CommitId? maxCommitId = null, CancellationToken cancellationToken = default)
 		{
 
 			FilterDefinition<TestDataRefDocument> filter = FilterDefinition<TestDataRefDocument>.Empty;
@@ -476,13 +474,15 @@ namespace HordeServer.Jobs.TestData
 				}
 			}
 
-			if (minChange != null)
+			if (minCommitId != null)
 			{
-				filter &= filterBuilder.Gte(x => x.BuildChangeList, minChange);
+				int minCommitOrder = minCommitId.GetPerforceChange(); // TODO: need to resolve against correct stream
+				filter &= filterBuilder.Gte(x => x.BuildCommitOrder, minCommitOrder);
 			}
-			if (maxChange != null)
+			if (maxCommitId != null)
 			{
-				filter &= filterBuilder.Lte(x => x.BuildChangeList, maxChange);
+				int maxCommitOrder = maxCommitId.GetPerforceChange(); // TODO: need to resolve against correct stream
+				filter &= filterBuilder.Lte(x => x.BuildCommitOrder, maxCommitOrder);
 			}
 
 			List<TestDataRefDocument> results;
@@ -549,19 +549,21 @@ namespace HordeServer.Jobs.TestData
 		}
 
 		/// <inheritdoc/>
-		public async Task<IReadOnlyList<ITestData>> FindAsync(StreamId? streamId, int? minChange, int? maxChange, JobId? jobId, JobStepId? stepId, string? key = null, int index = 0, int count = 10, CancellationToken cancellationToken = default)
+		public async Task<IReadOnlyList<ITestData>> FindAsync(StreamId? streamId, CommitId? minCommitId, CommitId? maxCommitId, JobId? jobId, JobStepId? stepId, string? key = null, int index = 0, int count = 10, CancellationToken cancellationToken = default)
 		{
 			FilterDefinition<TestDataDocument> filter = FilterDefinition<TestDataDocument>.Empty;
 			if (streamId != null)
 			{
 				filter &= Builders<TestDataDocument>.Filter.Eq(x => x.StreamId, streamId.Value);
-				if (minChange != null)
+				if (minCommitId != null)
 				{
-					filter &= Builders<TestDataDocument>.Filter.Gte(x => x.Change, minChange.Value);
+					CommitIdWithOrder minCommitIdWithOrder = await _commitService.GetOrderedAsync(streamId.Value, minCommitId, cancellationToken);
+					filter &= Builders<TestDataDocument>.Filter.Gte(x => x.CommitOrder, minCommitIdWithOrder.Order);
 				}
-				if (maxChange != null)
+				if (maxCommitId != null)
 				{
-					filter &= Builders<TestDataDocument>.Filter.Lte(x => x.Change, maxChange.Value);
+					CommitIdWithOrder maxCommitIdWithOrder = await _commitService.GetOrderedAsync(streamId.Value, maxCommitId, cancellationToken);
+					filter &= Builders<TestDataDocument>.Filter.Lte(x => x.CommitOrder, maxCommitIdWithOrder.Order);
 				}
 			}
 			if (jobId != null)
@@ -577,7 +579,7 @@ namespace HordeServer.Jobs.TestData
 				filter &= Builders<TestDataDocument>.Filter.Eq(x => x.Key, key);
 			}
 
-			SortDefinition<TestDataDocument> sort = Builders<TestDataDocument>.Sort.Ascending(x => x.StreamId).Descending(x => x.Change);
+			SortDefinition<TestDataDocument> sort = Builders<TestDataDocument>.Sort.Ascending(x => x.StreamId).Descending(x => x.CommitOrder);
 
 			return await _testDataDocuments.Find(filter).Sort(sort).Skip(index).Limit(count).ToListAsync(cancellationToken);
 		}
@@ -1052,7 +1054,7 @@ namespace HordeServer.Jobs.TestData
 		{
 
 			// do not add preflight to temporal data
-			if (job.PreflightChange != 0)
+			if (job.PreflightCommitId != null)
 			{
 				return;
 			}
@@ -1508,7 +1510,7 @@ namespace HordeServer.Jobs.TestData
 				testRef.JobId = job.Id;
 				testRef.StepId = step.Id;
 				testRef.Metadata = metaData.Id;
-				testRef.BuildChangeList = testData.BuildChangeList <= 0 ? job.Change : testData.BuildChangeList;
+				testRef.BuildCommitId = testData.BuildChangeList <= 0 ? job.CommitId : CommitIdWithOrder.FromPerforceChange(testData.BuildChangeList);
 				testRef.Duration = TimeSpan.FromSeconds(testData.TotalDurationSeconds);
 				testRef.TestId = test.Id;
 				testRef.Outcome = testData.HasSucceeded ? TestOutcome.Success : TestOutcome.Unspecified;
@@ -1644,7 +1646,7 @@ namespace HordeServer.Jobs.TestData
 						testRef.JobId = job.Id;
 						testRef.StepId = step.Id;
 						testRef.Metadata = metaData.Id;
-						testRef.BuildChangeList = job.Change;
+						testRef.BuildCommitId = job.CommitId;
 						testRef.Duration = TimeSpan.FromSeconds(session.TestSessionInfo!.TimeElapseSec);
 						testRef.Outcome = TestOutcome.Unspecified;
 						testRef.SuiteId = testSuite.Id;
diff --git a/Engine/Source/Programs/Horde/Plugins/Build/HordeServer.Build/Jobs/TestData/TestDataController.cs b/Engine/Source/Programs/Horde/Plugins/Build/HordeServer.Build/Jobs/TestData/TestDataController.cs
--- a/Engine/Source/Programs/Horde/Plugins/Build/HordeServer.Build/Jobs/TestData/TestDataController.cs
+++ b/Engine/Source/Programs/Horde/Plugins/Build/HordeServer.Build/Jobs/TestData/TestDataController.cs
@@ -1,6 +1,7 @@
 // Copyright Epic Games, Inc. All Rights Reserved.
 
 using EpicGames.Core;
+using EpicGames.Horde.Commits;
 using EpicGames.Horde.Jobs;
 using EpicGames.Horde.Jobs.TestData;
 using EpicGames.Horde.Streams;
@@ -303,8 +304,8 @@ namespace HordeServer.Jobs.TestData
 			[FromQuery(Name = "Sid")] string[]? suiteIds = null,
 			[FromQuery] DateTimeOffset? minCreateTime = null,
 			[FromQuery] DateTimeOffset? maxCreateTime = null,
-			[FromQuery] int? minChange = null,
-			[FromQuery] int? maxChange = null)
+			[FromQuery] CommitId? minChange = null,
+			[FromQuery] CommitId? maxChange = null)
 		{
 			StreamId[] streamIdValues = Array.ConvertAll(streamIds, x => new StreamId(x));
 
@@ -335,7 +336,7 @@ namespace HordeServer.Jobs.TestData
 					JobId = testData.JobId?.ToString(),
 					StepId = testData.StepId?.ToString(),
 					Duration = testData.Duration,
-					BuildChangeList = testData.BuildChangeList,
+					BuildCommitId = testData.BuildCommitId,
 					MetaId = testData.Metadata.ToString(),
 					TestId = testData.TestId?.ToString(),
 					Outcome = testData.TestId != null ? testData.Outcome : null,
@@ -395,7 +396,7 @@ namespace HordeServer.Jobs.TestData
 		[HttpGet]
 		[Route("/api/v1/testdata")]
 		[ProducesResponseType(typeof(List<GetTestDataResponse>), 200)]
-		public async Task<ActionResult<List<object>>> FindTestDataAsync([FromQuery] string? streamId = null, [FromQuery] int? minChange = null, [FromQuery] int? maxChange = null, JobId? jobId = null, JobStepId? jobStepId = null, string? key = null, int index = 0, int count = 10, PropertyFilter? filter = null, CancellationToken cancellationToken = default)
+		public async Task<ActionResult<List<object>>> FindTestDataAsync([FromQuery] string? streamId = null, [FromQuery] CommitId? minChange = null, [FromQuery] CommitId? maxChange = null, JobId? jobId = null, JobStepId? jobStepId = null, string? key = null, int index = 0, int count = 10, PropertyFilter? filter = null, CancellationToken cancellationToken = default)
 		{
 			StreamId? streamIdValue = null;
 			if (streamId != null)
@@ -417,7 +418,7 @@ namespace HordeServer.Jobs.TestData
 						TemplateRefId = testData.TemplateRefId.ToString(),
 						JobId = testData.JobId.ToString(),
 						StepId = testData.StepId.ToString(),
-						Change = testData.Change,
+						CommitId = testData.CommitId,
 						Key = testData.Key,
 						Data = BsonSerializer.Deserialize<Dictionary<string, object>>(testData.Data)
 					}
@@ -457,7 +458,7 @@ namespace HordeServer.Jobs.TestData
 				TemplateRefId = testData.TemplateRefId.ToString(),
 				JobId = testData.JobId.ToString(),
 				StepId = testData.StepId.ToString(),
-				Change = testData.Change,
+				CommitId = testData.CommitId,
 				Key = testData.Key,
 				Data = BsonSerializer.Deserialize<Dictionary<string, object>>(testData.Data)
 			}, filter);
diff --git a/Engine/Source/Programs/Horde/Plugins/Build/HordeServer.Build/Jobs/TestData/TestDataService.cs b/Engine/Source/Programs/Horde/Plugins/Build/HordeServer.Build/Jobs/TestData/TestDataService.cs
--- a/Engine/Source/Programs/Horde/Plugins/Build/HordeServer.Build/Jobs/TestData/TestDataService.cs
+++ b/Engine/Source/Programs/Horde/Plugins/Build/HordeServer.Build/Jobs/TestData/TestDataService.cs
@@ -1,6 +1,7 @@
 // Copyright Epic Games, Inc. All Rights Reserved.
 
 using EpicGames.Core;
+using EpicGames.Horde.Commits;
 using EpicGames.Horde.Streams;
 using Microsoft.Extensions.Hosting;
 using Microsoft.Extensions.Logging;
@@ -139,16 +140,16 @@ namespace HordeServer.Jobs.TestData
 		/// <param name="suiteIds"></param>
 		/// <param name="minCreateTime"></param>
 		/// <param name="maxCreateTime"></param>
-		/// <param name="minChange"></param>
-		/// <param name="maxChange"></param>
+		/// <param name="minCommitId"></param>
+		/// <param name="maxCommitId"></param>
 		/// <param name="cancellationToken">Cancellation token for the operation</param>
 		/// <returns></returns>
-		public async Task<IReadOnlyList<ITestDataRef>> FindTestRefsAsync(StreamId[] streamIds, TestMetaId[] metaIds, string[]? testIds = null, string[]? suiteIds = null, DateTime? minCreateTime = null, DateTime? maxCreateTime = null, int? minChange = null, int? maxChange = null, CancellationToken cancellationToken = default)
+		public async Task<IReadOnlyList<ITestDataRef>> FindTestRefsAsync(StreamId[] streamIds, TestMetaId[] metaIds, string[]? testIds = null, string[]? suiteIds = null, DateTime? minCreateTime = null, DateTime? maxCreateTime = null, CommitId? minCommitId = null, CommitId? maxCommitId = null, CancellationToken cancellationToken = default)
 		{
 			TestId[]? tids = testIds?.ConvertAll(x => TestId.Parse(x));
 			TestSuiteId[]? sids = suiteIds?.ConvertAll(x => TestSuiteId.Parse(x));
 
-			return await _testData.FindTestRefsAsync(streamIds, metaIds, tids, sids, minCreateTime, maxCreateTime, minChange, maxChange, cancellationToken);
+			return await _testData.FindTestRefsAsync(streamIds, metaIds, tids, sids, minCreateTime, maxCreateTime, minCommitId, maxCommitId, cancellationToken);
 		}
 
 		/// <summary>
diff --git a/Engine/Source/Programs/Horde/Plugins/Build/HordeServer.Build/Notifications/NotificationService.cs b/Engine/Source/Programs/Horde/Plugins/Build/HordeServer.Build/Notifications/NotificationService.cs
--- a/Engine/Source/Programs/Horde/Plugins/Build/HordeServer.Build/Notifications/NotificationService.cs
+++ b/Engine/Source/Programs/Horde/Plugins/Build/HordeServer.Build/Notifications/NotificationService.cs
@@ -511,7 +511,7 @@ namespace HordeServer.Notifications
 			IReadOnlyList<IUser> usersToNotify = await GetUsersToNotifyAsync(jobCompleteEvent, job.NotificationTriggerId, true, cancellationToken);
 			foreach (IUser userToNotify in usersToNotify)
 			{
-				if (job.PreflightChange != 0)
+				if (job.PreflightCommitId != null)
 				{
 					if (userToNotify.Id != job.StartedByUserId)
 					{
@@ -521,7 +521,7 @@ namespace HordeServer.Notifications
 				EnqueueTasks((sink, ctx) => sink.NotifyJobCompleteAsync(userToNotify, job, graph, outcome, ctx));
 			}
 
-			if (job.PreflightChange == 0)
+			if (job.PreflightCommitId == null)
 			{
 				EnqueueTasks((sink, ctx) => sink.NotifyJobCompleteAsync(job, graph, outcome, ctx));
 			}
@@ -690,7 +690,7 @@ namespace HordeServer.Notifications
 
 			foreach (IUser slackUser in usersToNotify)
 			{
-				if (job.PreflightChange != 0)
+				if (job.PreflightCommitId != null)
 				{
 					if (slackUser.Id != job.StartedByUserId)
 					{
@@ -757,7 +757,7 @@ namespace HordeServer.Notifications
 					IReadOnlyList<IUser> usersToNotify = await GetUsersToNotifyAsync(eventId, triggerId, fireTrigger, cancellationToken);
 
 					// filter preflight label notifications to only include initiator
-					if (usersToNotify.Count > 0 && job.PreflightChange != 0 && job.StartedByUserId != null)
+					if (usersToNotify.Count > 0 && job.PreflightCommitId != null && job.StartedByUserId != null)
 					{
 						usersToNotify = usersToNotify.Where(x => x.Id == job.StartedByUserId).ToList();
 					}
diff --git a/Engine/Source/Programs/Horde/Plugins/Build/HordeServer.Build/Notifications/Sinks/SlackNotificationSink.cs b/Engine/Source/Programs/Horde/Plugins/Build/HordeServer.Build/Notifications/Sinks/SlackNotificationSink.cs
--- a/Engine/Source/Programs/Horde/Plugins/Build/HordeServer.Build/Notifications/Sinks/SlackNotificationSink.cs
+++ b/Engine/Source/Programs/Horde/Plugins/Build/HordeServer.Build/Notifications/Sinks/SlackNotificationSink.cs
@@ -11,6 +11,7 @@ using System.Text.RegularExpressions;
 using EpicGames.Core;
 using EpicGames.Horde;
 using EpicGames.Horde.Agents;
+using EpicGames.Horde.Commits;
 using EpicGames.Horde.Issues;
 using EpicGames.Horde.Jobs;
 using EpicGames.Horde.Jobs.Templates;
@@ -23,6 +24,7 @@ using EpicGames.Slack;
 using EpicGames.Slack.Blocks;
 using EpicGames.Slack.Elements;
 using HordeServer.Agents;
+using HordeServer.Commits;
 using HordeServer.Configuration;
 using HordeServer.Devices;
 using HordeServer.Issues;
@@ -208,6 +210,7 @@ namespace HordeServer.Notifications.Sinks
 
 		readonly IRedisService _redisService;
 		readonly IssueService _issueService;
+		readonly ICommitService _commitService;
 		readonly IUserCollection _userCollection;
 		readonly ILogCollection _logCollection;
 		readonly IWebHostEnvironment _environment;
@@ -242,10 +245,11 @@ namespace HordeServer.Notifications.Sinks
 		/// <summary>
 		/// Constructor
 		/// </summary>
-		public SlackNotificationSink(IMongoService mongoService, IRedisService redisService, IssueService issueService, IUserCollection userCollection, ILogCollection logCollection, IExternalIssueService externalIssueService, IWebHostEnvironment environment, IOptions<StaticBuildConfig> settings, IClock clock, IServerInfo serverInfo, IOptionsMonitor<BuildConfig> buildConfig, ILogger<SlackNotificationSink> logger)
+		public SlackNotificationSink(IMongoService mongoService, IRedisService redisService, IssueService issueService, ICommitService commitService, IUserCollection userCollection, ILogCollection logCollection, IExternalIssueService externalIssueService, IWebHostEnvironment environment, IOptions<StaticBuildConfig> settings, IClock clock, IServerInfo serverInfo, IOptionsMonitor<BuildConfig> buildConfig, ILogger<SlackNotificationSink> logger)
 		{
 			_redisService = redisService;
 			_issueService = issueService;
+			_commitService = commitService;
 			_userCollection = userCollection;
 			_logCollection = logCollection;
 			_externalIssueService = externalIssueService;
@@ -581,7 +585,7 @@ namespace HordeServer.Notifications.Sinks
 						autoSubmitMessage = $"\n\n```{job.AutoSubmitMessage}```";
 					}
 
-					attachment.AddSection($"Files in CL *{job.PreflightChange}* were *not submitted*. Please resolve these issues and submit manually.{autoSubmitMessage}");
+					attachment.AddSection($"Files in CL *{job.PreflightCommitId}* were *not submitted*. Please resolve these issues and submit manually.{autoSubmitMessage}");
 				}
 			}
 
@@ -903,16 +907,46 @@ namespace HordeServer.Notifications.Sinks
 			await UpdateReportsAsync(buildConfig, issue, details.Spans, cancellationToken);
 		}
 
-		static IIssueSpan? GetFixFailedSpan(IIssue issue, IReadOnlyList<IIssueSpan> spans)
+		async Task<IIssueStep?> GetFixFailedStepAsync(IIssue issue, IEnumerable<IIssueSpan> spans, CancellationToken cancellationToken)
 		{
-			if (issue.FixChange != null)
+			if (issue.FixCommitId != null)
+			{
+				foreach(IGrouping<StreamId, IIssueSpan> group in spans.GroupBy(x => x.StreamId))
+				{
+					IIssueStream? stream = issue.Streams.FirstOrDefault(x => x.StreamId == group.Key);
+					if (stream != null && (stream.ContainsFix ?? false))
+					{
+						CommitIdWithOrder fixedCommitIdWithOrder = await _commitService.GetOrderedAsync(group.Key, issue.FixCommitId, cancellationToken);
+						foreach (IIssueSpan span in group)
+						{
+							if (span.LastFailure.CommitId >= fixedCommitIdWithOrder)
+							{
+								return span.LastFailure;
+							}
+						}
+					}
+				}
+			}
+			return null;
+		}
+
+		async Task<IIssueSpan?> GetFixFailedSpanAsync(IIssue issue, IReadOnlyList<IIssueSpan> spans, CancellationToken cancellationToken)
+		{
+			if (issue.FixCommitId != null)
 			{
 				HashSet<StreamId> originStreams = new HashSet<StreamId>(issue.Streams.Where(x => x.MergeOrigin ?? false).Select(x => x.StreamId));
-				foreach (IIssueSpan originSpan in spans.Where(x => originStreams.Contains(x.StreamId)).OrderBy(x => x.StreamId.ToString()))
+				foreach(IGrouping<StreamId, IIssueSpan> group in spans.GroupBy(x => x.StreamId).OrderBy(x => x.Key.ToString()))
 				{
-					if (originSpan.LastFailure.Change > issue.FixChange.Value)
+					if (originStreams.Contains(group.Key))
 					{
-						return originSpan;
+						CommitIdWithOrder fixedCommitIdWithOrder = await _commitService.GetOrderedAsync(group.Key, issue.FixCommitId, cancellationToken);
+						foreach (IIssueSpan span in group)
+						{
+							if (span.LastFailure.CommitId > fixedCommitIdWithOrder)
+							{
+								return span;
+							}
+						}
 					}
 				}
 			}
@@ -1075,9 +1109,9 @@ namespace HordeServer.Notifications.Sinks
 			}
 
 			IIssueSpan? fixFailedSpan = null;
-			if (issue.FixChange != null)
+			if (issue.FixCommitId != null)
 			{
-				fixFailedSpan = GetFixFailedSpan(issue, spans);
+				fixFailedSpan = await GetFixFailedSpanAsync(issue, spans, cancellationToken);
 			}
 
 			if (!isSecondThread)
@@ -1155,7 +1189,7 @@ namespace HordeServer.Notifications.Sinks
 					{
 						string mention = await FormatMentionAsync(issue.OwnerId.Value, workflow.AllowMentions, cancellationToken);
 
-						string changes = String.Join(", ", suspects.Where(x => x.AuthorId == issue.OwnerId).Select(x => FormatChange(x.Change)));
+						string changes = String.Join(", ", suspects.Where(x => x.AuthorId == issue.OwnerId).Select(x => FormatChange(x.CommitId)));
 						if (changes.Length > 0)
 						{
 							mention += $" ({changes})";
@@ -1173,7 +1207,7 @@ namespace HordeServer.Notifications.Sinks
 							foreach (IGrouping<UserId, IIssueSuspect> suspectGroup in suspectGroups)
 							{
 								string mention = await FormatMentionAsync(suspectGroup.Key, workflow.AllowMentions, cancellationToken);
-								string changes = String.Join(", ", suspectGroup.Select(x => FormatChange(x.Change)));
+								string changes = String.Join(", ", suspectGroup.Select(x => FormatChange(x.CommitId)));
 								suspectList.Add($"{mention} ({changes})");
 								inviteUserIds.Add(suspectGroup.Key);
 							}
@@ -1231,18 +1265,18 @@ namespace HordeServer.Notifications.Sinks
 					}
 				}
 
-				if (issue.FixChange != null)
+				if (issue.FixCommitId != null)
 				{
 					if (fixFailedSpan == null)
 					{
-						string fixedEventId = $"issue_{issue.Id}_fixed_{issue.FixChange}";
-						string fixedMessage = $"Marked as fixed in {FormatChange(issue.FixChange.Value)}";
+						string fixedEventId = $"issue_{issue.Id}_fixed_{issue.FixCommitId}";
+						string fixedMessage = $"Marked as fixed in {FormatChange(issue.FixCommitId)}";
 						await PostSingleMessageToThreadAsync(triageChannel, fixedEventId, threadId, fixedMessage, cancellationToken);
 					}
 					else
 					{
-						string fixFailedEventId = $"issue_{issue.Id}_fixfailed_{issue.FixChange}";
-						string fixFailedMessage = $"Issue not fixed by {FormatChange(issue.FixChange.Value)}; see {FormatJobStep(fixFailedSpan.LastFailure, fixFailedSpan.NodeName)} at CL {fixFailedSpan.LastFailure.Change} in {fixFailedSpan.StreamName}.";
+						string fixFailedEventId = $"issue_{issue.Id}_fixfailed_{issue.FixCommitId}";
+						string fixFailedMessage = $"Issue not fixed by {FormatChange(issue.FixCommitId)}; see {FormatJobStep(fixFailedSpan.LastFailure, fixFailedSpan.NodeName)} at CL {fixFailedSpan.LastFailure.CommitId} in {fixFailedSpan.StreamName}.";
 						if (issue.OwnerId.HasValue)
 						{
 							string mention = await FormatMentionAsync(issue.OwnerId.Value, workflow.AllowMentions, cancellationToken);
@@ -1258,14 +1292,14 @@ namespace HordeServer.Notifications.Sinks
 							if ((stream.MergeOrigin ?? false) && !(stream.ContainsFix ?? false))
 							{
 								string streamName = spans.FirstOrDefault(x => x.StreamId == stream.StreamId)?.StreamName ?? stream.StreamId.ToString();
-								string missingEventId = $"issue_{issue.Id}_fixmissing_{issue.FixChange}_{stream.StreamId}";
+								string missingEventId = $"issue_{issue.Id}_fixmissing_{issue.FixCommitId}_{stream.StreamId}";
 								string missingMessage = $"Note: Fix may need manually merging to {streamName}";
 								await PostSingleMessageToThreadAsync(triageChannel, missingEventId, threadId, missingMessage, cancellationToken);
 							}
 						}
 					}
 				}
-				else if (issue.FixedSystemic)
+				else if (issue.FixSystemic)
 				{
 					string fixedEventId = $"issue_{issue.Id}_fixed_systemic";
 					string fixedMessage = $"Marked fixed as a systemic issue";
@@ -1376,7 +1410,7 @@ namespace HordeServer.Notifications.Sinks
 			}
 			else
 			{
-				return new Uri(_settings.P4SwarmUrl, $"files/{span.StreamName.TrimStart('/')}?range=@{span.LastSuccess.Change + 1},@{span.FirstFailure.Change}#commits");
+				return new Uri(_settings.P4SwarmUrl, $"files/{span.StreamName.TrimStart('/')}?range=@{span.LastSuccess.CommitId.GetPerforceChange() + 1},@{span.FirstFailure.CommitId}#commits");
 			}
 		}
 
@@ -1474,23 +1508,23 @@ namespace HordeServer.Notifications.Sinks
 				}
 			}
 
-			if (issue.FixChange != null)
+			if (issue.FixCommitId != null)
 			{
-				IIssueStep? fixFailedStep = issue.FindFixFailedStep(details.Spans);
+				IIssueStep? fixFailedStep = await GetFixFailedStepAsync(issue, details.Spans, cancellationToken);
 
 				string text;
 				if (fixFailedStep != null)
 				{
 					Uri fixFailedUrl = new Uri(_serverInfo.DashboardUrl, $"job/{fixFailedStep.JobId}?step={fixFailedStep.StepId}&issue={issue.Id}");
-					text = $":cross: Marked fixed in *CL {issue.FixChange.Value}*, but seen again at *<{fixFailedUrl}|CL {fixFailedStep.Change}>*";
+					text = $":cross: Marked fixed in *CL {issue.FixCommitId}*, but seen again at *<{fixFailedUrl}|CL {fixFailedStep.CommitId}>*";
 				}
 				else
 				{
-					text = $":tick: Marked fixed in *CL {issue.FixChange.Value}*.";
+					text = $":tick: Marked fixed in *CL {issue.FixCommitId}*.";
 				}
 				attachment.AddSection(text);
 			}
-			else if (issue.FixedSystemic)
+			else if (issue.FixSystemic)
 			{
 				attachment.AddSection(":tick: Marked as a systemic issue.");
 			}
@@ -1510,7 +1544,7 @@ namespace HordeServer.Notifications.Sinks
 					}
 					else
 					{
-						List<int> changes = details.Suspects.Where(x => x.AuthorId == userId).Select(x => x.Change).OrderBy(x => x).ToList();
+						List<int> changes = details.Suspects.Where(x => x.AuthorId == userId).Select(x => x.CommitId.GetPerforceChange()).OrderBy(x => x).ToList();
 						if (changes.Count > 0)
 						{
 							string text = $"Horde has determined that {StringUtils.FormatList(changes.Select(x => $"CL {x}"), "or")} is the most likely cause for this issue.";
@@ -1570,15 +1604,15 @@ namespace HordeServer.Notifications.Sinks
 				{
 					if (!details.Issue.Promoted)
 					{
-						declinedLines.Add($"Possibly {await FormatNameAsync(suspect.AuthorId, cancellationToken)} (CL {suspect.Change})");
+						declinedLines.Add($"Possibly {await FormatNameAsync(suspect.AuthorId, cancellationToken)} (CL {suspect.CommitId})");
 					}
 					else if (suspect.DeclinedAt == null)
 					{
-						declinedLines.Add($":heavy_minus_sign: Ignored by {await FormatNameAsync(suspect.AuthorId, cancellationToken)} (CL {suspect.Change})");
+						declinedLines.Add($":heavy_minus_sign: Ignored by {await FormatNameAsync(suspect.AuthorId, cancellationToken)} (CL {suspect.CommitId})");
 					}
 					else
 					{
-						declinedLines.Add($":downvote: Declined by {await FormatNameAsync(suspect.AuthorId, cancellationToken)} at {FormatSlackTime(suspect.DeclinedAt.Value)} (CL {suspect.Change})");
+						declinedLines.Add($":downvote: Declined by {await FormatNameAsync(suspect.AuthorId, cancellationToken)} at {FormatSlackTime(suspect.DeclinedAt.Value)} (CL {suspect.CommitId})");
 					}
 				}
 				attachment.AddSection(String.Join("\n", declinedLines));
@@ -1630,7 +1664,7 @@ namespace HordeServer.Notifications.Sinks
 			}
 		}
 
-		string FormatChange(int change)
+		string FormatChange(CommitId change)
 		{
 			if (_settings.P4SwarmUrl != null)
 			{
@@ -2009,7 +2043,7 @@ namespace HordeServer.Notifications.Sinks
 			}
 
 			string status = "*Unassigned*";
-			if (issue.FixChange != null)
+			if (issue.FixCommitId != null)
 			{
 				if (issue.Streams.Any(x => x.FixFailed ?? false))
 				{
@@ -2021,18 +2055,18 @@ namespace HordeServer.Notifications.Sinks
 				}
 				else
 				{
-					status = $"Fixed in CL {issue.FixChange.Value}";
+					status = $"Fixed in CL {issue.FixCommitId}";
 				}
 			}
-			else if (issue.FixedSystemic)
+			else if (issue.FixSystemic)
 			{
-				status = "Closed as systemic issue";
+				status = "Marked fixed as systemic";
 			}
 			else if (issue.OwnerId != null)
 			{
 				if (owner == null)
 				{
-					status = $"Assigned to user {issue.OwnerId.Value}";
+					status = $"Assigned to {GetDefaultUserName(issue.OwnerId.Value)}";
 				}
 				else
 				{
@@ -2497,13 +2531,13 @@ namespace HordeServer.Notifications.Sinks
 
 		static string GetJobChangeText(IJob job)
 		{
-			if (job.PreflightChange == 0)
+			if (job.PreflightCommitId == null)
 			{
-				return $"CL {job.Change}";
+				return $"CL {job.CommitId}";
 			}
 			else
 			{
-				return $"Preflight CL {job.PreflightChange} against CL {job.Change}";
+				return $"Preflight CL {job.PreflightCommitId} against CL {job.CommitId}";
 			}
 		}
 
@@ -2728,7 +2762,7 @@ namespace HordeServer.Notifications.Sinks
 			{
 				return null;
 			}
-			if (issue.FixChange != null && GetFixFailedSpan(issue, spans) == null)
+			if (issue.IsMarkedFixed() && await GetFixFailedSpanAsync(issue, spans, cancellationToken) == null)
 			{
 				return null;
 			}
@@ -3020,16 +3054,9 @@ namespace HordeServer.Notifications.Sinks
 							string? fixChangeStr;
 							if (payload.View.State.TryGetValue("fix_cl", "fix_cl_action", out fixChangeStr))
 							{
-								int fixChange;
-								if (!Int32.TryParse(fixChangeStr, out fixChange) || fixChange < 0)
-								{
-									Dictionary<string, string> errors = new Dictionary<string, string>();
-									errors.Add("fix_cl", $"'{fixChangeStr}' is not a valid fix changelist.");
-									return new { response_action = "errors", errors };
-								}
-
-								await _issueService.UpdateIssueAsync(issueId, fixChange: fixChange, resolvedById: resolvedById, initiatedById: userId, cancellationToken: cancellationToken);
-								_logger.LogInformation("Marked issue {IssueId} fixed by user {UserId} in {Change}", issueId, resolvedById, fixChange);
+								CommitId commitId = new CommitId(fixChangeStr);
+								await _issueService.UpdateIssueAsync(issueId, fixCommitId: commitId, resolvedById: resolvedById, initiatedById: userId, cancellationToken: cancellationToken);
+								_logger.LogInformation("Marked issue {IssueId} fixed by user {UserId} in {Change}", issueId, resolvedById, commitId);
 							}
 						}
 						else if (TryMatch(payload.View.CallbackId, @"^issue_(\d+)_ack_([a-fA-F0-9]{24})$", out match))
diff --git a/Engine/Source/Programs/Horde/Plugins/Build/HordeServer.Build/Perforce/PerforceService.cs b/Engine/Source/Programs/Horde/Plugins/Build/HordeServer.Build/Perforce/PerforceService.cs
--- a/Engine/Source/Programs/Horde/Plugins/Build/HordeServer.Build/Perforce/PerforceService.cs
+++ b/Engine/Source/Programs/Horde/Plugins/Build/HordeServer.Build/Perforce/PerforceService.cs
@@ -145,8 +145,13 @@ namespace HordeServer.Perforce
 			readonly StreamConfig _streamConfig;
 
 			public StreamId StreamId => _streamConfig.Id;
+
 			public int Number { get; }
+			CommitIdWithOrder ICommit.Id => CommitIdWithOrder.FromPerforceChange(Number);
+
 			public int OriginalChange { get; }
+			CommitIdWithOrder ICommit.OriginalCommitId => CommitIdWithOrder.FromPerforceChange(OriginalChange);
+
 			public UserId AuthorId { get; }
 			public UserId OwnerId { get; }
 			public string Description { get; }
@@ -1180,21 +1185,35 @@ namespace HordeServer.Perforce
 				Logger = logger;
 			}
 
-			public async Task<int> CreateNewAsync(string path, string description, CancellationToken cancellationToken = default)
+			public async Task<CommitIdWithOrder> CreateNewAsync(string path, string description, CancellationToken cancellationToken = default)
 			{
 				Match match = Regex.Match(path, @"^(//[^/]+/[^/]+)/(.+)$");
 				if (match.Success)
 				{
-					return await PerforceService.CreateNewChangeAsync(StreamConfig.ClusterName, match.Groups[1].Value, match.Groups[2].Value, description, cancellationToken);
+					return CommitIdWithOrder.FromPerforceChange(await PerforceService.CreateNewChangeAsync(StreamConfig.ClusterName, match.Groups[1].Value, match.Groups[2].Value, description, cancellationToken));
 				}
 				else
 				{
-					return await PerforceService.CreateNewChangeAsync(StreamConfig.ClusterName, StreamConfig.Name, path, description, cancellationToken);
+					return CommitIdWithOrder.FromPerforceChange(await PerforceService.CreateNewChangeAsync(StreamConfig.ClusterName, StreamConfig.Name, path, description, cancellationToken));
 				}
 			}
 
-			public virtual async IAsyncEnumerable<ICommit> FindAsync(int? minChange, int? maxChange, int? maxResults, IReadOnlyList<CommitTag>? tags, [EnumeratorCancellation] CancellationToken cancellationToken = default)
+			public IAsyncEnumerable<ICommit> FindAsync(CommitId? minChange = null, bool includeMinChange = true, CommitId? maxChange = null, bool includeMaxChange = true, int? maxResults = null, IReadOnlyList<CommitTag>? tags = null, CancellationToken cancellationToken = default)
+			{
+				return FindAsync(minChange?.GetPerforceChange(), includeMinChange, maxChange?.GetPerforceChange(), includeMaxChange, maxResults, tags, cancellationToken);
+			}
+
+			public virtual async IAsyncEnumerable<ICommit> FindAsync(int? minChange, bool includeMinChange, int? maxChange, bool includeMaxChange, int? maxResults, IReadOnlyList<CommitTag>? tags, [EnumeratorCancellation] CancellationToken cancellationToken = default)
 			{
+				if (minChange != null && minChange.Value > 0 && !includeMinChange)
+				{
+					minChange = minChange.Value + 1;
+				}
+				if (maxChange != null && maxChange.Value > 0 && !includeMaxChange)
+				{
+					maxChange = maxChange.Value - 1;
+				}
+
 				using TelemetrySpan span = PerforceService._tracer.StartActiveSpan($"{nameof(Perforce.PerforceService)}.{nameof(PerforceService.CommitSource)}.{nameof(FindAsync)}");
 				span.SetAttribute("stream", StreamConfig.ClusterName);
 				span.SetAttribute("minChange", minChange ?? -2);
@@ -1256,21 +1275,31 @@ namespace HordeServer.Perforce
 				}
 			}
 
+			/// <inheritdoc/>
+			public Task<ICommit> GetAsync(CommitId changeNumber, CancellationToken cancellationToken = default)
+				=> GetAsync(changeNumber.GetPerforceChange(), cancellationToken);
+
 			/// <inheritdoc/>
 			public virtual Task<ICommit> GetAsync(int changeNumber, CancellationToken cancellationToken = default)
-			{
-				return PerforceService.GetChangeDetailsAsync(StreamConfig, changeNumber, cancellationToken);
-			}
+				=> PerforceService.GetChangeDetailsAsync(StreamConfig, changeNumber, cancellationToken);
+
+			/// <inheritdoc/>
+			public ValueTask<CommitIdWithOrder> GetOrderedAsync(CommitId commit, CancellationToken cancellationToken = default)
+				=> ValueTask.FromResult(new CommitIdWithOrder(commit.Name, commit.GetPerforceChange()));
+
+			/// <inheritdoc/>
+			public IAsyncEnumerable<ICommit> SubscribeAsync(CommitId minCommit, IReadOnlyList<CommitTag>? tags = null, CancellationToken cancellationToken = default)
+				=> SubscribeAsync(minCommit.GetPerforceChange(), tags, cancellationToken);
 
 			/// <inheritdoc/>
 			public virtual async IAsyncEnumerable<ICommit> SubscribeAsync(int minChange, IReadOnlyList<CommitTag>? tags = null, [EnumeratorCancellation] CancellationToken cancellationToken = default)
 			{
 				for (; ; )
 				{
-					await foreach (ICommit commit in FindAsync(minChange + 1, null, 10, tags, cancellationToken))
+					await foreach (ICommit commit in FindAsync(minChange: CommitId.FromPerforceChange(minChange), includeMinChange: false, maxResults: 10, tags: tags, cancellationToken: cancellationToken))
 					{
 						yield return commit;
-						minChange = commit.Number;
+						minChange = commit.Id.GetPerforceChange();
 					}
 					await Task.Delay(TimeSpan.FromSeconds(10.0), cancellationToken);
 				}
diff --git a/Engine/Source/Programs/Horde/Plugins/Build/HordeServer.Build/Perforce/PerforceServiceCache.cs b/Engine/Source/Programs/Horde/Plugins/Build/HordeServer.Build/Perforce/PerforceServiceCache.cs
--- a/Engine/Source/Programs/Horde/Plugins/Build/HordeServer.Build/Perforce/PerforceServiceCache.cs
+++ b/Engine/Source/Programs/Horde/Plugins/Build/HordeServer.Build/Perforce/PerforceServiceCache.cs
@@ -138,8 +138,13 @@ namespace HordeServer.Perforce
 			public ObjectId Id { get; set; }
 
 			public StreamId StreamId { get; set; }
+
 			public int Number { get; set; }
+			CommitIdWithOrder ICommit.Id => CommitIdWithOrder.FromPerforceChange(Number);
+
 			public int OriginalChange { get; set; }
+			CommitIdWithOrder ICommit.OriginalCommitId => CommitIdWithOrder.FromPerforceChange(OriginalChange);
+
 			public UserId AuthorId { get; set; }
 			public UserId OwnerId { get; set; }
 			public string Description { get; set; }
@@ -160,8 +165,8 @@ namespace HordeServer.Perforce
 			public CachedCommitDoc(ICommit commit, List<CommitTag> commitTags)
 			{
 				StreamId = commit.StreamId;
-				Number = commit.Number;
-				OriginalChange = commit.OriginalChange;
+				Number = commit.Id.GetPerforceChange();
+				OriginalChange = commit.OriginalCommitId.GetPerforceChange();
 				AuthorId = commit.AuthorId;
 				OwnerId = commit.OwnerId;
 				Description = commit.Description;
@@ -610,8 +615,17 @@ namespace HordeServer.Perforce
 				_owner = owner;
 			}
 
-			public override async IAsyncEnumerable<ICommit> FindAsync(int? minChange, int? maxChange, int? maxResults, IReadOnlyList<CommitTag>? tags, [EnumeratorCancellation] CancellationToken cancellationToken = default)
+			public override async IAsyncEnumerable<ICommit> FindAsync(int? minChange = null, bool includeMinChange = true, int? maxChange = null, bool includeMaxChange = true, int? maxResults = null, IReadOnlyList<CommitTag>? tags = null, [EnumeratorCancellation] CancellationToken cancellationToken = default)
 			{
+				if (minChange != null && !includeMinChange)
+				{
+					minChange = minChange.Value + 1;
+				}
+				if (maxChange != null && !includeMaxChange)
+				{
+					maxChange = maxChange.Value - 1;
+				}
+
 				if (minChange != null && maxChange != null && minChange.Value > maxChange.Value)
 				{
 					yield break;
@@ -624,7 +638,7 @@ namespace HordeServer.Perforce
 				if (state.Clusters.TryGetValue(StreamConfig.ClusterName, out ClusterState? clusterState))
 				{
 					int minReplicatedChange;
-					if (clusterState.MinChanges.TryGetValue(StreamConfig.Id, out minReplicatedChange) && (maxChange == null || maxChange > minReplicatedChange))
+					if (clusterState.MinChanges.TryGetValue(StreamConfig.Id, out minReplicatedChange) && (maxChange == null || maxChange.Value > minReplicatedChange))
 					{
 						FilterDefinition<CachedCommitDoc> filter = Builders<CachedCommitDoc>.Filter.Eq(x => x.StreamId, StreamConfig.Id);
 
@@ -688,12 +702,12 @@ namespace HordeServer.Perforce
 						// Expand the range of cached changes if necessary
 						if ((maxResults == null || maxResults.Value > 0) && maxChange > 0)
 						{
-							await foreach (ICommit commit in base.FindAsync(minChange, maxChange, maxResults, null, cancellationToken))
+							await foreach (ICommit commit in base.FindAsync(minChange, true, maxChange, true, maxResults, null, cancellationToken))
 							{
 								CachedCommitDoc cachedCommit = await CachedCommitDoc.FromCommitAsync(commit, cancellationToken);
 								await _owner.AddCachedCommitAsync(cachedCommit, cancellationToken);
-								await _owner._mongoService.UpdateSingletonAsync<CacheState>(x => TryUpdateRange(x, StreamConfig, commit.Number, maxChange), cancellationToken);
-								_owner._logger.LogDebug("Adding new cached commit for {StreamId} at change {Change}", StreamConfig.Id, commit.Number);
+								await _owner._mongoService.UpdateSingletonAsync<CacheState>(x => TryUpdateRange(x, StreamConfig, commit.Id.GetPerforceChange(), maxChange), cancellationToken);
+								_owner._logger.LogDebug("Adding new cached commit for {StreamId} at change {Change}", StreamConfig.Id, commit.Id);
 
 								if (tags == null || tags.Any(x => cachedCommit.CommitTags.Contains(x)))
 								{
@@ -717,7 +731,7 @@ namespace HordeServer.Perforce
 				{
 					_owner._logger.LogDebug("Querying Perforce server for {StreamId} commits from {MinChange} to {MaxChange} (max: {MaxResults}, tags: {Tags})", StreamConfig.Id, minChange ?? -2, maxChange ?? -2, maxResults ?? -1, (tags == null || tags.Count == 0) ? "none" : String.Join("/", tags.Select(x => x.ToString())));
 
-					await foreach (ICommit commit in base.FindAsync(minChange, maxChange, maxResults, tags, cancellationToken))
+					await foreach (ICommit commit in base.FindAsync(minChange, true, maxChange, true, maxResults, tags, cancellationToken))
 					{
 						yield return commit;
 					}
@@ -776,10 +790,10 @@ namespace HordeServer.Perforce
 						Task task = updateEvent.Task;
 
 						int numResults = 10;
-						await foreach (ICommit commit in FindAsync(minChange + 1, null, numResults, tags, cancellationToken))
+						await foreach (ICommit commit in FindAsync(minChange, includeMinChange: false, maxResults: numResults, tags: tags, cancellationToken: cancellationToken))
 						{
 							yield return commit;
-							minChange = commit.Number;
+							minChange = commit.Id.GetPerforceChange();
 							numResults--;
 						}
 
diff --git a/Engine/Source/Programs/Horde/Plugins/Build/HordeServer.Build/Replicators/PerforceReplicator.cs b/Engine/Source/Programs/Horde/Plugins/Build/HordeServer.Build/Replicators/PerforceReplicator.cs
--- a/Engine/Source/Programs/Horde/Plugins/Build/HordeServer.Build/Replicators/PerforceReplicator.cs
+++ b/Engine/Source/Programs/Horde/Plugins/Build/HordeServer.Build/Replicators/PerforceReplicator.cs
@@ -5,6 +5,7 @@ using System.Diagnostics;
 using System.Security.Cryptography;
 using System.Text;
 using EpicGames.Core;
+using EpicGames.Horde.Commits;
 using EpicGames.Horde.Replicators;
 using EpicGames.Horde.Storage;
 using EpicGames.Horde.Storage.Bundles;
@@ -372,18 +373,18 @@ namespace HordeServer.Replicators
 				ICommit? nextCommit;
 				if (replicator.LastChange != null)
 				{
-					nextCommit = await commits.SubscribeAsync(replicator.LastChange.Value, cancellationToken: cancellationToken).FirstAsync(cancellationToken);
-					_logger.LogInformation("Replicating next change for {ReplicatorId} at CL {Change}", replicator.Id, nextCommit.Number);
+					nextCommit = await commits.SubscribeAsync(CommitIdWithOrder.FromPerforceChange(replicator.LastChange.Value), cancellationToken: cancellationToken).FirstAsync(cancellationToken);
+					_logger.LogInformation("Replicating next change for {ReplicatorId} at CL {Change}", replicator.Id, nextCommit.Id);
 				}
 				else
 				{
 					nextCommit = await commits.GetLatestAsync(cancellationToken);
-					_logger.LogInformation("Starting {ReplicatorId} replication from latest CL {Change}", replicator.Id, nextCommit.Number);
+					_logger.LogInformation("Starting {ReplicatorId} replication from latest CL {Change}", replicator.Id, nextCommit.Id);
 				}
 
 				await store.DeleteRefAsync(incRefName, cancellationToken);
 
-				UpdateReplicatorOptions updateOptions = new UpdateReplicatorOptions { CurrentChange = nextCommit.Number };
+				UpdateReplicatorOptions updateOptions = new UpdateReplicatorOptions { CurrentChange = nextCommit.Id.GetPerforceChange() };
 				replicator = await UpdateReplicatorAsync(replicator, updateOptions, cancellationToken);
 			}
 
diff --git a/Engine/Source/Programs/Horde/Plugins/Build/HordeServer.Build/Streams/IStream.cs b/Engine/Source/Programs/Horde/Plugins/Build/HordeServer.Build/Streams/IStream.cs
--- a/Engine/Source/Programs/Horde/Plugins/Build/HordeServer.Build/Streams/IStream.cs
+++ b/Engine/Source/Programs/Horde/Plugins/Build/HordeServer.Build/Streams/IStream.cs
@@ -1,5 +1,6 @@
 // Copyright Epic Games, Inc. All Rights Reserved.
 
+using EpicGames.Horde.Commits;
 using EpicGames.Horde.Jobs;
 using EpicGames.Horde.Jobs.Templates;
 using EpicGames.Horde.Streams;
@@ -99,7 +100,7 @@ namespace HordeServer.Streams
 		/// <summary>
 		/// Last changelist number that this was triggered for
 		/// </summary>
-		int LastTriggerChange { get; }
+		CommitIdWithOrder? LastTriggerCommitId { get; }
 
 		/// <summary>
 		/// Gets the last trigger time, in UTC
diff --git a/Engine/Source/Programs/Horde/Plugins/Build/HordeServer.Build/Streams/IStreamCollection.cs b/Engine/Source/Programs/Horde/Plugins/Build/HordeServer.Build/Streams/IStreamCollection.cs
--- a/Engine/Source/Programs/Horde/Plugins/Build/HordeServer.Build/Streams/IStreamCollection.cs
+++ b/Engine/Source/Programs/Horde/Plugins/Build/HordeServer.Build/Streams/IStreamCollection.cs
@@ -1,5 +1,6 @@
 // Copyright Epic Games, Inc. All Rights Reserved.
 
+using EpicGames.Horde.Commits;
 using EpicGames.Horde.Jobs;
 using EpicGames.Horde.Jobs.Templates;
 
@@ -42,11 +43,11 @@ namespace HordeServer.Streams
 		/// <param name="stream">The stream to update</param>
 		/// <param name="templateRefId">The template ref id</param>
 		/// <param name="lastTriggerTimeUtc">New last trigger time for the schedule</param>
-		/// <param name="lastTriggerChange">New last trigger changelist for the schedule</param>
+		/// <param name="lastTriggerCommitId">New last trigger commit for the schedule</param>
 		/// <param name="newActiveJobs">New list of active jobs</param>
 		/// <param name="cancellationToken">Cancellation token for the operation</param>
 		/// <returns>The updated stream if successful, null otherwise</returns>
-		Task<IStream?> TryUpdateScheduleTriggerAsync(IStream stream, TemplateId templateRefId, DateTime? lastTriggerTimeUtc, int? lastTriggerChange, List<JobId> newActiveJobs, CancellationToken cancellationToken = default);
+		Task<IStream?> TryUpdateScheduleTriggerAsync(IStream stream, TemplateId templateRefId, DateTime? lastTriggerTimeUtc, CommitIdWithOrder? lastTriggerCommitId, List<JobId> newActiveJobs, CancellationToken cancellationToken = default);
 
 		/// <summary>
 		/// Attempts to update a stream template ref
@@ -93,12 +94,12 @@ namespace HordeServer.Streams
 		/// <param name="stream">The stream to update</param>
 		/// <param name="templateRefId">The template ref id</param>
 		/// <param name="lastTriggerTimeUtc"></param>
-		/// <param name="lastTriggerChange"></param>
+		/// <param name="lastTriggerCommitId"></param>
 		/// <param name="addJobs">Jobs to add</param>
 		/// <param name="removeJobs">Jobs to remove</param>
 		/// <param name="cancellationToken">Cancellation token for the operation</param>
 		/// <returns>True if the stream was updated</returns>
-		public static async Task<IStream?> UpdateScheduleTriggerAsync(this IStreamCollection streamCollection, IStream stream, TemplateId templateRefId, DateTime? lastTriggerTimeUtc = null, int? lastTriggerChange = null, List<JobId>? addJobs = null, List<JobId>? removeJobs = null, CancellationToken cancellationToken = default)
+		public static async Task<IStream?> UpdateScheduleTriggerAsync(this IStreamCollection streamCollection, IStream stream, TemplateId templateRefId, DateTime? lastTriggerTimeUtc = null, CommitIdWithOrder? lastTriggerCommitId = null, List<JobId>? addJobs = null, List<JobId>? removeJobs = null, CancellationToken cancellationToken = default)
 		{
 			IStream? newStream = stream;
 			while (newStream != null)
@@ -123,7 +124,7 @@ namespace HordeServer.Streams
 					newActiveJobs = newActiveJobs.Union(addJobs);
 				}
 
-				newStream = await streamCollection.TryUpdateScheduleTriggerAsync(newStream, templateRefId, lastTriggerTimeUtc, lastTriggerChange, newActiveJobs.ToList(), cancellationToken);
+				newStream = await streamCollection.TryUpdateScheduleTriggerAsync(newStream, templateRefId, lastTriggerTimeUtc, lastTriggerCommitId, newActiveJobs.ToList(), cancellationToken);
 
 				if (newStream != null)
 				{
diff --git a/Engine/Source/Programs/Horde/Plugins/Build/HordeServer.Build/Streams/StreamCollection.cs b/Engine/Source/Programs/Horde/Plugins/Build/HordeServer.Build/Streams/StreamCollection.cs
--- a/Engine/Source/Programs/Horde/Plugins/Build/HordeServer.Build/Streams/StreamCollection.cs
+++ b/Engine/Source/Programs/Horde/Plugins/Build/HordeServer.Build/Streams/StreamCollection.cs
@@ -1,6 +1,7 @@
 // Copyright Epic Games, Inc. All Rights Reserved.
 
 using EpicGames.Core;
+using EpicGames.Horde.Commits;
 using EpicGames.Horde.Jobs;
 using EpicGames.Horde.Jobs.Templates;
 using EpicGames.Horde.Streams;
@@ -147,7 +148,16 @@ namespace HordeServer.Streams
 
 		class TemplateScheduleDoc : ITemplateSchedule
 		{
-			public int LastTriggerChange { get; set; }
+			public CommitIdWithOrder LastTriggerCommitId
+			{
+				get => (LastTriggerCommitName != null) ? new CommitIdWithOrder(LastTriggerCommitName, LastTriggerCommitOrder) : CommitIdWithOrder.FromPerforceChange(LastTriggerCommitOrder);
+				set => (LastTriggerCommitName, LastTriggerCommitOrder) = (value.Name, value.Order);
+			}
+
+			public string? LastTriggerCommitName { get; set; }
+
+			[BsonElement("LastTriggerChange")]
+			public int LastTriggerCommitOrder { get; set; }
 
 			[BsonIgnoreIfNull, Obsolete("Use LastTriggerTimeUtc instead")]
 			public DateTimeOffset? LastTriggerTime { get; set; }
@@ -382,7 +392,7 @@ namespace HordeServer.Streams
 		}
 
 		/// <inheritdoc/>
-		public async Task<IStream?> TryUpdateScheduleTriggerAsync(IStream streamInterface, TemplateId templateId, DateTime? lastTriggerTimeUtc, int? lastTriggerChange, List<JobId> newActiveJobs, CancellationToken cancellationToken)
+		public async Task<IStream?> TryUpdateScheduleTriggerAsync(IStream streamInterface, TemplateId templateId, DateTime? lastTriggerTimeUtc, CommitIdWithOrder? lastTriggerCommitId, List<JobId> newActiveJobs, CancellationToken cancellationToken)
 		{
 			StreamDoc stream = (StreamDoc)streamInterface;
 			TemplateRefDoc template = stream.Templates[templateId];
@@ -394,18 +404,24 @@ namespace HordeServer.Streams
 			{
 				FieldDefinition<StreamDoc, DateTime> lastTriggerTimeField = $"{nameof(stream.Templates)}.{templateId}.{nameof(template.Schedule)}.{nameof(schedule.LastTriggerTimeUtc)}";
 				updates.Add(Builders<StreamDoc>.Update.Set(lastTriggerTimeField, lastTriggerTimeUtc.Value));
+
 				schedule.LastTriggerTimeUtc = lastTriggerTimeUtc.Value;
 			}
-			if (lastTriggerChange.HasValue && lastTriggerChange.Value > schedule.LastTriggerChange)
+			if (lastTriggerCommitId != null && lastTriggerCommitId > schedule.LastTriggerCommitId)
 			{
-				FieldDefinition<StreamDoc, int> lastTriggerChangeField = $"{nameof(stream.Templates)}.{templateId}.{nameof(template.Schedule)}.{nameof(schedule.LastTriggerChange)}";
-				updates.Add(Builders<StreamDoc>.Update.Set(lastTriggerChangeField, lastTriggerChange.Value));
-				schedule.LastTriggerChange = lastTriggerChange.Value;
+				FieldDefinition<StreamDoc, string> lastTriggerCommitNameField = $"{nameof(stream.Templates)}.{templateId}.{nameof(template.Schedule)}.{nameof(schedule.LastTriggerCommitName)}";
+				updates.Add(Builders<StreamDoc>.Update.Set(lastTriggerCommitNameField, lastTriggerCommitId.Name));
+
+				FieldDefinition<StreamDoc, int> lastTriggerChangeField = $"{nameof(stream.Templates)}.{templateId}.{nameof(template.Schedule)}.LastTriggerChange";
+				updates.Add(Builders<StreamDoc>.Update.Set(lastTriggerChangeField, lastTriggerCommitId.Order));
+
+				schedule.LastTriggerCommitId = lastTriggerCommitId;
 			}
 			if (newActiveJobs != null)
 			{
 				FieldDefinition<StreamDoc, List<JobId>> field = $"{nameof(stream.Templates)}.{templateId}.{nameof(template.Schedule)}.{nameof(schedule.ActiveJobs)}";
 				updates.Add(Builders<StreamDoc>.Update.Set(field, newActiveJobs));
+
 				schedule.ActiveJobs = newActiveJobs;
 			}
 
diff --git a/Engine/Source/Programs/Horde/Plugins/Build/HordeServer.Build/Streams/StreamMessages.cs b/Engine/Source/Programs/Horde/Plugins/Build/HordeServer.Build/Streams/StreamMessages.cs
--- a/Engine/Source/Programs/Horde/Plugins/Build/HordeServer.Build/Streams/StreamMessages.cs
+++ b/Engine/Source/Programs/Horde/Plugins/Build/HordeServer.Build/Streams/StreamMessages.cs
@@ -2,6 +2,7 @@
 
 using System.ComponentModel.DataAnnotations;
 using EpicGames.Core;
+using EpicGames.Horde.Commits;
 using EpicGames.Horde.Jobs;
 using EpicGames.Horde.Jobs.Templates;
 using EpicGames.Horde.Projects;
@@ -478,8 +479,12 @@ namespace HordeServer.Streams
 	{
 		readonly ITemplateSchedule _schedule;
 
-		/// <inheritdoc cref="ITemplateSchedule.LastTriggerChange"/>
-		public int LastTriggerChange => _schedule.LastTriggerChange;
+		/// <inheritdoc cref="ITemplateSchedule.LastTriggerCommitId"/>
+		[Obsolete("Use LastTriggerCommit instead")]
+		public int LastTriggerChange => _schedule.LastTriggerCommitId?.GetPerforceChangeOrMinusOne() ?? 0;
+
+		/// <inheritdoc cref="ITemplateSchedule.LastTriggerCommitId"/>
+		public CommitIdWithOrder? LastTriggerCommit => _schedule.LastTriggerCommitId;
 
 		/// <inheritdoc cref="ITemplateSchedule.LastTriggerTimeUtc"/>
 		public DateTime LastTriggerTimeUtc => _schedule.LastTriggerTimeUtc;
diff --git a/Engine/Source/Programs/Horde/Plugins/Build/HordeServer.Build/Streams/StreamsController.cs b/Engine/Source/Programs/Horde/Plugins/Build/HordeServer.Build/Streams/StreamsController.cs
--- a/Engine/Source/Programs/Horde/Plugins/Build/HordeServer.Build/Streams/StreamsController.cs
+++ b/Engine/Source/Programs/Horde/Plugins/Build/HordeServer.Build/Streams/StreamsController.cs
@@ -247,7 +247,7 @@ namespace HordeServer.Streams
 		[HttpGet]
 		[Route("/api/v1/streams/{streamId}/changes")]
 		[ProducesResponseType(typeof(List<GetCommitResponse>), 200)]
-		public async Task<ActionResult<List<object>>> GetChangesAsync(StreamId streamId, [FromQuery] int? min = null, [FromQuery] int? max = null, [FromQuery] int results = 50, [FromQuery] string? tags = null, PropertyFilter? filter = null, CancellationToken cancellationToken = default)
+		public async Task<ActionResult<List<object>>> GetChangesAsync(StreamId streamId, [FromQuery] CommitId? min = null, [FromQuery] CommitId? max = null, [FromQuery] int results = 50, [FromQuery] string? tags = null, PropertyFilter? filter = null, CancellationToken cancellationToken = default)
 		{
 			StreamConfig? streamConfig;
 			if (!_buildConfig.Value.TryGetStream(streamId, out streamConfig))
@@ -278,7 +278,7 @@ namespace HordeServer.Streams
 
 		static GetCommitResponse CreateGetCommitResponse(ICommit commit, IUser author, IReadOnlyList<CommitTag>? tags, IReadOnlyList<string>? files)
 		{
-			GetCommitResponse response = new GetCommitResponse(commit.Number, author.ToThinApiResponse(), commit.Description);
+			GetCommitResponse response = new GetCommitResponse(author.ToThinApiResponse(), commit.Description) { Id = commit.Id };
 			if (tags != null)
 			{
 				response.Tags = new List<CommitTag>(tags);
@@ -302,7 +302,7 @@ namespace HordeServer.Streams
 		[HttpGet]
 		[Route("/api/v1/streams/{streamId}/changes/{changeNumber}")]
 		[ProducesResponseType(typeof(GetCommitResponse), 200)]
-		public async Task<ActionResult<object>> GetChangeDetailsAsync(StreamId streamId, int changeNumber, int maxFiles = 100, PropertyFilter? filter = null, CancellationToken cancellationToken = default)
+		public async Task<ActionResult<object>> GetChangeDetailsAsync(StreamId streamId, CommitId changeNumber, int maxFiles = 100, PropertyFilter? filter = null, CancellationToken cancellationToken = default)
 		{
 			StreamConfig? streamConfig;
 			if (!_buildConfig.Value.TryGetStream(streamId, out streamConfig))
@@ -340,7 +340,7 @@ namespace HordeServer.Streams
 		[HttpGet]
 		[Route("/api/v1/streams/{streamId}/history")]
 		[ProducesResponseType(typeof(List<GetJobStepRefResponse>), 200)]
-		public async Task<ActionResult<List<object>>> GetStepHistoryAsync(StreamId streamId, [FromQuery] string templateId, [FromQuery] string step, [FromQuery] int? change = null, [FromQuery] int count = 10, [FromQuery] PropertyFilter? filter = null)
+		public async Task<ActionResult<List<object>>> GetStepHistoryAsync(StreamId streamId, [FromQuery] string templateId, [FromQuery] string step, [FromQuery] CommitId? change = null, [FromQuery] int count = 10, [FromQuery] PropertyFilter? filter = null)
 		{
 			StreamConfig? streamConfig;
 			if (!_buildConfig.Value.TryGetStream(streamId, out streamConfig))
@@ -354,7 +354,13 @@ namespace HordeServer.Streams
 
 			TemplateId templateIdValue = new TemplateId(templateId);
 
-			List<IJobStepRef> steps = await _jobStepRefCollection.GetStepsForNodeAsync(streamId, templateIdValue, step, change, true, count);
+			CommitIdWithOrder? commitIdWithOrder = null;
+			if (change != null)
+			{
+				commitIdWithOrder = await _commitService.GetOrderedAsync(streamId, change, HttpContext.RequestAborted);
+			}
+
+			List<IJobStepRef> steps = await _jobStepRefCollection.GetStepsForNodeAsync(streamId, templateIdValue, step, commitIdWithOrder, true, count, HttpContext.RequestAborted);
 			return steps.ConvertAll(x => PropertyFilter.Apply(CreateGetJobStepRefResponse(x), filter));
 		}
 
@@ -368,7 +374,7 @@ namespace HordeServer.Streams
 			response.JobId = jobStepRef.Id.JobId;
 			response.BatchId = jobStepRef.Id.BatchId;
 			response.StepId = jobStepRef.Id.StepId;
-			response.Change = jobStepRef.Change;
+			response.CommitId = jobStepRef.CommitId;
 			response.LogId = jobStepRef.LogId.ToString();
 			response.PoolId = jobStepRef.PoolId?.ToString();
 			response.AgentId = jobStepRef.AgentId?.ToString();
diff --git a/Engine/Source/Programs/Horde/Plugins/Build/HordeServer.Build/Ugs/UgsController.cs b/Engine/Source/Programs/Horde/Plugins/Build/HordeServer.Build/Ugs/UgsController.cs
--- a/Engine/Source/Programs/Horde/Plugins/Build/HordeServer.Build/Ugs/UgsController.cs
+++ b/Engine/Source/Programs/Horde/Plugins/Build/HordeServer.Build/Ugs/UgsController.cs
@@ -1,5 +1,6 @@
 // Copyright Epic Games, Inc. All Rights Reserved.
 
+using EpicGames.Horde.Commits;
 using EpicGames.Horde.Issues;
 using EpicGames.Horde.Users;
 using HordeServer.Issues;
@@ -283,7 +284,7 @@ namespace HordeServer.Ugs
 		{
 			HashSet<ObjectId> unresolvedSpans = new HashSet<ObjectId>(issue.Spans.Where(x => x.NextSuccess == null).Select(x => x.Id));
 
-			IIssueStep? step = issue.Steps.OrderByDescending(x => unresolvedSpans.Contains(x.SpanId)).ThenByDescending(x => x.Change).FirstOrDefault();
+			IIssueStep? step = issue.Steps.OrderByDescending(x => unresolvedSpans.Contains(x.SpanId)).ThenByDescending(x => x.CommitId).FirstOrDefault();
 			if (step == null)
 			{
 				return null;
@@ -301,7 +302,7 @@ namespace HordeServer.Ugs
 		/// <returns>Response object</returns>
 		GetUgsIssueBuildResponse CreateBuildResponse(IIssueSpan span, IIssueStep step, IssueBuildOutcome outcome)
 		{
-			GetUgsIssueBuildResponse response = new GetUgsIssueBuildResponse(span.StreamName, step.Change, outcome);
+			GetUgsIssueBuildResponse response = new GetUgsIssueBuildResponse(span.StreamName, step.CommitId.GetPerforceChange(), outcome);
 			response.Id = step.LogId.GetHashCode();
 			response.JobName = $"{step.JobName}: {span.NodeName}";
 			response.JobUrl = new Uri(_serverInfo.DashboardUrl, $"job/{step.JobId}");
@@ -349,7 +350,7 @@ namespace HordeServer.Ugs
 				newResolvedById = request.Resolved.Value ? IIssue.ResolvedByUnknownId : UserId.Empty;
 			}
 
-			if (!await _issueService.UpdateIssueAsync(issueId, ownerId: newOwnerId, nominatedById: newNominatedById, acknowledged: request.Acknowledged, declinedById: newDeclinedById, fixChange: request.FixChange, resolvedById: newResolvedById, initiatedById: User.GetUserId()))
+			if (!await _issueService.UpdateIssueAsync(issueId, ownerId: newOwnerId, nominatedById: newNominatedById, acknowledged: request.Acknowledged, declinedById: newDeclinedById, fixCommitId: CommitId.FromPerforceChange(request.FixChange), resolvedById: newResolvedById, initiatedById: User.GetUserId()))
 			{
 				return NotFound();
 			}
diff --git a/Engine/Source/Programs/Horde/Plugins/Build/HordeServer.Build/Ugs/UgsMessages.cs b/Engine/Source/Programs/Horde/Plugins/Build/HordeServer.Build/Ugs/UgsMessages.cs
--- a/Engine/Source/Programs/Horde/Plugins/Build/HordeServer.Build/Ugs/UgsMessages.cs
+++ b/Engine/Source/Programs/Horde/Plugins/Build/HordeServer.Build/Ugs/UgsMessages.cs
@@ -510,7 +510,7 @@ namespace HordeServer.Ugs
 			Owner = owner?.Login;
 			NominatedBy = nominatedBy?.Login;
 			AcknowledgedAt = issue.AcknowledgedAt;
-			FixChange = issue.FixChange;
+			FixChange = issue.FixCommitId?.GetPerforceChange();
 			ResolvedAt = issue.ResolvedAt;
 			BNotify = notify;
 			BWarning = issue.Severity == IssueSeverity.Warning;
diff --git a/Engine/Source/Programs/Shared/EpicGames.Horde/Artifacts/ArtifactMessages.cs b/Engine/Source/Programs/Shared/EpicGames.Horde/Artifacts/ArtifactMessages.cs
--- a/Engine/Source/Programs/Shared/EpicGames.Horde/Artifacts/ArtifactMessages.cs
+++ b/Engine/Source/Programs/Shared/EpicGames.Horde/Artifacts/ArtifactMessages.cs
@@ -3,6 +3,7 @@
 using System;
 using System.Collections.Generic;
 using EpicGames.Core;
+using EpicGames.Horde.Commits;
 using EpicGames.Horde.Storage;
 using EpicGames.Horde.Streams;
 
@@ -17,10 +18,31 @@ namespace EpicGames.Horde.Artifacts
 	/// <param name="Type">Additional search keys tagged on the artifact</param>
 	/// <param name="Description">Description for the artifact</param>
 	/// <param name="StreamId">Stream to create the artifact for</param>
-	/// <param name="Change">Change number for the artifact</param>
 	/// <param name="Keys">Keys used to identify the artifact</param>
 	/// <param name="Metadata">Metadata for the artifact</param>
-	public record CreateArtifactRequest(ArtifactName Name, ArtifactType Type, string? Description, StreamId? StreamId, int? Change, List<string> Keys, List<string> Metadata);
+	public record CreateArtifactRequest(ArtifactName Name, ArtifactType Type, string? Description, StreamId? StreamId, List<string> Keys, List<string> Metadata)
+	{
+		/// <summary>
+		/// Legacy Perforce changelist number.
+		/// </summary>
+		[Obsolete("Use CommitId instead")]
+		public int Change
+		{
+			get => _change ?? _commitId?.TryGetPerforceChange() ?? -1;
+			set => _change = value;
+		}
+		int? _change;
+
+		/// <summary>
+		/// Commit for the new artifact
+		/// </summary>
+		public CommitId CommitId
+		{
+			get => _commitId ?? CommitId.FromPerforceChange(_change) ?? CommitId.Empty;
+			set => _commitId = value;
+		}
+		CommitId? _commitId;
+	}
 
 	/// <summary>
 	/// Information about a created artifact
@@ -81,7 +103,23 @@ namespace EpicGames.Horde.Artifacts
 		/// <summary>
 		/// Change number
 		/// </summary>
-		public int Change { get; }
+		[Obsolete("Use CommitId instead")]
+		public int Change
+		{
+			get => _change ?? _commitId?.TryGetPerforceChange() ?? -1;
+			set => _change = value;
+		}
+		int? _change;
+
+		/// <summary>
+		/// Commit
+		/// </summary>
+		public CommitId CommitId
+		{
+			get => _commitId ?? CommitId.FromPerforceChange(_change) ?? CommitId.Empty;
+			set => _commitId = value;
+		}
+		CommitId? _commitId;
 
 		/// <summary>
 		/// Keys used to collate artifacts
@@ -101,14 +139,13 @@ namespace EpicGames.Horde.Artifacts
 		/// <summary>
 		/// Constructor
 		/// </summary>
-		public GetArtifactResponse(ArtifactId id, ArtifactName name, ArtifactType type, string? description, StreamId streamId, int change, IReadOnlyList<string> keys, IReadOnlyList<string> metadata, DateTime createdAtUtc)
+		public GetArtifactResponse(ArtifactId id, ArtifactName name, ArtifactType type, string? description, StreamId streamId, IReadOnlyList<string> keys, IReadOnlyList<string> metadata, DateTime createdAtUtc)
 		{
 			Id = id;
 			Name = name;
 			Type = type;
 			Description = description;
 			StreamId = streamId;
-			Change = change;
 			Keys = keys;
 			Metadata = metadata;
 			CreatedAtUtc = createdAtUtc;
diff --git a/Engine/Source/Programs/Shared/EpicGames.Horde/Commits/CommitMessages.cs b/Engine/Source/Programs/Shared/EpicGames.Horde/Commits/CommitMessages.cs
--- a/Engine/Source/Programs/Shared/EpicGames.Horde/Commits/CommitMessages.cs
+++ b/Engine/Source/Programs/Shared/EpicGames.Horde/Commits/CommitMessages.cs
@@ -1,5 +1,6 @@
 // Copyright Epic Games, Inc. All Rights Reserved.
 
+using System;
 using System.Collections.Generic;
 using EpicGames.Horde.Users;
 
@@ -13,9 +14,25 @@ namespace EpicGames.Horde.Commits
 	public class GetCommitResponse
 	{
 		/// <summary>
-		/// The source changelist number
+		/// The commit id
 		/// </summary>
-		public int Number { get; set; }
+		public CommitIdWithOrder Id
+		{
+			get => _id ?? CommitIdWithOrder.FromPerforceChange(_number) ?? CommitIdWithOrder.Empty;
+			set => _id = value;
+		}
+		CommitIdWithOrder? _id;
+
+		/// <summary>
+		/// The changelist number
+		/// </summary>
+		[Obsolete("Use Id instead")]
+		public int Number
+		{
+			get => _number ?? _id?.TryGetPerforceChange() ?? -1;
+			set => _number = value;
+		}
+		int? _number;
 
 		/// <summary>
 		/// Name of the user that authored this change [DEPRECATED]
@@ -45,9 +62,8 @@ namespace EpicGames.Horde.Commits
 		/// <summary>
 		/// Constructor
 		/// </summary>
-		public GetCommitResponse(int number, GetThinUserInfoResponse authorInfo, string description)
+		public GetCommitResponse(GetThinUserInfoResponse authorInfo, string description)
 		{
-			Number = number;
 			Author = authorInfo.Name;
 			AuthorInfo = authorInfo;
 			Description = description;
diff --git a/Engine/Source/Programs/Shared/EpicGames.Horde/HordeHttpClient.cs b/Engine/Source/Programs/Shared/EpicGames.Horde/HordeHttpClient.cs
--- a/Engine/Source/Programs/Shared/EpicGames.Horde/HordeHttpClient.cs
+++ b/Engine/Source/Programs/Shared/EpicGames.Horde/HordeHttpClient.cs
@@ -12,6 +12,7 @@ using System.Threading.Tasks;
 using EpicGames.Core;
 using EpicGames.Horde.Agents.Telemetry;
 using EpicGames.Horde.Artifacts;
+using EpicGames.Horde.Commits;
 using EpicGames.Horde.Dashboard;
 using EpicGames.Horde.Jobs;
 using EpicGames.Horde.Jobs.Graphs;
@@ -132,13 +133,13 @@ namespace EpicGames.Horde
 		/// <param name="type">Additional search keys tagged on the artifact</param>
 		/// <param name="description">Description for the artifact</param>
 		/// <param name="streamId">Stream to create the artifact for</param>
-		/// <param name="change">Change number for the artifact</param>
+		/// <param name="commitId">Commit for the artifact</param>
 		/// <param name="keys">Keys used to identify the artifact</param>
 		/// <param name="metadata">Metadata for the artifact</param>
 		/// <param name="cancellationToken">Cancellation token for the operation</param>
-		public Task<CreateArtifactResponse> CreateArtifactAsync(ArtifactName name, ArtifactType type, string? description, StreamId? streamId = null, int? change = null, List<string>? keys = null, List<string>? metadata = null, CancellationToken cancellationToken = default)
+		public Task<CreateArtifactResponse> CreateArtifactAsync(ArtifactName name, ArtifactType type, string? description, StreamId streamId, CommitId commitId, List<string>? keys = null, List<string>? metadata = null, CancellationToken cancellationToken = default)
 		{
-			return PostAsync<CreateArtifactResponse, CreateArtifactRequest>(_httpClient, $"api/v2/artifacts", new CreateArtifactRequest(name, type, description, streamId, change, keys ?? new List<string>(), metadata ?? new List<string>()), cancellationToken);
+			return PostAsync<CreateArtifactResponse, CreateArtifactRequest>(_httpClient, $"api/v2/artifacts", new CreateArtifactRequest(name, type, description, streamId, keys ?? new List<string>(), metadata ?? new List<string>()) { CommitId = commitId }, cancellationToken);
 		}
 
 		/// <summary>
