Horde: Add an options record for IIssueCollection.TryUpdateIssue.

#rnx

diff --git a/Engine/Source/Programs/Horde/Plugins/Build/HordeServer.Build/Issues/IIssueCollection.cs b/Engine/Source/Programs/Horde/Plugins/Build/HordeServer.Build/Issues/IIssueCollection.cs
--- a/Engine/Source/Programs/Horde/Plugins/Build/HordeServer.Build/Issues/IIssueCollection.cs
+++ b/Engine/Source/Programs/Horde/Plugins/Build/HordeServer.Build/Issues/IIssueCollection.cs
@@ -17,7 +17,321 @@
 
 namespace HordeServer.Issues
 {
+
 	/// <summary>
+	/// Interface for an issue collection
+	/// </summary>
+	public interface IIssueCollection
+	{
+		/// <summary>
+		/// Attempts to enters a critical section over the issue collection. Used for performing operations that require coordination of several documents (eg. attaching spans to issues).
+		/// </summary>
+		/// <returns>The disposable critical section</returns>
+		Task<IAsyncDisposable> EnterCriticalSectionAsync();
+
+		#region Issues
+
+		/// <summary>
+		/// Creates a new issue
+		/// </summary>
+		/// <param name="summary">Summary text for the issue</param>
+		/// <param name="cancellationToken">Cancellation token for the operation</param>
+		/// <returns>The new issue instance</returns>
+		Task<IIssue> AddIssueAsync(string summary, CancellationToken cancellationToken = default);
+
+		/// <summary>
+		/// Retrieves and issue by id
+		/// </summary>
+		/// <param name="issueId">Unique id of the issue</param>
+		/// <param name="cancellationToken">Cancellation token for the operation</param>
+		/// <returns>The issue matching the given id, or null</returns>
+		Task<IIssue?> GetIssueAsync(int issueId, CancellationToken cancellationToken = default);
+
+		/// <summary>
+		/// Finds the suspects for an issue
+		/// </summary>
+		/// <param name="issueId">The issue to retrieve suspects for</param>
+		/// <param name="cancellationToken">Cancellation token for the operation</param>
+		/// <returns>List of suspects</returns>
+		Task<IReadOnlyList<IIssueSuspect>> FindSuspectsAsync(int issueId, CancellationToken cancellationToken = default);
+
+		/// <summary>
+		/// Searches for open issues
+		/// </summary>
+		/// <param name="ids">Set of issue ids to find</param>
+		/// <param name="ownerId">The user to find issues for</param>
+		/// <param name="streamId">The stream affected by the issue</param>
+		/// <param name="minCommitId">Minimum changelist affected by the issue</param>
+		/// <param name="maxCommitId">Maximum changelist affected by the issue</param>
+		/// <param name="resolved">Include issues that are now resolved</param>
+		/// <param name="promoted">Include only promoted issues</param>
+		/// <param name="index">Index within the results to return</param>
+		/// <param name="count">Number of results</param>
+		/// <param name="cancellationToken">Cancellation token for the operation</param>
+		/// <returns>List of streams open in the given stream at the given changelist</returns>
+		Task<IReadOnlyList<IIssue>> FindIssuesAsync(IEnumerable<int>? ids = null, UserId? ownerId = null, StreamId? streamId = null, CommitId? minCommitId = null, CommitId? maxCommitId = null, bool? resolved = null, bool? promoted = null, int? index = null, int? count = null, CancellationToken cancellationToken = default);
+
+		/// <summary>
+		/// Searches for open issues
+		/// </summary>
+		/// <param name="commits">List of suspect commits</param>
+		/// <param name="cancellationToken">Cancellation token for the operation</param>
+		/// <returns>List of issues that are affected by the given changes</returns>
+		Task<IReadOnlyList<IIssue>> FindIssuesForChangesAsync(List<CommitIdWithOrder> commits, CancellationToken cancellationToken = default);
+
+		/// <summary>
+		/// Try to update the state of an issue
+		/// </summary>
+		/// <param name="issue">The issue to update</param>
+		/// <param name="initiatedByUserId">User that is making the change</param>
+		/// <param name="options">Options for updating an issue</param>
+		/// <param name="cancellationToken">Cancellation token for the operation</param>
+		/// <returns>True if the issue was updated</returns>
+		Task<IIssue?> TryUpdateIssueAsync(IIssue issue, UserId? initiatedByUserId, UpdateIssueOptions options, CancellationToken cancellationToken = default);
+
+		/// <summary>
+		/// Updates derived data for an issue (ie. data computed from the spans attached to it). Also clears the issue's 'modified' state.
+		/// </summary>
+		/// <param name="issue">Issue to update</param>
+		/// <param name="newSummary">New summary for the issue</param>
+		/// <param name="newSeverity">New severity for the issue</param>
+		/// <param name="newFingerprints">New fingerprints for the issue</param>
+		/// <param name="newStreams">New streams for the issue</param>
+		/// <param name="newSuspects">New suspects for the issue</param>
+		/// <param name="newResolvedAt">Time for the last resolved change</param>
+		/// <param name="newVerifiedAt">Time that the issue was resolved</param>
+		/// <param name="newLastSeenAt">Last time the issue was seen</param>
+		/// <param name="cancellationToken">Cancellation token for the operation</param>
+		/// <returns>Updated issue, or null if the issue is modified in the interim</returns>
+		Task<IIssue?> TryUpdateIssueDerivedDataAsync(IIssue issue, string newSummary, IssueSeverity newSeverity, List<NewIssueFingerprint> newFingerprints, List<NewIssueStream> newStreams, List<NewIssueSuspectData> newSuspects, DateTime? newResolvedAt, DateTime? newVerifiedAt, DateTime newLastSeenAt, CancellationToken cancellationToken = default);
+
+		#endregion
+
+		#region Spans
+
+		/// <summary>
+		/// Creates a new span from the given failure
+		/// </summary>
+		/// <param name="issueId">The issue that the span belongs to</param>
+		/// <param name="newSpan">Information about the new span</param>
+		/// <param name="cancellationToken">Cancellation token for the operation</param>
+		/// <returns>New span, or null if the sequence token is not valid</returns>
+		Task<IIssueSpan> AddSpanAsync(int issueId, NewIssueSpanData newSpan, CancellationToken cancellationToken = default);
+
+		/// <summary>
+		/// Gets a particular span
+		/// </summary>
+		/// <param name="spanId">Unique id of the span</param>
+		/// <param name="cancellationToken">Cancellation token for the operation</param>
+		/// <returns>New span, or null if the sequence token is not valid</returns>
+		Task<IIssueSpan?> GetSpanAsync(ObjectId spanId, CancellationToken cancellationToken = default);
+
+		/// <summary>
+		/// Updates the given span. Note that data in the span's issue may be derived from this, and the issue should be updated afterwards.
+		/// </summary>
+		/// <param name="span">Span to update</param>
+		/// <param name="newLastSuccess">New last successful step</param>
+		/// <param name="newFailure">New failed step</param>
+		/// <param name="newNextSuccess">New next successful step</param>
+		/// <param name="newSuspects">New suspects for the span</param>
+		/// <param name="newIssueId">The new issue id for this span</param>
+		/// <param name="cancellationToken">Cancellation token for the operation</param>
+		/// <returns>The updated span, or null on failure</returns>
+		Task<IIssueSpan?> TryUpdateSpanAsync(IIssueSpan span, NewIssueStepData? newLastSuccess = null, NewIssueStepData? newFailure = null, NewIssueStepData? newNextSuccess = null, List<NewIssueSpanSuspectData>? newSuspects = null, int? newIssueId = null, CancellationToken cancellationToken = default);
+
+		/// <summary>
+		/// Gets all the spans for a particular issue
+		/// </summary>
+		/// <param name="issueId">Issue id</param>
+		/// <param name="cancellationToken">Cancellation token for the operation</param>
+		/// <returns>List of spans</returns>
+		Task<IReadOnlyList<IIssueSpan>> FindSpansAsync(int issueId, CancellationToken cancellationToken = default);
+
+		/// <summary>
+		/// Retrieves multiple spans
+		/// </summary>
+		/// <param name="spanIds">The span ids</param>
+		/// <param name="cancellationToken">Cancellation token for the operation</param>
+		/// <returns>List of spans</returns>
+		Task<IReadOnlyList<IIssueSpan>> FindSpansAsync(IEnumerable<ObjectId> spanIds, CancellationToken cancellationToken = default);
+
+		/// <summary>
+		/// Finds the open issues for a given stream
+		/// </summary>
+		/// <param name="streamId">The stream id</param>
+		/// <param name="templateId">The template id</param>
+		/// <param name="name">Name of the node</param>
+		/// <param name="commitId">Changelist number to query</param>
+		/// <param name="cancellationToken">Cancellation token for the operation</param>
+		/// <returns>List of open issues</returns>
+		Task<IReadOnlyList<IIssueSpan>> FindOpenSpansAsync(StreamId streamId, TemplateId templateId, string name, CommitId commitId, CancellationToken cancellationToken = default);
+
+		/// <summary>
+		/// Searches for open issues
+		/// </summary>
+		/// <param name="ids">Set of issue ids to find</param>
+		/// <param name="issueIds">The issue ids to retrieve spans for</param>
+		/// <param name="streamId">The stream affected by the issue</param>
+		/// <param name="minCommitId">Minimum commit affected by the issue</param>
+		/// <param name="maxCommitId">Maximum commit affected by the issue</param>
+		/// <param name="resolved">Include issues that are now resolved</param>
+		/// <param name="index">Index within the results to return</param>
+		/// <param name="count">Number of results</param>
+		/// <param name="cancellationToken">Cancellation token for the operation</param>
+		/// <returns>List of streams open in the given stream at the given changelist</returns>
+		Task<IReadOnlyList<IIssueSpan>> FindSpansAsync(IEnumerable<ObjectId>? ids = null, IEnumerable<int>? issueIds = null, StreamId? streamId = null, CommitId? minCommitId = null, CommitId? maxCommitId = null, bool? resolved = null, int? index = null, int? count = null, CancellationToken cancellationToken = default);
+
+		#endregion
+
+		#region Steps
+
+		/// <summary>
+		/// Adds a new step
+		/// </summary>
+		/// <param name="spanId">Initial span for the step</param>
+		/// <param name="newStep">Information about the new step</param>
+		/// <param name="cancellationToken">Cancellation token for the operation</param>
+		/// <returns>New step object</returns>
+		Task<IIssueStep> AddStepAsync(ObjectId spanId, NewIssueStepData newStep, CancellationToken cancellationToken = default);
+
+		/// <summary>
+		/// Find steps for the given spans
+		/// </summary>
+		/// <param name="spanIds">Span ids</param>
+		/// <param name="cancellationToken">Cancellation token for the operation</param>
+		/// <returns>List of steps</returns>
+		Task<IReadOnlyList<IIssueStep>> FindStepsAsync(IEnumerable<ObjectId> spanIds, CancellationToken cancellationToken = default);
+
+		/// <summary>
+		/// Find steps for the given spans
+		/// </summary>
+		/// <param name="jobId">The job id</param>
+		/// <param name="batchId">The batch id</param>
+		/// <param name="stepId">The step id</param>
+		/// <param name="cancellationToken">Cancellation token for the operation</param>
+		/// <returns>List of steps</returns>
+		Task<IReadOnlyList<IIssueStep>> FindStepsAsync(JobId jobId, JobStepBatchId? batchId, JobStepId? stepId, CancellationToken cancellationToken = default);
+
+		#endregion
+
+		/// <summary>
+		/// Gets the logger for a particular issue
+		/// </summary>
+		/// <param name="issueId">The issue id</param>
+		/// <returns>Logger for this issue</returns>
+		IAuditLogChannel<int> GetLogger(int issueId);
+	}
+
+	/// <summary>
+	/// Options for updating an issue
+	/// </summary>
+	/// <param name="Severity">New severity for the issue</param>
+	/// <param name="Summary">New summary for the issue</param>
+	/// <param name="UserSummary">New user summary for the issue</param>
+	/// <param name="Description">New description for the issue</param>
+	/// <param name="Promoted">New promoted state of the issue</param>
+	/// <param name="OwnerId">New owner of the issue</param>
+	/// <param name="NominatedById">Person that nominated the new owner</param>
+	/// <param name="Acknowledged">Whether the issue has been acknowledged</param>
+	/// <param name="DeclinedById">Name of a user that has declined the issue</param>
+	/// <param name="FixCommitId">Fix changelist for the issue. Pass a default value to clear the fix changelist.</param>
+	/// <param name="FixSystemic">Whether the issue should be marked fixed as a systemic issue</param>
+	/// <param name="ResolvedById">User that resolved the issue (may be ObjectId.Empty to clear)</param>
+	/// <param name="ExcludeSpanIds">List of span ids to exclude from this issue</param>
+	/// <param name="LastSeenAt"></param>
+	/// <param name="ExternalIssueKey">Issue key for external issue tracking</param>
+	/// <param name="QuarantinedById">The user that quarantined the issue</param>
+	/// <param name="ForceClosedById">The user that force closed the issue</param>
+	/// <param name="WorkflowThreadUrl">The workflow thread url associated with the issue</param>
+	public record class UpdateIssueOptions
+	(
+		IssueSeverity? Severity = null, 
+		string? Summary = null, 
+		string? UserSummary = null, 
+		string? Description = null, 
+		bool? Promoted = null, 
+		UserId? OwnerId = null, 
+		UserId? NominatedById = null, 
+		bool? Acknowledged = null, 
+		UserId? DeclinedById = null, 
+		CommitId? FixCommitId = null, 
+		bool? FixSystemic = null, 
+		UserId? ResolvedById = null, 
+		List<ObjectId>? ExcludeSpanIds = null, 
+		DateTime? LastSeenAt = null, 
+		string? ExternalIssueKey = null, 
+		UserId? QuarantinedById = null, 
+		UserId? ForceClosedById = null, 
+		Uri? WorkflowThreadUrl = null
+	);
+
+	/// <summary>
+	/// Extension methods for issues
+	/// </summary>
+	public static class IssueCollectionExtensions
+	{
+		/// <summary>
+		/// Find suspects for the given issue
+		/// </summary>
+		/// <param name="issueCollection"></param>
+		/// <param name="issue"></param>
+		/// <param name="cancellationToken">Cancellation token for the operation</param>
+		/// <returns></returns>
+		public static Task<IReadOnlyList<IIssueSuspect>> FindSuspectsAsync(this IIssueCollection issueCollection, IIssue issue, CancellationToken cancellationToken = default)
+		{
+			return issueCollection.FindSuspectsAsync(issue.Id, cancellationToken);
+		}
+
+		/// <summary>
+		/// Find steps for the given spans
+		/// </summary>
+		/// <param name="issueCollection"></param>
+		/// <param name="spanId">Span ids</param>
+		/// <param name="cancellationToken">Cancellation token for the operation</param>
+		/// <returns>List of steps</returns>
+		public static Task<IReadOnlyList<IIssueStep>> FindStepsAsync(this IIssueCollection issueCollection, ObjectId spanId, CancellationToken cancellationToken = default)
+		{
+			return issueCollection.FindStepsAsync(new[] { spanId }, cancellationToken);
+		}
+
+		/// <summary>
+		/// Searches for open issues affecting a job
+		/// </summary>
+		/// <param name="issueCollection"></param>
+		/// <param name="job"></param>
+		/// <param name="graph"></param>
+		/// <param name="stepId"></param>
+		/// <param name="batchId"></param>
+		/// <param name="labelIdx"></param>
+		/// <param name="ownerId"></param>
+		/// <param name="resolved">Whether to include results that are resolved</param>
+		/// <param name="promoted">Whether to filter by promoted issues</param>
+		/// <param name="index">Index within the results to return</param>
+		/// <param name="count">Number of results</param>
+		/// <param name="cancellationToken">Cancellation token for the operation</param>
+		/// <returns></returns>
+		public static async Task<IReadOnlyList<IIssue>> FindIssuesForJobAsync(this IIssueCollection issueCollection, IJob job, IGraph graph, JobStepId? stepId = null, JobStepBatchId? batchId = null, int? labelIdx = null, UserId? ownerId = null, bool? resolved = null, bool? promoted = null, int? index = null, int? count = null, CancellationToken cancellationToken = default)
+		{
+			IReadOnlyList<IIssueStep> steps = await issueCollection.FindStepsAsync(job.Id, batchId, stepId, cancellationToken);
+			IReadOnlyList<IIssueSpan> spans = await issueCollection.FindSpansAsync(steps.Select(x => x.SpanId), cancellationToken);
+
+			if (labelIdx != null)
+			{
+				HashSet<string> nodeNames = new HashSet<string>(job.GetNodesForLabel(graph, labelIdx.Value).Select(x => graph.GetNode(x).Name));
+				spans = spans.Where(x => nodeNames.Contains(x.NodeName)).ToList();
+			}
+
+			List<int> issueIds = new List<int>(spans.Select(x => x.IssueId));
+			if (issueIds.Count == 0)
+			{
+				return new List<IIssue>();
+			}
+
+			return await issueCollection.FindIssuesAsync(ids: issueIds, ownerId: ownerId, resolved: resolved, promoted: promoted, index: index, count: count, cancellationToken: cancellationToken);
+		}
+	}
+
+	/// <summary>
 	/// Fingerprint for an issue
 	/// </summary>
 	public class NewIssueFingerprint : IIssueFingerprint, IEquatable<IIssueFingerprint>
@@ -417,291 +731,4 @@
 			StreamId = streamId;
 		}
 	}
-
-	/// <summary>
-	/// Interface for an issue collection
-	/// </summary>
-	public interface IIssueCollection
-	{
-		/// <summary>
-		/// Attempts to enters a critical section over the issue collection. Used for performing operations that require coordination of several documents (eg. attaching spans to issues).
-		/// </summary>
-		/// <returns>The disposable critical section</returns>
-		Task<IAsyncDisposable> EnterCriticalSectionAsync();
-
-		#region Issues
-
-		/// <summary>
-		/// Creates a new issue
-		/// </summary>
-		/// <param name="summary">Summary text for the issue</param>
-		/// <param name="cancellationToken">Cancellation token for the operation</param>
-		/// <returns>The new issue instance</returns>
-		Task<IIssue> AddIssueAsync(string summary, CancellationToken cancellationToken = default);
-
-		/// <summary>
-		/// Retrieves and issue by id
-		/// </summary>
-		/// <param name="issueId">Unique id of the issue</param>
-		/// <param name="cancellationToken">Cancellation token for the operation</param>
-		/// <returns>The issue matching the given id, or null</returns>
-		Task<IIssue?> GetIssueAsync(int issueId, CancellationToken cancellationToken = default);
-
-		/// <summary>
-		/// Finds the suspects for an issue
-		/// </summary>
-		/// <param name="issueId">The issue to retrieve suspects for</param>
-		/// <param name="cancellationToken">Cancellation token for the operation</param>
-		/// <returns>List of suspects</returns>
-		Task<IReadOnlyList<IIssueSuspect>> FindSuspectsAsync(int issueId, CancellationToken cancellationToken = default);
-
-		/// <summary>
-		/// Searches for open issues
-		/// </summary>
-		/// <param name="ids">Set of issue ids to find</param>
-		/// <param name="ownerId">The user to find issues for</param>
-		/// <param name="streamId">The stream affected by the issue</param>
-		/// <param name="minCommitId">Minimum changelist affected by the issue</param>
-		/// <param name="maxCommitId">Maximum changelist affected by the issue</param>
-		/// <param name="resolved">Include issues that are now resolved</param>
-		/// <param name="promoted">Include only promoted issues</param>
-		/// <param name="index">Index within the results to return</param>
-		/// <param name="count">Number of results</param>
-		/// <param name="cancellationToken">Cancellation token for the operation</param>
-		/// <returns>List of streams open in the given stream at the given changelist</returns>
-		Task<IReadOnlyList<IIssue>> FindIssuesAsync(IEnumerable<int>? ids = null, UserId? ownerId = null, StreamId? streamId = null, CommitId? minCommitId = null, CommitId? maxCommitId = null, bool? resolved = null, bool? promoted = null, int? index = null, int? count = null, CancellationToken cancellationToken = default);
-
-		/// <summary>
-		/// Searches for open issues
-		/// </summary>
-		/// <param name="commits">List of suspect commits</param>
-		/// <param name="cancellationToken">Cancellation token for the operation</param>
-		/// <returns>List of issues that are affected by the given changes</returns>
-		Task<IReadOnlyList<IIssue>> FindIssuesForChangesAsync(List<CommitIdWithOrder> commits, CancellationToken cancellationToken = default);
-
-		/// <summary>
-		/// Try to update the state of an issue
-		/// </summary>
-		/// <param name="issue">The issue to update</param>
-		/// <param name="initiatedByUserId">User that is making the change</param>
-		/// <param name="newSeverity">New severity for the issue</param>
-		/// <param name="newSummary">New summary for the issue</param>
-		/// <param name="newUserSummary">New user summary for the issue</param>
-		/// <param name="newDescription">New description for the issue</param>
-		/// <param name="newPromoted">New promoted state of the issue</param>
-		/// <param name="newOwnerId">New owner of the issue</param>
-		/// <param name="newNominatedById">Person that nominated the new owner</param>
-		/// <param name="newAcknowledged">Whether the issue has been acknowledged</param>
-		/// <param name="newDeclinedById">Name of a user that has declined the issue</param>
-		/// <param name="newFixCommitId">Fix changelist for the issue. Pass a default value to clear the fix changelist.</param>
-		/// <param name="newFixSystemic">Whether the issue should be marked fixed as a systemic issue</param>
-		/// <param name="newResolvedById">User that resolved the issue (may be ObjectId.Empty to clear)</param>
-		/// <param name="newExcludeSpanIds">List of span ids to exclude from this issue</param>
-		/// <param name="newLastSeenAt"></param>
-		/// <param name="newExternalIssueKey">Issue key for external issue tracking</param>
-		/// <param name="newQuarantinedById">The user that quarantined the issue</param>
-		/// <param name="newForceClosedById">The user that force closed the issue</param>
-		/// <param name="newWorkflowThreadUrl">The workflow thread url associated with the issue</param>
-		/// <param name="cancellationToken">Cancellation token for the operation</param>
-		/// <returns>True if the issue was updated</returns>
-		Task<IIssue?> TryUpdateIssueAsync(IIssue issue, UserId? initiatedByUserId, IssueSeverity? newSeverity = null, string? newSummary = null, string? newUserSummary = null, string? newDescription = null, bool? newPromoted = null, UserId? newOwnerId = null, UserId? newNominatedById = null, bool? newAcknowledged = null, UserId? newDeclinedById = null, CommitId? newFixCommitId = null, bool? newFixSystemic = null, UserId? newResolvedById = null, List<ObjectId>? newExcludeSpanIds = null, DateTime? newLastSeenAt = null, string? newExternalIssueKey = null, UserId? newQuarantinedById = null, UserId? newForceClosedById = null, Uri? newWorkflowThreadUrl = null, CancellationToken cancellationToken = default);
-
-		/// <summary>
-		/// Updates derived data for an issue (ie. data computed from the spans attached to it). Also clears the issue's 'modified' state.
-		/// </summary>
-		/// <param name="issue">Issue to update</param>
-		/// <param name="newSummary">New summary for the issue</param>
-		/// <param name="newSeverity">New severity for the issue</param>
-		/// <param name="newFingerprints">New fingerprints for the issue</param>
-		/// <param name="newStreams">New streams for the issue</param>
-		/// <param name="newSuspects">New suspects for the issue</param>
-		/// <param name="newResolvedAt">Time for the last resolved change</param>
-		/// <param name="newVerifiedAt">Time that the issue was resolved</param>
-		/// <param name="newLastSeenAt">Last time the issue was seen</param>
-		/// <param name="cancellationToken">Cancellation token for the operation</param>
-		/// <returns>Updated issue, or null if the issue is modified in the interim</returns>
-		Task<IIssue?> TryUpdateIssueDerivedDataAsync(IIssue issue, string newSummary, IssueSeverity newSeverity, List<NewIssueFingerprint> newFingerprints, List<NewIssueStream> newStreams, List<NewIssueSuspectData> newSuspects, DateTime? newResolvedAt, DateTime? newVerifiedAt, DateTime newLastSeenAt, CancellationToken cancellationToken = default);
-
-		#endregion
-
-		#region Spans
-
-		/// <summary>
-		/// Creates a new span from the given failure
-		/// </summary>
-		/// <param name="issueId">The issue that the span belongs to</param>
-		/// <param name="newSpan">Information about the new span</param>
-		/// <param name="cancellationToken">Cancellation token for the operation</param>
-		/// <returns>New span, or null if the sequence token is not valid</returns>
-		Task<IIssueSpan> AddSpanAsync(int issueId, NewIssueSpanData newSpan, CancellationToken cancellationToken = default);
-
-		/// <summary>
-		/// Gets a particular span
-		/// </summary>
-		/// <param name="spanId">Unique id of the span</param>
-		/// <param name="cancellationToken">Cancellation token for the operation</param>
-		/// <returns>New span, or null if the sequence token is not valid</returns>
-		Task<IIssueSpan?> GetSpanAsync(ObjectId spanId, CancellationToken cancellationToken = default);
-
-		/// <summary>
-		/// Updates the given span. Note that data in the span's issue may be derived from this, and the issue should be updated afterwards.
-		/// </summary>
-		/// <param name="span">Span to update</param>
-		/// <param name="newLastSuccess">New last successful step</param>
-		/// <param name="newFailure">New failed step</param>
-		/// <param name="newNextSuccess">New next successful step</param>
-		/// <param name="newSuspects">New suspects for the span</param>
-		/// <param name="newIssueId">The new issue id for this span</param>
-		/// <param name="cancellationToken">Cancellation token for the operation</param>
-		/// <returns>The updated span, or null on failure</returns>
-		Task<IIssueSpan?> TryUpdateSpanAsync(IIssueSpan span, NewIssueStepData? newLastSuccess = null, NewIssueStepData? newFailure = null, NewIssueStepData? newNextSuccess = null, List<NewIssueSpanSuspectData>? newSuspects = null, int? newIssueId = null, CancellationToken cancellationToken = default);
-
-		/// <summary>
-		/// Gets all the spans for a particular issue
-		/// </summary>
-		/// <param name="issueId">Issue id</param>
-		/// <param name="cancellationToken">Cancellation token for the operation</param>
-		/// <returns>List of spans</returns>
-		Task<IReadOnlyList<IIssueSpan>> FindSpansAsync(int issueId, CancellationToken cancellationToken = default);
-
-		/// <summary>
-		/// Retrieves multiple spans
-		/// </summary>
-		/// <param name="spanIds">The span ids</param>
-		/// <param name="cancellationToken">Cancellation token for the operation</param>
-		/// <returns>List of spans</returns>
-		Task<IReadOnlyList<IIssueSpan>> FindSpansAsync(IEnumerable<ObjectId> spanIds, CancellationToken cancellationToken = default);
-
-		/// <summary>
-		/// Finds the open issues for a given stream
-		/// </summary>
-		/// <param name="streamId">The stream id</param>
-		/// <param name="templateId">The template id</param>
-		/// <param name="name">Name of the node</param>
-		/// <param name="commitId">Changelist number to query</param>
-		/// <param name="cancellationToken">Cancellation token for the operation</param>
-		/// <returns>List of open issues</returns>
-		Task<IReadOnlyList<IIssueSpan>> FindOpenSpansAsync(StreamId streamId, TemplateId templateId, string name, CommitId commitId, CancellationToken cancellationToken = default);
-
-		/// <summary>
-		/// Searches for open issues
-		/// </summary>
-		/// <param name="ids">Set of issue ids to find</param>
-		/// <param name="issueIds">The issue ids to retrieve spans for</param>
-		/// <param name="streamId">The stream affected by the issue</param>
-		/// <param name="minCommitId">Minimum commit affected by the issue</param>
-		/// <param name="maxCommitId">Maximum commit affected by the issue</param>
-		/// <param name="resolved">Include issues that are now resolved</param>
-		/// <param name="index">Index within the results to return</param>
-		/// <param name="count">Number of results</param>
-		/// <param name="cancellationToken">Cancellation token for the operation</param>
-		/// <returns>List of streams open in the given stream at the given changelist</returns>
-		Task<IReadOnlyList<IIssueSpan>> FindSpansAsync(IEnumerable<ObjectId>? ids = null, IEnumerable<int>? issueIds = null, StreamId? streamId = null, CommitId? minCommitId = null, CommitId? maxCommitId = null, bool? resolved = null, int? index = null, int? count = null, CancellationToken cancellationToken = default);
-
-		#endregion
-
-		#region Steps
-
-		/// <summary>
-		/// Adds a new step
-		/// </summary>
-		/// <param name="spanId">Initial span for the step</param>
-		/// <param name="newStep">Information about the new step</param>
-		/// <param name="cancellationToken">Cancellation token for the operation</param>
-		/// <returns>New step object</returns>
-		Task<IIssueStep> AddStepAsync(ObjectId spanId, NewIssueStepData newStep, CancellationToken cancellationToken = default);
-
-		/// <summary>
-		/// Find steps for the given spans
-		/// </summary>
-		/// <param name="spanIds">Span ids</param>
-		/// <param name="cancellationToken">Cancellation token for the operation</param>
-		/// <returns>List of steps</returns>
-		Task<IReadOnlyList<IIssueStep>> FindStepsAsync(IEnumerable<ObjectId> spanIds, CancellationToken cancellationToken = default);
-
-		/// <summary>
-		/// Find steps for the given spans
-		/// </summary>
-		/// <param name="jobId">The job id</param>
-		/// <param name="batchId">The batch id</param>
-		/// <param name="stepId">The step id</param>
-		/// <param name="cancellationToken">Cancellation token for the operation</param>
-		/// <returns>List of steps</returns>
-		Task<IReadOnlyList<IIssueStep>> FindStepsAsync(JobId jobId, JobStepBatchId? batchId, JobStepId? stepId, CancellationToken cancellationToken = default);
-
-		#endregion
-
-		/// <summary>
-		/// Gets the logger for a particular issue
-		/// </summary>
-		/// <param name="issueId">The issue id</param>
-		/// <returns>Logger for this issue</returns>
-		IAuditLogChannel<int> GetLogger(int issueId);
-	}
-
-	/// <summary>
-	/// Extension methods for issues
-	/// </summary>
-	public static class IssueCollectionExtensions
-	{
-		/// <summary>
-		/// Find suspects for the given issue
-		/// </summary>
-		/// <param name="issueCollection"></param>
-		/// <param name="issue"></param>
-		/// <param name="cancellationToken">Cancellation token for the operation</param>
-		/// <returns></returns>
-		public static Task<IReadOnlyList<IIssueSuspect>> FindSuspectsAsync(this IIssueCollection issueCollection, IIssue issue, CancellationToken cancellationToken = default)
-		{
-			return issueCollection.FindSuspectsAsync(issue.Id, cancellationToken);
-		}
-
-		/// <summary>
-		/// Find steps for the given spans
-		/// </summary>
-		/// <param name="issueCollection"></param>
-		/// <param name="spanId">Span ids</param>
-		/// <param name="cancellationToken">Cancellation token for the operation</param>
-		/// <returns>List of steps</returns>
-		public static Task<IReadOnlyList<IIssueStep>> FindStepsAsync(this IIssueCollection issueCollection, ObjectId spanId, CancellationToken cancellationToken = default)
-		{
-			return issueCollection.FindStepsAsync(new[] { spanId }, cancellationToken);
-		}
-
-		/// <summary>
-		/// Searches for open issues affecting a job
-		/// </summary>
-		/// <param name="issueCollection"></param>
-		/// <param name="job"></param>
-		/// <param name="graph"></param>
-		/// <param name="stepId"></param>
-		/// <param name="batchId"></param>
-		/// <param name="labelIdx"></param>
-		/// <param name="ownerId"></param>
-		/// <param name="resolved">Whether to include results that are resolved</param>
-		/// <param name="promoted">Whether to filter by promoted issues</param>
-		/// <param name="index">Index within the results to return</param>
-		/// <param name="count">Number of results</param>
-		/// <param name="cancellationToken">Cancellation token for the operation</param>
-		/// <returns></returns>
-		public static async Task<IReadOnlyList<IIssue>> FindIssuesForJobAsync(this IIssueCollection issueCollection, IJob job, IGraph graph, JobStepId? stepId = null, JobStepBatchId? batchId = null, int? labelIdx = null, UserId? ownerId = null, bool? resolved = null, bool? promoted = null, int? index = null, int? count = null, CancellationToken cancellationToken = default)
-		{
-			IReadOnlyList<IIssueStep> steps = await issueCollection.FindStepsAsync(job.Id, batchId, stepId, cancellationToken);
-			IReadOnlyList<IIssueSpan> spans = await issueCollection.FindSpansAsync(steps.Select(x => x.SpanId), cancellationToken);
-
-			if (labelIdx != null)
-			{
-				HashSet<string> nodeNames = new HashSet<string>(job.GetNodesForLabel(graph, labelIdx.Value).Select(x => graph.GetNode(x).Name));
-				spans = spans.Where(x => nodeNames.Contains(x.NodeName)).ToList();
-			}
-
-			List<int> issueIds = new List<int>(spans.Select(x => x.IssueId));
-			if (issueIds.Count == 0)
-			{
-				return new List<IIssue>();
-			}
-
-			return await issueCollection.FindIssuesAsync(ids: issueIds, ownerId: ownerId, resolved: resolved, promoted: promoted, index: index, count: count, cancellationToken: cancellationToken);
-		}
-	}
 }
diff --git a/Engine/Source/Programs/Horde/Plugins/Build/HordeServer.Build/Issues/IssueCollection.cs b/Engine/Source/Programs/Horde/Plugins/Build/HordeServer.Build/Issues/IssueCollection.cs
--- a/Engine/Source/Programs/Horde/Plugins/Build/HordeServer.Build/Issues/IssueCollection.cs
+++ b/Engine/Source/Programs/Horde/Plugins/Build/HordeServer.Build/Issues/IssueCollection.cs
@@ -1216,65 +1216,64 @@ namespace HordeServer.Issues
 		}
 
 		/// <inheritdoc/>
-		public async Task<IIssue?> TryUpdateIssueAsync(IIssue issue, UserId? initiatedByUserId, IssueSeverity? newSeverity = null, string? newSummary = null, string? newUserSummary = null, string? newDescription = null, bool? newManuallyPromoted = null, UserId? newOwnerId = null, UserId? newNominatedById = null, bool? newAcknowledged = null, UserId? newDeclinedById = null, CommitId? newFixCommit = null, bool? newFixSystemic = null, UserId? newResolvedById = null, List<ObjectId>? newExcludeSpanIds = null, DateTime? newLastSeenAt = null, string? newExternaIssueKey = null, UserId? newQuarantinedById = null, UserId? newForceClosedById = null, Uri? newWorkflowThreadUrl = null, CancellationToken cancellationToken = default)
+		public async Task<IIssue?> TryUpdateIssueAsync(IIssue issue, UserId? initiatedByUserId, UpdateIssueOptions options, CancellationToken cancellationToken = default)
 		{
 			IssueDocument issueDocument = (IssueDocument)issue;
 
-			if (newDeclinedById != null && newDeclinedById == issueDocument.OwnerId)
+			if (options.DeclinedById != null && options.DeclinedById == issueDocument.OwnerId)
 			{
-				newOwnerId = UserId.Empty;
+				options = options with { OwnerId = UserId.Empty };
 			}
 
-			if (issue.ResolvedById == null && (newForceClosedById != null && newResolvedById == null))
+			if (issue.ResolvedById == null && (options.ForceClosedById != null && options.ResolvedById == null))
 			{
-				newResolvedById = newForceClosedById;
+				options = options with { ResolvedById = options.ForceClosedById };
 			}
 
 			DateTime utcNow = DateTime.UtcNow;
 
 			List<UpdateDefinition<IssueDocument>> updates = new List<UpdateDefinition<IssueDocument>>();
-			if (newSeverity != null)
+			if (options.Severity != null)
 			{
-				updates.Add(Builders<IssueDocument>.Update.Set(x => x.Severity, newSeverity.Value));
+				updates.Add(Builders<IssueDocument>.Update.Set(x => x.Severity, options.Severity.Value));
 			}
-			if (newSummary != null)
+			if (options.Summary != null)
 			{
-				updates.Add(Builders<IssueDocument>.Update.Set(x => x.Summary, newSummary));
+				updates.Add(Builders<IssueDocument>.Update.Set(x => x.Summary, options.Summary));
 			}
-			if (newUserSummary != null)
+			if (options.UserSummary != null)
 			{
-				if (newUserSummary.Length == 0)
+				if (options.UserSummary.Length == 0)
 				{
 					updates.Add(Builders<IssueDocument>.Update.Unset(x => x.UserSummary!));
 				}
 				else
 				{
-					updates.Add(Builders<IssueDocument>.Update.Set(x => x.UserSummary, newUserSummary));
+					updates.Add(Builders<IssueDocument>.Update.Set(x => x.UserSummary, options.UserSummary));
 				}
 			}
-			if (newDescription != null)
+			if (options.Description != null)
 			{
-				if (newDescription.Length == 0)
+				if (options.Description.Length == 0)
 				{
 					updates.Add(Builders<IssueDocument>.Update.Unset(x => x.Description));
 				}
 				else
 				{
-					updates.Add(Builders<IssueDocument>.Update.Set(x => x.Description, newDescription));
+					updates.Add(Builders<IssueDocument>.Update.Set(x => x.Description, options.Description));
 				}
 			}
-			if (newManuallyPromoted != null)
+			if (options.Promoted != null)
 			{
-				updates.Add(Builders<IssueDocument>.Update.Set(x => x.ManuallyPromoted, newManuallyPromoted.Value));
+				updates.Add(Builders<IssueDocument>.Update.Set(x => x.ManuallyPromoted, options.Promoted.Value));
 			}
-			if (newResolvedById != null)
+			if (options.ResolvedById != null)
 			{
-				newOwnerId ??= newResolvedById;
-				newAcknowledged ??= true;
+				options = options with { OwnerId = options.OwnerId ?? options.ResolvedById, Acknowledged = options.Acknowledged ?? true };
 			}
-			if (newOwnerId != null)
+			if (options.OwnerId != null)
 			{
-				if (newOwnerId.Value == UserId.Empty)
+				if (options.OwnerId.Value == UserId.Empty)
 				{
 					updates.Add(Builders<IssueDocument>.Update.Unset(x => x.OwnerId!));
 					updates.Add(Builders<IssueDocument>.Update.Unset(x => x.NominatedAt!));
@@ -1282,23 +1281,23 @@ namespace HordeServer.Issues
 				}
 				else
 				{
-					updates.Add(Builders<IssueDocument>.Update.Set(x => x.OwnerId!, newOwnerId.Value));
+					updates.Add(Builders<IssueDocument>.Update.Set(x => x.OwnerId!, options.OwnerId.Value));
 
 					updates.Add(Builders<IssueDocument>.Update.Set(x => x.NominatedAt, DateTime.UtcNow));
-					if (newNominatedById == null)
+					if (options.NominatedById == null)
 					{
 						updates.Add(Builders<IssueDocument>.Update.Unset(x => x.NominatedById!));
 					}
 					else
 					{
-						updates.Add(Builders<IssueDocument>.Update.Set(x => x.NominatedById, newNominatedById.Value));
+						updates.Add(Builders<IssueDocument>.Update.Set(x => x.NominatedById, options.NominatedById.Value));
 					}
-					newAcknowledged ??= false;
+					options = options with { Acknowledged = options.Acknowledged ?? false };
 				}
 			}
-			if (newAcknowledged != null)
+			if (options.Acknowledged != null)
 			{
-				if (newAcknowledged.Value)
+				if (options.Acknowledged.Value)
 				{
 					if (issueDocument.AcknowledgedAt == null)
 					{
@@ -1313,20 +1312,20 @@ namespace HordeServer.Issues
 					}
 				}
 			}
-			if (newFixCommit != null)
+			if (options.FixCommitId != null)
 			{
-				if (String.IsNullOrEmpty(newFixCommit.Name))
+				if (String.IsNullOrEmpty(options.FixCommitId.Name))
 				{
 					updates.Add(Builders<IssueDocument>.Update.Unset(x => x.FixCommitName).Unset(x => x.FixChange).Unset(x => x.FixSystemic));
 				}
 				else
 				{
-					updates.Add(Builders<IssueDocument>.Update.Set(x => x.FixCommitName, newFixCommit.Name).Unset(x => x.FixChange).Unset(x => x.FixSystemic));
+					updates.Add(Builders<IssueDocument>.Update.Set(x => x.FixCommitName, options.FixCommitId.Name).Unset(x => x.FixChange).Unset(x => x.FixSystemic));
 				}
 			}
-			else if (newFixSystemic != null)
+			else if (options.FixSystemic != null)
 			{
-				if (newFixSystemic.Value)
+				if (options.FixSystemic.Value)
 				{
 					updates.Add(Builders<IssueDocument>.Update.Unset(x => x.FixCommitName).Unset(x => x.FixChange).Set(x => x.FixSystemic, true));
 				}
@@ -1335,14 +1334,14 @@ namespace HordeServer.Issues
 					updates.Add(Builders<IssueDocument>.Update.Unset(x => x.FixCommitName).Unset(x => x.FixChange).Unset(x => x.FixSystemic));
 				}
 			}
-			if (newResolvedById != null)
+			if (options.ResolvedById != null)
 			{
-				if (newResolvedById.Value != UserId.Empty)
+				if (options.ResolvedById.Value != UserId.Empty)
 				{
-					if (issueDocument.ResolvedAt == null || issueDocument.ResolvedById != newResolvedById)
+					if (issueDocument.ResolvedAt == null || issueDocument.ResolvedById != options.ResolvedById)
 					{
 						updates.Add(Builders<IssueDocument>.Update.Set(x => x.ResolvedAt, utcNow));
-						updates.Add(Builders<IssueDocument>.Update.Set(x => x.ResolvedById, newResolvedById.Value));
+						updates.Add(Builders<IssueDocument>.Update.Set(x => x.ResolvedById, options.ResolvedById.Value));
 					}
 				}
 				else
@@ -1357,65 +1356,65 @@ namespace HordeServer.Issues
 					}
 				}
 			}
-			if (newExcludeSpanIds != null)
+			if (options.ExcludeSpanIds != null)
 			{
-				List<ObjectId> newCombinedExcludeSpanIds = newExcludeSpanIds;
+				List<ObjectId> newCombinedExcludeSpanIds = options.ExcludeSpanIds;
 				if (issue.ExcludeSpans != null)
 				{
 					newCombinedExcludeSpanIds = newCombinedExcludeSpanIds.Union(issue.ExcludeSpans).ToList();
 				}
 				updates.Add(Builders<IssueDocument>.Update.Set(x => x.ExcludeSpans, newCombinedExcludeSpanIds));
 			}
-			if (newLastSeenAt != null)
+			if (options.LastSeenAt != null)
 			{
-				updates.Add(Builders<IssueDocument>.Update.Set(x => x.LastSeenAt, newLastSeenAt.Value));
+				updates.Add(Builders<IssueDocument>.Update.Set(x => x.LastSeenAt, options.LastSeenAt.Value));
 			}
 
-			if (newDeclinedById != null)
+			if (options.DeclinedById != null)
 			{
-				GetLogger(issue.Id).LogInformation("Declined by {UserId}", newDeclinedById.Value);
-				await _issueSuspects.UpdateManyAsync(x => x.IssueId == issue.Id && x.AuthorId == newDeclinedById.Value, Builders<IssueSuspectDocument>.Update.Set(x => x.DeclinedAt, DateTime.UtcNow), null, cancellationToken);
+				GetLogger(issue.Id).LogInformation("Declined by {UserId}", options.DeclinedById.Value);
+				await _issueSuspects.UpdateManyAsync(x => x.IssueId == issue.Id && x.AuthorId == options.DeclinedById.Value, Builders<IssueSuspectDocument>.Update.Set(x => x.DeclinedAt, DateTime.UtcNow), null, cancellationToken);
 			}
-			if (newQuarantinedById != null)
+			if (options.QuarantinedById != null)
 			{
-				if (newQuarantinedById.Value == UserId.Empty)
+				if (options.QuarantinedById.Value == UserId.Empty)
 				{
 					updates.Add(Builders<IssueDocument>.Update.Unset(x => x.QuarantinedByUserId));
 					updates.Add(Builders<IssueDocument>.Update.Unset(x => x.QuarantineTimeUtc));
 				}
 				else
 				{
-					updates.Add(Builders<IssueDocument>.Update.Set(x => x.QuarantinedByUserId!, newQuarantinedById.Value));
+					updates.Add(Builders<IssueDocument>.Update.Set(x => x.QuarantinedByUserId!, options.QuarantinedById.Value));
 					updates.Add(Builders<IssueDocument>.Update.Set(x => x.QuarantineTimeUtc, DateTime.UtcNow));
 				}
 			}
-			else if ((newResolvedById != null && newResolvedById.Value != UserId.Empty) || (newForceClosedById != null && newForceClosedById.Value != UserId.Empty))
+			else if ((options.ResolvedById != null && options.ResolvedById.Value != UserId.Empty) || (options.ForceClosedById != null && options.ForceClosedById.Value != UserId.Empty))
 			{
 				// Clear quarantine if being resolved or if being force closed
 				updates.Add(Builders<IssueDocument>.Update.Unset(x => x.QuarantinedByUserId));
 				updates.Add(Builders<IssueDocument>.Update.Unset(x => x.QuarantineTimeUtc));
 			}
 
-			if (newForceClosedById != null)
+			if (options.ForceClosedById != null)
 			{
-				if (newForceClosedById.Value == UserId.Empty)
+				if (options.ForceClosedById.Value == UserId.Empty)
 				{
 					updates.Add(Builders<IssueDocument>.Update.Unset(x => x.ForceClosedByUserId));
 				}
 				else
 				{
-					updates.Add(Builders<IssueDocument>.Update.Set(x => x.ForceClosedByUserId, newForceClosedById.Value));
+					updates.Add(Builders<IssueDocument>.Update.Set(x => x.ForceClosedByUserId, options.ForceClosedById.Value));
 				}
 			}
 
-			if (newExternaIssueKey != null)
+			if (options.ExternalIssueKey != null)
 			{
-				updates.Add(Builders<IssueDocument>.Update.Set(x => x.ExternalIssueKey, newExternaIssueKey.Length == 0 ? null : newExternaIssueKey));
+				updates.Add(Builders<IssueDocument>.Update.Set(x => x.ExternalIssueKey, options.ExternalIssueKey.Length == 0 ? null : options.ExternalIssueKey));
 			}
 
-			if (newWorkflowThreadUrl != null)
+			if (options.WorkflowThreadUrl != null)
 			{
-				updates.Add(Builders<IssueDocument>.Update.Set(x => x.WorkflowThreadUrl, newWorkflowThreadUrl.ToString().Length == 0 ? null : newWorkflowThreadUrl));
+				updates.Add(Builders<IssueDocument>.Update.Set(x => x.WorkflowThreadUrl, options.WorkflowThreadUrl.ToString().Length == 0 ? null : options.WorkflowThreadUrl));
 			}
 
 			if (updates.Count == 0)
diff --git a/Engine/Source/Programs/Horde/Plugins/Build/HordeServer.Build/Issues/IssueService.cs b/Engine/Source/Programs/Horde/Plugins/Build/HordeServer.Build/Issues/IssueService.cs
--- a/Engine/Source/Programs/Horde/Plugins/Build/HordeServer.Build/Issues/IssueService.cs
+++ b/Engine/Source/Programs/Horde/Plugins/Build/HordeServer.Build/Issues/IssueService.cs
@@ -333,7 +333,7 @@ namespace HordeServer.Issues
 
 				if (openIssue.LastSeenAt < utcNow - TimeSpan.FromDays(7.0))
 				{
-					await _issueCollection.TryUpdateIssueAsync(openIssue, null, newResolvedById: IIssue.ResolvedByTimeoutId, cancellationToken: cancellationToken);
+					await _issueCollection.TryUpdateIssueAsync(openIssue, null, new UpdateIssueOptions { ResolvedById = IIssue.ResolvedByTimeoutId }, cancellationToken: cancellationToken);
 					await UpdateIssueDerivedDataAsync(openIssue, cancellationToken);
 					openIssues.RemoveAt(idx--);
 					continue;
@@ -467,7 +467,7 @@ namespace HordeServer.Issues
 					return false;
 				}
 
-				issue = await _issueCollection.TryUpdateIssueAsync(issue, initiatedById, newUserSummary: summary, newDescription: description, newPromoted: promoted, newOwnerId: ownerId ?? resolvedById, newNominatedById: nominatedById, newDeclinedById: declinedById, newAcknowledged: acknowledged, newFixCommitId: fixCommitId, newFixSystemic: fixSystemic, newResolvedById: resolvedById, newExcludeSpanIds: removeSpanIds, newExternalIssueKey: externalIssueKey, newQuarantinedById: quarantinedById, newForceClosedById: forceClosedById, newWorkflowThreadUrl: workflowThreadUrl, cancellationToken: cancellationToken);
+				issue = await _issueCollection.TryUpdateIssueAsync(issue, initiatedById, new UpdateIssueOptions { UserSummary = summary, Description = description, Promoted = promoted, OwnerId = ownerId ?? resolvedById, NominatedById = nominatedById, DeclinedById = declinedById, Acknowledged = acknowledged, FixCommitId = fixCommitId, FixSystemic = fixSystemic, ResolvedById = resolvedById, ExcludeSpanIds = removeSpanIds, ExternalIssueKey = externalIssueKey, QuarantinedById = quarantinedById, ForceClosedById = forceClosedById, WorkflowThreadUrl = workflowThreadUrl }, cancellationToken: cancellationToken);
 				if (issue != null)
 				{
 					break;
diff --git a/Engine/Source/Programs/Horde/Plugins/Build/HordeServer.Build/Issues/IssueTagService.cs b/Engine/Source/Programs/Horde/Plugins/Build/HordeServer.Build/Issues/IssueTagService.cs
--- a/Engine/Source/Programs/Horde/Plugins/Build/HordeServer.Build/Issues/IssueTagService.cs
+++ b/Engine/Source/Programs/Horde/Plugins/Build/HordeServer.Build/Issues/IssueTagService.cs
@@ -129,7 +129,7 @@ namespace HordeServer.Issues
 							break;
 						}
 
-						issue = await _issueCollection.TryUpdateIssueAsync(issue, commit.AuthorId, newFixCommitId: commit.Id, newResolvedById: commit.AuthorId, cancellationToken: cancellationToken);
+						issue = await _issueCollection.TryUpdateIssueAsync(issue, commit.AuthorId, new UpdateIssueOptions { FixCommitId = commit.Id, ResolvedById = commit.AuthorId }, cancellationToken: cancellationToken);
 						if (issue != null)
 						{
 							_logger.LogInformation("Commit {Change} by {Author} in {StreamId} fixes issue id {IssueId}", commit.Id, commit.AuthorId, streamConfig.Id, issueId);
diff --git a/Engine/Source/Programs/Horde/Plugins/Build/HordeServer.Build/Notifications/Sinks/SlackNotificationSink.cs b/Engine/Source/Programs/Horde/Plugins/Build/HordeServer.Build/Notifications/Sinks/SlackNotificationSink.cs
--- a/Engine/Source/Programs/Horde/Plugins/Build/HordeServer.Build/Notifications/Sinks/SlackNotificationSink.cs
+++ b/Engine/Source/Programs/Horde/Plugins/Build/HordeServer.Build/Notifications/Sinks/SlackNotificationSink.cs
@@ -1093,7 +1093,7 @@ namespace HordeServer.Notifications.Sinks
 							break;
 						}
 
-						updateIssue = await _issueService.Collection.TryUpdateIssueAsync(updateIssue, null, newWorkflowThreadUrl: new Uri(permalink), cancellationToken: cancellationToken);
+						updateIssue = await _issueService.Collection.TryUpdateIssueAsync(updateIssue, null, new UpdateIssueOptions { WorkflowThreadUrl = new Uri(permalink) }, cancellationToken: cancellationToken);
 						if (updateIssue != null)
 						{
 							break;
