Horde: Fix issue with advanced job priority setting when rerunning a job
#rnx
#skipci

diff --git a/Engine/Source/Programs/Horde/HordeDashboard/src/components/build/NewBuildV2.tsx b/Engine/Source/Programs/Horde/HordeDashboard/src/components/build/NewBuildV2.tsx
--- a/Engine/Source/Programs/Horde/HordeDashboard/src/components/build/NewBuildV2.tsx
+++ b/Engine/Source/Programs/Horde/HordeDashboard/src/components/build/NewBuildV2.tsx
@@ -102,6 +102,8 @@ class BuildOptions {
          this.preflightChange = this.queryShelvedChange;
       }
 
+      this.advJobPriority = jobDetails?.jobData?.priority;
+
       if (query.get("autosubmit") === "true") {
          this.autoSubmit = true;
       }
