Horde: Add a record containing options for searches. Move search logic into JobCollection.

#rnx

diff --git a/Engine/Source/Programs/Horde/Plugins/Build/HordeServer.Build.Tests/Jobs/JobTaskSourceTests.cs b/Engine/Source/Programs/Horde/Plugins/Build/HordeServer.Build.Tests/Jobs/JobTaskSourceTests.cs
--- a/Engine/Source/Programs/Horde/Plugins/Build/HordeServer.Build.Tests/Jobs/JobTaskSourceTests.cs
+++ b/Engine/Source/Programs/Horde/Plugins/Build/HordeServer.Build.Tests/Jobs/JobTaskSourceTests.cs
@@ -114,7 +114,7 @@ namespace HordeServer.Tests.Jobs
 			IGraph graph = await GraphCollection.AppendAsync(null, newGroups);
 
 			// remove the default fixture jobs
-			IReadOnlyList<IJob> jobs = await JobCollection.FindAsync();
+			IReadOnlyList<IJob> jobs = await JobCollection.FindAsync(new FindJobOptions());
 			for (int i = 0; i < jobs.Count; i++)
 			{
 				await jobs[i].TryDeleteAsync();
diff --git a/Engine/Source/Programs/Horde/Plugins/Build/HordeServer.Build.Tests/Jobs/Schedules/SchedulerTests.cs b/Engine/Source/Programs/Horde/Plugins/Build/HordeServer.Build.Tests/Jobs/Schedules/SchedulerTests.cs
--- a/Engine/Source/Programs/Horde/Plugins/Build/HordeServer.Build.Tests/Jobs/Schedules/SchedulerTests.cs
+++ b/Engine/Source/Programs/Horde/Plugins/Build/HordeServer.Build.Tests/Jobs/Schedules/SchedulerTests.cs
@@ -42,7 +42,7 @@ namespace HordeServer.Tests.Jobs.Schedules
 
 			_template = TemplateCollection.GetOrAddAsync(new TemplateConfig { Name = "Test template" }).Result;
 
-			_initialJobIds = new HashSet<JobId>(JobCollection.FindAsync().Result.Select(x => x.Id));
+			_initialJobIds = new HashSet<JobId>(JobCollection.FindAsync(new FindJobOptions()).Result.Select(x => x.Id));
 
 			PerforceService.Changes.Clear();
 			PerforceService.AddChange(StreamId, 100, bob, "", new[] { "code.cpp" });
@@ -719,7 +719,7 @@ namespace HordeServer.Tests.Jobs.Schedules
 
 		async Task<List<IJob>> GetNewJobsAsync()
 		{
-			List<IJob> jobs = (await JobCollection.FindAsync()).ToList();
+			List<IJob> jobs = (await JobCollection.FindAsync(new FindJobOptions())).ToList();
 			jobs.RemoveAll(x => _initialJobIds.Contains(x.Id));
 			_initialJobIds.UnionWith(jobs.Select(x => x.Id));
 			return jobs.OrderBy(x => x.CommitId).ToList();
diff --git a/Engine/Source/Programs/Horde/Plugins/Build/HordeServer.Build/Agents/Fleet/JobQueueStrategy.cs b/Engine/Source/Programs/Horde/Plugins/Build/HordeServer.Build/Agents/Fleet/JobQueueStrategy.cs
--- a/Engine/Source/Programs/Horde/Plugins/Build/HordeServer.Build/Agents/Fleet/JobQueueStrategy.cs
+++ b/Engine/Source/Programs/Horde/Plugins/Build/HordeServer.Build/Agents/Fleet/JobQueueStrategy.cs
@@ -113,7 +113,7 @@ namespace HordeServer.Agents.Fleet
 			span.SetAttribute("after", jobsCreatedAfter);
 
 			Dictionary<StreamId, StreamConfig> streams = _buildConfig.CurrentValue.Streams.ToDictionary(x => x.Id, x => (StreamConfig)x);
-			IReadOnlyList<IJob> recentJobs = await _jobs.FindAsync(minCreateTime: jobsCreatedAfter, batchState: JobStepBatchState.Ready, cancellationToken: cancellationToken);
+			IReadOnlyList<IJob> recentJobs = await _jobs.FindAsync(new FindJobOptions(MinCreateTime: jobsCreatedAfter, BatchState: JobStepBatchState.Ready), cancellationToken: cancellationToken);
 			span.SetAttribute("numJobs", recentJobs.Count);
 			span.SetAttribute("numUniqueGraphs", recentJobs.Select(x => x.GraphHash).Distinct().Count());
 
diff --git a/Engine/Source/Programs/Horde/Plugins/Build/HordeServer.Build/Issues/IssueReportService.cs b/Engine/Source/Programs/Horde/Plugins/Build/HordeServer.Build/Issues/IssueReportService.cs
--- a/Engine/Source/Programs/Horde/Plugins/Build/HordeServer.Build/Issues/IssueReportService.cs
+++ b/Engine/Source/Programs/Horde/Plugins/Build/HordeServer.Build/Issues/IssueReportService.cs
@@ -224,7 +224,7 @@ namespace HordeServer.Issues
 
 		private async Task<Dictionary<WorkflowId, WorkflowStats>> GetWorkflowStatsAsync(StreamConfig streamConfig, DateTime minTime, CancellationToken cancellationToken)
 		{
-			IReadOnlyList<IJob> jobs = await _jobCollection.FindAsync(streamId: streamConfig.Id, minCreateTime: minTime, cancellationToken: cancellationToken);
+			IReadOnlyList<IJob> jobs = await _jobCollection.FindAsync(new FindJobOptions(StreamId: streamConfig.Id, MinCreateTime: minTime), cancellationToken: cancellationToken);
 
 			Dictionary<WorkflowId, WorkflowStats> workflowIdToStats = new Dictionary<WorkflowId, WorkflowStats>();
 			foreach (IGrouping<TemplateId, IJob> templateGroup in jobs.GroupBy(x => x.TemplateId))
diff --git a/Engine/Source/Programs/Horde/Plugins/Build/HordeServer.Build/Jobs/DebugJobsController.cs b/Engine/Source/Programs/Horde/Plugins/Build/HordeServer.Build/Jobs/DebugJobsController.cs
--- a/Engine/Source/Programs/Horde/Plugins/Build/HordeServer.Build/Jobs/DebugJobsController.cs
+++ b/Engine/Source/Programs/Horde/Plugins/Build/HordeServer.Build/Jobs/DebugJobsController.cs
@@ -187,7 +187,7 @@ namespace HordeServer.Jobs
 			}
 
 			minCreateTime ??= DateTimeOffset.UtcNow.Subtract(TimeSpan.FromDays(1));
-			IReadOnlyList<IJob> jobs = await _jobService.FindJobsAsync(minCreateTime: minCreateTime, maxCreateTime: maxCreateTime);
+			IReadOnlyList<IJob> jobs = await _jobService.FindJobsAsync(new FindJobOptions(MinCreateTime: minCreateTime, MaxCreateTime: maxCreateTime));
 			List<JobTiming> allJobTimings = await CalculateJobTimingsAsync(jobs);
 
 			if (onlySetupBuild)
@@ -343,7 +343,7 @@ namespace HordeServer.Jobs
 				return BadRequest();
 			}
 
-			IReadOnlyList<IJob> jobs = await _jobService.FindJobsAsync(minCreateTime: minTime, maxCreateTime: maxTime);
+			IReadOnlyList<IJob> jobs = await _jobService.FindJobsAsync(new FindJobOptions(MinCreateTime: minTime, MaxCreateTime: maxTime));
 			foreach (IJob job in jobs)
 			{
 				StreamConfig? streamConfig;
diff --git a/Engine/Source/Programs/Horde/Plugins/Build/HordeServer.Build/Jobs/IJobCollection.cs b/Engine/Source/Programs/Horde/Plugins/Build/HordeServer.Build/Jobs/IJobCollection.cs
--- a/Engine/Source/Programs/Horde/Plugins/Build/HordeServer.Build/Jobs/IJobCollection.cs
+++ b/Engine/Source/Programs/Horde/Plugins/Build/HordeServer.Build/Jobs/IJobCollection.cs
@@ -150,30 +150,12 @@ namespace HordeServer.Jobs
 		/// <summary>
 		/// Searches for jobs matching the given criteria
 		/// </summary>
-		/// <param name="jobIds">List of job ids to return</param>
-		/// <param name="streamId">The stream containing the job</param>
-		/// <param name="name">Name of the job</param>
-		/// <param name="templates">Templates to look for</param>
-		/// <param name="minCommitId">The minimum commit</param>
-		/// <param name="maxCommitId">The maximum commit</param>
-		/// <param name="preflightCommitId">Preflight change to find</param>
-		/// <param name="preflightOnly">Whether to only include preflights</param>
-		/// <param name="includePreflight">Whether to include preflights in the results</param>
-		/// <param name="startedByUser">User id for which to include jobs</param>
-		/// <param name="preflightStartedByUser">User for which to include preflight jobs</param>
-		/// <param name="minCreateTime">The minimum creation time</param>
-		/// <param name="maxCreateTime">The maximum creation time</param>
-		/// <param name="modifiedBefore">Filter the results by modified time</param>
-		/// <param name="modifiedAfter">Filter the results by modified time</param>
-		/// <param name="batchState">One or more batches matches this state</param>
+		/// <param name="options">Options for the search</param>
 		/// <param name="index">Index of the first result to return</param>
 		/// <param name="count">Number of results to return</param>
-		/// <param name="consistentRead">If the database read should be made to the replica server</param>
-		/// <param name="indexHint">Name of index to be specified as a hint to the database query planner</param>
-		/// <param name="excludeUserJobs">Whether to exclude user jobs from the find</param>
 		/// <param name="cancellationToken">Cancellation token for the operation</param>
 		/// <returns>List of jobs matching the given criteria</returns>
-		Task<IReadOnlyList<IJob>> FindAsync(JobId[]? jobIds = null, StreamId? streamId = null, string? name = null, TemplateId[]? templates = null, CommitId? minCommitId = null, CommitId? maxCommitId = null, CommitId? preflightCommitId = null, bool? preflightOnly = null, bool? includePreflight = null, UserId? preflightStartedByUser = null, UserId? startedByUser = null, DateTimeOffset? minCreateTime = null, DateTimeOffset? maxCreateTime = null, DateTimeOffset? modifiedBefore = null, DateTimeOffset? modifiedAfter = null, JobStepBatchState? batchState = null, int? index = null, int? count = null, bool consistentRead = true, string? indexHint = null, bool? excludeUserJobs = null, CancellationToken cancellationToken = default);
+		Task<IReadOnlyList<IJob>> FindAsync(FindJobOptions options, int? index = null, int? count = null, CancellationToken cancellationToken = default);
 
 		/// <summary>
 		/// Searches for jobs matching the given criteria
@@ -220,4 +202,53 @@ namespace HordeServer.Jobs
 		/// <returns>Async task</returns>
 		Task UpgradeDocumentsAsync();
 	}
+
+	/// <summary>
+	/// Options for finding jobs
+	/// </summary>
+	/// <param name="JobIds">List of job ids to return</param>
+	/// <param name="StreamId">The stream containing the job</param>
+	/// <param name="Name">Name of the job</param>
+	/// <param name="Templates">Templates to look for</param>
+	/// <param name="MinCommitId">The minimum commit</param>
+	/// <param name="MaxCommitId">The maximum commit</param>
+	/// <param name="PreflightCommitId">Preflight change to find</param>
+	/// <param name="PreflightOnly">Whether to only include preflights</param>
+	/// <param name="IncludePreflight">Whether to include preflights in the results</param>
+	/// <param name="StartedByUser">User id for which to include jobs</param>
+	/// <param name="PreflightStartedByUser">User for which to include preflight jobs</param>
+	/// <param name="MinCreateTime">The minimum creation time</param>
+	/// <param name="MaxCreateTime">The maximum creation time</param>
+	/// <param name="ModifiedBefore">Filter the results by modified time</param>
+	/// <param name="ModifiedAfter">Filter the results by modified time</param>
+	/// <param name="Target">The target to query</param>
+	/// <param name="State">State to query</param>
+	/// <param name="Outcome">Outcomes to return</param>
+	/// <param name="BatchState">One or more batches matches this state</param>
+	/// <param name="ExcludeUserJobs">Whether to exclude user jobs from the find</param>
+	/// <param name="ExcludeCancelled">Whether to exclude cancelled jobs</param>
+	public record class FindJobOptions
+	(
+		JobId[]? JobIds = null,
+		StreamId? StreamId = null,
+		string? Name = null,
+		TemplateId[]? Templates = null,
+		CommitId? MinCommitId = null,
+		CommitId? MaxCommitId = null,
+		CommitId? PreflightCommitId = null,
+		bool? PreflightOnly = null,
+		bool? IncludePreflight = null,
+		UserId? PreflightStartedByUser = null,
+		UserId? StartedByUser = null,
+		DateTimeOffset? MinCreateTime = null,
+		DateTimeOffset? MaxCreateTime = null,
+		DateTimeOffset? ModifiedBefore = null,
+		DateTimeOffset? ModifiedAfter = null,
+		string? Target = null,
+		JobStepState[]? State = null,
+		JobStepOutcome[]? Outcome = null,
+		JobStepBatchState? BatchState = null,
+		bool? ExcludeUserJobs = null,
+		bool? ExcludeCancelled = null
+	);
 }
diff --git a/Engine/Source/Programs/Horde/Plugins/Build/HordeServer.Build/Jobs/JobCollection.cs b/Engine/Source/Programs/Horde/Plugins/Build/HordeServer.Build/Jobs/JobCollection.cs
--- a/Engine/Source/Programs/Horde/Plugins/Build/HordeServer.Build/Jobs/JobCollection.cs
+++ b/Engine/Source/Programs/Horde/Plugins/Build/HordeServer.Build/Jobs/JobCollection.cs
@@ -898,106 +898,202 @@ namespace HordeServer.Jobs
 		}
 
 		/// <inheritdoc/>
-		public async Task<IReadOnlyList<IJob>> FindAsync(JobId[]? jobIds, StreamId? streamId, string? name, TemplateId[]? templates, CommitId? minCommitId, CommitId? maxCommitId, CommitId? preflightCommitId, bool? preflightOnly, bool? includePreflight, UserId? preflightStartedByUser, UserId? startedByUser, DateTimeOffset? minCreateTime, DateTimeOffset? maxCreateTime, DateTimeOffset? modifiedBefore, DateTimeOffset? modifiedAfter, JobStepBatchState? batchState, int? index, int? count, bool consistentRead, string? indexHint, bool? excludeUserJobs, CancellationToken cancellationToken)
+		public async Task<IReadOnlyList<IJob>> FindAsync(FindJobOptions options, int? index = null, int? count = null, CancellationToken cancellationToken = default)
+		{
+			// JobId[]? jobIds = null, StreamId? streamId = null, string? name = null, TemplateId[]? templates = null, CommitId? minCommitId = null, CommitId? maxCommitId = null, CommitId? preflightCommitId = null, bool? preflightOnly = null, bool? includePreflights = null, UserId? preflightStartedByUser = null, UserId? startedByUser = null, DateTimeOffset? minCreateTime = null, DateTimeOffset? maxCreateTime = null, string? target = null, JobStepBatchState? batchState = null, JobStepState[]? state = null, JobStepOutcome[]? outcome = null, DateTimeOffset? modifiedBefore = null, DateTimeOffset? modifiedAfter = null, int? index = null, int? count = null, bool consistentRead = true, bool? excludeUserJobs = null, bool? excludeCancelled = null, 
+			using TelemetrySpan span = _tracer.StartActiveSpan($"{nameof(JobCollection)}.{nameof(FindAsync)}");
+			span.SetAttribute("JobIds", (options.JobIds == null) ? null : String.Join(',', options.JobIds));
+			span.SetAttribute("StreamId", options.StreamId);
+			span.SetAttribute("Name", options.Name);
+			span.SetAttribute("Templates", options.Templates);
+			span.SetAttribute("MinCommitId", options.MinCommitId?.ToString());
+			span.SetAttribute("MaxCommitId", options.MaxCommitId?.ToString());
+			span.SetAttribute("PreflightCommitId", options.PreflightCommitId?.ToString());
+			span.SetAttribute("PreflightStartedByUser", options.PreflightStartedByUser?.ToString());
+			span.SetAttribute("StartedByUser", options.StartedByUser?.ToString());
+			span.SetAttribute("MinCreateTime", options.MinCreateTime);
+			span.SetAttribute("MaxCreateTime", options.MaxCreateTime);
+			span.SetAttribute("Target", options.Target);
+			span.SetAttribute("State", options.State?.ToString());
+			span.SetAttribute("Outcome", options.Outcome?.ToString());
+			span.SetAttribute("ModifiedBefore", options.ModifiedBefore);
+			span.SetAttribute("ModifiedAfter", options.ModifiedAfter);
+			span.SetAttribute("Index", index);
+			span.SetAttribute("Count", count);
+
+			if (options.Target == null && (options.State == null || options.State.Length == 0) && (options.Outcome == null || options.Outcome.Length == 0))
+			{
+				return await FindInternalAsync(options, index, count, cancellationToken: cancellationToken);
+			}
+			else
+			{
+				List<IJob> results = new List<IJob>();
+				_logger.LogInformation("Performing scan for job with ");
+
+				bool excludeMaxCommitId = false;
+
+				int maxCount = (count ?? 1);
+				while (results.Count < maxCount)
+				{
+					IReadOnlyList<IJob> scanJobs = await FindInternalAsync(options, 0, 5, cancellationToken: cancellationToken);
+					if (scanJobs.Count == 0)
+					{
+						break;
+					}
+
+					foreach (IJob job in scanJobs.OrderByDescending(x => x.CommitId))
+					{
+						if (excludeMaxCommitId && job.CommitId == options.MaxCommitId)
+						{
+							continue;
+						}
+
+						if (options.ExcludeCancelled != null && options.ExcludeCancelled.Value && WasCancelled(job))
+						{
+							continue;
+						}
+
+						(JobStepState, JobStepOutcome)? result;
+						if (options.Target == null)
+						{
+							result = job.GetTargetState();
+						}
+						else
+						{
+							result = job.GetTargetState(await _graphCollection.GetAsync(job.GraphHash, cancellationToken), options.Target);
+						}
+
+						if (result != null)
+						{
+							(JobStepState jobState, JobStepOutcome jobOutcome) = result.Value;
+							if ((options.State == null || options.State.Length == 0 || options.State.Contains(jobState)) && (options.Outcome == null || options.Outcome.Length == 0 || options.Outcome.Contains(jobOutcome)))
+							{
+								results.Add(job);
+								if (results.Count == maxCount)
+								{
+									break;
+								}
+							}
+						}
+					}
+
+					options = options with { MaxCommitId = scanJobs.Min(x => x.CommitId) };
+					excludeMaxCommitId = true;
+				}
+
+				return results;
+			}
+		}
+
+		/// <summary>
+		/// Test whether a job was cancelled
+		/// </summary>
+		static bool WasCancelled(IJob job)
+		{
+			return job.AbortedByUserId != null || job.Batches.Any(x => x.Steps.Any(y => y.AbortedByUserId != null));
+		}
+
+		/// <inheritdoc/>
+		public async Task<IReadOnlyList<IJob>> FindInternalAsync(FindJobOptions options, int? index = null, int? count = null, bool consistentRead = false, string? indexHint = null, CancellationToken cancellationToken = default)
 		{
 			FilterDefinitionBuilder<JobDocument> filterBuilder = Builders<JobDocument>.Filter;
 
 			FilterDefinition<JobDocument> filter = filterBuilder.Empty;
-			if (jobIds != null && jobIds.Length > 0)
+			if (options.JobIds != null && options.JobIds.Length > 0)
 			{
-				filter &= filterBuilder.In(x => x.Id, jobIds);
+				filter &= filterBuilder.In(x => x.Id, options.JobIds);
 			}
-			if (streamId != null)
+			if (options.StreamId != null)
 			{
-				filter &= filterBuilder.Eq(x => x.StreamId, streamId.Value);
-				if (minCommitId != null)
+				filter &= filterBuilder.Eq(x => x.StreamId, options.StreamId.Value);
+				if (options.MinCommitId != null)
 				{
-					CommitIdWithOrder minCommitIdWithOrder = await _commitService.GetOrderedAsync(streamId.Value, minCommitId, cancellationToken);
+					CommitIdWithOrder minCommitIdWithOrder = await _commitService.GetOrderedAsync(options.StreamId.Value, options.MinCommitId, cancellationToken);
 					filter &= filterBuilder.Gte(x => x.CommitOrder, minCommitIdWithOrder.Order);
 				}
-				if (maxCommitId != null)
+				if (options.MaxCommitId != null)
 				{
-					CommitIdWithOrder maxCommitIdWithOrder = await _commitService.GetOrderedAsync(streamId.Value, maxCommitId, cancellationToken);
+					CommitIdWithOrder maxCommitIdWithOrder = await _commitService.GetOrderedAsync(options.StreamId.Value, options.MaxCommitId, cancellationToken);
 					filter &= filterBuilder.Lte(x => x.CommitOrder, maxCommitIdWithOrder.Order);
 				}
 			}
-			if (name != null)
+			if (options.Name != null)
 			{
-				if (name.StartsWith("$", StringComparison.InvariantCulture))
+				if (options.Name.StartsWith("$", StringComparison.InvariantCulture))
 				{
-					BsonRegularExpression regex = new BsonRegularExpression(name.Substring(1), "i");
+					BsonRegularExpression regex = new BsonRegularExpression(options.Name.Substring(1), "i");
 					filter &= filterBuilder.Regex(x => x.Name, regex);
 				}
 				else
 				{
-					filter &= filterBuilder.Eq(x => x.Name, name);
+					filter &= filterBuilder.Eq(x => x.Name, options.Name);
 				}
 			}
-			if (templates != null)
+			if (options.Templates != null)
 			{
-				if (templates.Length == 1)
+				if (options.Templates.Length == 1)
 				{
-					filter &= filterBuilder.Eq(x => x.TemplateId, templates[0]);
+					filter &= filterBuilder.Eq(x => x.TemplateId, options.Templates[0]);
 				}
 				else
 				{
-					filter &= filterBuilder.In(x => x.TemplateId, templates);
+					filter &= filterBuilder.In(x => x.TemplateId, options.Templates);
 				}
 			}
-			if (preflightCommitId != null)
+			if (options.PreflightCommitId != null)
 			{
-				int? preflightChange = preflightCommitId.TryGetPerforceChange();
+				int? preflightChange = options.PreflightCommitId.TryGetPerforceChange();
 				if (preflightChange == null)
 				{
-					filter &= filterBuilder.Eq(x => x.PreflightCommitName, preflightCommitId.Name);
+					filter &= filterBuilder.Eq(x => x.PreflightCommitName, options.PreflightCommitId.Name);
 				}
 				else
 				{
 					filter &= filterBuilder.Eq(x => x.PreflightChange, preflightChange);
 				}
 			}
-			if (includePreflight != null && !includePreflight.Value)
+			if (options.IncludePreflight != null && !options.IncludePreflight.Value)
 			{
 				filter &= filterBuilder.Eq(x => x.PreflightChange, 0);
 			}
-			if (preflightOnly != null && preflightOnly.Value)
+			if (options.PreflightOnly != null && options.PreflightOnly.Value)
 			{
 				filter &= filterBuilder.Ne(x => x.PreflightChange, 0);
 			}
-			if (excludeUserJobs != null && excludeUserJobs.Value)
+			if (options.ExcludeUserJobs != null && options.ExcludeUserJobs.Value)
 			{
 				filter &= filterBuilder.Eq(x => x.StartedByUserId, null);
 			}
 			else
 			{
-				if (preflightStartedByUser != null)
+				if (options.PreflightStartedByUser != null)
 				{
-					filter &= filterBuilder.Or(filterBuilder.Eq(x => x.PreflightChange, 0), filterBuilder.Eq(x => x.StartedByUserId, preflightStartedByUser));
+					filter &= filterBuilder.Or(filterBuilder.Eq(x => x.PreflightChange, 0), filterBuilder.Eq(x => x.StartedByUserId, options.PreflightStartedByUser));
 				}
-				if (startedByUser != null)
+				if (options.StartedByUser != null)
 				{
-					filter &= filterBuilder.Eq(x => x.StartedByUserId, startedByUser);
+					filter &= filterBuilder.Eq(x => x.StartedByUserId, options.StartedByUser);
 				}
 			}
-			if (minCreateTime != null)
+			if (options.MinCreateTime != null)
 			{
-				filter &= filterBuilder.Gte(x => x.CreateTimeUtc!, minCreateTime.Value.UtcDateTime);
+				filter &= filterBuilder.Gte(x => x.CreateTimeUtc!, options.MinCreateTime.Value.UtcDateTime);
 			}
-			if (maxCreateTime != null)
+			if (options.MaxCreateTime != null)
 			{
-				filter &= filterBuilder.Lte(x => x.CreateTimeUtc!, maxCreateTime.Value.UtcDateTime);
+				filter &= filterBuilder.Lte(x => x.CreateTimeUtc!, options.MaxCreateTime.Value.UtcDateTime);
 			}
-			if (modifiedBefore != null)
+			if (options.ModifiedBefore != null)
 			{
-				filter &= filterBuilder.Lte(x => x.UpdateTimeUtc!, modifiedBefore.Value.UtcDateTime);
+				filter &= filterBuilder.Lte(x => x.UpdateTimeUtc!, options.ModifiedBefore.Value.UtcDateTime);
 			}
-			if (modifiedAfter != null)
+			if (options.ModifiedAfter != null)
 			{
-				filter &= filterBuilder.Gte(x => x.UpdateTimeUtc!, modifiedAfter.Value.UtcDateTime);
+				filter &= filterBuilder.Gte(x => x.UpdateTimeUtc!, options.ModifiedAfter.Value.UtcDateTime);
 			}
-			if (batchState != null)
+			if (options.BatchState != null)
 			{
-				filter &= filterBuilder.ElemMatch(x => x.Batches, batch => batch.State == batchState);
+				filter &= filterBuilder.ElemMatch(x => x.Batches, batch => batch.State == options.BatchState);
 			}
 
 			List<JobDocument> documents;
@@ -1047,13 +1143,30 @@ namespace HordeServer.Jobs
 		/// <inheritdoc/>
 		public async Task<IReadOnlyList<IJob>> FindLatestByStreamWithTemplatesAsync(StreamId streamId, TemplateId[] templates, UserId? preflightStartedByUser, DateTimeOffset? maxCreateTime, DateTimeOffset? modifiedAfter, int? index, int? count, bool consistentRead, CancellationToken cancellationToken)
 		{
+			using TelemetrySpan span = _tracer.StartActiveSpan($"{nameof(JobCollection)}.{nameof(FindLatestByStreamWithTemplatesAsync)}");
+			span.SetAttribute("StreamId", streamId);
+			span.SetAttribute("Templates", templates);
+			span.SetAttribute("PreflightStartedByUser", preflightStartedByUser?.ToString());
+			span.SetAttribute("MaxCreateTime", maxCreateTime);
+			span.SetAttribute("ModifiedAfter", modifiedAfter);
+			span.SetAttribute("Index", index);
+			span.SetAttribute("Count", count);
+
 			string indexHint = _streamThenTemplateThenCreationTimeIndex.Name;
 
 			// This find call uses an index hint. Modifying the parameter passed to FindAsync can affect execution time a lot as the query planner is forced to use the specified index.
-			// Casting to interface to benefit from default parameter values
-			return await (this as IJobCollection).FindAsync(
-				streamId: streamId, templates: templates, preflightStartedByUser: preflightStartedByUser, modifiedAfter: modifiedAfter, maxCreateTime: maxCreateTime,
-				index: index, count: count, indexHint: indexHint, consistentRead: consistentRead, cancellationToken: cancellationToken);
+			return await FindInternalAsync(
+				new FindJobOptions(
+					StreamId: streamId, 
+					Templates: templates, 
+					PreflightStartedByUser: preflightStartedByUser, 
+					ModifiedAfter: modifiedAfter, 
+					MaxCreateTime: maxCreateTime),
+				index: index,
+				count: count,
+				indexHint: indexHint,
+				consistentRead: consistentRead, 
+				cancellationToken: cancellationToken);
 		}
 
 		/// <inheritdoc/>
diff --git a/Engine/Source/Programs/Horde/Plugins/Build/HordeServer.Build/Jobs/JobExpirationService.cs b/Engine/Source/Programs/Horde/Plugins/Build/HordeServer.Build/Jobs/JobExpirationService.cs
--- a/Engine/Source/Programs/Horde/Plugins/Build/HordeServer.Build/Jobs/JobExpirationService.cs
+++ b/Engine/Source/Programs/Horde/Plugins/Build/HordeServer.Build/Jobs/JobExpirationService.cs
@@ -53,7 +53,7 @@ namespace HordeServer.Jobs
 					DateTimeOffset maxCreateTime = new DateTimeOffset(currentTime) - TimeSpan.FromDays(expireAfterDays);
 					for (; ; )
 					{
-						IReadOnlyList<IJob> jobs = await _jobCollection.FindAsync(streamId: streamConfig.Id, maxCreateTime: maxCreateTime, count: 50, cancellationToken: cancellationToken);
+						IReadOnlyList<IJob> jobs = await _jobCollection.FindAsync(new FindJobOptions(StreamId: streamConfig.Id, MaxCreateTime: maxCreateTime), count: 50, cancellationToken: cancellationToken);
 						if (jobs.Count == 0)
 						{
 							break;
diff --git a/Engine/Source/Programs/Horde/Plugins/Build/HordeServer.Build/Jobs/JobService.cs b/Engine/Source/Programs/Horde/Plugins/Build/HordeServer.Build/Jobs/JobService.cs
--- a/Engine/Source/Programs/Horde/Plugins/Build/HordeServer.Build/Jobs/JobService.cs
+++ b/Engine/Source/Programs/Horde/Plugins/Build/HordeServer.Build/Jobs/JobService.cs
@@ -246,7 +246,7 @@ namespace HordeServer.Jobs
 			IReadOnlyList<IJob> jobsToAbort = new List<IJob>();
 			if (newJob.PreflightCommitId != null)
 			{
-				jobsToAbort = await _jobs.FindAsync(preflightCommitId: newJob.PreflightCommitId, cancellationToken: cancellationToken);
+				jobsToAbort = await _jobs.FindAsync(new FindJobOptions(PreflightCommitId: newJob.PreflightCommitId), cancellationToken: cancellationToken);
 			}
 
 			foreach (IJob job in jobsToAbort)
@@ -433,7 +433,7 @@ namespace HordeServer.Jobs
 				{
 					List<JobStepOutcome> outcomes = query.Outcomes ?? new List<JobStepOutcome> { JobStepOutcome.Success };
 
-					IReadOnlyList<IJob> jobs = await FindJobsAsync(streamId: streamId, templates: new[] { query.TemplateId.Value }, target: query.Target, state: new[] { JobStepState.Completed }, outcome: outcomes.ToArray(), count: 1, excludeUserJobs: true, excludeCancelled: true, cancellationToken: cancellationToken);
+					IReadOnlyList<IJob> jobs = await FindJobsAsync(new FindJobOptions(StreamId: streamId, Templates: new[] { query.TemplateId.Value }, Target: query.Target, State: new[] { JobStepState.Completed }, Outcome: outcomes.ToArray(), ExcludeUserJobs: true, ExcludeCancelled: true), count: 1, cancellationToken: cancellationToken);
 					if (jobs.Count > 0)
 					{
 						_logger.LogInformation("Last successful build of {TemplateId} target {Target} was job {JobId} at change {Change}", query.TemplateId, query.Target, jobs[0].Id, jobs[0].CommitId);
@@ -531,124 +531,14 @@ namespace HordeServer.Jobs
 		/// <summary>
 		/// Searches for jobs matching the given criteria
 		/// </summary>
-		/// <param name="jobIds">List of job ids to return</param>
-		/// <param name="streamId">The stream containing the job</param>
-		/// <param name="name">Name of the job</param>
-		/// <param name="templates">Templates to look for</param>
-		/// <param name="minCommitId">The minimum commit</param>
-		/// <param name="maxCommitId">The maximum commit</param>		
-		/// <param name="preflightCommitId">The preflight change to look for</param>
-		/// <param name="preflightOnly">Whether to only include preflights</param>
-		/// <param name="includePreflights">Whether to include preflight jobs</param>
-		/// <param name="preflightStartedByUser">User for which to include preflight jobs</param>		
-		/// <param name="startedByUser">User for which to include jobs</param>
-		/// <param name="minCreateTime">The minimum creation time</param>
-		/// <param name="maxCreateTime">The maximum creation time</param>
-		/// <param name="target">The target to query</param>
-		/// <param name="batchState">One or more batches matches this state</param>
-		/// <param name="state">State to query</param>
-		/// <param name="outcome">Outcomes to return</param>
-		/// <param name="modifiedBefore">Filter the results by last modified time</param>
-		/// <param name="modifiedAfter">Filter the results by last modified time</param>
+		/// <param name="options">Options for the search</param>
 		/// <param name="index">Index of the first result to return</param>
 		/// <param name="count">Number of results to return</param>
-		/// <param name="consistentRead">If the database read should be made to the replica server</param>
-		/// <param name="excludeUserJobs">Whether to exclude user jobs from the find</param>
-		/// <param name="excludeCancelled">Whether to exclude cancelled jobs</param>
 		/// <param name="cancellationToken">Cancellation token for the operation</param>
 		/// <returns>List of jobs matching the given criteria</returns>
-		public async Task<IReadOnlyList<IJob>> FindJobsAsync(JobId[]? jobIds = null, StreamId? streamId = null, string? name = null, TemplateId[]? templates = null, CommitId? minCommitId = null, CommitId? maxCommitId = null, CommitId? preflightCommitId = null, bool? preflightOnly = null, bool? includePreflights = null, UserId? preflightStartedByUser = null, UserId? startedByUser = null, DateTimeOffset? minCreateTime = null, DateTimeOffset? maxCreateTime = null, string? target = null, JobStepBatchState? batchState = null, JobStepState[]? state = null, JobStepOutcome[]? outcome = null, DateTimeOffset? modifiedBefore = null, DateTimeOffset? modifiedAfter = null, int? index = null, int? count = null, bool consistentRead = true, bool? excludeUserJobs = null, bool? excludeCancelled = null, CancellationToken cancellationToken = default)
-		{
-			using TelemetrySpan span = _tracer.StartActiveSpan($"{nameof(JobService)}.{nameof(FindJobsAsync)}");
-			span.SetAttribute("JobIds", (jobIds == null) ? null : String.Join(',', jobIds));
-			span.SetAttribute("StreamId", streamId);
-			span.SetAttribute("Name", name);
-			span.SetAttribute("Templates", templates);
-			span.SetAttribute("MinCommitId", minCommitId?.ToString());
-			span.SetAttribute("MaxCommitId", maxCommitId?.ToString());
-			span.SetAttribute("PreflightCommitId", preflightCommitId?.ToString());
-			span.SetAttribute("PreflightStartedByUser", preflightStartedByUser?.ToString());
-			span.SetAttribute("StartedByUser", startedByUser?.ToString());
-			span.SetAttribute("MinCreateTime", minCreateTime);
-			span.SetAttribute("MaxCreateTime", maxCreateTime);
-			span.SetAttribute("Target", target);
-			span.SetAttribute("State", state?.ToString());
-			span.SetAttribute("Outcome", outcome?.ToString());
-			span.SetAttribute("ModifiedBefore", modifiedBefore);
-			span.SetAttribute("ModifiedAfter", modifiedAfter);
-			span.SetAttribute("Index", index);
-			span.SetAttribute("Count", count);
-
-			if (target == null && (state == null || state.Length == 0) && (outcome == null || outcome.Length == 0))
-			{
-				return await _jobs.FindAsync(jobIds, streamId, name, templates, minCommitId, maxCommitId, preflightCommitId, preflightOnly, includePreflights, preflightStartedByUser, startedByUser, minCreateTime, maxCreateTime, modifiedBefore, modifiedAfter, batchState, index, count, consistentRead, null, excludeUserJobs, cancellationToken);
-			}
-			else
-			{
-				List<IJob> results = new List<IJob>();
-				_logger.LogInformation("Performing scan for job with ");
-
-				bool excludeMaxCommitId = false;
-
-				int maxCount = (count ?? 1);
-				while (results.Count < maxCount)
-				{
-					IReadOnlyList<IJob> scanJobs = await _jobs.FindAsync(jobIds, streamId, name, templates, minCommitId, maxCommitId, preflightCommitId, preflightOnly, includePreflights, preflightStartedByUser, startedByUser, minCreateTime, maxCreateTime, modifiedBefore, modifiedAfter, batchState, 0, 5, consistentRead, null, excludeUserJobs, cancellationToken);
-					if (scanJobs.Count == 0)
-					{
-						break;
-					}
-
-					foreach (IJob job in scanJobs.OrderByDescending(x => x.CommitId))
-					{
-						if (excludeMaxCommitId && job.CommitId == maxCommitId)
-						{
-							continue;
-						}
-
-						if (excludeCancelled != null && excludeCancelled.Value && WasCancelled(job))
-						{
-							continue;
-						}
-
-						(JobStepState, JobStepOutcome)? result;
-						if (target == null)
-						{
-							result = job.GetTargetState();
-						}
-						else
-						{
-							result = job.GetTargetState(await GetGraphAsync(job, cancellationToken), target);
-						}
-
-						if (result != null)
-						{
-							(JobStepState jobState, JobStepOutcome jobOutcome) = result.Value;
-							if ((state == null || state.Length == 0 || state.Contains(jobState)) && (outcome == null || outcome.Length == 0 || outcome.Contains(jobOutcome)))
-							{
-								results.Add(job);
-								if (results.Count == maxCount)
-								{
-									break;
-								}
-							}
-						}
-					}
-
-					maxCommitId = scanJobs.Min(x => x.CommitId);
-					excludeMaxCommitId = true;
-				}
-
-				return results;
-			}
-		}
-
-		/// <summary>
-		/// Test whether a job was cancelled
-		/// </summary>
-		static bool WasCancelled(IJob job)
+		public Task<IReadOnlyList<IJob>> FindJobsAsync(FindJobOptions options, int? index = null, int? count = null, CancellationToken cancellationToken = default)
 		{
-			return job.AbortedByUserId != null || job.Batches.Any(x => x.Steps.Any(y => y.AbortedByUserId != null));
+			return _jobs.FindAsync(options, index, count, cancellationToken);
 		}
 
 		/// <summary>
@@ -666,15 +556,6 @@ namespace HordeServer.Jobs
 		/// <returns>List of jobs matching the given criteria</returns>
 		public async Task<IReadOnlyList<IJob>> FindJobsByStreamWithTemplatesAsync(StreamId streamId, TemplateId[] templates, UserId? preflightStartedByUser = null, DateTimeOffset? maxCreateTime = null, DateTimeOffset? modifiedAfter = null, int? index = null, int? count = null, bool consistentRead = true, CancellationToken cancellationToken = default)
 		{
-			using TelemetrySpan span = _tracer.StartActiveSpan($"{nameof(JobService)}.{nameof(FindJobsByStreamWithTemplatesAsync)}");
-			span.SetAttribute("StreamId", streamId);
-			span.SetAttribute("Templates", templates);
-			span.SetAttribute("PreflightStartedByUser", preflightStartedByUser?.ToString());
-			span.SetAttribute("MaxCreateTime", maxCreateTime);
-			span.SetAttribute("ModifiedAfter", modifiedAfter);
-			span.SetAttribute("Index", index);
-			span.SetAttribute("Count", count);
-
 			return await _jobs.FindLatestByStreamWithTemplatesAsync(streamId, templates, preflightStartedByUser, maxCreateTime, modifiedAfter, index, count, consistentRead, cancellationToken);
 		}
 
diff --git a/Engine/Source/Programs/Horde/Plugins/Build/HordeServer.Build/Jobs/JobsController.cs b/Engine/Source/Programs/Horde/Plugins/Build/HordeServer.Build/Jobs/JobsController.cs
--- a/Engine/Source/Programs/Horde/Plugins/Build/HordeServer.Build/Jobs/JobsController.cs
+++ b/Engine/Source/Programs/Horde/Plugins/Build/HordeServer.Build/Jobs/JobsController.cs
@@ -11,7 +11,6 @@ using EpicGames.Horde.Streams;
 using EpicGames.Horde.Users;
 using HordeServer.Acls;
 using HordeServer.Agents;
-using HordeServer.Artifacts;
 using HordeServer.Commits;
 using HordeServer.Jobs.Graphs;
 using HordeServer.Jobs.Templates;
@@ -1095,9 +1094,9 @@ namespace HordeServer.Jobs
 			}
 
 			IReadOnlyList<IJob> jobs;
-			jobs = await _jobService.FindJobsAsync(jobIdValues, streamIdValue, name, templateRefIds, minChange,
-				maxChange, preflightChange, preflightOnly, includePreflight, preflightStartedByUserIdValue, startedByUserIdValue, minCreateTime?.UtcDateTime, maxCreateTime?.UtcDateTime, target, null, state, outcome,
-				modifiedBefore, modifiedAfter, index, count, false);
+			jobs = await _jobService.FindJobsAsync(new FindJobOptions(jobIdValues, streamIdValue, name, templateRefIds, minChange,
+				maxChange, preflightChange, preflightOnly, includePreflight, preflightStartedByUserIdValue, startedByUserIdValue, minCreateTime?.UtcDateTime, maxCreateTime?.UtcDateTime, Target: target, State: state, Outcome: outcome,
+				ModifiedBefore: modifiedBefore, ModifiedAfter: modifiedAfter), index, count);
 
 			return await CreateAuthorizedJobResponsesAsync(jobs, filter);
 		}
diff --git a/Engine/Source/Programs/Horde/Plugins/Build/HordeServer.Build/Jobs/Schedules/ScheduleService.cs b/Engine/Source/Programs/Horde/Plugins/Build/HordeServer.Build/Jobs/Schedules/ScheduleService.cs
--- a/Engine/Source/Programs/Horde/Plugins/Build/HordeServer.Build/Jobs/Schedules/ScheduleService.cs
+++ b/Engine/Source/Programs/Horde/Plugins/Build/HordeServer.Build/Jobs/Schedules/ScheduleService.cs
@@ -514,7 +514,7 @@ namespace HordeServer.Jobs.Schedules
 			{
 				cancellationToken.ThrowIfCancellationRequested();
 
-				IReadOnlyList<IJob> jobs = await _jobCollection.FindAsync(streamId: streamId, templates: new[] { gate.TemplateId }, minCommitId: minCommitId, maxCommitId: maxCommitId, count: 2, cancellationToken: cancellationToken);
+				IReadOnlyList<IJob> jobs = await _jobCollection.FindAsync(new FindJobOptions(StreamId: streamId, Templates: new[] { gate.TemplateId }, MinCommitId: minCommitId, MaxCommitId: maxCommitId), count: 2, cancellationToken: cancellationToken);
 
 				IJob? job = jobs.FirstOrDefault(x => maxCommitId == null || x.CommitId < maxCommitId);
 				if (job == null)
