Horde: Remove logic for cloning changes before running preflights. This has been disabled for some time due to corruption issues on edge servers.

#rnx

diff --git a/Engine/Source/Programs/Horde/Drivers/JobDriver/Execution/PerforceExecutor.cs b/Engine/Source/Programs/Horde/Drivers/JobDriver/Execution/PerforceExecutor.cs
--- a/Engine/Source/Programs/Horde/Drivers/JobDriver/Execution/PerforceExecutor.cs
+++ b/Engine/Source/Programs/Horde/Drivers/JobDriver/Execution/PerforceExecutor.cs
@@ -125,8 +125,7 @@ namespace JobDriver.Execution
 				}
 
 				// Sync the workspace
-				int syncPreflightChange = (Batch.ClonedPreflightChange != 0) ? Batch.ClonedPreflightChange : Batch.PreflightChange;
-				await _workspace.SyncAsync(perforce, Batch.Change, syncPreflightChange, null, cancellationToken);
+				await _workspace.SyncAsync(perforce, Batch.Change, Batch.PreflightChange, null, cancellationToken);
 
 				// Remove any cached BuildGraph manifests
 				DirectoryReference manifestDir = DirectoryReference.Combine(_workspace.WorkspaceDir, "Engine", "Saved", "BuildGraph");
diff --git a/Engine/Source/Programs/Horde/Drivers/JobDriver/Execution/WorkspaceExecutor.cs b/Engine/Source/Programs/Horde/Drivers/JobDriver/Execution/WorkspaceExecutor.cs
--- a/Engine/Source/Programs/Horde/Drivers/JobDriver/Execution/WorkspaceExecutor.cs
+++ b/Engine/Source/Programs/Horde/Drivers/JobDriver/Execution/WorkspaceExecutor.cs
@@ -61,8 +61,7 @@ namespace JobDriver.Execution
 			{
 				scope.Span.SetTag(Datadog.Trace.OpenTracing.DatadogTags.ResourceName, _workspace.Identifier);
 
-				int preflightChange = (Batch.ClonedPreflightChange != 0) ? Batch.ClonedPreflightChange : Batch.PreflightChange;
-				await _workspace.SyncAsync(Batch.Change, preflightChange, new SyncOptions(), cancellationToken);
+				await _workspace.SyncAsync(Batch.Change, Batch.PreflightChange, new SyncOptions(), cancellationToken);
 
 				// TODO: Purging of cache for ManagedWorkspace did happen here in WorkspaceInfo
 
diff --git a/Engine/Source/Programs/Horde/Plugins/Build/HordeServer.Build/Jobs/IJob.cs b/Engine/Source/Programs/Horde/Plugins/Build/HordeServer.Build/Jobs/IJob.cs
--- a/Engine/Source/Programs/Horde/Plugins/Build/HordeServer.Build/Jobs/IJob.cs
+++ b/Engine/Source/Programs/Horde/Plugins/Build/HordeServer.Build/Jobs/IJob.cs
@@ -114,11 +114,6 @@ namespace HordeServer.Jobs
 		/// </summary>
 		public CommitId? PreflightCommitId { get; }
 
-		/// <summary>
-		/// The cloned preflight changelist number (if the prefight change is duplicated via p4 reshelve)
-		/// </summary>
-		public CommitId? ClonedPreflightCommitId { get; }
-
 		/// <summary>
 		/// Description for the shelved change if running a preflight
 		/// </summary>
@@ -1182,7 +1177,6 @@ namespace HordeServer.Jobs
 			response.Change = job.CommitId.GetPerforceChange();
 			response.CodeChange = job.CodeCommitId?.GetPerforceChange() ?? 0;
 			response.PreflightChange = job.PreflightCommitId?.GetPerforceChange() ?? 0;
-			response.ClonedPreflightChange = job.ClonedPreflightCommitId?.GetPerforceChange() ?? 0;
 			response.Arguments.Add(job.Arguments);
 			return response;
 		}
diff --git a/Engine/Source/Programs/Horde/Plugins/Build/HordeServer.Build/Jobs/IJobCollection.cs b/Engine/Source/Programs/Horde/Plugins/Build/HordeServer.Build/Jobs/IJobCollection.cs
--- a/Engine/Source/Programs/Horde/Plugins/Build/HordeServer.Build/Jobs/IJobCollection.cs
+++ b/Engine/Source/Programs/Horde/Plugins/Build/HordeServer.Build/Jobs/IJobCollection.cs
@@ -23,9 +23,6 @@ namespace HordeServer.Jobs
 		/// <inheritdoc cref="IJob.PreflightCommitId"/>
 		public CommitId? PreflightCommitId { get; set; }
 
-		/// <inheritdoc cref="IJob.ClonedPreflightCommitId"/>
-		public CommitId? ClonedPreflightCommitId { get; set; }
-
 		/// <inheritdoc cref="IJob.PreflightDescription"/>
 		public string? PreflightDescription { get; set; }
 
diff --git a/Engine/Source/Programs/Horde/Plugins/Build/HordeServer.Build/Jobs/JobCollection.cs b/Engine/Source/Programs/Horde/Plugins/Build/HordeServer.Build/Jobs/JobCollection.cs
--- a/Engine/Source/Programs/Horde/Plugins/Build/HordeServer.Build/Jobs/JobCollection.cs
+++ b/Engine/Source/Programs/Horde/Plugins/Build/HordeServer.Build/Jobs/JobCollection.cs
@@ -66,7 +66,6 @@ namespace HordeServer.Jobs
 			CommitIdWithOrder IJob.CommitId => _document.CommitId;
 			CommitIdWithOrder? IJob.CodeCommitId => _document.CodeCommitId;
 			CommitId? IJob.PreflightCommitId => _document.PreflightCommitId;
-			CommitId? IJob.ClonedPreflightCommitId => _document.ClonedPreflightCommitId;
 			string? IJob.PreflightDescription => _document.PreflightDescription;
 			Priority IJob.Priority => _document.Priority;
 			bool IJob.AutoSubmit => _document.AutoSubmit;
@@ -456,9 +455,6 @@ namespace HordeServer.Jobs
 			// -1 for non-P4 preflights
 			public int PreflightChange { get; set; }
 
-			[BsonElement("ClonedPreflightChange"), BsonIgnoreIfNull]
-			public CommitId? ClonedPreflightCommitId { get; set; }
-
 			public string? PreflightDescription { get; set; }
 			public Priority Priority { get; set; }
 
@@ -542,7 +538,6 @@ namespace HordeServer.Jobs
 				CommitId = commitId;
 				CodeCommitId = codeCommitId;
 				PreflightCommitId = options.PreflightCommitId;
-				ClonedPreflightCommitId = options.ClonedPreflightCommitId;
 				PreflightDescription = options.PreflightDescription;
 				StartedByUserId = options.StartedByUserId;
 				StartedByBisectTaskId = options.StartedByBisectTaskId;
diff --git a/Engine/Source/Programs/Horde/Plugins/Build/HordeServer.Build/Jobs/JobRpcService.cs b/Engine/Source/Programs/Horde/Plugins/Build/HordeServer.Build/Jobs/JobRpcService.cs
--- a/Engine/Source/Programs/Horde/Plugins/Build/HordeServer.Build/Jobs/JobRpcService.cs
+++ b/Engine/Source/Programs/Horde/Plugins/Build/HordeServer.Build/Jobs/JobRpcService.cs
@@ -200,7 +200,6 @@ namespace HordeServer.Jobs
 			response.Change = job.CommitId.GetPerforceChange();
 			response.CodeChange = job.CodeCommitId?.GetPerforceChange() ?? 0;
 			response.PreflightChange = job.PreflightCommitId?.GetPerforceChange() ?? 0;
-			response.ClonedPreflightChange = job.ClonedPreflightCommitId?.GetPerforceChange() ?? 0;
 			response.Arguments.AddRange(job.Arguments);
 			if (agentConfig.TempStorageDir != null)
 			{
diff --git a/Engine/Source/Programs/Horde/Plugins/Build/HordeServer.Build/Jobs/JobService.cs b/Engine/Source/Programs/Horde/Plugins/Build/HordeServer.Build/Jobs/JobService.cs
--- a/Engine/Source/Programs/Horde/Plugins/Build/HordeServer.Build/Jobs/JobService.cs
+++ b/Engine/Source/Programs/Horde/Plugins/Build/HordeServer.Build/Jobs/JobService.cs
@@ -142,13 +142,6 @@ namespace HordeServer.Jobs
 			};
 		}
 
-#pragma warning disable IDE0060
-		static bool ShouldClonePreflightChange(StreamId streamId)
-		{
-			return false; //			return StreamId == new StreamId("ue5-main");
-		}
-#pragma warning restore IDE0060
-
 		/// <summary>
 		/// Creates a new job
 		/// </summary>
@@ -175,7 +168,6 @@ namespace HordeServer.Jobs
 			span.SetAttribute("Change", commitId.ToString());
 			span.SetAttribute("CodeChange", codeCommitId?.ToString());
 			span.SetAttribute("PreflightChange", options.PreflightCommitId?.GetPerforceChange());
-			span.SetAttribute("ClonedPreflightChange", options.ClonedPreflightCommitId?.GetPerforceChange());
 			span.SetAttribute("StartedByUserId", options.StartedByUserId.ToString());
 			span.SetAttribute("Priority", options.Priority.ToString());
 			span.SetAttribute("ShowUgsBadges", options.ShowUgsBadges);
@@ -199,18 +191,12 @@ namespace HordeServer.Jobs
 			JobId jobIdValue = jobId ?? JobIdUtils.GenerateNewId();
 			using IDisposable? scope = _logger.BeginScope("CreateJobAsync({JobId})", jobIdValue);
 
-			if (options.PreflightCommitId != null && ShouldClonePreflightChange(streamConfig.Id))
-			{
-				options.ClonedPreflightCommitId = await CloneShelvedChangeAsync(streamConfig.ClusterName, options.ClonedPreflightCommitId ?? options.PreflightCommitId, cancellationToken);
-			}
-
-			_logger.LogInformation("Creating job at CL {Change}, code CL {CodeChange}, preflight CL {PreflightChange}, cloned CL {ClonedPreflightChange}", commitId, codeCommitId, options.PreflightCommitId, options.ClonedPreflightCommitId);
+			_logger.LogInformation("Creating job at CL {Change}, code CL {CodeChange}, preflight CL {PreflightChange}", commitId, codeCommitId, options.PreflightCommitId);
 
 			Dictionary<string, string> properties = new Dictionary<string, string>();
 			properties["Change"] = commitId.ToString();
 			properties["CodeChange"] = codeCommitId?.ToString() ?? String.Empty;
 			properties["PreflightChange"] = options.PreflightCommitId?.ToString() ?? String.Empty;
-			properties["ClonedPreflightChange"] = options.ClonedPreflightCommitId?.ToString() ?? String.Empty;
 			properties["StreamId"] = streamConfig.Id.ToString();
 			properties["TemplateId"] = templateRefId.ToString();
 			properties["JobId"] = jobIdValue.ToString();
@@ -978,7 +964,7 @@ namespace HordeServer.Jobs
 							{
 								job = await AutoSubmitChangeAsync(streamConfig, job, graph, cancellationToken);
 							}
-							else if (job.ClonedPreflightCommitId == null && job.StartedByUserId.HasValue && outcome == JobStepOutcome.Success && job.AbortedByUserId == null)
+							else if (job.StartedByUserId.HasValue && outcome == JobStepOutcome.Success && job.AbortedByUserId == null)
 							{
 								IUserSettings settings = await _userCollection.GetSettingsAsync(job.StartedByUserId.Value, cancellationToken);
 								if (settings.AlwaysTagPreflightCL)
@@ -987,10 +973,6 @@ namespace HordeServer.Jobs
 									await _perforceService.UpdateChangelistDescriptionAsync(streamConfig.ClusterName, job.PreflightCommitId.GetPerforceChange(), x => x.TrimEnd() + $"\n#preflight {job.Id}", cancellationToken);
 								}
 							}
-							else if (job.ClonedPreflightCommitId != null)
-							{
-								await DeleteShelvedChangeAsync(streamConfig.ClusterName, job.ClonedPreflightCommitId);
-							}
 						}
 					}
 
@@ -1088,46 +1070,25 @@ namespace HordeServer.Jobs
 			string message;
 			try
 			{
-				CommitId? clonedPreflightChange = job.ClonedPreflightCommitId;
-				if (clonedPreflightChange == null)
-				{
-					if (ShouldClonePreflightChange(job.StreamId))
-					{
-						clonedPreflightChange = await CloneShelvedChangeAsync(streamConfig.ClusterName, job.PreflightCommitId, cancellationToken);
-					}
-					else
-					{
-						clonedPreflightChange = job.PreflightCommitId;
-					}
-				}
+				CommitId preflightCommitId = job.PreflightCommitId;
 
-				_logger.LogInformation("Updating description for {ClonedPreflightChange}", clonedPreflightChange);
+				_logger.LogInformation("Updating description for {PreflightChange}", preflightCommitId);
 
-				await _perforceService.UpdateChangelistDescriptionAsync(streamConfig.ClusterName, clonedPreflightChange.GetPerforceChange(), x => x.TrimEnd() + $"\n#preflight {job.Id}", cancellationToken);
+				await _perforceService.UpdateChangelistDescriptionAsync(streamConfig.ClusterName, preflightCommitId.GetPerforceChange(), x => x.TrimEnd() + $"\n#preflight {job.Id}", cancellationToken);
 
-				_logger.LogInformation("Submitting change {Change} (through {ChangeCopy}) after successful completion of {JobId}", job.PreflightCommitId, clonedPreflightChange, job.Id);
-				(change, message) = await _perforceService.SubmitShelvedChangeAsync(streamConfig, clonedPreflightChange.GetPerforceChange(), job.PreflightCommitId.GetPerforceChange(), cancellationToken);
+				_logger.LogInformation("Submitting change {Change} (through {ChangeCopy}) after successful completion of {JobId}", job.PreflightCommitId, preflightCommitId, job.Id);
+				(change, message) = await _perforceService.SubmitShelvedChangeAsync(streamConfig, preflightCommitId.GetPerforceChange(), job.PreflightCommitId.GetPerforceChange(), cancellationToken);
 
-				_logger.LogInformation("Attempt to submit {Change} (through {ChangeCopy}): {Message}", job.PreflightCommitId, clonedPreflightChange, message);
+				_logger.LogInformation("Attempt to submit {Change} (through {ChangeCopy}): {Message}", job.PreflightCommitId, preflightCommitId, message);
 
 				if (!String.IsNullOrEmpty(message))
 				{
 					message = Regex.Replace(message, @"^Submit validation failed.*\n(?:\s*'[^']*' validation failed:\s*\n)?", "");
 				}
 
-				if (ShouldClonePreflightChange(job.StreamId))
+				if (change != null)
 				{
-					if (change != null && job.ClonedPreflightCommitId != null)
-					{
-						await DeleteShelvedChangeAsync(streamConfig.ClusterName, job.PreflightCommitId);
-					}
-				}
-				else
-				{
-					if (change != null)
-					{
-						await DeleteShelvedChangeAsync(streamConfig.ClusterName, job.PreflightCommitId);
-					}
+					await DeleteShelvedChangeAsync(streamConfig.ClusterName, job.PreflightCommitId);
 				}
 			}
 			catch (Exception ex)
diff --git a/Engine/Source/Programs/Horde/Plugins/Build/HordeServer.Build/Jobs/JobsController.cs b/Engine/Source/Programs/Horde/Plugins/Build/HordeServer.Build/Jobs/JobsController.cs
--- a/Engine/Source/Programs/Horde/Plugins/Build/HordeServer.Build/Jobs/JobsController.cs
+++ b/Engine/Source/Programs/Horde/Plugins/Build/HordeServer.Build/Jobs/JobsController.cs
@@ -594,7 +594,6 @@ namespace HordeServer.Jobs
 			response.CommitId = job.CommitId;
 			response.CodeCommitId = job.CodeCommitId;
 			response.PreflightCommitId = job.PreflightCommitId;
-			response.ClonedPreflightCommitId = job.ClonedPreflightCommitId;
 			response.PreflightDescription = job.PreflightDescription;
 			response.TemplateHash = job.TemplateHash?.ToString() ?? String.Empty;
 			response.GraphHash = job.GraphHash.ToString();
diff --git a/Engine/Source/Programs/Shared/EpicGames.Horde/Jobs/JobMessages.cs b/Engine/Source/Programs/Shared/EpicGames.Horde/Jobs/JobMessages.cs
--- a/Engine/Source/Programs/Shared/EpicGames.Horde/Jobs/JobMessages.cs
+++ b/Engine/Source/Programs/Shared/EpicGames.Horde/Jobs/JobMessages.cs
@@ -485,27 +485,6 @@ namespace EpicGames.Horde.Jobs
 		}
 		CommitId? _preflightCommitId;
 
-		/// <summary>
-		/// The cloned preflight changelist number
-		/// </summary>
-		[Obsolete("Use ClonedPreflightCommitId instead")]
-		public int? ClonedPreflightChange
-		{
-			get => _clonedPreflightChange ?? _clonedPreflightCommitId?.GetPerforceChangeOrMinusOne();
-			set => _clonedPreflightChange = value;
-		}
-		int? _clonedPreflightChange;
-
-		/// <summary>
-		/// The cloned preflight commit
-		/// </summary>
-		public CommitId? ClonedPreflightCommitId
-		{
-			get => _clonedPreflightCommitId ?? Horde.Commits.CommitId.FromPerforceChange(_clonedPreflightChange);
-			set => _clonedPreflightCommitId = value;
-		}
-		CommitId? _clonedPreflightCommitId;
-
 		/// <summary>
 		/// Description of the preflight
 		/// </summary>
