Horde: Add a utility method for getting a stream config object.

#rnx

diff --git a/Engine/Source/Programs/Horde/Plugins/Build/HordeServer.Build/BuildConfig.cs b/Engine/Source/Programs/Horde/Plugins/Build/HordeServer.Build/BuildConfig.cs
--- a/Engine/Source/Programs/Horde/Plugins/Build/HordeServer.Build/BuildConfig.cs
+++ b/Engine/Source/Programs/Horde/Plugins/Build/HordeServer.Build/BuildConfig.cs
@@ -16,6 +16,7 @@ using EpicGames.Perforce;
 using HordeServer.Acls;
 using HordeServer.Agents;
 using HordeServer.Agents.Pools;
+using HordeServer.Commits;
 using HordeServer.Devices;
 using HordeServer.Perforce;
 using HordeServer.Plugins;
@@ -214,6 +215,21 @@ namespace HordeServer
 		/// <returns>True if the stream configuration was found</returns>
 		public bool TryGetProject(ProjectId projectId, [NotNullWhen(true)] out ProjectConfig? config) => _projectLookup.TryGetValue(projectId, out config);
 
+		/// <summary>
+		/// Get configuration for a stream from this object
+		/// </summary>
+		/// <param name="streamId">The stream identifier</param>
+		/// <returns>Configuration for the stream</returns>
+		public StreamConfig GetStream(StreamId streamId)
+		{
+			StreamConfig? streamConfig;
+			if (!TryGetStream(streamId, out streamConfig))
+			{
+				throw new StreamNotFoundException(streamId);
+			}
+			return streamConfig;
+		}
+
 		/// <summary>
 		/// Attempts to get configuration for a stream from this object
 		/// </summary>
