Horde: Change artifact paths to Type/Stream/Commit/Name/Id, to reflect permissions hierarchy.

diff --git a/Engine/Source/Programs/Horde/Plugins/Build/HordeServer.Build/Artifacts/ArtifactCollection.cs b/Engine/Source/Programs/Horde/Plugins/Build/HordeServer.Build/Artifacts/ArtifactCollection.cs
--- a/Engine/Source/Programs/Horde/Plugins/Build/HordeServer.Build/Artifacts/ArtifactCollection.cs
+++ b/Engine/Source/Programs/Horde/Plugins/Build/HordeServer.Build/Artifacts/ArtifactCollection.cs
@@ -128,7 +128,7 @@ namespace HordeServer.Artifacts
 		/// <summary>
 		/// Gets the base path for a set of artifacts
 		/// </summary>
-		public static string GetArtifactPath(StreamId streamId, ArtifactName name, ArtifactType type) => $"{streamId}/{name}/{type}";
+		public static string GetArtifactPath(StreamId streamId, ArtifactType type) => $"{type}/{streamId}";
 
 		/// <inheritdoc/>
 		public async Task<IArtifact> AddAsync(ArtifactName name, ArtifactType type, string? description, StreamId streamId, CommitId commitId, IEnumerable<string> keys, IEnumerable<string> metadata, AclScopeName scopeName, CancellationToken cancellationToken)
@@ -145,7 +145,7 @@ namespace HordeServer.Artifacts
 			ArtifactId id = new ArtifactId(BinaryIdUtils.CreateNew());
 
 			NamespaceId namespaceId = Namespace.Artifacts;
-			RefName refName = new RefName($"{GetArtifactPath(streamId, name, type)}/{commitId}/{id}");
+			RefName refName = new RefName($"{GetArtifactPath(streamId, type)}/{commitId}/{name}/{id}");
 
 			CommitIdWithOrder commitIdWithOrder = await _commitService.GetOrderedAsync(streamId, commitId, cancellationToken);
 
diff --git a/Engine/Source/Programs/Horde/Plugins/Build/HordeServer.Build/Artifacts/ArtifactsController.cs b/Engine/Source/Programs/Horde/Plugins/Build/HordeServer.Build/Artifacts/ArtifactsController.cs
--- a/Engine/Source/Programs/Horde/Plugins/Build/HordeServer.Build/Artifacts/ArtifactsController.cs
+++ b/Engine/Source/Programs/Horde/Plugins/Build/HordeServer.Build/Artifacts/ArtifactsController.cs
@@ -149,7 +149,7 @@ namespace HordeServer.Artifacts
 			RefName? prevRefName = await GetPrevRefNameForArtifactAsync(artifact, cancellationToken);
 
 			List<AclClaimConfig> claims = new List<AclClaimConfig>();
-			claims.Add(new AclClaimConfig(HordeClaimTypes.ReadNamespace, $"{artifact.NamespaceId}:{ArtifactCollection.GetArtifactPath(streamId, name, type)}"));
+			claims.Add(new AclClaimConfig(HordeClaimTypes.ReadNamespace, $"{artifact.NamespaceId}:{ArtifactCollection.GetArtifactPath(streamId, type)}"));
 			claims.Add(new AclClaimConfig(HordeClaimTypes.WriteNamespace, $"{artifact.NamespaceId}:{artifact.RefName}"));
 
 			string token = await _aclService.IssueBearerTokenAsync(claims, TimeSpan.FromHours(8.0), cancellationToken);
