Horde: Update release notes.

#rnx

diff --git a/Engine/Source/Programs/Horde/Docs/ReleaseNotes.md b/Engine/Source/Programs/Horde/Docs/ReleaseNotes.md
--- a/Engine/Source/Programs/Horde/Docs/ReleaseNotes.md
+++ b/Engine/Source/Programs/Horde/Docs/ReleaseNotes.md
@@ -1,5 +1,12 @@
 # Release Notes
 
+## 2024-07-29
+
+* Fix changelist number not being returned in artifact responses. (35140759)
+* Change artifact paths to Type/Stream/Commit/Name/Id, to reflect permissions hierarchy. (35129994)
+* Add extension methods to allow extracting IBlobRef<DirectoryNode> directly into directories. (35094979)
+* Preliminary support for VCS commit ids that aren't an integer. All endpoints still support passing changelist numbers for now. (35091676)
+
 ## 2024-07-25
 
 * Add a separate flag for marking an issue as fixed as a systemic issue, rather than having to pass a negative value as a fix changelist. Still supports passing/returning negative fix changelists as well for now, but will be removed in future. (35029295)
