Horde: Fix issue with advanced tab switching to basic on update
#rnx
#skipci

diff --git a/Engine/Source/Programs/Horde/HordeDashboard/src/components/build/NewBuildV2.tsx b/Engine/Source/Programs/Horde/HordeDashboard/src/components/build/NewBuildV2.tsx
--- a/Engine/Source/Programs/Horde/HordeDashboard/src/components/build/NewBuildV2.tsx
+++ b/Engine/Source/Programs/Horde/HordeDashboard/src/components/build/NewBuildV2.tsx
@@ -213,6 +213,10 @@ class BuildOptions {
    queryShelvedChange?: number;
    autoSubmit?: boolean;
 
+   static initialized() {
+      return !!BuildOptions.instance;
+   }
+
    static get() {
       if (!BuildOptions.instance) {
          throw `Build Options have not been instantiated`;
@@ -2044,9 +2048,15 @@ const NewBuildV2Inner: React.FC<{ setUseLegacyDialog: (value: boolean) => void }
 export const NewBuildV2: React.FC<{ streamId: string; show: boolean; onClose: (newJobId: string | undefined) => void, jobKey?: string; jobDetails?: JobDetailsV2, readOnly?: boolean }> = ({ streamId, jobKey, show, onClose, jobDetails, readOnly }) => {
 
    const { projectStore } = useBackend();
-   const options = useConst(new BuildOptions(streamId, projectStore, onClose, jobDetails, jobKey, readOnly));
    const [useNewBuildV1, setNewBuildV1] = useState<boolean>(false);
 
+   // initially had this using useConst, though not working
+   if (!BuildOptions.initialized()) {
+      new BuildOptions(streamId, projectStore, onClose, jobDetails, jobKey, readOnly);
+   }
+
+   const options = BuildOptions.get();
+
    useEffect(() => {
       return () => {
          options.clear();
