Horde: Add statically typed accessor methods for RpcLease.Id.

#rnx

diff --git a/Engine/Source/Programs/Horde/HordeAgent/HordeAgent.csproj b/Engine/Source/Programs/Horde/HordeAgent/HordeAgent.csproj
--- a/Engine/Source/Programs/Horde/HordeAgent/HordeAgent.csproj
+++ b/Engine/Source/Programs/Horde/HordeAgent/HordeAgent.csproj
@@ -109,16 +109,19 @@
   </ItemGroup>
 
   <ItemGroup>
-    <Content Include="..\HordeServer\Protos\horde\*.proto">
+    <Content Include="..\Plugins\Compute\HordeServer.Compute\Protos\horde\*.proto">
       <Link>Protos\%(FileName)%(Extension)</Link>
     </Content>
-    <Content Include="..\HordeServer\Protos\horde\common\**\*.proto">
+    <Content Include="..\Plugins\Compute\HordeServer.Compute\Protos\horde\common\**\*.proto">
       <Link>Protos\horde\common\%(RecursiveDir)%(FileName)%(Extension)</Link>
     </Content>
-    <Content Include="..\HordeServer\Protos\horde\messages\**\*.proto">
+    <Compile Include="..\Plugins\Compute\HordeServer.Compute\Protos\horde\messages\**\*.cs">
+      <Link>Protos\horde\messages\%(RecursiveDir)%(FileName)%(Extension)</Link>
+    </Compile>
+    <Content Include="..\Plugins\Compute\HordeServer.Compute\Protos\horde\messages\**\*.proto">
       <Link>Protos\horde\messages\%(RecursiveDir)%(FileName)%(Extension)</Link>
     </Content>
-    <Content Include="..\HordeServer\Protos\horde\tasks\**\*.proto">
+    <Content Include="..\Plugins\Compute\HordeServer.Compute\Protos\horde\tasks\**\*.proto">
       <Link>Protos\horde\tasks\%(RecursiveDir)%(FileName)%(Extension)</Link>
     </Content>
 <!--    
@@ -126,7 +129,9 @@
       <Link>Protos\horde\%(RecursiveDir)%(FileName)%(Extension)</Link>
     </Compile>
     -->
-    <Compile Include="$(IntermediateOutputPath)VersionInfo.cs" />
+    <Compile Include="$(IntermediateOutputPath)VersionInfo.cs">
+      <Link>Properties/VersionInfo.cs</Link>
+    </Compile>
   </ItemGroup>
 
   <ItemGroup>
diff --git a/Engine/Source/Programs/Horde/HordeAgent/Leases/LeaseHandler.cs b/Engine/Source/Programs/Horde/HordeAgent/Leases/LeaseHandler.cs
--- a/Engine/Source/Programs/Horde/HordeAgent/Leases/LeaseHandler.cs
+++ b/Engine/Source/Programs/Horde/HordeAgent/Leases/LeaseHandler.cs
@@ -59,7 +59,7 @@ namespace HordeAgent.Leases
 		/// </summary>
 		protected LeaseHandler(RpcLease rpcLease)
 		{
-			Id = LeaseId.Parse(rpcLease.Id);
+			Id = rpcLease.Id;
 			RpcLease = rpcLease;
 			RpcPayload = rpcLease.Payload;
 			Result = Task.FromException<LeaseResult>(new InvalidOperationException("Lease has not been started"));
@@ -105,8 +105,8 @@ namespace HordeAgent.Leases
 		/// </summary>
 		async Task<LeaseResult> HandleLeaseAsync(ISession session, ILogger logger, LeaseLoggerFactory leaseLoggerFactory)
 		{
-			using IScope scope = GlobalTracer.Instance.BuildSpan("HandleLease").WithResourceName(RpcLease.Id).StartActive();
-			scope.Span.SetTag("LeaseId", RpcLease.Id);
+			using IScope scope = GlobalTracer.Instance.BuildSpan("HandleLease").WithResourceName(RpcLease.Id.ToString()).StartActive();
+			scope.Span.SetTag("LeaseId", RpcLease.Id.ToString());
 			scope.Span.SetTag("AgentId", session.AgentId.ToString());
 			//			using IDisposable TraceProperty = LogContext.PushProperty("dd.trace_id", CorrelationIdentifier.TraceId.ToString());
 			//			using IDisposable SpanProperty = LogContext.PushProperty("dd.span_id", CorrelationIdentifier.SpanId.ToString());
diff --git a/Engine/Source/Programs/Horde/Plugins/Compute/HordeServer.Compute/Agents/AgentService.cs b/Engine/Source/Programs/Horde/Plugins/Compute/HordeServer.Compute/Agents/AgentService.cs
--- a/Engine/Source/Programs/Horde/Plugins/Compute/HordeServer.Compute/Agents/AgentService.cs
+++ b/Engine/Source/Programs/Horde/Plugins/Compute/HordeServer.Compute/Agents/AgentService.cs
@@ -491,8 +491,8 @@ namespace HordeServer.Agents
 		/// <returns>True if a lease was assigned, false otherwise</returns>
 		public async Task<IAgent?> WaitForLeaseAsync(IAgent? agent, IList<HordeCommon.Rpc.Messages.RpcLease> newLeases, CancellationToken cancellationToken = default)
 		{
-			HashSet<string> knownLeases = new HashSet<string>(newLeases.Select(x => x.Id), StringComparer.OrdinalIgnoreCase);
-			while (agent != null && agent.Leases.All(x => knownLeases.Contains(x.Id.ToString())))
+			HashSet<LeaseId> knownLeases = new HashSet<LeaseId>(newLeases.Select(x => x.Id));
+			while (agent != null && agent.Leases.All(x => knownLeases.Contains(x.Id)))
 			{
 				if (!agent.SessionExpiresAt.HasValue)
 				{
@@ -755,7 +755,7 @@ namespace HordeServer.Agents
 				List<AgentLease> leases = new List<AgentLease>(agent.Leases);
 
 				// Remove any completed leases from the agent
-				Dictionary<LeaseId, RpcLease> leaseIdToNewState = newLeases.ToDictionary(x => LeaseId.Parse(x.Id), x => x);
+				Dictionary<LeaseId, RpcLease> leaseIdToNewState = newLeases.ToDictionary(x => x.Id, x => x);
 				for (int idx = 0; idx < leases.Count; idx++)
 				{
 					AgentLease lease = leases[idx];
diff --git a/Engine/Source/Programs/Horde/Plugins/Compute/HordeServer.Compute/Agents/IAgent.cs b/Engine/Source/Programs/Horde/Plugins/Compute/HordeServer.Compute/Agents/IAgent.cs
--- a/Engine/Source/Programs/Horde/Plugins/Compute/HordeServer.Compute/Agents/IAgent.cs
+++ b/Engine/Source/Programs/Horde/Plugins/Compute/HordeServer.Compute/Agents/IAgent.cs
@@ -914,7 +914,7 @@ namespace HordeServer.Agents
 		public RpcLease ToRpcMessage()
 		{
 			RpcLease lease = new RpcLease();
-			lease.Id = Id.ToString();
+			lease.Id = Id;
 			lease.Payload = Google.Protobuf.WellKnownTypes.Any.Parser.ParseFrom(Payload);
 			lease.State = (RpcLeaseState)State;
 			return lease;
diff --git a/Engine/Source/Programs/Horde/Plugins/Compute/HordeServer.Compute/Protos/horde/messages/lease.proto.cs b/Engine/Source/Programs/Horde/Plugins/Compute/HordeServer.Compute/Protos/horde/messages/lease.proto.cs
new file mode 100644
--- /dev/null
+++ b/Engine/Source/Programs/Horde/Plugins/Compute/HordeServer.Compute/Protos/horde/messages/lease.proto.cs
@@ -0,0 +1,21 @@
+// Copyright Epic Games, Inc. All Rights Reserved.
+
+using EpicGames.Horde.Agents.Leases;
+
+namespace HordeCommon.Rpc.Messages
+{
+	/// <summary>
+	/// Additional methods for RpcLease
+	/// </summary>
+	partial class RpcLease
+	{
+		/// <summary>
+		/// Typed acesssor for the <see cref="RpcLease.IdString"/> field.
+		/// </summary>
+		public LeaseId Id
+		{
+			get => LeaseId.Parse(IdString);
+			set => IdString = value.ToString();
+		}
+	}
+}
