Horde: Improve log issue, warning, error, and search navigation, also normalize query string line index
#rnx
#skipci

diff --git a/Engine/Source/Programs/Horde/HordeDashboard/src/components/AutomationReport.tsx b/Engine/Source/Programs/Horde/HordeDashboard/src/components/AutomationReport.tsx
--- a/Engine/Source/Programs/Horde/HordeDashboard/src/components/AutomationReport.tsx
+++ b/Engine/Source/Programs/Horde/HordeDashboard/src/components/AutomationReport.tsx
@@ -331,7 +331,7 @@ const ErrorPane: React.FC<{ failure: TestFailureEvent }> = ({ failure }) => {
          return <div>???</div>;
       }
 
-      const url = `/log/${failure.logId}?lineindex=${event.lineIndex}`;
+      const url = `/log/${failure.logId}?lineIndex=${event.lineIndex + 1}`;
 
       const lines = event.lines.filter(line => line.message?.trim().length).map(line => <Stack styles={{ root: { paddingLeft: 8, paddingRight: 8, lineBreak: "anywhere", whiteSpace: "pre-wrap", lineHeight: 18, fontSize: 10, fontFamily: "Horde Cousine Regular, monospace, monospace" } }}> <Link className="log-link" target="_blank" to={url}>{renderLine(navigate, line, undefined, {})}</Link></Stack>);
 
diff --git a/Engine/Source/Programs/Horde/HordeDashboard/src/components/ErrorPane.tsx b/Engine/Source/Programs/Horde/HordeDashboard/src/components/ErrorPane.tsx
--- a/Engine/Source/Programs/Horde/HordeDashboard/src/components/ErrorPane.tsx
+++ b/Engine/Source/Programs/Horde/HordeDashboard/src/components/ErrorPane.tsx
@@ -112,7 +112,7 @@ export const ErrorPane: React.FC<{ jobDetails: JobDetails; stepId: string; showE
          return <div>???</div>;
       }
 
-      const url = `/log/${item.logId}?lineindex=${item.lineIndex}`;
+      const url = `/log/${item.logId}?lineIndex=${item.lineIndex + 1}`;
 
       const lines = item.lines.filter(line => line.message?.trim().length).map(line => <Stack styles={{ root: { paddingLeft: 8, paddingRight: 8, lineBreak: "anywhere", whiteSpace: "pre-wrap", lineHeight: 18, fontSize: 10, fontFamily: "Horde Cousine Regular, monospace, monospace" } }}> <Link className="log-link" to={url}>{renderLine(navigate, line, undefined, {})}</Link></Stack>);
 
@@ -206,7 +206,7 @@ export const JobEventListPanel: React.FC<{ jobDetails: JobDetails, stepIds: stri
       const logId = step.logId!;
       const event = item.event!;
 
-      const url = `/log/${logId}?lineindex=${event.lineIndex}`;
+      const url = `/log/${logId}?lineIndex=${event.lineIndex + 1}`;
 
       const lines = event.lines.map(line => <Stack styles={{ root: { paddingLeft: 8, paddingRight: 8, lineBreak: "anywhere", whiteSpace: "pre-wrap", lineHeight: 18, fontSize: 10, fontFamily: "Horde Cousine Regular, monospace, monospace" } }}> {renderLine(navigate, line, undefined, {})}</Stack>);
       return (
diff --git a/Engine/Source/Programs/Horde/HordeDashboard/src/components/IssueViewV2.tsx b/Engine/Source/Programs/Horde/HordeDashboard/src/components/IssueViewV2.tsx
--- a/Engine/Source/Programs/Horde/HordeDashboard/src/components/IssueViewV2.tsx
+++ b/Engine/Source/Programs/Horde/HordeDashboard/src/components/IssueViewV2.tsx
@@ -1430,7 +1430,7 @@ export const ErrorPane: React.FC<{ events?: GetLogEventResponse[]; onClose?: ()
          return <div>???</div>;
       }
 
-      const url = `/log/${item.logId}?lineindex=${item.lineIndex}`;
+      const url = `/log/${item.logId}?lineIndex=${item.lineIndex + 1}`;
 
       const lines = item.lines.filter(line => line.message?.trim().length).map(line => <Stack key={`errorpane_line_${item.lineIndex}_${lineKey++}`} styles={{ root: { paddingLeft: 8, paddingRight: 8, lineBreak: "normal", whiteSpace: "pre-wrap", lineHeight: 18, fontSize: 10, fontFamily: "Horde Cousine Regular, monospace, monospace" } }}> <Link style={{ color: modeColors.text }} to={url}>{renderLine(navigate, line, undefined, {})}</Link></Stack>);
 
@@ -1539,7 +1539,7 @@ const StepPanel: React.FC<{ streamId: string, hstep: GetIssueStepResponse }> = o
 
       let viewLogURL = `/log/${hstep.logId}`;
       if (logEvents?.length) {
-         viewLogURL += `?lineindex=${logEvents[0].lineIndex}`;
+         viewLogURL += `?lineIndex=${logEvents[0].lineIndex + 1}`;
       }
 
       return <div style={{ paddingTop: 8, height: "100%" }}>
diff --git a/Engine/Source/Programs/Horde/HordeDashboard/src/components/LogSource.tsx b/Engine/Source/Programs/Horde/HordeDashboard/src/components/LogSource.tsx
--- a/Engine/Source/Programs/Horde/HordeDashboard/src/components/LogSource.tsx
+++ b/Engine/Source/Programs/Horde/HordeDashboard/src/components/LogSource.tsx
@@ -291,13 +291,13 @@ export abstract class LogSource {
       return new Promise<LogSource>(async (resolve, reject) => {
 
          let data: GetLogFileResponse | undefined;
-         
+
          try {
             data = await backend.getLogData(logId);
          } catch (reason) {
             reject(reason)
             return;
-         }         
+         }
 
          const value = Number(`0x${data.jobId}`);
 
@@ -317,8 +317,11 @@ export abstract class LogSource {
          }
 
          let line: number | undefined;
-         if (query.get("lineindex")) {
-            line = parseInt(query.get("lineindex")!);
+
+         if (query.get("lineIndex")) {
+            line = parseInt(query.get("lineIndex")!);
+         } else if (query.get("lineindex")) {
+            line = parseInt(query.get("lineindex")!) + 1;
          }
 
          source.startLine = line;
@@ -424,8 +427,7 @@ export class JobLogSource extends LogSource {
       return undefined;
    }
 
-   get agentTelemetry(): AgentTelemetryHandler | undefined 
-   {
+   get agentTelemetry(): AgentTelemetryHandler | undefined {
       return this._agentTelemety;
    }
 
@@ -516,8 +518,8 @@ export class JobLogSource extends LogSource {
       if (this.batch) {
          this.jobName = `Batch-${this.batch.id}`;
       }
-      
-      this.agentTelemetry?.set(this.agentId ?? "", new Date(this.startTime as any), this.step?.finishTime ? new Date(this.step.finishTime) : undefined);      
+
+      this.agentTelemetry?.set(this.agentId ?? "", new Date(this.startTime as any), this.step?.finishTime ? new Date(this.step.finishTime) : undefined);
 
    }
 
@@ -607,7 +609,7 @@ export class JobLogSource extends LogSource {
    agentId?: string;
    batch?: BatchData;
    step?: StepData;
-   
+
    artifactsV2?: GetArtifactResponse[];
 
    _agentTelemety: AgentTelemetryHandler = new AgentTelemetryHandler();
@@ -660,7 +662,7 @@ class LeaseLogSource extends LogSource {
                      name: "Missing Agent"
                   } as any;
                }
-               
+
 
                this.leaseUpdated();
                this.initComplete();
diff --git a/Engine/Source/Programs/Horde/HordeDashboard/src/components/LogStyle.tsx b/Engine/Source/Programs/Horde/HordeDashboard/src/components/LogStyle.tsx
--- a/Engine/Source/Programs/Horde/HordeDashboard/src/components/LogStyle.tsx
+++ b/Engine/Source/Programs/Horde/HordeDashboard/src/components/LogStyle.tsx
@@ -84,7 +84,7 @@ export const getLogStyles = () => {
          }
       },
       errorButtonDisabled: {         
-         color: dashboard.darktheme ? "#909398" : undefined,
+         color: dashboard.darktheme ? "#909398" : "#616E85 !important",
          backgroundColor: dashboard.darktheme ? "#1F2223" : "#f3f2f1"
       },
       warningButton: {
diff --git a/Engine/Source/Programs/Horde/HordeDashboard/src/components/LogView.tsx b/Engine/Source/Programs/Horde/HordeDashboard/src/components/LogView.tsx
--- a/Engine/Source/Programs/Horde/HordeDashboard/src/components/LogView.tsx
+++ b/Engine/Source/Programs/Horde/HordeDashboard/src/components/LogView.tsx
@@ -30,6 +30,33 @@ import { PrintException } from './PrintException';
 import { StepRefStatusIcon } from './StatusIcon';
 import { TopNav } from './TopNav';
 
+function getQueryLine(): number | undefined {
+
+   const search = new URLSearchParams(window.location.search);
+
+   if (search.get("lineindex")) {
+      return parseInt(search.get("lineindex")!) + 1;
+   }
+
+   if (search.get("lineIndex")) {
+      return parseInt(search.get("lineIndex")!);
+   }
+
+   return undefined;
+
+}
+
+function updateLineQuery(line: number, navigate: any) {
+
+   const search = new URLSearchParams(window.location.search);
+   search.set("lineIndex", line.toString());
+   search.delete("lineindex");
+   const url = `${window.location.pathname}?` + search.toString();
+
+   navigate(url, { replace: true })
+
+}
+
 class LogHandler {
 
    constructor() {
@@ -48,12 +75,76 @@ class LogHandler {
    }
 
    logSource?: LogSource;
-   currentWarning: number | undefined;
-   currentError: number | undefined;
+   currentLine?: number;
    trailing?: boolean;
    scroll?: number;
    initialRender = true;
 
+   getCurrentEvent() {
+      return this.getLogEvent(this.currentLine)
+   }
+
+   getNextLogEvent(warnings?: boolean) {
+
+      let currentLine = -1;
+      if (this.currentLine !== undefined) {
+         currentLine = this.currentLine - 1;
+      }
+
+      const startLine = currentLine + 1;
+
+      return this.events.filter(e => {
+         if (warnings && e.severity !== EventSeverity.Warning) {
+            return false;
+         }
+         if (!warnings && e.severity !== EventSeverity.Error) {
+            return false;
+         }
+         return true;
+      }).find(e => e.lineIndex >= startLine);
+   }
+
+   getPrevLogEvent(warnings?: boolean) {
+
+      let currentLine = 0;
+      if (this.currentLine !== undefined) {
+         currentLine = this.currentLine - 1;
+      }
+
+      const startLine = currentLine - 1;
+
+      if (startLine < 0) {
+         return undefined;
+      }
+
+      return this.events.filter(e => {
+         if (warnings && e.severity !== EventSeverity.Warning) {
+            return false;
+         }
+         if (!warnings && e.severity !== EventSeverity.Error) {
+            return false;
+         }
+         return true;
+      }).reverse().find(e => e.lineIndex <= startLine);
+   }
+
+   getLogEvent(line: number | undefined) {
+
+      if (line === undefined) {
+         return undefined;
+      }
+
+      line--;
+
+      return this.events.find(e => line >= e.lineIndex && line < (e.lineIndex + e.lineCount));
+   }
+
+   get events(): GetLogEventResponse[] {
+      const events = this.logSource?.errors.map(e => e) ?? [];
+      events.push(...(this.logSource?.warnings.map(e => e) ?? []))
+      return events;
+   }
+
    infoLine?: number;
 
    // could be a preference
@@ -131,12 +222,7 @@ let logListKey = 0;
 let globalHandler: LogHandler | undefined;
 let globalSearchState: { search?: string, results?: number[], curRequest?: any } | undefined;
 
-const searchUp = () => {
-
-   if (globalHandler) {
-      globalHandler.currentWarning = undefined;
-      globalHandler.currentError = undefined;
-   }
+const searchUp = (navigate: any) => {
 
    if (globalSearchState?.results?.length) {
 
@@ -146,7 +232,15 @@ const searchUp = () => {
       }
 
       globalHandler?.stopTrailing();
-      let lineIdx = globalSearchState.results[curSearchIdx] - 10;
+      let lineIdx = globalSearchState.results[curSearchIdx];
+
+      if (globalHandler) {
+         globalHandler.currentLine = lineIdx + 1;
+      }
+
+      updateLineQuery(lineIdx + 1, navigate);
+
+      lineIdx -= 10;
       if (lineIdx < 0) {
          lineIdx = 0;
       }
@@ -162,12 +256,7 @@ const searchUp = () => {
 
 }
 
-const searchDown = () => {
-
-   if (globalHandler) {
-      globalHandler.currentWarning = undefined;
-      globalHandler.currentError = undefined;
-   }
+const searchDown = (navigate: any) => {
 
    if (globalSearchState?.results?.length) {
 
@@ -177,7 +266,16 @@ const searchDown = () => {
       }
 
       globalHandler?.stopTrailing();
-      let lineIdx = globalSearchState.results[curSearchIdx] - 10;
+
+      let lineIdx = globalSearchState.results[curSearchIdx];
+
+      if (globalHandler) {
+         globalHandler.currentLine = lineIdx + 1;
+      }
+
+      updateLineQuery(lineIdx + 1, navigate);
+
+      lineIdx -= 10;
       if (lineIdx < 0) {
          lineIdx = 0;
       }
@@ -425,6 +523,39 @@ const LogProgressIndicator: React.FC<{ logSource: LogSource }> = observer(({ log
 });
 
 
+const LogLineIndicator: React.FC<{ lineIndex: number }> = ({ lineIndex }) => {
+
+   useQuery();
+
+   if (!globalHandler) {
+      return null;
+   }
+
+   let prefix = "";
+   const event = globalHandler.getLogEvent(globalHandler.currentLine);
+
+   const queryLine = getQueryLine();
+
+   if (queryLine) {
+      if (lineIndex === queryLine) {
+         prefix = ">>> ";
+      }
+   } else if (event && lineIndex >= event.lineIndex && (lineIndex < event.lineIndex + event.lineCount)) {
+      prefix = ">>> ";
+   } else if (globalSearchState?.results?.length && curSearchIdx < globalSearchState.results.length) {
+      if (lineIndex === globalSearchState.results[curSearchIdx]) {
+         prefix = ">>> ";
+      }
+   }
+
+   if (!prefix) {
+      return null;
+   }
+
+   return <span>{">>>"}</span>
+}
+
+
 export const LogList: React.FC<{ logId: string }> = observer(({ logId }) => {
 
    const windowSize = useWindowSize();
@@ -463,9 +594,9 @@ export const LogList: React.FC<{ logId: string }> = observer(({ logId }) => {
          if (e.keyCode === 114) {
             e.preventDefault();
             if (e.shiftKey) {
-               searchUp();
+               searchUp(navigate);
             } else {
-               searchDown();
+               searchDown(navigate);
             }
 
          }
@@ -526,6 +657,7 @@ export const LogList: React.FC<{ logId: string }> = observer(({ logId }) => {
          }
 
          handler.logSource = source;
+         handler.currentLine = source.startLine;
          setHandler(handler);
       }).catch((reason) => {
          setLogError(reason);
@@ -577,8 +709,6 @@ export const LogList: React.FC<{ logId: string }> = observer(({ logId }) => {
                </div>
             </Stack>
 
-
-
          }
 
          const line = item.line;
@@ -664,32 +794,6 @@ export const LogList: React.FC<{ logId: string }> = observer(({ logId }) => {
             </Stack>
          }
 
-         let prefix = "";
-         const lineIndex = item.lineNumber - 1;
-         if (errors.length && handler.currentError !== undefined) {
-            const error = errors[handler.currentError];
-            if (lineIndex >= error.lineIndex && lineIndex < error.lineIndex + error.lineCount) {
-               prefix = ">>> ";
-            }
-         }
-         else if (warnings.length && handler.currentWarning !== undefined) {
-            const warning = warnings[handler.currentWarning];
-            if (lineIndex >= warning.lineIndex && lineIndex < warning.lineIndex + warning.lineCount) {
-               prefix = ">>> ";
-            }
-         } else if (globalSearchState?.results?.length && curSearchIdx < globalSearchState.results.length) {
-            if (lineIndex === globalSearchState.results[curSearchIdx]) {
-               prefix = ">>> ";
-            }
-         } else {
-
-            if (query.get("lineindex")) {
-               if (lineIndex === parseInt(query.get("lineindex")!)) {
-                  prefix = ">>> ";
-               }
-            }
-         }
-
          const eyeColor = modeColors.text + "44";
 
          return (
@@ -705,15 +809,15 @@ export const LogList: React.FC<{ logId: string }> = observer(({ logId }) => {
                   }
                }}
                onClick={() => {
-                  const search = new URLSearchParams(window.location.search);
-                  search.set("lineindex", (item.lineNumber - 1).toString());
-                  const url = `${window.location.pathname}?` + search.toString();
-
-                  navigate(url, { replace: true })
+                  updateLineQuery(item.lineNumber, navigate);
+                  handler.currentLine = item.lineNumber;
                }}>
                <div style={{ position: "relative" }}>
                   <Stack className={styles.logLine} style={{ position: "relative" }} tokens={{ childrenGap: 8 }} horizontal disableShrink={true}>
-                     <Stack styles={{ root: { color: "#c0c0c0", width: 80, textAlign: "right", userSelect: "none", fontSize: handler.fontSize } }}>{prefix + item.lineNumber}</Stack>
+                     <Stack horizontal styles={{ root: { color: "#c0c0c0", width: 80, textAlign: "right", userSelect: "none", fontSize: handler.fontSize } }}>
+                        <LogLineIndicator lineIndex={item.lineNumber} />
+                        <Stack styles={{ root: { color: "#c0c0c0", width: 80, textAlign: "right", userSelect: "none", fontSize: handler.fontSize } }}>{item.lineNumber}</Stack>
+                     </Stack>
                      <Stack className={style} horizontal disableShrink={true}>
                         <Stack className={gutterStyle}></Stack>
                         {(!item.issueId || !ev) && <Stack styles={{ root: { color: "#8a8a8a", width: tsWidth, whiteSpace: "nowrap", fontSize: handler.fontSize, userSelect: "none" } }}> {timestamp}</Stack>}
@@ -803,8 +907,7 @@ export const LogList: React.FC<{ logId: string }> = observer(({ logId }) => {
 
    const doQuery = (newValue: string) => {
 
-      handler.currentWarning = undefined;
-      handler.currentError = undefined;
+      handler.currentLine = undefined;
 
       if (!newValue) {
 
@@ -844,6 +947,8 @@ export const LogList: React.FC<{ logId: string }> = observer(({ logId }) => {
                      lineIdx = 0;
                   }
 
+                  updateLineQuery(lines[curSearchIdx] + 1, navigate);
+
                   // oof
                   inTimeout = true;
                   setTimeout(() => {
@@ -870,6 +975,9 @@ export const LogList: React.FC<{ logId: string }> = observer(({ logId }) => {
          curSearchIdx = 0;
 
          if (searchState.results?.length) {
+
+            updateLineQuery(searchState.results[curSearchIdx] + 1, navigate);
+
             handler.stopTrailing();
             let lineIdx = searchState.results[curSearchIdx] - 10;
             if (lineIdx < 0) {
@@ -976,98 +1084,81 @@ export const LogList: React.FC<{ logId: string }> = observer(({ logId }) => {
       })
    }
 
-   function updateError() {
-      if (!handler || handler.currentError === undefined) return;
-      handler.currentWarning = undefined;
-      handler.stopTrailing();
-      let lineIdx = errors[handler.currentError].lineIndex;
-      const search = new URLSearchParams(window.location.search);
-      search.set("lineindex", (lineIdx).toString());
-      const url = `${window.location.pathname}?` + search.toString();
-      lineIdx -= 10;
-      if (lineIdx < 0) {
-         lineIdx = 0;
+   function updateEvent() {
+
+      const event = handler?.getCurrentEvent();
+      if (!handler || !event) {
+         return;
       }
-      handler.externalUpdate();
-      listRef?.scrollToIndex(lineIdx, () => handler.lineHeight, ScrollToMode.top);
 
-      navigate(url, { replace: true })
+      handler.stopTrailing();
+
+      updateLineQuery(event.lineIndex + 1, navigate);
+      handler.externalUpdate();
+      listRef?.scrollToIndex(event.lineIndex - 10 < 0 ? 0 : event.lineIndex - 10, () => handler.lineHeight, ScrollToMode.top);
    }
 
    function prevError() {
-      if (!handler || !errors.length) return;
 
-      if (handler.currentError === undefined) {
-         handler.currentError = errors.length - 1;
-      } else {
-         handler.currentError--;
-         if (handler.currentError < 0) {
-            handler.currentError = errors.length - 1;
-         }
+      if (!handler) {
+         return;
       }
-      updateError();
-   }
 
-   function nextError() {
-      if (!handler) return;
+      const event = handler.getPrevLogEvent();
 
-      if (handler.currentError === undefined) {
-         handler.currentError = 0;
-      } else {
-         handler.currentError++;
-         handler.currentError %= errors.length;
+      if (event) {
+         handler.currentLine = event.lineIndex + 1;
       }
 
-      updateError();
+      updateEvent();
+
    }
 
-   function updateWarning() {
-      if (!handler || handler.currentWarning === undefined) return;
+   function nextError() {
 
-      handler.currentError = undefined;
+      if (!handler) {
+         return;
+      }
 
-      handler.stopTrailing();
-      let lineIdx = warnings[handler.currentWarning].lineIndex;
-      const search = new URLSearchParams(window.location.search);
-      search.set("lineindex", (lineIdx).toString());
-      const url = `${window.location.pathname}?` + search.toString();
+      const event = handler.getNextLogEvent();
 
-      lineIdx -= 10
-      if (lineIdx < 0) {
-         lineIdx = 0;
+      if (event) {
+         handler.currentLine = event.lineIndex + 1;
       }
-      handler.externalUpdate();
-      listRef?.scrollToIndex(lineIdx, () => handler.lineHeight, ScrollToMode.top);
 
-      navigate(url, { replace: true })
+      updateEvent();
    }
 
+
    function prevWarning() {
 
-      if (!handler || !warnings.length) return;
+      if (!handler) {
+         return;
+      }
+
+      const event = handler.getPrevLogEvent(true);
 
-      if (handler.currentWarning === undefined) {
-         handler.currentWarning = warnings.length - 1;
-      } else {
-         handler.currentWarning--;
-         if (handler.currentWarning < 0) {
-            handler.currentWarning = warnings.length - 1;
-         }
+      if (event) {
+         handler.currentLine = event.lineIndex + 1;
       }
 
-      updateWarning();
+      updateEvent();
+
    }
 
    function nextWarning() {
-      if (!handler) return;
+      if (!handler) {
+         return;
+      }
 
-      if (handler.currentWarning === undefined) {
-         handler.currentWarning = 0;
-      } else {
-         handler.currentWarning++;
-         handler.currentWarning %= warnings.length;
+      const event = handler.getNextLogEvent(true);
+
+      if (event) {
+         handler.currentLine = event.lineIndex + 1;
       }
-      updateWarning();
+
+      updateEvent();
+
    }
 
    let warningsText = "";
@@ -1099,6 +1190,11 @@ export const LogList: React.FC<{ logId: string }> = observer(({ logId }) => {
       heightAdjust += 120;
    }
 
+   const nextErrorEnabled = handler?.getNextLogEvent() ? true : false;
+   const prevErrorEnablesd = handler?.getPrevLogEvent() ? true : false;
+   const nextWarningEnabled = handler?.getNextLogEvent(true) ? true : false;
+   const prevWarningEnabled = handler?.getPrevLogEvent(true) ? true : false;
+
    return <Stack>
       {!!fixme && <IssueModalV2 issueId={query.get("issue")} popHistoryOnClose={issueHistory} />}
       {!!fixme && logHistory && <StepHistoryModal jobDetails={fixme!} stepId={fixme!.stepByLogId(logId)?.id} onClose={() => setLogHistory(false)} />}
@@ -1112,7 +1208,7 @@ export const LogList: React.FC<{ logId: string }> = observer(({ logId }) => {
                <Stack horizontal style={{ paddingBottom: 4 }}>
                   <Stack className={hordeClasses.button} horizontal horizontalAlign={"start"} verticalAlign="center" tokens={{ childrenGap: 8 }}>
                      <Stack horizontal tokens={{ childrenGap: 2 }}>
-                        <DefaultButton disabled={!errors.length} className={errors.length ? handler.style.errorButton : handler.style.errorButtonDisabled}
+                        <DefaultButton disabled={!nextErrorEnabled} className={nextErrorEnabled ? handler.style.errorButton : handler.style.errorButtonDisabled}
                            text={`${errorText} ${errors.length === 1 ? "Error" : "Errors"}`}
                            onClick={() => {
                               nextError();
@@ -1122,7 +1218,7 @@ export const LogList: React.FC<{ logId: string }> = observer(({ logId }) => {
                         </DefaultButton>
 
                         {!!errors.length && <Stack>
-                           <IconButton className={handler.style.errorButton} style={{ height: 30, fontSize: 19, padding: 8 }} iconProps={{ iconName: 'ChevronUp' }} onClick={(event: any) => {
+                           <IconButton disabled={!prevErrorEnablesd} className={prevErrorEnablesd ? handler.style.errorButton : handler.style.errorButtonDisabled} style={{ height: 30, fontSize: 19, padding: 8 }} iconProps={{ iconName: 'ChevronUp' }} onClick={(event: any) => {
                               event?.stopPropagation();
                               prevError();
 
@@ -1132,16 +1228,16 @@ export const LogList: React.FC<{ logId: string }> = observer(({ logId }) => {
 
 
                         <Stack horizontal style={{ paddingLeft: 12 }} tokens={{ childrenGap: 2 }}>
-                           <DefaultButton disabled={!warnings.length} className={warnings.length ? handler.style.warningButton : handler.style.warningButtonDisabled}
+                           <DefaultButton disabled={!nextWarningEnabled} className={nextWarningEnabled ? handler.style.warningButton : handler.style.warningButtonDisabled}
                               text={`${warningsText} ${warnings.length === 1 ? "Warning" : "Warnings"}`}
                               onClick={() => {
                                  nextWarning();
                               }}
-                              style={{ color: (dashboard.darktheme && warnings.length) ? "#F9F9FB" : buttonNoneText, padding: 15 }} >
+                              style={{ color: (dashboard.darktheme && nextWarningEnabled) ? "#F9F9FB" : buttonNoneText, padding: 15 }} >
                               {!!warnings.length && <Icon style={{ fontSize: 19, paddingLeft: 12 }} iconName='ChevronDown' />}
                            </DefaultButton>
                            {!!warnings.length && <Stack>
-                              <IconButton className={handler.style.warningButton} style={{ height: 30, fontSize: 19, padding: 8, color: (dashboard.darktheme && warnings.length) ? "#F9F9FB" : buttonNoneText }} iconProps={{ iconName: 'ChevronUp' }} onClick={(event: any) => {
+                              <IconButton disabled={!prevWarningEnabled} className={prevWarningEnabled ? handler.style.warningButton : handler.style.warningButtonDisabled} style={{ height: 30, fontSize: 19, padding: 8, color: (dashboard.darktheme && nextWarningEnabled) ? "#F9F9FB" : buttonNoneText }} iconProps={{ iconName: 'ChevronUp' }} onClick={(event: any) => {
                                  event?.stopPropagation();
                                  prevWarning();
 
@@ -1187,7 +1283,7 @@ export const LogList: React.FC<{ logId: string }> = observer(({ logId }) => {
 
                                        if (ev.key === "Enter" && !searchState.curRequest && !inTimeout) {
                                           if (searchBox.current?.value === searchState.search) {
-                                             searchDown();
+                                             searchDown(navigate);
                                           } else {
                                              doQuery(searchBox.current?.value ?? "");
                                           }
@@ -1206,10 +1302,10 @@ export const LogList: React.FC<{ logId: string }> = observer(({ logId }) => {
                                  />
                                  <Stack horizontal style={{ borderWidth: 1, borderStyle: "solid", borderColor: dashboard.darktheme ? "#3F3F3F" : "rgb(96, 94, 92)", height: 32, borderLeft: 0 }}>
                                     <IconButton style={{ height: 30 }} iconProps={{ iconName: 'ChevronUp' }} onClick={(event: any) => {
-                                       searchUp();
+                                       searchUp(navigate);
                                     }} />
                                     <IconButton style={{ height: 30 }} iconProps={{ iconName: 'ChevronDown' }} onClick={(event: any) => {
-                                       searchDown();
+                                       searchDown(navigate);
                                     }} />
 
                                  </Stack>
@@ -1297,7 +1393,8 @@ export const LogList: React.FC<{ logId: string }> = observer(({ logId }) => {
                                     if (handler.initialRender && listRef) {
                                        handler.initialRender = false;
                                        if (logSource?.startLine !== undefined) {
-                                          listRef.scrollToIndex(logSource.startLine - 1, () => handler.lineHeight, ScrollToMode.center);
+
+                                          listRef?.scrollToIndex(logSource?.startLine - 10 < 0 ? 0 : logSource?.startLine - 10, () => handler.lineHeight, ScrollToMode.top);
                                        } else if (handler.trailing) {
                                           listRef.scrollToIndex(logSource.logData!.lineCount - 1, undefined, ScrollToMode.bottom);
                                        }
diff --git a/Engine/Source/Programs/Horde/HordeDashboard/src/components/jobDetailsV2/StepErrorPanel.tsx b/Engine/Source/Programs/Horde/HordeDashboard/src/components/jobDetailsV2/StepErrorPanel.tsx
--- a/Engine/Source/Programs/Horde/HordeDashboard/src/components/jobDetailsV2/StepErrorPanel.tsx
+++ b/Engine/Source/Programs/Horde/HordeDashboard/src/components/jobDetailsV2/StepErrorPanel.tsx
@@ -219,7 +219,7 @@ const ErrorPane: React.FC<{ jobDetails: JobDetailsV2; view: StepSummaryErrorsVie
          return <div>???</div>;
       }
 
-      const url = `/log/${item.logId}?lineindex=${item.lineIndex}`;
+      const url = `/log/${item.logId}?lineIndex=${item.lineIndex + 1}`;
 
       const lines = item.lines.filter(line => line.message?.trim().length).map(line => <Stack key={`steperrorpane_line_${item.lineIndex}`} styles={{ root: { paddingLeft: 8, paddingRight: 8, lineBreak: "anywhere", whiteSpace: "pre-wrap", lineHeight: 18, fontSize: 10, fontFamily: "Horde Cousine Regular, monospace, monospace" } }}> <Link style={{ color: modeColors.text }} to={url}>{renderLine(navigate, line, undefined, {})}</Link></Stack>);
 
