Horde: Telemetry plugin update
#rnx

diff --git a/Engine/Source/Programs/Horde/HordeServer/Defaults/default-metrics.telemetry.json b/Engine/Source/Programs/Horde/HordeServer/Defaults/default-metrics.telemetry.json
--- a/Engine/Source/Programs/Horde/HordeServer/Defaults/default-metrics.telemetry.json
+++ b/Engine/Source/Programs/Horde/HordeServer/Defaults/default-metrics.telemetry.json
@@ -682,5 +682,689 @@
 			"function": "Average",
 			"interval": "00:30:00"
 		}
-	]
+	],
+	"views": [
+		{
+			"id": "engine",
+			"name": "Engine",
+			"telemetryStoreId": "engine",
+			"variables": [
+				{
+					"name": "Stream",
+					"group": "Build_BranchName"
+				},
+				{
+					"name": "Project",
+					"group": "ProjectName"
+				},
+				{
+					"name": "Map",
+					"group": "MapName"
+				},
+				{
+					"name": "LoadingName",
+					"group": "LoadingName",
+					"defaults": [
+						"BootEditor",
+						"LoadMap",
+						"TimeToPIE"
+					]
+				},
+				{
+					"name": "AssetType",
+					"group": "AssetType"
+				},
+				{
+					"name": "Cook Platform",
+					"group": "Cook_Profile_TargetPlatforms"
+				},
+				{
+					"name": "Horde JobName",
+					"group": "Horde_JobName"
+				},
+				{
+					"name": "Horde StepName",
+					"group": "Horde_StepName"
+				},
+				{
+					"name": "StreamId",
+					"group": "StreamId"
+				},
+				{
+					"name": "JobName",
+					"group": "JobName"
+				},
+				{
+					"name": "StepName",
+					"group": "StepName"
+				},
+				{
+					"name": "Target",
+					"group": "Target"
+				},
+				{
+					"name": "Mode",
+					"group": "Mode",
+					"defaults": [
+						"Rebuild"
+					]
+				},
+				{
+					"name": "Platform",
+					"group": "Platform",
+					"defaults": [
+						"Win64"
+					]
+				},
+				{
+					"name": "Configuration",
+					"group": "Configuration",
+					"defaults": [
+						"Development"
+					]
+				}
+			],
+			"categories": [
+				{
+					"name": "Editor Iteration",
+					"charts": [
+						{
+							"name": "User Load Time (75th Percentile)",
+							"metrics": [
+								{
+									"id": "core-loading-user-timetopie-percentile",
+									"alias": "TimeToPIE",
+									"threshold": 360
+								},
+								{
+									"id": "core-loading-user-editorboottime-percentile",
+									"alias": "BootEditor",
+									"threshold": 120
+								},
+								{
+									"id": "core-loading-user-loadmaptime-percentile",
+									"alias": "LoadMap",
+									"threshold": 240
+								},
+								{
+									"id": "core-loading-user-piestartuptime-percentile",
+									"alias": "PIEStartup",
+									"threshold": 360
+								}
+							],
+							"graph": "indicator",
+							"max": 720
+						},
+						{
+							"name": "User Load Time (Average)",
+							"metrics": [
+								{
+									"id": "core-loading-user-editorboottime",
+									"alias": "BootEditor",
+									"threshold": 120
+								},
+								{
+									"id": "core-loading-user-loadmaptime",
+									"alias": "LoadMap",
+									"threshold": 240
+								},
+								{
+									"id": "core-loading-user-piestartuptime",
+									"alias": "PIEStartup",
+									"threshold": 360
+								},
+								{
+									"id": "core-loading-user-timetopie",
+									"alias": "TimeToPIE",
+									"threshold": 360
+								}
+							],
+							"graph": "indicator",
+							"max": 900
+						},
+						{
+							"name": "User Loading Time (KPIs)",
+							"metrics": [
+								{
+									"id": "core-loading-user-editorboottime-kpi",
+									"alias": "BootEditor"
+								},
+								{
+									"id": "core-loading-user-loadmaptime-kpi",
+									"alias": "LoadMap"
+								},
+								{
+									"id": "core-loading-user-piestartuptime-kpi",
+									"alias": "PIEStartup"
+								},
+								{
+									"id": "core-loading-user-timetopie-kpi",
+									"alias": "TimeToPIE"
+								}
+							]
+						},
+						{
+							"name": "User Loading Count (Live)",
+							"display": "Value",
+							"max": 100,
+							"metrics": [
+								{
+									"id": "core-loading-user-loadingcount"
+								}
+							]
+						},
+						{
+							"name": "User Editor Boot Time (Live)",
+							"metrics": [
+								{
+									"id": "core-loading-user-editorboottime"
+								}
+							]
+						},
+						{
+							"name": "User Editor LoadMap Time (Live)",
+							"metrics": [
+								{
+									"id": "core-loading-user-loadmaptime"
+								}
+							]
+						},
+						{
+							"name": "User Editor LoadMap Count (Live)",
+							"display": "Value",
+							"metrics": [
+								{
+									"id": "core-loading-user-loadmap-count"
+								}
+							]
+						},
+						{
+							"name": "User PIE Startup Time (Live)",
+							"metrics": [
+								{
+									"id": "core-loading-user-piestartuptime"
+								}
+							]
+						},
+						{
+							"name": "User PIE LoadMap Time (Live)",
+							"metrics": [
+								{
+									"id": "core-loading-user-pieloadmaptime"
+								}
+							]
+						},
+						{
+							"name": "User Boot To PIE Time (Live)",
+							"metrics": [
+								{
+									"id": "core-loading-user-timetopie"
+								}
+							]
+						}
+					]
+				},
+				{
+					"name": "Editor Performance",
+					"charts": [
+						{
+							"name": "Count",
+							"display": "Value",
+							"metrics": [
+								{
+									"id": "editor-performance-warning-count"
+								}
+							]
+						},
+						{
+							"name": "Editor",
+							"display": "Value",
+							"metrics": [
+								{
+									"id": "editor-performance-editor-warning"
+								}
+							]
+						},
+						{
+							"name": "PIE",
+							"display": "Value",
+							"metrics": [
+								{
+									"id": "editor-performance-pie-warning"
+								}
+							]
+						},
+						{
+							"name": "Cache",
+							"display": "Value",
+							"metrics": [
+								{
+									"id": "editor-performance-cache-warning"
+								}
+							]
+						}
+					]
+				},
+				{
+					"name": "Cook",
+					"charts": [
+						{
+							"name": "Cook Time (KPIs)",
+							"metrics": [
+								{
+									"id": "core-cooking-lyra-time-windows-kpi",
+									"alias": "Lyra Cook Windows",
+									"threshold": 7200
+								},
+								{
+									"id": "core-cooking-ancientgame-time-windows-kpi",
+									"alias": "AncientGame Cook Windows",
+									"threshold": 7200
+								},
+								{
+									"id": "core-cooking-citysample-time-windows-kpi",
+									"alias": "CitySample Cook Windows",
+									"threshold": 7200
+								},
+								{
+									"id": "core-cooking-lyra-time-ps5-kpi",
+									"alias": "Lyra Cook Windows",
+									"threshold": 7200
+								},
+								{
+									"id": "core-cooking-ancientgame-time-ps5-kpi",
+									"alias": "AncientGame Cook Windows",
+									"threshold": 7200
+								},
+								{
+									"id": "core-cooking-citysample-time-ps5-kpi",
+									"alias": "CitySample Cook Windows",
+									"threshold": 7200
+								},
+								{
+									"id": "core-cooking-nightly-cook-time-windows-kpi",
+									"alias": "Nightly Cook Windows",
+									"threshold": 7200
+								},
+								{
+									"id": "core-cooking-nightly-cook-time-ps5-kpi",
+									"alias": "Nightly Cook PS5",
+									"threshold": 7200
+								},
+								{
+									"id": "core-cooking-autotest-cook-time-kpi",
+									"alias": "AutoTest Cook",
+									"threshold": 4500
+								},
+								{
+									"id": "core-cooking-autotest-fastcook-time-kpi",
+									"alias": "AutoTest FastCook",
+									"threshold": 1800
+								}
+							],
+							"graph": "indicator",
+							"max": 10800
+						},
+						{
+							"name": "Cook Shaders Compiled (KPIs)",
+							"display": "Value",
+							"metrics": [
+								{
+									"id": "core-cooking-nightly-cook-shaders-compiled-windows-kpi",
+									"alias": "Nightly Cook Windows",
+									"threshold": 1000
+								},
+								{
+									"id": "core-cooking-nightly-cook-shaders-compiled-ps5-kpi",
+									"alias": "Nightly Cook PS5",
+									"threshold": 1000
+								},
+								{
+									"id": "core-cooking-autotest-cook-shaders-compiled-kpi",
+									"alias": "AutoTest Cook",
+									"threshold": 1000
+								},
+								{
+									"id": "core-cooking-autotest-fastcook-shaders-compiled-kpi",
+									"alias": "AutoTest FastCook",
+									"threshold": 1000
+								}
+							],
+							"graph": "indicator",
+							"max": 2000
+						},
+						{
+							"name": "Cook Time (KPIs)",
+							"metrics": [
+								{
+									"id": "core-cooking-nightly-cook-time-windows-kpi",
+									"alias": "Nightly Cook Windows"
+								},
+								{
+									"id": "core-cooking-nightly-cook-time-ps5-kpi",
+									"alias": "Nightly Cook PS5"
+								},
+								{
+									"id": "core-cooking-autotest-cook-time-kpi",
+									"alias": "AutoTest Cook"
+								},
+								{
+									"id": "core-cooking-autotest-fastcook-time-kpi",
+									"alias": "AutoTest FastCook"
+								}
+							]
+						},
+						{
+							"name": "Cook Shaders Compiled (KPIs)",
+							"display": "Value",
+							"metrics": [
+								{
+									"id": "core-cooking-nightly-cook-shaders-compiled-windows-kpi",
+									"alias": "Nightly Cook Windows"
+								},
+								{
+									"id": "core-cooking-nightly-cook-shaders-compiled-ps5-kpi",
+									"alias": "Nightly Cook PS5"
+								},
+								{
+									"id": "core-cooking-autotest-cook-shaders-compiled-kpi",
+									"alias": "AutoTest Cook"
+								},
+								{
+									"id": "core-cooking-autotest-fastcook-shaders-compiled-kpi",
+									"alias": "AutoTest FastCook"
+								}
+							]
+						},
+						{
+							"name": "Incremental Cook Time",
+							"metrics": [
+								{
+									"id": "core-cooking-incremental-cook-time"
+								}
+							]
+						},
+						{
+							"name": "Incremental Cook Shaders Compiled",
+							"metrics": [
+								{
+									"id": "core-cooking-incremental-shaders-compiled"
+								}
+							]
+						},
+						{
+							"name": "Package And Test Cook Time",
+							"metrics": [
+								{
+									"id": "core-cooking-package-and-test-cook-time"
+								}
+							]
+						},
+						{
+							"name": "Package And Test Shaders Compiled",
+							"display": "Value",
+							"metrics": [
+								{
+									"id": "core-cooking-package-and-test-shaders-compiled"
+								}
+							]
+						},
+						{
+							"name": "Nightly Cook Time",
+							"metrics": [
+								{
+									"id": "core-cooking-nightly-cook-time"
+								}
+							]
+						},
+						{
+							"name": "Nightly Cook Shaders Compiled",
+							"display": "Value",
+							"metrics": [
+								{
+									"id": "core-cooking-nightly-shaders-compiled"
+								}
+							]
+						},
+						{
+							"name": "Standard Cook Time",
+							"metrics": [
+								{
+									"id": "core-cooking-standard-cook-time"
+								}
+							]
+						},
+						{
+							"name": "Standard Cook Shaders Compiled",
+							"metrics": [
+								{
+									"id": "core-cooking-standard-shaders-compiled"
+								}
+							]
+						}
+					]
+				},
+				{
+					"name": "Cache",
+					"charts": [
+						{
+							"name": "User DDC Backend Efficiency",
+							"display": "Ratio",
+							"metrics": [
+								{
+									"id": "core-loading-user-totalddchitratio",
+									"alias": "Root"
+								},
+								{
+									"id": "core-loading-user-cloudlocalddchitratio",
+									"alias": "Cloud"
+								},
+								{
+									"id": "core-loading-user-zenlocalddchitratio",
+									"alias": "Local"
+								},
+								{
+									"id": "core-loading-user-sharedddchitratio",
+									"alias": "Shared"
+								}
+							]
+						},
+						{
+							"name": "User DDC Asset Efficiency",
+							"display": "Ratio",
+							"metrics": [
+								{
+									"id": "core-loading-user-ddc-asset-efficiency"
+								}
+							]
+						}
+					]
+				},
+				{
+					"name": "Build",
+					"charts": [
+						{
+							"name": "Incremental Build",
+							"metrics": [
+								{
+									"id": "horde-incremental-build-jobs"
+								},
+								{
+									"id": "horde-incremental-ci-jobs"
+								}
+							]
+						},
+						{
+							"name": "AutoTest Build",
+							"metrics": [
+								{
+									"id": "horde-autotest-ci-jobs"
+								}
+							]
+						},
+						{
+							"name": "Package And Test Build",
+							"metrics": [
+								{
+									"id": "horde-package-and-test-jobs"
+								}
+							]
+						},
+						{
+							"name": "Nighty Build",
+							"metrics": [
+								{
+									"id": "horde-nightly-build-jobs"
+								}
+							]
+						}
+					]
+				},
+				{
+					"name": "Compile",
+					"charts": [
+						{
+							"name": "Full Compile",
+							"metrics": [
+								{
+									"id": "horde-monolithic-compile-jobs"
+								}
+							]
+						},
+						{
+							"name": "Incremental Compile",
+							"metrics": [
+								{
+									"id": "horde-incremental-compile-jobs"
+								}
+							]
+						}
+					]
+				},
+				{
+					"name": "UBA",
+					"charts": [
+						{
+							"name": "Remote Distibution All (Average)",
+							"display": "Ratio",
+							"metrics": [
+								{
+									"id": "ubt-uba-remote-usage-all"
+								}
+							]
+						},
+						{
+							"name": "Cache Rate All (Average)",
+							"display": "Ratio",
+							"metrics": [
+								{
+									"id": "ubt-uba-cache-rate-all"
+								}
+							]
+						},
+						{
+							"name": "Actions All (Sum)",
+							"display": "Value",
+							"metrics": [
+								{
+									"id": "ubt-uba-total-actions-all",
+									"alias": "Total"
+								},
+								{
+									"id": "ubt-uba-remote-actions-all",
+									"alias": "Remote"
+								}
+							]
+						},
+						{
+							"name": "Failure Rate All (Average)",
+							"display": "Ratio",
+							"metrics": [
+								{
+									"id": "ubt-uba-failure-rate-all"
+								}
+							]
+						},
+						{
+							"name": "WallTime (Average)",
+							"metrics": [
+								{
+									"id": "ubt-uba-duration"
+								}
+							]
+						},
+						{
+							"name": "Remote Distribution (Average)",
+							"display": "Ratio",
+							"metrics": [
+								{
+									"id": "ubt-uba-remote-usage"
+								}
+							]
+						},
+						{
+							"name": "Cache Rate (Average)",
+							"display": "Ratio",
+							"metrics": [
+								{
+									"id": "ubt-uba-cache-rate"
+								}
+							]
+						},
+						{
+							"name": "Actions (Average)",
+							"display": "Value",
+							"metrics": [
+								{
+									"id": "ubt-uba-total-actions-avg",
+									"alias": "Total"
+								},
+								{
+									"id": "ubt-uba-remote-actions-avg",
+									"alias": "Remote"
+								}
+							]
+						},
+						{
+							"name": "Actions (Sum)",
+							"display": "Value",
+							"metrics": [
+								{
+									"id": "ubt-uba-total-actions",
+									"alias": "Total"
+								},
+								{
+									"id": "ubt-uba-remote-actions",
+									"alias": "Remote"
+								}
+							]
+						},
+						{
+							"name": "Failure Rate (Average)",
+							"display": "Ratio",
+							"metrics": [
+								{
+									"id": "ubt-uba-failure-rate"
+								}
+							]
+						}
+					]
+				},
+				{
+					"name": "UBT",
+					"charts": [
+						{
+							"name": "Wall Time (Average)",
+							"metrics": [
+								{
+									"id": "ubt-completed-duration"
+
+								}
+							]
+						}
+					]
+				}
+			]
+		}
+	]	
 }
\ No newline at end of file
