Horde: Updating schema documentation

#rnx

diff --git a/Engine/Source/Programs/Horde/Docs/Config/Schema/Dashboard.md b/Engine/Source/Programs/Horde/Docs/Config/Schema/Dashboard.md
--- a/Engine/Source/Programs/Horde/Docs/Config/Schema/Dashboard.md
+++ b/Engine/Source/Programs/Horde/Docs/Config/Schema/Dashboard.md
@@ -15,7 +15,6 @@ Name | Description
 `showTests` | `boolean`<br>Show automated tests on the server menu
 `agentCategories` | [DashboardAgentCategoryConfig](#dashboardagentcategoryconfig)`[]`<br>Configuration for different agent pages
 `poolCategories` | [DashboardPoolCategoryConfig](#dashboardpoolcategoryconfig)`[]`<br>Configuration for different pool pages
-`analytics` | [TelemetryViewConfig](#telemetryviewconfig)`[]`<br>Configuration for telemetry views
 `include` | [ConfigInclude](#configinclude)`[]`<br>Includes for other configuration files
 `macros` | [ConfigMacro](#configmacro)`[]`<br>Macros within this configuration
 
@@ -37,79 +36,6 @@ Name | Description
 `name` | `string`<br>Name of the category
 `condition` | `string`<br>Condition string to be evaluated for this page
 
-## TelemetryViewConfig
-
-A telemetry view of related metrics, divided into categofies
-
-Name | Description
----- | -----------
-`id` | `string`<br>Identifier for the view
-`name` | `string`<br>The name of the view
-`telemetryStoreId` | `string`<br>The telemetry store this view uses
-`variables` | [TelemetryVariableConfig](#telemetryvariableconfig)`[]`<br>The variables used to filter the view data
-`categories` | [TelemetryCategoryConfig](#telemetrycategoryconfig)`[]`<br>The categories contained within the view
-
-## TelemetryVariableConfig
-
-A telemetry view variable used for filtering the charting data
-
-Name | Description
----- | -----------
-`name` | `string`<br>The name of the variable for display purposes
-`group` | `string`<br>The associated data group attached to the variable
-`defaults` | `string[]`<br>The default values to select
-
-## TelemetryCategoryConfig
-
-A chart categody, will be displayed on the dashbord under an associated pivot
-
-Name | Description
----- | -----------
-`name` | `string`<br>The name of the category
-`charts` | [TelemetryChartConfig](#telemetrychartconfig)`[]`<br>The charts contained within the category
-
-## TelemetryChartConfig
-
-Telemetry chart configuraton
-
-Name | Description
----- | -----------
-`name` | `string`<br>The name of the chart, will be displayed on the dashboard
-`display` | [TelemetryMetricUnitType](#telemetrymetricunittype-enum)<br>The unit to display
-`graph` | [TelemetryMetricGraphType](#telemetrymetricgraphtype-enum)<br>The graph type
-`metrics` | [TelemetryChartMetricConfig](#telemetrychartmetricconfig)`[]`<br>List of configured metrics
-`min` | `integer`<br>The min unit value for clamping chart
-`max` | `integer`<br>The max unit value for clamping chart
-
-## TelemetryMetricUnitType (Enum)
-
-The units used to present the telemetry
-
-Name | Description
----- | -----------
-`Time` | Time duration
-`Ratio` | Ratio 0-100%
-`Value` | Artbitrary numeric value
-
-## TelemetryMetricGraphType (Enum)
-
-The type of
-
-Name | Description
----- | -----------
-`Line` | A line graph
-`Indicator` | Key performance indicator (KPI) chart with thrasholds
-
-## TelemetryChartMetricConfig
-
-Metric attached to a telemetry chart
-
-Name | Description
----- | -----------
-`id` | `string`<br>Associated metric id
-`threshold` | `integer`<br>The threshold for KPI values
-`alias` | `string`<br>The metric alias for display purposes
-
 ## ConfigInclude
 
 Directive to merge config data from another source
diff --git a/Engine/Source/Programs/Horde/Docs/Config/Schema/Telemetry.md b/Engine/Source/Programs/Horde/Docs/Config/Schema/Telemetry.md
--- a/Engine/Source/Programs/Horde/Docs/Config/Schema/Telemetry.md
+++ b/Engine/Source/Programs/Horde/Docs/Config/Schema/Telemetry.md
@@ -9,6 +9,7 @@ Name | Description
 `id` | `string`<br>Identifier for this store
 `acl` | [AclConfig](#aclconfig)<br>Permissions for this store
 `metrics` | [MetricConfig](#metricconfig)`[]`<br>Metrics to aggregate on the Horde server
+`views` | [TelemetryViewConfig](#telemetryviewconfig)`[]`<br>Configuration for telemetry views
 `include` | [ConfigInclude](#configinclude)`[]`<br>Includes for other configuration files
 `macros` | [ConfigMacro](#configmacro)`[]`<br>Macros within this configuration
 
@@ -80,6 +81,79 @@ Name | Description
 `Average` | Average all the samples
 `Percentile` | Estimates the value at a certain percentile
 
+## TelemetryViewConfig
+
+A telemetry view of related metrics, divided into categofies
+
+Name | Description
+---- | -----------
+`id` | `string`<br>Identifier for the view
+`name` | `string`<br>The name of the view
+`telemetryStoreId` | `string`<br>The telemetry store this view uses
+`variables` | [TelemetryVariableConfig](#telemetryvariableconfig)`[]`<br>The variables used to filter the view data
+`categories` | [TelemetryCategoryConfig](#telemetrycategoryconfig)`[]`<br>The categories contained within the view
+
+## TelemetryVariableConfig
+
+A telemetry view variable used for filtering the charting data
+
+Name | Description
+---- | -----------
+`name` | `string`<br>The name of the variable for display purposes
+`group` | `string`<br>The associated data group attached to the variable
+`defaults` | `string[]`<br>The default values to select
+
+## TelemetryCategoryConfig
+
+A chart categody, will be displayed on the dashbord under an associated pivot
+
+Name | Description
+---- | -----------
+`name` | `string`<br>The name of the category
+`charts` | [TelemetryChartConfig](#telemetrychartconfig)`[]`<br>The charts contained within the category
+
+## TelemetryChartConfig
+
+Telemetry chart configuraton
+
+Name | Description
+---- | -----------
+`name` | `string`<br>The name of the chart, will be displayed on the dashboard
+`display` | [TelemetryMetricUnitType](#telemetrymetricunittype-enum)<br>The unit to display
+`graph` | [TelemetryMetricGraphType](#telemetrymetricgraphtype-enum)<br>The graph type
+`metrics` | [TelemetryChartMetricConfig](#telemetrychartmetricconfig)`[]`<br>List of configured metrics
+`min` | `integer`<br>The min unit value for clamping chart
+`max` | `integer`<br>The max unit value for clamping chart
+
+## TelemetryMetricUnitType (Enum)
+
+The units used to present the telemetry
+
+Name | Description
+---- | -----------
+`Time` | Time duration
+`Ratio` | Ratio 0-100%
+`Value` | Artbitrary numeric value
+
+## TelemetryMetricGraphType (Enum)
+
+The type of
+
+Name | Description
+---- | -----------
+`Line` | A line graph
+`Indicator` | Key performance indicator (KPI) chart with thrasholds
+
+## TelemetryChartMetricConfig
+
+Metric attached to a telemetry chart
+
+Name | Description
+---- | -----------
+`id` | `string`<br>Associated metric id
+`threshold` | `integer`<br>The threshold for KPI values
+`alias` | `string`<br>The metric alias for display purposes
+
 ## ConfigInclude
 
 Directive to merge config data from another source
