Small UX Change to Horde Analytics dashboard
Removed the EditorPerformance category
Renamed the Editor Iteration category to Editor

diff --git a/Engine/Source/Programs/Horde/HordeServer/Defaults/default-analytics.dashboard.json b/Engine/Source/Programs/Horde/HordeServer/Defaults/default-analytics.dashboard.json
--- a/Engine/Source/Programs/Horde/HordeServer/Defaults/default-analytics.dashboard.json
+++ b/Engine/Source/Programs/Horde/HordeServer/Defaults/default-analytics.dashboard.json
@@ -82,7 +82,7 @@
 			],
 			"categories": [
 				{
-					"name": "Editor Iteration",
+					"name": "Editor",
 					"charts": [
 						{
 							"name": "User Load Time (75th Percentile)",
@@ -220,47 +220,6 @@
 						}
 					]
 				},
-				{
-					"name": "Editor Performance",
-					"charts": [
-						{
-							"name": "Count",
-							"display": "Value",
-							"metrics": [
-								{
-									"id": "editor-performance-warning-count"
-								}
-							]
-						},
-						{
-							"name": "Editor",
-							"display": "Value",
-							"metrics": [
-								{
-									"id": "editor-performance-editor-warning"
-								}
-							]
-						},
-						{
-							"name": "PIE",
-							"display": "Value",
-							"metrics": [
-								{
-									"id": "editor-performance-pie-warning"
-								}
-							]
-						},
-						{
-							"name": "Cache",
-							"display": "Value",
-							"metrics": [
-								{
-									"id": "editor-performance-cache-warning"
-								}
-							]
-						}
-					]
-				},
 				{
 					"name": "Cook",
 					"charts": [
