Horde: Remove debug entries from appsettings.json.

#rnx

diff --git a/Engine/Source/Programs/Horde/HordeServer/appsettings.json b/Engine/Source/Programs/Horde/HordeServer/appsettings.json
--- a/Engine/Source/Programs/Horde/HordeServer/appsettings.json
+++ b/Engine/Source/Programs/Horde/HordeServer/appsettings.json
@@ -3,12 +3,7 @@
 		"ConsoleLogLevel": "Debug",
 		"DashboardUrl": "https://localhost:3000",
 		"CorsEnabled": true,
-		"CorsOrigin": "https://localhost:5003;http://localhost:3000",
-		"Plugins": {
-			"analytics": {
-				"Foo": 1234
-			}
-		}
+		"CorsOrigin": "https://localhost:5003;http://localhost:3000"
 	},
 	"Serilog": {
 		"MinimumLevel": {
