Horde: Fix documentation link.

#rnx

diff --git a/Engine/Source/Programs/Horde/Docs/Config/Analytics.md b/Engine/Source/Programs/Horde/Docs/Config/Analytics.md
--- a/Engine/Source/Programs/Horde/Docs/Config/Analytics.md
+++ b/Engine/Source/Programs/Horde/Docs/Config/Analytics.md
@@ -29,7 +29,7 @@ running metrics for each time interval. This aggregation is performed according
 ## Charting
 
 The Horde dashboard render charts showing metrics collected on the server. These views are configured
-using the `Dashboard.Analytics` section of the globals.json file (see [TelemetryViewConfig](Schema/Dashboard.md#telemetryviewconfig)).
+using the `Dashboard.Analytics` section of the globals.json file (see [TelemetryViewConfig](Schema/Telemetry.md#telemetryviewconfig)).
 
 ## Telemetry Sinks
 
