Horde: Ensure resolved issues are excluded from queries when resolved by timeout.

diff --git a/Engine/Source/Programs/Horde/Plugins/Build/HordeServer.Build/Issues/IssueCollection.cs b/Engine/Source/Programs/Horde/Plugins/Build/HordeServer.Build/Issues/IssueCollection.cs
--- a/Engine/Source/Programs/Horde/Plugins/Build/HordeServer.Build/Issues/IssueCollection.cs
+++ b/Engine/Source/Programs/Horde/Plugins/Build/HordeServer.Build/Issues/IssueCollection.cs
@@ -1132,9 +1132,12 @@ namespace HordeServer.Issues
 						for (int idx = results.Count - 1; idx >= 0; idx--)
 						{
 							IssueDocument issue = results[idx];
-							if ((issue.ResolvedAt != null) != resolved.Value && issue.ResolvedById != IIssue.ResolvedByTimeoutId)
+							if ((issue.ResolvedAt != null) != resolved.Value)
 							{
-								_logger.LogWarning("Issue {IssueId} has resolved state out of sync with spans", issue.Id);
+								if (issue.ResolvedById != IIssue.ResolvedByTimeoutId) // TODO: We should always be marking spans as resolved when timing out an issue, but handling here for now.
+								{
+									_logger.LogWarning("Issue {IssueId} has resolved state out of sync with spans", issue.Id);
+								}
 								results.RemoveAt(idx);
 							}
 						}
