Horde: Pass agent capabilities object into agent collection.

#rnx

diff --git a/Engine/Source/Programs/Horde/Plugins/Build/HordeServer.Build.Tests/Agents/AgentCollectionTests.cs b/Engine/Source/Programs/Horde/Plugins/Build/HordeServer.Build.Tests/Agents/AgentCollectionTests.cs
--- a/Engine/Source/Programs/Horde/Plugins/Build/HordeServer.Build.Tests/Agents/AgentCollectionTests.cs
+++ b/Engine/Source/Programs/Horde/Plugins/Build/HordeServer.Build.Tests/Agents/AgentCollectionTests.cs
@@ -3,6 +3,7 @@
 using EpicGames.Horde.Agents;
 using EpicGames.Horde.Agents.Leases;
 using EpicGames.Horde.Agents.Pools;
+using HordeCommon.Rpc.Messages;
 using HordeServer.Agents;
 using HordeServer.Agents.Sessions;
 
@@ -56,7 +57,7 @@ public class AgentCollectionTests : BuildTestSetup
 		await UpdateAgentAsync();
 
 		await _agent.TryCreateSessionAsync(new CreateSessionOptions(SessionIdUtils.GenerateNewId(), DateTime.UtcNow, AgentStatus.Ok,
-			new List<string>(), new Dictionary<string, int>(), new List<PoolId>(), DateTime.UtcNow, null));
+			new RpcAgentCapabilities(), new List<PoolId>(), DateTime.UtcNow, null));
 
 		List<LeaseId> leases = await AgentCollection.FindActiveLeaseIdsAsync();
 
diff --git a/Engine/Source/Programs/Horde/Plugins/Build/HordeServer.Build.Tests/Agents/AgentServiceTest.cs b/Engine/Source/Programs/Horde/Plugins/Build/HordeServer.Build.Tests/Agents/AgentServiceTest.cs
--- a/Engine/Source/Programs/Horde/Plugins/Build/HordeServer.Build.Tests/Agents/AgentServiceTest.cs
+++ b/Engine/Source/Programs/Horde/Plugins/Build/HordeServer.Build.Tests/Agents/AgentServiceTest.cs
@@ -44,8 +44,7 @@ public class AgentServiceTest : BuildTestSetup
 	{
 		Fixture fixture = await CreateFixtureAsync();
 
-		IAgent agent = await AgentService.CreateSessionAsync(fixture.Agent1, AgentStatus.Ok, new List<string>(),
-			new Dictionary<string, int>(),
+		IAgent agent = await AgentService.CreateSessionAsync(fixture.Agent1, AgentStatus.Ok, new RpcAgentCapabilities(), 
 			"test");
 
 		Assert.IsTrue(AgentService.AuthorizeSession(agent, GetUser(agent), out string _));
@@ -68,7 +67,7 @@ public class AgentServiceTest : BuildTestSetup
 		Assert.IsFalse(agent.LastStatusChange.HasValue);
 
 		// A session has been created, status change timestamp is current time
-		agent = await AgentService.CreateSessionAsync(agent, AgentStatus.Ok, new List<string>(), new Dictionary<string, int>(), "v1");
+		agent = await AgentService.CreateSessionAsync(agent, AgentStatus.Ok, new RpcAgentCapabilities(), "v1");
 		Assert.AreEqual(AgentStatus.Ok, agent.Status);
 		Assert.AreEqual(ToUnixTime(Clock.UtcNow), ToUnixTime(agent.LastStatusChange!.Value));
 		DateTime lastStatusChange = agent.LastStatusChange!.Value;
@@ -76,7 +75,7 @@ public class AgentServiceTest : BuildTestSetup
 		await Clock.AdvanceAsync(TimeSpan.FromMinutes(1));
 
 		// The session is re-created, status change timestamp is same as when it first got created
-		agent = await AgentService.CreateSessionAsync(agent, AgentStatus.Ok, new List<string>(), new Dictionary<string, int>(), "v1");
+		agent = await AgentService.CreateSessionAsync(agent, AgentStatus.Ok, new RpcAgentCapabilities(), "v1");
 		Assert.AreEqual(AgentStatus.Ok, agent.Status);
 		Assert.AreEqual(ToUnixTime(lastStatusChange), ToUnixTime(agent.LastStatusChange!.Value));
 	}
@@ -85,7 +84,7 @@ public class AgentServiceTest : BuildTestSetup
 	private async Task<IAgent> CreateAgentSessionAsync()
 	{
 		IAgent agent = await AgentService.CreateAgentAsync("agentServiceTest-" + s_agentId++, false, "");
-		agent = await AgentService.CreateSessionAsync(agent, AgentStatus.Ok, new List<string>(), new Dictionary<string, int>(), "v1");
+		agent = await AgentService.CreateSessionAsync(agent, AgentStatus.Ok, new RpcAgentCapabilities(), "v1");
 		return agent;
 	}
 
@@ -113,7 +112,7 @@ public class AgentServiceTest : BuildTestSetup
 		DateTime lastStatusChange = agent.LastStatusChange!.Value;
 		await Clock.AdvanceAsync(TimeSpan.FromMinutes(1));
 
-		agent = (await AgentService.UpdateSessionAsync(agent, agent.SessionId!.Value, status, null, null, new List<RpcLease>()))!;
+		agent = (await AgentService.UpdateSessionAsync(agent, agent.SessionId!.Value, status, null, new List<RpcLease>()))!;
 		if (expectTimestampUpdate)
 		{
 			AssertNotEqual(lastStatusChange, agent.LastStatusChange);
@@ -137,11 +136,11 @@ public class AgentServiceTest : BuildTestSetup
 		}
 
 		List<string> props = new() { $"{KnownPropertyNames.AwsInstanceType}=m5.large" };
-		IAgent agent = await AgentService.CreateSessionAsync(fixture.Agent1, AgentStatus.Ok, props, new Dictionary<string, int>(), "test");
+		IAgent agent = await AgentService.CreateSessionAsync(fixture.Agent1, AgentStatus.Ok, new RpcAgentCapabilities(props), "test");
 		Assert.IsFalse(await AuditLogContains("AWS EC2 instance type changed"));
 
 		props = new() { $"{KnownPropertyNames.AwsInstanceType}=c6.xlarge" };
-		agent = await AgentService.CreateSessionAsync(agent, AgentStatus.Ok, props, new Dictionary<string, int>(), "test");
+		agent = await AgentService.CreateSessionAsync(agent, AgentStatus.Ok, new RpcAgentCapabilities(props), "test");
 		Assert.IsTrue(await AuditLogContains("AWS EC2 instance type changed"));
 	}
 
@@ -150,8 +149,8 @@ public class AgentServiceTest : BuildTestSetup
 	{
 		IAgent agent1 = await AgentService.CreateAgentAsync("agent1", false, "");
 		IAgent agent2 = await AgentService.CreateAgentAsync("agent2", false, "");
-		await AgentService.CreateSessionAsync(agent1, AgentStatus.Ok, new List<string>() { "aws-instance-type=c5.24xlarge", "osfamily=windows" }, new Dictionary<string, int>(), "test");
-		await AgentService.CreateSessionAsync(agent2, AgentStatus.Ok, new List<string>() { "aws-instance-type=c4.4xLARge", "osfamily=WinDowS" }, new Dictionary<string, int>(), "test");
+		await AgentService.CreateSessionAsync(agent1, AgentStatus.Ok, new RpcAgentCapabilities(new List<string>() { "aws-instance-type=c5.24xlarge", "osfamily=windows" }), "test");
+		await AgentService.CreateSessionAsync(agent2, AgentStatus.Ok, new RpcAgentCapabilities(new List<string>() { "aws-instance-type=c4.4xLARge", "osfamily=WinDowS" }), "test");
 
 		List<AgentRateConfig> agentRateConfigs = new()
 		{
diff --git a/Engine/Source/Programs/Horde/Plugins/Build/HordeServer.Build.Tests/BuildTestSetup.cs b/Engine/Source/Programs/Horde/Plugins/Build/HordeServer.Build.Tests/BuildTestSetup.cs
--- a/Engine/Source/Programs/Horde/Plugins/Build/HordeServer.Build.Tests/BuildTestSetup.cs
+++ b/Engine/Source/Programs/Horde/Plugins/Build/HordeServer.Build.Tests/BuildTestSetup.cs
@@ -7,6 +7,7 @@ using Amazon.EC2;
 using EpicGames.Horde.Agents;
 using EpicGames.Horde.Agents.Pools;
 using EpicGames.Horde.Artifacts;
+using HordeCommon.Rpc.Messages;
 using HordeServer.Acls;
 using HordeServer.Agents;
 using HordeServer.Agents.Enrollment;
@@ -300,7 +301,6 @@ namespace HordeServer.Tests
 				Clock.UtcNow = now + adjustClockBy.Value;
 			}
 
-			Dictionary<string, int> resources = new();
 			List<string> tempProps = new(properties ?? new List<string>());
 			if (awsInstanceId != null)
 			{
@@ -313,7 +313,7 @@ namespace HordeServer.Tests
 			agent = await agent.TryUpdateAsync(new UpdateAgentOptions { Enabled = enabled, ExplicitPools = poolId != null ? [poolId.Value] : [] });
 			Assert.IsNotNull(agent);
 
-			agent = await AgentService.CreateSessionAsync(agent, AgentStatus.Ok, tempProps, resources, null);
+			agent = await AgentService.CreateSessionAsync(agent, AgentStatus.Ok, new RpcAgentCapabilities(tempProps), null);
 			Assert.IsNotNull(agent);
 
 			if (workspaces is { Count: > 0 })
diff --git a/Engine/Source/Programs/Horde/Plugins/Build/HordeServer.Build.Tests/Fleet/AwsAutoScalingLifecycleTest.cs b/Engine/Source/Programs/Horde/Plugins/Build/HordeServer.Build.Tests/Fleet/AwsAutoScalingLifecycleTest.cs
--- a/Engine/Source/Programs/Horde/Plugins/Build/HordeServer.Build.Tests/Fleet/AwsAutoScalingLifecycleTest.cs
+++ b/Engine/Source/Programs/Horde/Plugins/Build/HordeServer.Build.Tests/Fleet/AwsAutoScalingLifecycleTest.cs
@@ -7,6 +7,7 @@ using Amazon.AutoScaling.Model;
 using EpicGames.Horde.Agents;
 using EpicGames.Horde.Agents.Leases;
 using EpicGames.Horde.Agents.Pools;
+using HordeCommon.Rpc.Messages;
 using HordeServer.Agents;
 using HordeServer.Agents.Fleet;
 using HordeServer.Utilities;
@@ -101,7 +102,7 @@ public class AwsAutoScalingLifecycleServiceTest : BuildTestSetup
 		LifecycleActionEvent lae = new() { Ec2InstanceId = "i-1234", LifecycleActionToken = "action-token-test", Origin = "AutoScalingGroup" };
 		List<string> props = new() { KnownPropertyNames.AwsInstanceId + "=" + lae.Ec2InstanceId };
 		IAgent agent = await CreateAgentAsync(new PoolId("pool1"), properties: props);
-		agent = await AgentService.CreateSessionAsync(agent, AgentStatus.Ok, props, new Dictionary<string, int>(), "v1");
+		agent = await AgentService.CreateSessionAsync(agent, AgentStatus.Ok, new RpcAgentCapabilities(props), "v1");
 		await _asgLifecycleService.InitiateTerminationAsync(lae, CancellationToken.None);
 		TimeSpan extraMargin = TimeSpan.FromSeconds(10);
 
@@ -114,7 +115,7 @@ public class AwsAutoScalingLifecycleServiceTest : BuildTestSetup
 		AssertLifecycleUpdate(lae, AwsAutoScalingLifecycleService.ActionContinue, _lifecycleUpdates[1]);
 
 		// Agent responded to shutdown request, now safe to notify AWS ASG the termination can take place
-		agent = await AgentService.CreateSessionAsync(agent, AgentStatus.Stopped, props, new Dictionary<string, int>(), "v1");
+		agent = await AgentService.CreateSessionAsync(agent, AgentStatus.Stopped, new RpcAgentCapabilities(props), "v1");
 		await Clock.AdvanceAsync(_asgLifecycleService.LifecycleUpdaterInterval + extraMargin);
 		AssertLifecycleUpdate(lae, AwsAutoScalingLifecycleService.ActionAbandon, _lifecycleUpdates[2]);
 
diff --git a/Engine/Source/Programs/Horde/Plugins/Build/HordeServer.Build.Tests/Jobs/JobServiceTests.cs b/Engine/Source/Programs/Horde/Plugins/Build/HordeServer.Build.Tests/Jobs/JobServiceTests.cs
--- a/Engine/Source/Programs/Horde/Plugins/Build/HordeServer.Build.Tests/Jobs/JobServiceTests.cs
+++ b/Engine/Source/Programs/Horde/Plugins/Build/HordeServer.Build.Tests/Jobs/JobServiceTests.cs
@@ -8,6 +8,7 @@ using EpicGames.Horde.Commits;
 using EpicGames.Horde.Jobs;
 using EpicGames.Horde.Jobs.Templates;
 using EpicGames.Horde.Streams;
+using HordeCommon.Rpc.Messages;
 using HordeServer.Agents;
 using HordeServer.Commits;
 using HordeServer.Jobs;
@@ -252,7 +253,7 @@ namespace HordeServer.Tests.Jobs
 			agent = await agent.TryUpdateAsync(new UpdateAgentOptions { Enabled = true, ExplicitPools = new List<PoolId> { new PoolId("win") } });
 			Assert.IsNotNull(agent);
 
-			await AgentService.CreateSessionAsync(agent, AgentStatus.Ok, new List<string>(), new Dictionary<string, int>(), null);
+			await AgentService.CreateSessionAsync(agent, AgentStatus.Ok, new RpcAgentCapabilities(), null);
 
 			ITemplate template = await TemplateCollection.GetOrAddAsync(new TemplateConfig { Name = "Test template" });
 			IGraph graph = await GraphCollection.AddAsync(template, null);
diff --git a/Engine/Source/Programs/Horde/Plugins/Build/HordeServer.Build.Tests/Jobs/JobTaskSourceTests.cs b/Engine/Source/Programs/Horde/Plugins/Build/HordeServer.Build.Tests/Jobs/JobTaskSourceTests.cs
--- a/Engine/Source/Programs/Horde/Plugins/Build/HordeServer.Build.Tests/Jobs/JobTaskSourceTests.cs
+++ b/Engine/Source/Programs/Horde/Plugins/Build/HordeServer.Build.Tests/Jobs/JobTaskSourceTests.cs
@@ -5,6 +5,7 @@ using EpicGames.Horde.Agents.Pools;
 using EpicGames.Horde.Commits;
 using EpicGames.Horde.Jobs;
 using EpicGames.Horde.Users;
+using HordeCommon.Rpc.Messages;
 using HordeServer.Agents;
 using HordeServer.Agents.Pools;
 using HordeServer.Jobs;
@@ -164,7 +165,7 @@ namespace HordeServer.Tests.Jobs
 				agent = await agent.TryUpdateAsync(new UpdateAgentOptions { Enabled = isAgentEnabled, ExplicitPools = new List<PoolId> { pool.Id } });
 				Assert.IsNotNull(agent);
 
-				await AgentService.CreateSessionAsync(agent, AgentStatus.Ok, new List<string>(), new Dictionary<string, int>(), null);
+				await AgentService.CreateSessionAsync(agent, AgentStatus.Ok, new RpcAgentCapabilities(), null);
 			}
 
 			JobTaskSource.OnJobScheduled += (pool, poolHasAgentsOnline, job, graph, batchId) =>
diff --git a/Engine/Source/Programs/Horde/Plugins/Compute/HordeServer.Compute/Agents/AgentCollection.cs b/Engine/Source/Programs/Horde/Plugins/Compute/HordeServer.Compute/Agents/AgentCollection.cs
--- a/Engine/Source/Programs/Horde/Plugins/Compute/HordeServer.Compute/Agents/AgentCollection.cs
+++ b/Engine/Source/Programs/Horde/Plugins/Compute/HordeServer.Compute/Agents/AgentCollection.cs
@@ -10,6 +10,7 @@ using EpicGames.Horde.Agents.Pools;
 using EpicGames.Horde.Agents.Sessions;
 using EpicGames.Redis;
 using Google.Protobuf.WellKnownTypes;
+using HordeCommon.Rpc.Messages;
 using HordeCommon.Rpc.Tasks;
 using HordeServer.Auditing;
 using HordeServer.Server;
@@ -548,17 +549,20 @@ namespace HordeServer.Agents
 			{
 				updates.Add(updateBuilder.Set(x => x.SessionExpiresAt, options.SessionExpiresAt.Value));
 			}
-			if (options.Properties != null)
+
+			List<string>? newProperties = agent.Properties;
+			if (options.Capabilities != null)
 			{
-				List<string> newProperties = options.Properties.OrderBy(x => x, StringComparer.OrdinalIgnoreCase).ToList();
+				options.Capabilities.Flatten(out newProperties, out Dictionary<string, int> newResources);
+
 				if (!(agent.Properties ?? Enumerable.Empty<string>()).SequenceEqual(newProperties, StringComparer.Ordinal))
 				{
 					updates.Add(updateBuilder.Set(x => x.Properties, newProperties));
 				}
-			}
-			if (options.Resources != null && !ResourcesEqual(options.Resources, agent.Resources))
-			{
-				updates.Add(updateBuilder.Set(x => x.Resources, new Dictionary<string, int>(options.Resources)));
+				if (!ResourcesEqual(newResources, agent.Resources))
+				{
+					updates.Add(updateBuilder.Set(x => x.Resources, newResources));
+				}
 			}
 			if (options.DynamicPools != null)
 			{
@@ -597,7 +601,7 @@ namespace HordeServer.Agents
 			}
 
 			// Update the pools
-			List<PoolId> pools = CreatePoolsList(options.DynamicPools ?? agent.DynamicPools, agent.ExplicitPools, options.Properties ?? agent.Properties ?? Enumerable.Empty<string>());
+			List<PoolId> pools = CreatePoolsList(options.DynamicPools ?? agent.DynamicPools, agent.ExplicitPools, newProperties ?? Enumerable.Empty<string>());
 			if (!Enumerable.SequenceEqual(pools, agent.Pools))
 			{
 				updates.Add(updateBuilder.Set(x => x.Pools, pools));
@@ -688,8 +692,7 @@ namespace HordeServer.Agents
 		/// <inheritdoc/>
 		async Task<AgentDocument?> TryCreateSessionAsync(AgentDocument agent, CreateSessionOptions options, CancellationToken cancellationToken)
 		{
-			List<string> newProperties = options.Properties.OrderBy(x => x, StringComparer.OrdinalIgnoreCase).ToList();
-			Dictionary<string, int> newResources = new(options.Resources);
+			options.Capabilities.Flatten(out List<string> newProperties, out Dictionary<string, int> newResources);
 			List<PoolId> newDynamicPools = new(options.DynamicPools);
 
 			// Reset the agent to use the new session
diff --git a/Engine/Source/Programs/Horde/Plugins/Compute/HordeServer.Compute/Agents/AgentService.cs b/Engine/Source/Programs/Horde/Plugins/Compute/HordeServer.Compute/Agents/AgentService.cs
--- a/Engine/Source/Programs/Horde/Plugins/Compute/HordeServer.Compute/Agents/AgentService.cs
+++ b/Engine/Source/Programs/Horde/Plugins/Compute/HordeServer.Compute/Agents/AgentService.cs
@@ -379,12 +379,11 @@ namespace HordeServer.Agents
 		/// </summary>
 		/// <param name="agent">The agent to create a session for</param>
 		/// <param name="status">Current status of the agent</param>
-		/// <param name="properties">Properties for the agent</param>
-		/// <param name="resources">Resources which the agent has</param>
+		/// <param name="capabilities">Capabilities for the agent</param>
 		/// <param name="version">Version of the software that's running</param>
 		/// <param name="cancellationToken">Cancellation token for the operation</param>
 		/// <returns>New agent state</returns>
-		public async Task<IAgent> CreateSessionAsync(IAgent agent, AgentStatus status, IReadOnlyList<string> properties, IReadOnlyDictionary<string, int> resources, string? version, CancellationToken cancellationToken = default)
+		public async Task<IAgent> CreateSessionAsync(IAgent agent, AgentStatus status, RpcAgentCapabilities capabilities, string? version, CancellationToken cancellationToken = default)
 		{
 			DateTime? lastStatusChange = null;
 			for (; ; )
@@ -413,14 +412,14 @@ namespace HordeServer.Agents
 					}
 
 					// Create a new session document
-					ISession newSession = await _sessions.AddAsync(SessionIdUtils.GenerateNewId(), agent.Id, _clock.UtcNow, properties, resources, version, cancellationToken);
+					ISession newSession = await _sessions.AddAsync(SessionIdUtils.GenerateNewId(), agent.Id, _clock.UtcNow, version, cancellationToken);
 					DateTime sessionExpiresAt = utcNow + SessionExpiryTime;
 
 					// Get the new pools for the agent
 					List<PoolId> dynamicPools = await GetDynamicPoolsAsync(agent, cancellationToken);
 
 					// Reset the agent to use the new session
-					newAgent = await agent.TryCreateSessionAsync(new CreateSessionOptions(newSession.Id, sessionExpiresAt, status, properties, resources, dynamicPools, lastStatusChange ?? utcNow, version), cancellationToken);
+					newAgent = await agent.TryCreateSessionAsync(new CreateSessionOptions(newSession.Id, sessionExpiresAt, status, capabilities, dynamicPools, lastStatusChange ?? utcNow, version), cancellationToken);
 					if (newAgent != null)
 					{
 						LogPropertyChanges(agentLogger, agent.Properties, newAgent.Properties);
@@ -686,7 +685,7 @@ namespace HordeServer.Agents
 		/// <summary>
 		/// 
 		/// </summary>
-		public async Task<IAgent?> UpdateSessionWithWaitAsync(IAgent inAgent, SessionId sessionId, AgentStatus status, IReadOnlyList<string>? properties, IReadOnlyDictionary<string, int>? resources, IList<HordeCommon.Rpc.Messages.RpcLease> newLeases, CancellationToken cancellationToken = default)
+		public async Task<IAgent?> UpdateSessionWithWaitAsync(IAgent inAgent, SessionId sessionId, AgentStatus status, RpcAgentCapabilities? capabilities, IList<HordeCommon.Rpc.Messages.RpcLease> newLeases, CancellationToken cancellationToken = default)
 		{
 			IAgent? agent = inAgent;
 
@@ -694,7 +693,7 @@ namespace HordeServer.Agents
 			uint updateIndex = agent.UpdateIndex;
 
 			// Update the agent session and return to the caller if anything changes
-			agent = await UpdateSessionAsync(agent, sessionId, status, properties, resources, newLeases, cancellationToken);
+			agent = await UpdateSessionAsync(agent, sessionId, status, capabilities, newLeases, cancellationToken);
 			if (agent != null && agent.UpdateIndex == updateIndex && (agent.Leases.Count > 0 || agent.Status != AgentStatus.Stopping))
 			{
 				agent = await WaitForLeaseAsync(agent, newLeases, cancellationToken);
@@ -708,12 +707,11 @@ namespace HordeServer.Agents
 		/// <param name="inAgent">The current agent state</param>
 		/// <param name="sessionId">Id of the session</param>
 		/// <param name="status">New status for the agent</param>
-		/// <param name="properties">New agent properties</param>
-		/// <param name="resources">New agent resources</param>
+		/// <param name="capabilities">New agent capabilities</param>
 		/// <param name="newLeases">New list of leases for this session</param>
 		/// <param name="cancellationToken">Cancellation token for the operation</param>
 		/// <returns>Updated agent state</returns>
-		public async Task<IAgent?> UpdateSessionAsync(IAgent inAgent, SessionId sessionId, AgentStatus status, IReadOnlyList<string>? properties, IReadOnlyDictionary<string, int>? resources, IList<HordeCommon.Rpc.Messages.RpcLease> newLeases, CancellationToken cancellationToken = default)
+		public async Task<IAgent?> UpdateSessionAsync(IAgent inAgent, SessionId sessionId, AgentStatus status, RpcAgentCapabilities? capabilities, IList<HordeCommon.Rpc.Messages.RpcLease> newLeases, CancellationToken cancellationToken = default)
 		{
 			DateTime utcNow = _clock.UtcNow;
 
@@ -805,7 +803,7 @@ namespace HordeServer.Agents
 				List<PoolId> dynamicPools = await GetDynamicPoolsAsync(agent, cancellationToken);
 
 				// Update the agent, and try to create new lease documents if we succeed
-				IAgent? newAgent = await agent.TryUpdateSessionAsync(new UpdateSessionOptions(status, sessionExpiresAt, properties, resources, dynamicPools, updateLeases ? leases : null), cancellationToken);
+				IAgent? newAgent = await agent.TryUpdateSessionAsync(new UpdateSessionOptions(status, sessionExpiresAt, capabilities, dynamicPools, updateLeases ? leases : null), cancellationToken);
 				if (newAgent != null)
 				{
 					agent = newAgent;
@@ -872,7 +870,7 @@ namespace HordeServer.Agents
 
 				// Update the session document
 				Agents.GetLogger(agent.Id).LogInformation("Terminated session {SessionId}", sessionId);
-				await _sessions.UpdateAsync(sessionId, finishTime, agent.Properties, agent.Resources, cancellationToken);
+				await _sessions.UpdateAsync(sessionId, finishTime, cancellationToken);
 				return agent;
 			}
 			return null;
diff --git a/Engine/Source/Programs/Horde/Plugins/Compute/HordeServer.Compute/Agents/AgentsController.cs b/Engine/Source/Programs/Horde/Plugins/Compute/HordeServer.Compute/Agents/AgentsController.cs
--- a/Engine/Source/Programs/Horde/Plugins/Compute/HordeServer.Compute/Agents/AgentsController.cs
+++ b/Engine/Source/Programs/Horde/Plugins/Compute/HordeServer.Compute/Agents/AgentsController.cs
@@ -347,7 +347,7 @@ namespace HordeServer.Agents
 
 		static GetSessionResponse CreateGetSessionResponse(ISession session)
 		{
-			return new GetSessionResponse(session.Id, session.StartTime, session.FinishTime, new List<string>(session.Properties), new Dictionary<string, int>(session.Resources), session.Version);
+			return new GetSessionResponse(session.Id, session.StartTime, session.FinishTime, session.Version);
 		}
 
 		/// <summary>
diff --git a/Engine/Source/Programs/Horde/Plugins/Compute/HordeServer.Compute/Agents/IAgent.cs b/Engine/Source/Programs/Horde/Plugins/Compute/HordeServer.Compute/Agents/IAgent.cs
--- a/Engine/Source/Programs/Horde/Plugins/Compute/HordeServer.Compute/Agents/IAgent.cs
+++ b/Engine/Source/Programs/Horde/Plugins/Compute/HordeServer.Compute/Agents/IAgent.cs
@@ -273,23 +273,21 @@ namespace HordeServer.Agents
 	/// <param name="SessionId">New session id</param>
 	/// <param name="SessionExpiresAt">Expiry time for the new session</param>
 	/// <param name="Status">Status of the agent</param>
-	/// <param name="Properties">Properties for the current session</param>
-	/// <param name="Resources">Resources for the agent</param>
+	/// <param name="Capabilities">Capabilities for the agent</param>
 	/// <param name="DynamicPools">New list of dynamic pools for the agent</param>
 	/// <param name="LastStatusChange">Time to force status change timestamp to</param>
 	/// <param name="Version">Current version of the agent software</param>
-	public record class CreateSessionOptions(SessionId SessionId, DateTime SessionExpiresAt, AgentStatus Status, IReadOnlyList<string> Properties, IReadOnlyDictionary<string, int> Resources, IReadOnlyList<PoolId> DynamicPools, DateTime LastStatusChange, string? Version);
+	public record class CreateSessionOptions(SessionId SessionId, DateTime SessionExpiresAt, AgentStatus Status, RpcAgentCapabilities Capabilities, IReadOnlyList<PoolId> DynamicPools, DateTime LastStatusChange, string? Version);
 
 	/// <summary>
 	/// Options for updating a new agent session
 	/// </summary>
 	/// <param name="Status">New status of the agent</param>
 	/// <param name="SessionExpiresAt">New expiry time for the current session</param>
-	/// <param name="Properties">Properties for the current session</param>
-	/// <param name="Resources">Resources for the agent</param>
+	/// <param name="Capabilities">Capbilities for the session</param>
 	/// <param name="DynamicPools">New list of dynamic pools for the agent</param>
 	/// <param name="Leases">New set of leases</param>
-	public record class UpdateSessionOptions(AgentStatus? Status = null, DateTime? SessionExpiresAt = null, IReadOnlyList<string>? Properties = null, IReadOnlyDictionary<string, int>? Resources = null, IReadOnlyList<PoolId>? DynamicPools = null, List<AgentLease>? Leases = null);
+	public record class UpdateSessionOptions(AgentStatus? Status = null, DateTime? SessionExpiresAt = null, RpcAgentCapabilities? Capabilities = null, IReadOnlyList<PoolId>? DynamicPools = null, List<AgentLease>? Leases = null);
 
 	/// <summary>
 	/// Extension methods for IAgent
diff --git a/Engine/Source/Programs/Horde/Plugins/Compute/HordeServer.Compute/Agents/Sessions/ISession.cs b/Engine/Source/Programs/Horde/Plugins/Compute/HordeServer.Compute/Agents/Sessions/ISession.cs
--- a/Engine/Source/Programs/Horde/Plugins/Compute/HordeServer.Compute/Agents/Sessions/ISession.cs
+++ b/Engine/Source/Programs/Horde/Plugins/Compute/HordeServer.Compute/Agents/Sessions/ISession.cs
@@ -30,16 +30,6 @@ namespace HordeServer.Agents.Sessions
 		/// </summary>
 		public DateTime? FinishTime { get; }
 
-		/// <summary>
-		/// Properties of this agent at the time the session started
-		/// </summary>
-		public IReadOnlyList<string> Properties { get; }
-
-		/// <summary>
-		/// Resources reported by the agent at the start of the session
-		/// </summary>
-		public IReadOnlyDictionary<string, int> Resources { get; }
-
 		/// <summary>
 		/// Version of the agent software
 		/// </summary>
diff --git a/Engine/Source/Programs/Horde/Plugins/Compute/HordeServer.Compute/Agents/Sessions/ISessionCollection.cs b/Engine/Source/Programs/Horde/Plugins/Compute/HordeServer.Compute/Agents/Sessions/ISessionCollection.cs
--- a/Engine/Source/Programs/Horde/Plugins/Compute/HordeServer.Compute/Agents/Sessions/ISessionCollection.cs
+++ b/Engine/Source/Programs/Horde/Plugins/Compute/HordeServer.Compute/Agents/Sessions/ISessionCollection.cs
@@ -16,11 +16,9 @@ namespace HordeServer.Agents.Sessions
 		/// <param name="id">The session id</param>
 		/// <param name="agentId">The agent this session is for</param>
 		/// <param name="startTime">Start time of this session</param>
-		/// <param name="properties">Properties of this agent at the time the session started</param>
-		/// <param name="resources">Resources which the agent has</param>
 		/// <param name="version">Version of the agent software</param>
 		/// <param name="cancellationToken">Cancellation token for the operation</param>
-		Task<ISession> AddAsync(SessionId id, AgentId agentId, DateTime startTime, IReadOnlyList<string>? properties, IReadOnlyDictionary<string, int>? resources, string? version, CancellationToken cancellationToken = default);
+		Task<ISession> AddAsync(SessionId id, AgentId agentId, DateTime startTime, string? version, CancellationToken cancellationToken = default);
 
 		/// <summary>
 		/// Gets information about a particular session
@@ -56,11 +54,9 @@ namespace HordeServer.Agents.Sessions
 		/// </summary>
 		/// <param name="sessionId">The session to update</param>
 		/// <param name="finishTime">Time at which the session finished</param>
-		/// <param name="properties">The agent properties</param>
-		/// <param name="resources">Resources which the agent has</param>
 		/// <param name="cancellationToken">Cancellation token for the operation</param>
 		/// <returns>Async task</returns>
-		Task UpdateAsync(SessionId sessionId, DateTime? finishTime, IReadOnlyList<string>? properties, IReadOnlyDictionary<string, int>? resources, CancellationToken cancellationToken = default);
+		Task UpdateAsync(SessionId sessionId, DateTime? finishTime, CancellationToken cancellationToken = default);
 
 		/// <summary>
 		/// Delete a session from the collection
diff --git a/Engine/Source/Programs/Horde/Plugins/Compute/HordeServer.Compute/Agents/Sessions/SessionCollection.cs b/Engine/Source/Programs/Horde/Plugins/Compute/HordeServer.Compute/Agents/Sessions/SessionCollection.cs
--- a/Engine/Source/Programs/Horde/Plugins/Compute/HordeServer.Compute/Agents/Sessions/SessionCollection.cs
+++ b/Engine/Source/Programs/Horde/Plugins/Compute/HordeServer.Compute/Agents/Sessions/SessionCollection.cs
@@ -27,31 +27,18 @@ namespace HordeServer.Agents.Sessions
 
 			public DateTime StartTime { get; set; }
 			public DateTime? FinishTime { get; set; }
-			public List<string> Properties { get; set; } = new List<string>();
-			public Dictionary<string, int> Resources { get; set; } = new Dictionary<string, int>();
 			public string Version { get; set; } = String.Empty;
 
-			IReadOnlyList<string> ISession.Properties => Properties;
-			IReadOnlyDictionary<string, int> ISession.Resources => Resources;
-
 			[BsonConstructor]
 			private SessionDocument()
 			{
 			}
 
-			public SessionDocument(SessionId id, AgentId agentId, DateTime startTime, IReadOnlyList<string>? properties, IReadOnlyDictionary<string, int>? resources, string? version)
+			public SessionDocument(SessionId id, AgentId agentId, DateTime startTime, string? version)
 			{
 				Id = id;
 				AgentId = agentId;
 				StartTime = startTime;
-				if (properties != null)
-				{
-					Properties = new List<string>(properties);
-				}
-				if (resources != null)
-				{
-					Resources = new Dictionary<string, int>(resources);
-				}
 				if (version != null)
 				{
 					Version = version;
@@ -78,9 +65,9 @@ namespace HordeServer.Agents.Sessions
 		}
 
 		/// <inheritdoc/>
-		public async Task<ISession> AddAsync(SessionId id, AgentId agentId, DateTime startTime, IReadOnlyList<string>? properties, IReadOnlyDictionary<string, int>? resources, string? version, CancellationToken cancellationToken = default)
+		public async Task<ISession> AddAsync(SessionId id, AgentId agentId, DateTime startTime, string? version, CancellationToken cancellationToken = default)
 		{
-			SessionDocument newSession = new SessionDocument(id, agentId, startTime, properties, resources, version);
+			SessionDocument newSession = new SessionDocument(id, agentId, startTime, version);
 			await _sessions.InsertOneAsync(newSession, null, cancellationToken);
 			return newSession;
 		}
@@ -118,21 +105,13 @@ namespace HordeServer.Agents.Sessions
 		}
 
 		/// <inheritdoc/>
-		public Task UpdateAsync(SessionId sessionId, DateTime? finishTime, IReadOnlyList<string>? properties, IReadOnlyDictionary<string, int>? resources, CancellationToken cancellationToken = default)
+		public Task UpdateAsync(SessionId sessionId, DateTime? finishTime, CancellationToken cancellationToken = default)
 		{
 			List<UpdateDefinition<SessionDocument>> updates = new List<UpdateDefinition<SessionDocument>>();
 			if (finishTime != null)
 			{
 				updates.Add(Builders<SessionDocument>.Update.Set(x => x.FinishTime, finishTime));
 			}
-			if (properties != null)
-			{
-				updates.Add(Builders<SessionDocument>.Update.Set(x => x.Properties, new List<string>(properties)));
-			}
-			if (resources != null)
-			{
-				updates.Add(Builders<SessionDocument>.Update.Set(x => x.Resources, new Dictionary<string, int>(resources)));
-			}
 			return _sessions.FindOneAndUpdateAsync(x => x.Id == sessionId, Builders<SessionDocument>.Update.Combine(updates), cancellationToken: cancellationToken);
 		}
 
diff --git a/Engine/Source/Programs/Horde/Plugins/Compute/HordeServer.Compute/Protos/horde/messages/agent_capabilities.proto.cs b/Engine/Source/Programs/Horde/Plugins/Compute/HordeServer.Compute/Protos/horde/messages/agent_capabilities.proto.cs
new file mode 100644
--- /dev/null
+++ b/Engine/Source/Programs/Horde/Plugins/Compute/HordeServer.Compute/Protos/horde/messages/agent_capabilities.proto.cs
@@ -0,0 +1,61 @@
+// Copyright Epic Games, Inc. All Rights Reserved.
+
+using System;
+using System.Collections.Generic;
+using System.Linq;
+using System.Text;
+using System.Threading.Tasks;
+using EpicGames.Horde.Agents;
+
+#pragma warning disable CS1591
+
+namespace HordeCommon.Rpc.Messages
+{
+	partial class RpcAgentCapabilities
+	{
+		public RpcAgentCapabilities(IEnumerable<string> properties)
+		{
+			Properties.AddRange(properties);
+		}
+
+		public void Flatten(out List<string> properties, out Dictionary<string, int> resources)
+		{
+			properties = new List<string>();
+			resources = new Dictionary<string, int>();
+
+			properties.AddRange(Properties);
+
+			if (Devices.Count <= 0)
+			{
+				return;
+			}
+
+			RpcDeviceCapabilities device = Devices[0];
+			if (device.Properties == null)
+			{
+				return;
+			}
+
+			properties.AddRange(device.Properties);
+			properties.Sort(StringComparer.OrdinalIgnoreCase);
+
+			CopyPropertyToResource(KnownPropertyNames.LogicalCores, properties, resources);
+			CopyPropertyToResource(KnownPropertyNames.Ram, properties, resources);
+		}
+
+		static void CopyPropertyToResource(string name, List<string> properties, Dictionary<string, int> resources)
+		{
+			foreach (string property in properties)
+			{
+				if (property.Length > name.Length && property.StartsWith(name, StringComparison.OrdinalIgnoreCase) && property[name.Length] == '=')
+				{
+					int value;
+					if (Int32.TryParse(property.AsSpan(name.Length + 1), out value))
+					{
+						resources[name] = value;
+					}
+				}
+			}
+		}
+	}
+}
diff --git a/Engine/Source/Programs/Horde/Plugins/Compute/HordeServer.Compute/Server/ConsistencyService.cs b/Engine/Source/Programs/Horde/Plugins/Compute/HordeServer.Compute/Server/ConsistencyService.cs
--- a/Engine/Source/Programs/Horde/Plugins/Compute/HordeServer.Compute/Server/ConsistencyService.cs
+++ b/Engine/Source/Programs/Horde/Plugins/Compute/HordeServer.Compute/Server/ConsistencyService.cs
@@ -73,7 +73,7 @@ namespace HordeServer.Server
 					if (agent == null || agent.SessionId != session.Id)
 					{
 						_logger.LogWarning("Forcing agent {AgentId} session {SessionId} to complete.", session.AgentId, session.Id);
-						await _sessionCollection.UpdateAsync(session.Id, utcNow, null, null, cancellationToken);
+						await _sessionCollection.UpdateAsync(session.Id, utcNow, cancellationToken);
 						sessionIdToInstance.Remove(session.Id);
 					}
 				}
diff --git a/Engine/Source/Programs/Horde/Plugins/Compute/HordeServer.Compute/Server/RpcService.cs b/Engine/Source/Programs/Horde/Plugins/Compute/HordeServer.Compute/Server/RpcService.cs
--- a/Engine/Source/Programs/Horde/Plugins/Compute/HordeServer.Compute/Server/RpcService.cs
+++ b/Engine/Source/Programs/Horde/Plugins/Compute/HordeServer.Compute/Server/RpcService.cs
@@ -147,48 +147,6 @@ namespace HordeServer.Server
 			}
 		}
 
-		static void CopyPropertyToResource(string name, List<string> properties, Dictionary<string, int> resources)
-		{
-			foreach (string property in properties)
-			{
-				if (property.Length > name.Length && property.StartsWith(name, StringComparison.OrdinalIgnoreCase) && property[name.Length] == '=')
-				{
-					int value;
-					if (Int32.TryParse(property.AsSpan(name.Length + 1), out value))
-					{
-						resources[name] = value;
-					}
-				}
-			}
-		}
-
-		static void GetCapabilities(RpcAgentCapabilities? capabilities, out List<string> properties, out Dictionary<string, int> resources)
-		{
-			properties = new List<string>();
-			resources = new Dictionary<string, int>();
-
-			if (capabilities == null)
-			{
-				return;
-			}
-			properties.AddRange(capabilities.Properties);
-
-			if (capabilities.Devices.Count <= 0)
-			{
-				return;
-			}
-
-			RpcDeviceCapabilities device = capabilities.Devices[0];
-			if (device.Properties == null)
-			{
-				return;
-			}
-
-			properties.AddRange(device.Properties);
-			CopyPropertyToResource(KnownPropertyNames.LogicalCores, properties, resources);
-			CopyPropertyToResource(KnownPropertyNames.Ram, properties, resources);
-		}
-
 		/// <summary>
 		/// Creates a new agent
 		/// </summary>
@@ -261,11 +219,8 @@ namespace HordeServer.Server
 				throw new StructuredRpcException(StatusCode.PermissionDenied, "User is not authenticated to create session for {AgentId}", agentId);
 			}
 
-			// Get the known properties for this agent
-			GetCapabilities(request.Capabilities, out List<string> properties, out Dictionary<string, int> resources);
-
 			// Create a new session
-			agent = await _agentService.CreateSessionAsync(agent, (AgentStatus)request.Status, properties, resources, request.Version, context.CancellationToken);
+			agent = await _agentService.CreateSessionAsync(agent, (AgentStatus)request.Status, request.Capabilities, request.Version, context.CancellationToken);
 			if (agent == null)
 			{
 				throw new StructuredRpcException(StatusCode.NotFound, "Agent {AgentId} not found", agentId);
@@ -327,18 +282,10 @@ namespace HordeServer.Server
 						throw new StructuredRpcException(StatusCode.PermissionDenied, "Not authenticated for {AgentId}. Reason {Reason}", request.AgentId, authReason);
 					}
 
-					// Get the new capabilities of this agent
-					List<string>? properties = null;
-					Dictionary<string, int>? resources = null;
-					if (request.Capabilities != null)
-					{
-						GetCapabilities(request.Capabilities, out properties, out resources);
-					}
-
 					// Update the session
 					try
 					{
-						agent = await _agentService.UpdateSessionWithWaitAsync(agent, sessionId, (AgentStatus)request.Status, properties, resources, request.Leases, cancellationSource.Token);
+						agent = await _agentService.UpdateSessionWithWaitAsync(agent, sessionId, (AgentStatus)request.Status, request.Capabilities, request.Leases, cancellationSource.Token);
 					}
 					catch (OperationCanceledException)
 					{
diff --git a/Engine/Source/Programs/Shared/EpicGames.Horde/Agents/Sessions/SessionMessages.cs b/Engine/Source/Programs/Shared/EpicGames.Horde/Agents/Sessions/SessionMessages.cs
--- a/Engine/Source/Programs/Shared/EpicGames.Horde/Agents/Sessions/SessionMessages.cs
+++ b/Engine/Source/Programs/Shared/EpicGames.Horde/Agents/Sessions/SessionMessages.cs
@@ -11,8 +11,6 @@ namespace EpicGames.Horde.Agents.Sessions
 	/// <param name="Id">Unique id for this session</param>
 	/// <param name="StartTime">Start time for this session</param>
 	/// <param name="FinishTime">Finishing time for this session</param>
-	/// <param name="Properties">Properties of this agent</param>
-	/// <param name="Resources">Resources held by the agent at the start of the session</param>
 	/// <param name="Version"> Version of the software running during this session </param>
-	public record GetSessionResponse(SessionId Id, DateTime StartTime, DateTime? FinishTime, List<string>? Properties, Dictionary<string, int>? Resources, string? Version);
+	public record GetSessionResponse(SessionId Id, DateTime StartTime, DateTime? FinishTime, string? Version);
 }
