Horde: Add a custom type with options for creating new leases.

#rnx

diff --git a/Engine/Source/Programs/Horde/Plugins/Build/HordeServer.Build.Tests/Agents/AgentCollectionTests.cs b/Engine/Source/Programs/Horde/Plugins/Build/HordeServer.Build.Tests/Agents/AgentCollectionTests.cs
--- a/Engine/Source/Programs/Horde/Plugins/Build/HordeServer.Build.Tests/Agents/AgentCollectionTests.cs
+++ b/Engine/Source/Programs/Horde/Plugins/Build/HordeServer.Build.Tests/Agents/AgentCollectionTests.cs
@@ -3,6 +3,7 @@
 using EpicGames.Horde.Agents;
 using EpicGames.Horde.Agents.Leases;
 using EpicGames.Horde.Agents.Pools;
+using Google.Protobuf.WellKnownTypes;
 using HordeCommon.Rpc.Messages;
 using HordeServer.Agents;
 using HordeServer.Agents.Sessions;
@@ -13,30 +14,30 @@ namespace HordeServer.Tests.Agents;
 public class AgentCollectionTests : BuildTestSetup
 {
 	private IAgent _agent;
-	private readonly AgentLease _lease1;
-	private readonly AgentLease _lease2;
-	private readonly AgentLease _leaseWithParent3;
-	private readonly AgentLease _leaseWithParent4;
+	private readonly CreateLeaseOptions _lease1;
+	private readonly CreateLeaseOptions _lease2;
+	private readonly CreateLeaseOptions _leaseWithParent3;
+	private readonly CreateLeaseOptions _leaseWithParent4;
 
 	public AgentCollectionTests()
 	{
 		_agent = CreateAgentAsync(new PoolId("pool1"), ephemeral: true).Result;
-		_lease1 = new AgentLease(LeaseId.Parse("aaaaaaaaaaaaaaaaaaaaaaaa"), null, "lease", null, null, null, LeaseState.Pending, null, false, null);
-		_lease2 = new AgentLease(LeaseId.Parse("bbbbbbbbbbbbbbbbbbbbbbbb"), null, "lease", null, null, null, LeaseState.Pending, null, false, null);
-		_leaseWithParent3 = new AgentLease(LeaseId.Parse("cccccccccccccccccccccccc"), _lease1.Id, "leaseWithParent3", null, null, null, LeaseState.Pending, null, false, null);
-		_leaseWithParent4 = new AgentLease(LeaseId.Parse("dddddddddddddddddddddddd"), _lease1.Id, "leaseWithParent4", null, null, null, LeaseState.Pending, null, false, null);
+		_lease1 = new CreateLeaseOptions(LeaseId.Parse("aaaaaaaaaaaaaaaaaaaaaaaa"), null, "lease", null, null, null, null, false, new Empty());
+		_lease2 = new CreateLeaseOptions(LeaseId.Parse("bbbbbbbbbbbbbbbbbbbbbbbb"), null, "lease", null, null, null, null, false, new Empty());
+		_leaseWithParent3 = new CreateLeaseOptions(LeaseId.Parse("cccccccccccccccccccccccc"), _lease1.Id, "leaseWithParent3", null, null, null, null, false, new Empty());
+		_leaseWithParent4 = new CreateLeaseOptions(LeaseId.Parse("dddddddddddddddddddddddd"), _lease1.Id, "leaseWithParent4", null, null, null, null, false, new Empty());
 	}
 
 	[TestMethod]
 	public async Task AddLeaseAsync()
 	{
-		await _agent.TryAddLeaseAsync(_lease1);
+		await _agent.TryCreateLeaseAsync(_lease1);
 		await UpdateAgentAsync();
 
-		await _agent.TryAddLeaseAsync(_leaseWithParent3);
+		await _agent.TryCreateLeaseAsync(_leaseWithParent3);
 		await UpdateAgentAsync();
 
-		await _agent.TryAddLeaseAsync(_leaseWithParent4);
+		await _agent.TryCreateLeaseAsync(_leaseWithParent4);
 		await UpdateAgentAsync();
 
 		List<LeaseId> leases = await AgentCollection.FindActiveLeaseIdsAsync();
@@ -50,10 +51,10 @@ public class AgentCollectionTests : BuildTestSetup
 	[TestMethod]
 	public async Task StartSessionAsync()
 	{
-		await _agent.TryAddLeaseAsync(_lease1);
+		await _agent.TryCreateLeaseAsync(_lease1);
 		await UpdateAgentAsync();
 
-		await _agent.TryAddLeaseAsync(_lease2);
+		await _agent.TryCreateLeaseAsync(_lease2);
 		await UpdateAgentAsync();
 
 		await _agent.TryCreateSessionAsync(new CreateSessionOptions(SessionIdUtils.GenerateNewId(), DateTime.UtcNow, AgentStatus.Ok,
@@ -67,10 +68,10 @@ public class AgentCollectionTests : BuildTestSetup
 	[TestMethod]
 	public async Task UpdateSession_WithEmptyLeasesAsync()
 	{
-		await _agent.TryAddLeaseAsync(_lease1);
+		await _agent.TryCreateLeaseAsync(_lease1);
 		await UpdateAgentAsync();
 
-		await _agent.TryAddLeaseAsync(_lease2);
+		await _agent.TryCreateLeaseAsync(_lease2);
 		await UpdateAgentAsync();
 
 		await _agent.TryUpdateSessionAsync(new UpdateSessionOptions { Leases = new List<AgentLease>() });
@@ -83,10 +84,10 @@ public class AgentCollectionTests : BuildTestSetup
 	[TestMethod]
 	public async Task UpdateSession_WithNewLeasesAsync()
 	{
-		await _agent.TryAddLeaseAsync(_lease1);
+		await _agent.TryCreateLeaseAsync(_lease1);
 		await UpdateAgentAsync();
 
-		await _agent.TryUpdateSessionAsync(new UpdateSessionOptions { Leases = new List<AgentLease> { _lease1, _lease2 } });
+		await _agent.TryUpdateSessionAsync(new UpdateSessionOptions { Leases = new List<AgentLease> { new AgentLease(_lease1), new AgentLease(_lease2) } });
 
 		List<LeaseId> leases = await AgentCollection.FindActiveLeaseIdsAsync();
 
@@ -98,13 +99,13 @@ public class AgentCollectionTests : BuildTestSetup
 	[TestMethod]
 	public async Task UpdateSession_WithOneLeaseRemovedAsync()
 	{
-		await _agent.TryAddLeaseAsync(_lease1);
+		await _agent.TryCreateLeaseAsync(_lease1);
 		await UpdateAgentAsync();
 
-		await _agent.TryAddLeaseAsync(_lease2);
+		await _agent.TryCreateLeaseAsync(_lease2);
 		await UpdateAgentAsync();
 
-		await _agent.TryUpdateSessionAsync(new UpdateSessionOptions { Leases = new List<AgentLease> { _lease1 } });
+		await _agent.TryUpdateSessionAsync(new UpdateSessionOptions { Leases = new List<AgentLease> { new AgentLease(_lease1) } });
 
 		List<LeaseId> leases = await AgentCollection.FindActiveLeaseIdsAsync();
 
@@ -115,10 +116,10 @@ public class AgentCollectionTests : BuildTestSetup
 	[TestMethod]
 	public async Task CancelLeaseAsync()
 	{
-		await _agent.TryAddLeaseAsync(_lease1);
+		await _agent.TryCreateLeaseAsync(_lease1);
 		await UpdateAgentAsync();
 
-		await _agent.TryAddLeaseAsync(_lease2);
+		await _agent.TryCreateLeaseAsync(_lease2);
 		await UpdateAgentAsync();
 
 		await _agent.TryCancelLeaseAsync(0);
@@ -132,10 +133,10 @@ public class AgentCollectionTests : BuildTestSetup
 	[TestMethod]
 	public async Task TerminateSessionAsync()
 	{
-		await _agent.TryAddLeaseAsync(_lease1);
+		await _agent.TryCreateLeaseAsync(_lease1);
 		await UpdateAgentAsync();
 
-		await _agent.TryAddLeaseAsync(_lease2);
+		await _agent.TryCreateLeaseAsync(_lease2);
 		await UpdateAgentAsync();
 
 		await _agent.TryTerminateSessionAsync();
@@ -148,13 +149,13 @@ public class AgentCollectionTests : BuildTestSetup
 	[TestMethod]
 	public async Task GetChildLeaseIdsAsync()
 	{
-		await _agent.TryAddLeaseAsync(_lease1);
+		await _agent.TryCreateLeaseAsync(_lease1);
 		await UpdateAgentAsync();
 
-		await _agent.TryAddLeaseAsync(_leaseWithParent3);
+		await _agent.TryCreateLeaseAsync(_leaseWithParent3);
 		await UpdateAgentAsync();
 
-		await _agent.TryAddLeaseAsync(_leaseWithParent4);
+		await _agent.TryCreateLeaseAsync(_leaseWithParent4);
 		await UpdateAgentAsync();
 
 		List<LeaseId> leases = await AgentCollection.GetChildLeaseIdsAsync(_leaseWithParent3.ParentId!.Value);
diff --git a/Engine/Source/Programs/Horde/Plugins/Build/HordeServer.Build.Tests/BuildTestSetup.cs b/Engine/Source/Programs/Horde/Plugins/Build/HordeServer.Build.Tests/BuildTestSetup.cs
--- a/Engine/Source/Programs/Horde/Plugins/Build/HordeServer.Build.Tests/BuildTestSetup.cs
+++ b/Engine/Source/Programs/Horde/Plugins/Build/HordeServer.Build.Tests/BuildTestSetup.cs
@@ -292,7 +292,7 @@ namespace HordeServer.Tests
 			List<AgentWorkspaceInfo>? workspaces = null,
 			TimeSpan? adjustClockBy = null,
 			string? awsInstanceId = null,
-			AgentLease? lease = null,
+			CreateLeaseOptions? lease = null,
 			bool ephemeral = false)
 		{
 			DateTime now = Clock.UtcNow;
@@ -328,7 +328,7 @@ namespace HordeServer.Tests
 
 			if (lease != null)
 			{
-				await agent.TryAddLeaseAsync(lease);
+				await agent.TryCreateLeaseAsync(lease);
 			}
 
 			Clock.UtcNow = now;
diff --git a/Engine/Source/Programs/Horde/Plugins/Build/HordeServer.Build.Tests/Fleet/AwsAutoScalingLifecycleTest.cs b/Engine/Source/Programs/Horde/Plugins/Build/HordeServer.Build.Tests/Fleet/AwsAutoScalingLifecycleTest.cs
--- a/Engine/Source/Programs/Horde/Plugins/Build/HordeServer.Build.Tests/Fleet/AwsAutoScalingLifecycleTest.cs
+++ b/Engine/Source/Programs/Horde/Plugins/Build/HordeServer.Build.Tests/Fleet/AwsAutoScalingLifecycleTest.cs
@@ -7,6 +7,7 @@ using Amazon.AutoScaling.Model;
 using EpicGames.Horde.Agents;
 using EpicGames.Horde.Agents.Leases;
 using EpicGames.Horde.Agents.Pools;
+using Google.Protobuf.WellKnownTypes;
 using HordeCommon.Rpc.Messages;
 using HordeServer.Agents;
 using HordeServer.Agents.Fleet;
@@ -172,7 +173,7 @@ public class AwsAutoScalingLifecycleServiceTest : BuildTestSetup
 	public async Task GetInstancesAvailableForTermination_AgentRunningJob_ReturnsNoInstanceIdAsync()
 	{
 		// Arrange
-		AgentLease lease = new(new LeaseId(BinaryIdUtils.CreateNew()), null, "test-lease", null, null, null, LeaseState.Active, null, false, null);
+		CreateLeaseOptions lease = new(new LeaseId(BinaryIdUtils.CreateNew()), null, "test-lease", null, null, null, null, false, new Empty());
 		_ = await CreateAgentAsync(new PoolId("pool1"), awsInstanceId: "i-1000", lease: lease);
 		TerminationPolicyEvent e = CreateTerminationPolicyEvent("i-1000");
 
diff --git a/Engine/Source/Programs/Horde/Plugins/Build/HordeServer.Build/Jobs/JobTaskSource.cs b/Engine/Source/Programs/Horde/Plugins/Build/HordeServer.Build/Jobs/JobTaskSource.cs
--- a/Engine/Source/Programs/Horde/Plugins/Build/HordeServer.Build/Jobs/JobTaskSource.cs
+++ b/Engine/Source/Programs/Horde/Plugins/Build/HordeServer.Build/Jobs/JobTaskSource.cs
@@ -13,8 +13,6 @@ using EpicGames.Horde.Jobs;
 using EpicGames.Horde.Logs;
 using EpicGames.Horde.Storage;
 using EpicGames.Horde.Streams;
-using Google.Protobuf;
-using Google.Protobuf.WellKnownTypes;
 using HordeCommon;
 using HordeCommon.Rpc.Messages;
 using HordeCommon.Rpc.Tasks;
@@ -165,12 +163,12 @@ namespace HordeServer.Jobs
 			/// <summary>
 			/// Task to wait for a lease to be assigned
 			/// </summary>
-			public Task<AgentLease?> Task => LeaseSource.Task;
+			public Task<CreateLeaseOptions?> Task => LeaseSource.Task;
 
 			/// <summary>
 			/// Completion source for the waiting agent. If a new queue item becomes available, the result will be passed through 
 			/// </summary>
-			public TaskCompletionSource<AgentLease?> LeaseSource { get; } = new TaskCompletionSource<AgentLease?>(TaskCreationOptions.RunContinuationsAsynchronously);
+			public TaskCompletionSource<CreateLeaseOptions?> LeaseSource { get; } = new TaskCompletionSource<CreateLeaseOptions?>(TaskCreationOptions.RunContinuationsAsynchronously);
 
 			/// <summary>
 			/// Constructor
@@ -715,13 +713,13 @@ namespace HordeServer.Jobs
 		}
 
 		/// <inheritdoc/>
-		public override Task<Task<AgentLease?>> AssignLeaseAsync(IAgent agent, CancellationToken cancellationToken)
+		public override Task<Task<CreateLeaseOptions?>> AssignLeaseAsync(IAgent agent, CancellationToken cancellationToken)
 		{
 			QueueWaiter waiter = new QueueWaiter(agent);
 			lock (_lockObject)
 			{
 				AssignAnyQueueItemToWaiter(waiter);
-				if (waiter.LeaseSource.Task.TryGetResult(out AgentLease? result))
+				if (waiter.LeaseSource.Task.TryGetResult(out CreateLeaseOptions? result))
 				{
 					if (result == null)
 					{
@@ -734,7 +732,7 @@ namespace HordeServer.Jobs
 			return Task.FromResult(WaitForLeaseAsync(waiter, cancellationToken));
 		}
 
-		private async Task<AgentLease?> WaitForLeaseAsync(QueueWaiter waiter, CancellationToken cancellationToken)
+		private async Task<CreateLeaseOptions?> WaitForLeaseAsync(QueueWaiter waiter, CancellationToken cancellationToken)
 		{
 			try
 			{
@@ -765,7 +763,7 @@ namespace HordeServer.Jobs
 		/// <param name="waiter">The agent waiting for work</param>
 		/// <param name="cancellationToken">Cancellation token for the operation</param>
 		/// <returns>New work to execute</returns>
-		private async Task<AgentLease?> TryCreateLeaseAsync(QueueItem item, QueueWaiter waiter, CancellationToken cancellationToken)
+		private async Task<CreateLeaseOptions?> TryCreateLeaseAsync(QueueItem item, QueueWaiter waiter, CancellationToken cancellationToken)
 		{
 			IJob job = item._job;
 			IJobStepBatch batch = item.Batch;
@@ -820,10 +818,8 @@ namespace HordeServer.Jobs
 				ExecuteJobTask? task = await CreateExecuteJobTaskAsync(leaseId, streamConfig, job, batch, agent, item._workspace, autoSdkWorkspace, logId, cancellationToken);
 				if (task != null)
 				{
-					byte[] payload = Any.Pack(task).ToByteArray();
-
 					// Create the lease and try to set it on the waiter. If this fails, the waiter has already moved on, and the lease can be cancelled.
-					AgentLease lease = new AgentLease(leaseId, null, leaseName.ToString(), job.StreamId, item._poolId, logId, LeaseState.Pending, null, true, payload);
+					CreateLeaseOptions lease = new CreateLeaseOptions(leaseId, null, leaseName.ToString(), job.StreamId, item._poolId, logId, null, true, task);
 					if (waiter.LeaseSource.TrySetResult(lease))
 					{
 						_logger.LogInformation("Assigned lease to agent");
diff --git a/Engine/Source/Programs/Horde/Plugins/Build/HordeServer.Build/Tasks/ConformTaskSource.cs b/Engine/Source/Programs/Horde/Plugins/Build/HordeServer.Build/Tasks/ConformTaskSource.cs
--- a/Engine/Source/Programs/Horde/Plugins/Build/HordeServer.Build/Tasks/ConformTaskSource.cs
+++ b/Engine/Source/Programs/Horde/Plugins/Build/HordeServer.Build/Tasks/ConformTaskSource.cs
@@ -5,8 +5,6 @@ using EpicGames.Horde.Agents;
 using EpicGames.Horde.Agents.Leases;
 using EpicGames.Horde.Jobs;
 using EpicGames.Horde.Logs;
-using Google.Protobuf;
-using Google.Protobuf.WellKnownTypes;
 using HordeCommon.Rpc.Messages;
 using HordeCommon.Rpc.Tasks;
 using HordeServer.Agents;
@@ -150,7 +148,7 @@ namespace HordeServer.Tasks
 		}
 
 		/// <inheritdoc/>
-		public override async Task<Task<AgentLease?>> AssignLeaseAsync(IAgent agent, CancellationToken cancellationToken)
+		public override async Task<Task<CreateLeaseOptions?>> AssignLeaseAsync(IAgent agent, CancellationToken cancellationToken)
 		{
 			if (!_settings.CurrentValue.EnableConformTasks || !_buildConfig.CurrentValue.EnableConformTasks)
 			{
@@ -177,9 +175,7 @@ namespace HordeServer.Tasks
 							task.LogId = log.Id.ToString();
 							task.RemoveUntrackedFiles = agent.RequestFullConform;
 
-							byte[] payload = Any.Pack(task).ToByteArray();
-
-							return LeaseAsync(new AgentLease(leaseId, null, "Updating workspaces", null, null, log.Id, LeaseState.Pending, null, true, payload));
+							return LeaseAsync(new CreateLeaseOptions(leaseId, null, "Updating workspaces", null, null, log.Id, null, true, task));
 						}
 						catch
 						{
diff --git a/Engine/Source/Programs/Horde/Plugins/Compute/HordeServer.Compute/Agents/AgentCollection.cs b/Engine/Source/Programs/Horde/Plugins/Compute/HordeServer.Compute/Agents/AgentCollection.cs
--- a/Engine/Source/Programs/Horde/Plugins/Compute/HordeServer.Compute/Agents/AgentCollection.cs
+++ b/Engine/Source/Programs/Horde/Plugins/Compute/HordeServer.Compute/Agents/AgentCollection.cs
@@ -9,8 +9,8 @@ using EpicGames.Horde.Agents.Leases;
 using EpicGames.Horde.Agents.Pools;
 using EpicGames.Horde.Agents.Sessions;
 using EpicGames.Redis;
+using Google.Protobuf;
 using Google.Protobuf.WellKnownTypes;
-using HordeCommon.Rpc.Messages;
 using HordeCommon.Rpc.Tasks;
 using HordeServer.Auditing;
 using HordeServer.Server;
@@ -73,9 +73,9 @@ namespace HordeServer.Agents
 				_document = document;
 			}
 
-			public async Task<IAgent?> TryAddLeaseAsync(AgentLease newLease, CancellationToken cancellationToken = default)
+			public async Task<IAgent?> TryCreateLeaseAsync(CreateLeaseOptions options, CancellationToken cancellationToken = default)
 			{
-				AgentDocument? newDocument = await _collection.TryAddLeaseAsync(_document, newLease, cancellationToken);
+				AgentDocument? newDocument = await _collection.TryCreateLeaseAsync(_document, options, cancellationToken);
 				return _collection.CreateAgentObject(newDocument);
 			}
 
@@ -761,8 +761,11 @@ namespace HordeServer.Agents
 		private static string RedisKeyLeaseChildren(LeaseId parentId) => $"agent/lease-children/{parentId.ToString()}";
 
 		/// <inheritdoc/>
-		async Task<AgentDocument?> TryAddLeaseAsync(AgentDocument agent, AgentLease newLease, CancellationToken cancellationToken)
+		async Task<AgentDocument?> TryCreateLeaseAsync(AgentDocument agent, CreateLeaseOptions options, CancellationToken cancellationToken)
 		{
+			byte[] payloadData = Any.Pack(options.Payload).ToByteArray();
+			AgentLease newLease = new AgentLease(options.Id, options.ParentId, options.Name, options.StreamId, options.PoolId, options.LogId, LeaseState.Pending, options.Resources, options.Exclusive, payloadData);
+
 			List<AgentLease> leases = new List<AgentLease>();
 			if (agent.Leases != null)
 			{
diff --git a/Engine/Source/Programs/Horde/Plugins/Compute/HordeServer.Compute/Agents/AgentService.cs b/Engine/Source/Programs/Horde/Plugins/Compute/HordeServer.Compute/Agents/AgentService.cs
--- a/Engine/Source/Programs/Horde/Plugins/Compute/HordeServer.Compute/Agents/AgentService.cs
+++ b/Engine/Source/Programs/Horde/Plugins/Compute/HordeServer.Compute/Agents/AgentService.cs
@@ -10,6 +10,7 @@ using EpicGames.Horde.Agents.Pools;
 using EpicGames.Horde.Agents.Sessions;
 using EpicGames.Redis;
 using EpicGames.Serialization;
+using Google.Protobuf;
 using Google.Protobuf.WellKnownTypes;
 using HordeCommon.Rpc;
 using HordeCommon.Rpc.Messages;
@@ -511,7 +512,7 @@ namespace HordeServer.Agents
 				cancellationSource.CancelAfter(maxWaitTime);
 
 				// Assign a new lease
-				(ITaskSource, AgentLease)? result = null;
+				(ITaskSource, CreateLeaseOptions)? result = null;
 				try
 				{
 					// Add the cancellation source to the set of waiting agents
@@ -521,12 +522,12 @@ namespace HordeServer.Agents
 					}
 
 					// Create all the tasks to wait for
-					List<Task<(ITaskSource, AgentLease)?>> tasks = new List<Task<(ITaskSource, AgentLease)?>>();
+					List<Task<(ITaskSource, CreateLeaseOptions)?>> tasks = new List<Task<(ITaskSource, CreateLeaseOptions)?>>();
 					foreach (ITaskSource taskSource in _taskSources)
 					{
 						if (CanUseTaskSource(agent, taskSource) && !cancellationSource.IsCancellationRequested)
 						{
-							Task<(ITaskSource, AgentLease)?> task = await GuardedAssignLeaseAsync(taskSource, agent, cancellationSource);
+							Task<(ITaskSource, CreateLeaseOptions)?> task = await GuardedAssignLeaseAsync(taskSource, agent, cancellationSource);
 							tasks.Add(task);
 						}
 					}
@@ -544,19 +545,19 @@ namespace HordeServer.Agents
 					await Task.WhenAll(tasks);
 
 					// Find the first result
-					foreach (Task<(ITaskSource, AgentLease)?> task in tasks)
+					foreach (Task<(ITaskSource, CreateLeaseOptions)?> task in tasks)
 					{
-						(ITaskSource, AgentLease)? taskResult;
+						(ITaskSource, CreateLeaseOptions)? taskResult;
 						if (task.TryGetResult(out taskResult) && taskResult != null)
 						{
-							(ITaskSource taskSource, AgentLease taskLease) = taskResult.Value;
+							(ITaskSource taskSource, CreateLeaseOptions taskLease) = taskResult.Value;
 							if (result == null)
 							{
 								result = (taskSource, taskLease);
 							}
 							else
 							{
-								await taskSource.CancelLeaseAsync(agent, taskLease.Id, Any.Parser.ParseFrom(taskLease.Payload), CancellationToken.None);
+								await taskSource.CancelLeaseAsync(agent, taskLease.Id, Any.Pack(taskLease.Payload), CancellationToken.None);
 							}
 						}
 					}
@@ -580,20 +581,20 @@ namespace HordeServer.Agents
 				}
 
 				// Get the resulting lease
-				(ITaskSource source, AgentLease lease) = result.Value;
+				(ITaskSource source, CreateLeaseOptions lease) = result.Value;
 
 				// Add the new lease to the agent
-				IAgent? newAgent = await agent.TryAddLeaseAsync(lease, cancellationToken);
+				IAgent? newAgent = await agent.TryCreateLeaseAsync(lease, cancellationToken);
 				if (newAgent != null)
 				{
-					await source.OnLeaseStartedAsync(newAgent, lease.Id, Any.Parser.ParseFrom(lease.Payload), Agents.GetLogger(agent.Id), CancellationToken.None);
+					await source.OnLeaseStartedAsync(newAgent, lease.Id, Any.Pack(lease.Payload), Agents.GetLogger(agent.Id), CancellationToken.None);
 					await CreateLeaseAsync(agent, lease, CancellationToken.None);
 					return newAgent;
 				}
 				else
 				{
 					_logger.LogInformation("Failed adding lease {LeaseId} for agent {AgentId}", lease.Id.ToString(), agent.Id.ToString());
-					await source.CancelLeaseAsync(agent, lease.Id, Any.Parser.ParseFrom(lease.Payload), CancellationToken.None);
+					await source.CancelLeaseAsync(agent, lease.Id, Any.Pack(lease.Payload), CancellationToken.None);
 				}
 
 				// Update the agent
@@ -602,32 +603,32 @@ namespace HordeServer.Agents
 			return agent;
 		}
 
-		async Task<Task<(ITaskSource, AgentLease)?>> GuardedAssignLeaseAsync(ITaskSource source, IAgent agent, CancellationTokenSource cancellationSource)
+		async Task<Task<(ITaskSource, CreateLeaseOptions)?>> GuardedAssignLeaseAsync(ITaskSource source, IAgent agent, CancellationTokenSource cancellationSource)
 		{
 			CancellationToken cancellationToken = cancellationSource.Token;
 			try
 			{
-				Task<AgentLease?> task = await source.AssignLeaseAsync(agent, cancellationToken);
+				Task<CreateLeaseOptions?> task = await source.AssignLeaseAsync(agent, cancellationToken);
 				return task.ContinueWith(x => WrapAssignedLease(source, x, cancellationSource), TaskScheduler.Default);
 			}
 			catch (OperationCanceledException) when (cancellationToken.IsCancellationRequested)
 			{
-				return Task.FromResult<(ITaskSource, AgentLease)?>(null);
+				return Task.FromResult<(ITaskSource, CreateLeaseOptions)?>(null);
 			}
 			catch (Exception ex)
 			{
 				_logger.LogError(ex, "Exception while trying to assign lease");
-				return Task.FromResult<(ITaskSource, AgentLease)?>(null);
+				return Task.FromResult<(ITaskSource, CreateLeaseOptions)?>(null);
 			}
 		}
 
-		(ITaskSource, AgentLease)? WrapAssignedLease(ITaskSource source, Task<AgentLease?> task, CancellationTokenSource cancellationSource)
+		(ITaskSource, CreateLeaseOptions)? WrapAssignedLease(ITaskSource source, Task<CreateLeaseOptions?> task, CancellationTokenSource cancellationSource)
 		{
 			if (task.IsCanceled)
 			{
 				return null;
 			}
-			else if (task.TryGetResult(out AgentLease? lease))
+			else if (task.TryGetResult(out CreateLeaseOptions? lease))
 			{
 				if (lease == null)
 				{
@@ -880,18 +881,19 @@ namespace HordeServer.Agents
 		/// Creates a new lease document
 		/// </summary>
 		/// <param name="agent">Agent that will be executing the lease</param>
-		/// <param name="agentLease">The new agent lease</param>
+		/// <param name="options">The new agent lease</param>
 		/// <param name="cancellationToken">Cancellation token for the operation</param>
 		/// <returns>New lease document</returns>
-		internal Task<ILease> CreateLeaseAsync(IAgent agent, AgentLease agentLease, CancellationToken cancellationToken = default)
+		internal Task<ILease> CreateLeaseAsync(IAgent agent, CreateLeaseOptions options, CancellationToken cancellationToken = default)
 		{
 			try
 			{
-				return _leases.AddAsync(agentLease.Id, agentLease.ParentId, agentLease.Name, agent.Id, agent.SessionId!.Value, agentLease.StreamId, agentLease.PoolId, agentLease.LogId, agentLease.StartTime, agentLease.Payload!, cancellationToken);
+				DateTime startTime = _clock.UtcNow;
+				return _leases.AddAsync(options.Id, options.ParentId, options.Name, agent.Id, agent.SessionId!.Value, options.StreamId, options.PoolId, options.LogId, startTime, Any.Pack(options.Payload).ToByteArray(), cancellationToken);
 			}
 			catch (Exception ex)
 			{
-				_logger.LogError(ex, "Unable to create lease {LeaseId} for agent {AgentId}; lease already exists?", agentLease.Id, agent.Id);
+				_logger.LogError(ex, "Unable to create lease {LeaseId} for agent {AgentId}; lease already exists?", options.Id, agent.Id);
 				throw;
 			}
 		}
diff --git a/Engine/Source/Programs/Horde/Plugins/Compute/HordeServer.Compute/Agents/IAgent.cs b/Engine/Source/Programs/Horde/Plugins/Compute/HordeServer.Compute/Agents/IAgent.cs
--- a/Engine/Source/Programs/Horde/Plugins/Compute/HordeServer.Compute/Agents/IAgent.cs
+++ b/Engine/Source/Programs/Horde/Plugins/Compute/HordeServer.Compute/Agents/IAgent.cs
@@ -12,6 +12,8 @@ using EpicGames.Horde.Compute;
 using EpicGames.Horde.Logs;
 using EpicGames.Horde.Streams;
 using EpicGames.Horde.Tools;
+using Google.Protobuf;
+using Google.Protobuf.WellKnownTypes;
 using HordeCommon.Rpc;
 using HordeCommon.Rpc.Messages;
 using MongoDB.Bson;
@@ -239,10 +241,10 @@ namespace HordeServer.Agents
 		/// <summary>
 		/// Attempts to add a lease to an agent
 		/// </summary>
-		/// <param name="newLease">The new lease document</param>
+		/// <param name="options">The new lease document</param>
 		/// <param name="cancellationToken">Cancellation token for the operation</param>
 		/// <returns>New agent state if it succeeded, otherwise null</returns>
-		Task<IAgent?> TryAddLeaseAsync(AgentLease newLease, CancellationToken cancellationToken = default);
+		Task<IAgent?> TryCreateLeaseAsync(CreateLeaseOptions options, CancellationToken cancellationToken = default);
 
 		/// <summary>
 		/// Attempts to cancel a lease
@@ -253,6 +255,11 @@ namespace HordeServer.Agents
 		Task<IAgent?> TryCancelLeaseAsync(int leaseIdx, CancellationToken cancellationToken = default);
 	}
 
+	/// <summary>
+	/// Information for a new lease
+	/// </summary>
+	public record class CreateLeaseOptions(LeaseId Id, LeaseId? ParentId, string Name, StreamId? StreamId, PoolId? PoolId, LogId? LogId, IReadOnlyDictionary<string, int>? Resources, bool Exclusive, IMessage Payload);
+
 	/// <summary>
 	/// Options for updating an agent
 	/// </summary>
@@ -905,6 +912,16 @@ namespace HordeServer.Agents
 			StartTime = DateTime.UtcNow;
 		}
 
+		/// <summary>
+		/// Create a lease from 
+		/// </summary>
+		/// <param name="options"></param>
+		/// <returns></returns>
+		public AgentLease(CreateLeaseOptions options)
+			: this(options.Id, options.ParentId, options.Name, options.StreamId, options.PoolId, options.LogId, LeaseState.Pending, options.Resources, options.Exclusive, Any.Pack(options.Payload).ToByteArray())
+		{
+		}
+
 		/// <summary>
 		/// Converts this lease to an RPC message
 		/// </summary>
diff --git a/Engine/Source/Programs/Horde/Plugins/Compute/HordeServer.Compute/Compute/ComputeService.cs b/Engine/Source/Programs/Horde/Plugins/Compute/HordeServer.Compute/Compute/ComputeService.cs
--- a/Engine/Source/Programs/Horde/Plugins/Compute/HordeServer.Compute/Compute/ComputeService.cs
+++ b/Engine/Source/Programs/Horde/Plugins/Compute/HordeServer.Compute/Compute/ComputeService.cs
@@ -16,7 +16,6 @@ using EpicGames.Horde.Jobs;
 using EpicGames.Horde.Logs;
 using EpicGames.Horde.Users;
 using Google.Protobuf;
-using Google.Protobuf.WellKnownTypes;
 using Horde.Common.Rpc;
 using HordeCommon.Rpc.Tasks;
 using HordeServer.Agents;
@@ -423,8 +422,7 @@ namespace HordeServer.Compute
 
 						ComputeTask computeTask = CreateComputeTask(assignedResources, log?.Id, arp.Encryption, certificate, arp.ParentLeaseId, protocol);
 
-						byte[] payload = Any.Pack(computeTask).ToByteArray();
-						AgentLease lease = new AgentLease(leaseId, arp.ParentLeaseId, "Compute task", null, null, log?.Id, LeaseState.Pending, assignedResources, arp.Requirements.Exclusive, payload);
+						CreateLeaseOptions createLeaseOptions = new CreateLeaseOptions(leaseId, arp.ParentLeaseId, "Compute task", null, null, log?.Id, assignedResources, arp.Requirements.Exclusive, computeTask);
 
 						using TelemetrySpan assignSpan = _tracer.StartActiveSpan("TryAssignAsync");
 
@@ -433,11 +431,11 @@ namespace HordeServer.Compute
 						{
 							using TelemetrySpan addLeaseSpan = _tracer.StartActiveSpan("Adding lease");
 
-							IAgent? newAgent = await agent.TryAddLeaseAsync(lease, cancellationToken);
+							IAgent? newAgent = await agent.TryCreateLeaseAsync(createLeaseOptions, cancellationToken);
 							if (newAgent != null)
 							{
 								await _agentCollection.PublishUpdateEventAsync(agent.Id);
-								await _agentService.CreateLeaseAsync(newAgent, lease, cancellationToken);
+								await _agentService.CreateLeaseAsync(newAgent, createLeaseOptions, cancellationToken);
 								span.SetAttribute("allocatedLeaseId", leaseId.ToString());
 								span.SetAttribute("allocatedAgentId", newAgent.Id.ToString());
 								span.SetAttribute("allocatedConnectionAddress", resource.ConnectionAddress);
diff --git a/Engine/Source/Programs/Horde/Plugins/Compute/HordeServer.Compute/Compute/ComputeTaskSource.cs b/Engine/Source/Programs/Horde/Plugins/Compute/HordeServer.Compute/Compute/ComputeTaskSource.cs
--- a/Engine/Source/Programs/Horde/Plugins/Compute/HordeServer.Compute/Compute/ComputeTaskSource.cs
+++ b/Engine/Source/Programs/Horde/Plugins/Compute/HordeServer.Compute/Compute/ComputeTaskSource.cs
@@ -142,7 +142,7 @@ namespace HordeServer.Compute
 			public IAgent Agent { get; }
 			public IPAddress Ip { get; }
 			public int Port { get; }
-			public TaskCompletionSource<AgentLease?> Lease { get; } = new TaskCompletionSource<AgentLease?>(TaskCreationOptions.RunContinuationsAsynchronously);
+			public TaskCompletionSource<CreateLeaseOptions?> Lease { get; } = new TaskCompletionSource<CreateLeaseOptions?>(TaskCreationOptions.RunContinuationsAsynchronously);
 
 			public Waiter(IAgent agent, IPAddress ip, int port)
 			{
@@ -178,7 +178,7 @@ namespace HordeServer.Compute
 		}
 
 		/// <inheritdoc/>
-		public override Task<Task<AgentLease?>> AssignLeaseAsync(IAgent agent, CancellationToken cancellationToken)
+		public override Task<Task<CreateLeaseOptions?>> AssignLeaseAsync(IAgent agent, CancellationToken cancellationToken)
 		{
 			return Task.FromResult(WaitInternalAsync(agent, cancellationToken));
 		}
@@ -192,7 +192,7 @@ namespace HordeServer.Compute
 			await _agentRelay.RemovePortMappingAsync(leaseId);
 		}
 
-		async Task<AgentLease?> WaitInternalAsync(IAgent agent, CancellationToken cancellationToken)
+		async Task<CreateLeaseOptions?> WaitInternalAsync(IAgent agent, CancellationToken cancellationToken)
 		{
 			string? ipStr = agent.GetPropertyValues("ComputeIp").FirstOrDefault();
 			if (ipStr == null || !IPAddress.TryParse(ipStr, out IPAddress? ip))
@@ -233,7 +233,7 @@ namespace HordeServer.Compute
 				if (waiter != null)
 				{
 					using IDisposable disposable = cancellationToken.Register(() => waiter.Lease.TrySetResult(null));
-					AgentLease? lease = await waiter.Lease.Task;
+					CreateLeaseOptions? lease = await waiter.Lease.Task;
 					if (lease != null)
 					{
 						_logger.LogInformation("Created compute lease for agent {AgentId}", agent.Id);
diff --git a/Engine/Source/Programs/Horde/Plugins/Compute/HordeServer.Compute/Protos/horde/messages/agent_capabilities.proto.cs b/Engine/Source/Programs/Horde/Plugins/Compute/HordeServer.Compute/Protos/horde/messages/agent_capabilities.proto.cs
--- a/Engine/Source/Programs/Horde/Plugins/Compute/HordeServer.Compute/Protos/horde/messages/agent_capabilities.proto.cs
+++ b/Engine/Source/Programs/Horde/Plugins/Compute/HordeServer.Compute/Protos/horde/messages/agent_capabilities.proto.cs
@@ -1,10 +1,5 @@
 // Copyright Epic Games, Inc. All Rights Reserved.
 
-using System;
-using System.Collections.Generic;
-using System.Linq;
-using System.Text;
-using System.Threading.Tasks;
 using EpicGames.Horde.Agents;
 
 #pragma warning disable CS1591
diff --git a/Engine/Source/Programs/Horde/Plugins/Compute/HordeServer.Compute/Tasks/ITaskSource.cs b/Engine/Source/Programs/Horde/Plugins/Compute/HordeServer.Compute/Tasks/ITaskSource.cs
--- a/Engine/Source/Programs/Horde/Plugins/Compute/HordeServer.Compute/Tasks/ITaskSource.cs
+++ b/Engine/Source/Programs/Horde/Plugins/Compute/HordeServer.Compute/Tasks/ITaskSource.cs
@@ -77,7 +77,7 @@ namespace HordeServer.Tasks
 		/// <param name="agent">The agent to assign a lease to</param>
 		/// <param name="cancellationToken">Cancellation token for the wait</param>
 		/// <returns>Task returning a new lease object.</returns>
-		Task<Task<AgentLease?>> AssignLeaseAsync(IAgent agent, CancellationToken cancellationToken);
+		Task<Task<CreateLeaseOptions?>> AssignLeaseAsync(IAgent agent, CancellationToken cancellationToken);
 
 		/// <summary>
 		/// Cancel a lease that was previously assigned to an agent, allowing it to be assigned out again
@@ -174,7 +174,7 @@ namespace HordeServer.Tasks
 		public MessageDescriptor Descriptor => s_message.Descriptor;
 
 		/// <inheritdoc/>
-		public abstract Task<Task<AgentLease?>> AssignLeaseAsync(IAgent agent, CancellationToken cancellationToken);
+		public abstract Task<Task<CreateLeaseOptions?>> AssignLeaseAsync(IAgent agent, CancellationToken cancellationToken);
 
 		/// <inheritdoc/>
 		public Task CancelLeaseAsync(IAgent agent, LeaseId leaseId, Any payload, CancellationToken cancellationToken) => CancelLeaseAsync(agent, leaseId, payload.Unpack<TMessage>(), cancellationToken);
@@ -233,10 +233,10 @@ namespace HordeServer.Tasks
 		/// </summary>
 		/// <param name="token">The cancellation token</param>
 		/// <returns>Lease task</returns>
-		protected static Task<AgentLease?> SkipAsync(CancellationToken token)
+		protected static Task<CreateLeaseOptions?> SkipAsync(CancellationToken token)
 		{
 			_ = token;
-			return Task.FromResult<AgentLease?>(null);
+			return Task.FromResult<CreateLeaseOptions?>(null);
 		}
 
 		/// <summary>
@@ -244,10 +244,10 @@ namespace HordeServer.Tasks
 		/// </summary>
 		/// <param name="token">The cancellation token</param>
 		/// <returns>Lease task</returns>
-		protected static async Task<Task<AgentLease?>> DrainAsync(CancellationToken token)
+		protected static async Task<Task<CreateLeaseOptions?>> DrainAsync(CancellationToken token)
 		{
 			await token.AsTask();
-			return Task.FromResult<AgentLease?>(null);
+			return Task.FromResult<CreateLeaseOptions?>(null);
 		}
 
 		/// <summary>
@@ -255,9 +255,9 @@ namespace HordeServer.Tasks
 		/// </summary>
 		/// <param name="lease">Lease to create the task from</param>
 		/// <returns></returns>
-		protected static Task<AgentLease?> LeaseAsync(AgentLease lease)
+		protected static Task<CreateLeaseOptions?> LeaseAsync(CreateLeaseOptions lease)
 		{
-			return Task.FromResult<AgentLease?>(lease);
+			return Task.FromResult<CreateLeaseOptions?>(lease);
 		}
 	}
 }
diff --git a/Engine/Source/Programs/Horde/Plugins/Compute/HordeServer.Compute/Tasks/RestartTaskSource.cs b/Engine/Source/Programs/Horde/Plugins/Compute/HordeServer.Compute/Tasks/RestartTaskSource.cs
--- a/Engine/Source/Programs/Horde/Plugins/Compute/HordeServer.Compute/Tasks/RestartTaskSource.cs
+++ b/Engine/Source/Programs/Horde/Plugins/Compute/HordeServer.Compute/Tasks/RestartTaskSource.cs
@@ -3,8 +3,6 @@
 using EpicGames.Horde.Agents.Leases;
 using EpicGames.Horde.Jobs;
 using EpicGames.Horde.Logs;
-using Google.Protobuf;
-using Google.Protobuf.WellKnownTypes;
 using HordeCommon.Rpc.Tasks;
 using HordeServer.Agents;
 using HordeServer.Logs;
@@ -27,7 +25,7 @@ namespace HordeServer.Tasks
 			OnLeaseStartedProperties.Add(nameof(RestartTask.LogId), x => LogId.Parse(x.LogId));
 		}
 
-		public override async Task<Task<AgentLease?>> AssignLeaseAsync(IAgent agent, CancellationToken cancellationToken)
+		public override async Task<Task<CreateLeaseOptions?>> AssignLeaseAsync(IAgent agent, CancellationToken cancellationToken)
 		{
 			if (!agent.RequestForceRestart)
 			{
@@ -47,9 +45,7 @@ namespace HordeServer.Tasks
 			RestartTask task = new RestartTask();
 			task.LogId = log.Id.ToString();
 
-			byte[] payload = Any.Pack(task).ToByteArray();
-
-			return LeaseAsync(new AgentLease(leaseId, null, "Restart", null, null, log.Id, LeaseState.Pending, null, true, payload));
+			return LeaseAsync(new CreateLeaseOptions(leaseId, null, "Restart", null, null, log.Id, null, true, task));
 		}
 	}
 }
diff --git a/Engine/Source/Programs/Horde/Plugins/Compute/HordeServer.Compute/Tasks/ShutdownTaskSource.cs b/Engine/Source/Programs/Horde/Plugins/Compute/HordeServer.Compute/Tasks/ShutdownTaskSource.cs
--- a/Engine/Source/Programs/Horde/Plugins/Compute/HordeServer.Compute/Tasks/ShutdownTaskSource.cs
+++ b/Engine/Source/Programs/Horde/Plugins/Compute/HordeServer.Compute/Tasks/ShutdownTaskSource.cs
@@ -3,8 +3,6 @@
 using EpicGames.Horde.Agents.Leases;
 using EpicGames.Horde.Jobs;
 using EpicGames.Horde.Logs;
-using Google.Protobuf;
-using Google.Protobuf.WellKnownTypes;
 using HordeCommon.Rpc.Tasks;
 using HordeServer.Agents;
 using HordeServer.Logs;
@@ -26,7 +24,7 @@ namespace HordeServer.Tasks
 			OnLeaseStartedProperties.Add(nameof(ShutdownTask.LogId), x => LogId.Parse(x.LogId));
 		}
 
-		public override async Task<Task<AgentLease?>> AssignLeaseAsync(IAgent agent, CancellationToken cancellationToken)
+		public override async Task<Task<CreateLeaseOptions?>> AssignLeaseAsync(IAgent agent, CancellationToken cancellationToken)
 		{
 			if (!agent.RequestShutdown)
 			{
@@ -43,9 +41,7 @@ namespace HordeServer.Tasks
 			ShutdownTask task = new ShutdownTask();
 			task.LogId = log.Id.ToString();
 
-			byte[] payload = Any.Pack(task).ToByteArray();
-
-			return LeaseAsync(new AgentLease(leaseId, null, "Shutdown", null, null, log.Id, LeaseState.Pending, null, true, payload));
+			return LeaseAsync(new CreateLeaseOptions(leaseId, null, "Shutdown", null, null, log.Id, null, true, task));
 		}
 	}
 }
diff --git a/Engine/Source/Programs/Horde/Plugins/Compute/HordeServer.Compute/Tasks/UpgradeTaskSource.cs b/Engine/Source/Programs/Horde/Plugins/Compute/HordeServer.Compute/Tasks/UpgradeTaskSource.cs
--- a/Engine/Source/Programs/Horde/Plugins/Compute/HordeServer.Compute/Tasks/UpgradeTaskSource.cs
+++ b/Engine/Source/Programs/Horde/Plugins/Compute/HordeServer.Compute/Tasks/UpgradeTaskSource.cs
@@ -6,8 +6,6 @@ using EpicGames.Horde.Agents.Leases;
 using EpicGames.Horde.Jobs;
 using EpicGames.Horde.Logs;
 using EpicGames.Horde.Tools;
-using Google.Protobuf;
-using Google.Protobuf.WellKnownTypes;
 using HordeCommon.Rpc.Tasks;
 using HordeServer.Agents;
 using HordeServer.Logs;
@@ -40,7 +38,7 @@ namespace HordeServer.Tasks
 			OnLeaseStartedProperties.Add(nameof(UpgradeTask.LogId), x => LogId.Parse(x.LogId));
 		}
 
-		public override async Task<Task<AgentLease?>> AssignLeaseAsync(IAgent agent, CancellationToken cancellationToken)
+		public override async Task<Task<CreateLeaseOptions?>> AssignLeaseAsync(IAgent agent, CancellationToken cancellationToken)
 		{
 			if (!_staticComputeConfig.Value.EnableUpgradeTasks)
 			{
@@ -71,8 +69,7 @@ namespace HordeServer.Tasks
 			task.SoftwareId = $"{tool.Id}:{deployment.Version}";
 			task.LogId = log.Id.ToString();
 
-			byte[] payload = Any.Pack(task).ToByteArray();
-			return LeaseAsync(new AgentLease(leaseId, null, $"Upgrade to {tool.Id} {deployment.Version}", null, null, log.Id, LeaseState.Pending, null, true, payload));
+			return LeaseAsync(new CreateLeaseOptions(leaseId, null, $"Upgrade to {tool.Id} {deployment.Version}", null, null, log.Id, null, true, task));
 		}
 
 		/// <summary>
