Horde: Remove properties from AgentLease.

#rnx

diff --git a/Engine/Source/Programs/Horde/Plugins/Compute/HordeServer.Compute/Agents/AgentCollection.cs b/Engine/Source/Programs/Horde/Plugins/Compute/HordeServer.Compute/Agents/AgentCollection.cs
--- a/Engine/Source/Programs/Horde/Plugins/Compute/HordeServer.Compute/Agents/AgentCollection.cs
+++ b/Engine/Source/Programs/Horde/Plugins/Compute/HordeServer.Compute/Agents/AgentCollection.cs
@@ -9,7 +9,6 @@ using EpicGames.Horde.Agents.Leases;
 using EpicGames.Horde.Agents.Pools;
 using EpicGames.Horde.Agents.Sessions;
 using EpicGames.Redis;
-using Google.Protobuf;
 using Google.Protobuf.WellKnownTypes;
 using HordeCommon.Rpc.Tasks;
 using HordeServer.Auditing;
@@ -763,8 +762,7 @@ namespace HordeServer.Agents
 		/// <inheritdoc/>
 		async Task<AgentDocument?> TryCreateLeaseAsync(AgentDocument agent, CreateLeaseOptions options, CancellationToken cancellationToken)
 		{
-			byte[] payloadData = Any.Pack(options.Payload).ToByteArray();
-			AgentLease newLease = new AgentLease(options.Id, options.ParentId, options.Name, options.StreamId, options.PoolId, options.LogId, LeaseState.Pending, options.Resources, options.Exclusive, payloadData);
+			AgentLease newLease = new AgentLease(options);
 
 			List<AgentLease> leases = new List<AgentLease>();
 			if (agent.Leases != null)
diff --git a/Engine/Source/Programs/Horde/Plugins/Compute/HordeServer.Compute/Agents/AgentService.cs b/Engine/Source/Programs/Horde/Plugins/Compute/HordeServer.Compute/Agents/AgentService.cs
--- a/Engine/Source/Programs/Horde/Plugins/Compute/HordeServer.Compute/Agents/AgentService.cs
+++ b/Engine/Source/Programs/Horde/Plugins/Compute/HordeServer.Compute/Agents/AgentService.cs
@@ -977,10 +977,6 @@ namespace HordeServer.Agents
 			{
 				finishTime = agent.SessionExpiresAt.Value;
 			}
-			if (lease.ExpiryTime.HasValue && lease.ExpiryTime.Value < finishTime)
-			{
-				finishTime = lease.ExpiryTime.Value;
-			}
 
 			// Update the lease
 			await _leases.TrySetOutcomeAsync(lease.Id, finishTime, outcome, output, cancellationToken);
diff --git a/Engine/Source/Programs/Horde/Plugins/Compute/HordeServer.Compute/Agents/AgentsController.cs b/Engine/Source/Programs/Horde/Plugins/Compute/HordeServer.Compute/Agents/AgentsController.cs
--- a/Engine/Source/Programs/Horde/Plugins/Compute/HordeServer.Compute/Agents/AgentsController.cs
+++ b/Engine/Source/Programs/Horde/Plugins/Compute/HordeServer.Compute/Agents/AgentsController.cs
@@ -32,6 +32,7 @@ namespace HordeServer.Agents
 		readonly AgentService _agentService;
 		readonly IAgentTelemetryCollection _agentTelemetryCollection;
 		readonly IUserCollection _userCollection;
+		readonly ILeaseCollection _leaseCollection;
 		readonly IOptionsSnapshot<ComputeConfig> _computeConfig;
 		readonly Tracer _tracer;
 		readonly ILogger<AgentsController> _logger;
@@ -39,11 +40,12 @@ namespace HordeServer.Agents
 		/// <summary>
 		/// Constructor
 		/// </summary>
-		public AgentsController(AgentService agentService, IAgentTelemetryCollection agentTelemetryCollection, IUserCollection userCollection, IOptionsSnapshot<ComputeConfig> computeConfig, Tracer tracer, ILogger<AgentsController> logger)
+		public AgentsController(AgentService agentService, IAgentTelemetryCollection agentTelemetryCollection, IUserCollection userCollection, ILeaseCollection leaseCollection, IOptionsSnapshot<ComputeConfig> computeConfig, Tracer tracer, ILogger<AgentsController> logger)
 		{
 			_agentService = agentService;
 			_agentTelemetryCollection = agentTelemetryCollection;
 			_userCollection = userCollection;
+			_leaseCollection = leaseCollection;
 			_computeConfig = computeConfig;
 			_tracer = tracer;
 			_logger = logger;
@@ -139,27 +141,26 @@ namespace HordeServer.Agents
 			}
 
 			List<GetAgentLeaseResponse> leases = new List<GetAgentLeaseResponse>();
-			foreach (AgentLease lease in agent.Leases)
+			foreach (AgentLease agentLease in agent.Leases)
 			{
 				try
 				{
-					Dictionary<string, string>? details = await _agentService.GetPayloadDetailsAsync(lease.Payload, cancellationToken);
-					leases.Add(CreateGetAgentLeaseResponse(lease, details));
+					ILease? lease = await _leaseCollection.GetAsync(agentLease.Id, cancellationToken);
+					if (lease != null)
+					{
+						Dictionary<string, string>? details = await _agentService.GetPayloadDetailsAsync(lease.Payload, cancellationToken);
+						leases.Add(CreateGetAgentLeaseResponse(lease, details, null));
+					}
 				}
-				catch (Exception e)
+				catch (Exception ex)
 				{
-					_logger.LogError(e, "Failed getting payload details for agent lease {LeaseId}", lease.Id.ToString());
+					_logger.LogError(ex, "Failed getting payload details for agent lease {LeaseId}", agentLease.Id);
 				}
 			}
 
 			return CreateGetAgentResponse(agent, leases, rate).ApplyFilter(filter);
 		}
 
-		internal static GetAgentLeaseResponse CreateGetAgentLeaseResponse(AgentLease lease, Dictionary<string, string>? details)
-		{
-			return new GetAgentLeaseResponse(lease.Id, lease.ParentId, null, null, lease.Name, lease.LogId, lease.StartTime, lease.ExpiryTime, lease.Active, details, null, lease.State);
-		}
-
 		internal static GetAgentLeaseResponse CreateGetAgentLeaseResponse(ILease lease, Dictionary<string, string>? details, double? agentRate)
 		{
 			return new GetAgentLeaseResponse(lease.Id, lease.ParentId, lease.AgentId, agentRate, lease.Name, lease.LogId, lease.StartTime, lease.FinishTime, lease.FinishTime == null, details, lease.Outcome, null);
diff --git a/Engine/Source/Programs/Horde/Plugins/Compute/HordeServer.Compute/Agents/IAgent.cs b/Engine/Source/Programs/Horde/Plugins/Compute/HordeServer.Compute/Agents/IAgent.cs
--- a/Engine/Source/Programs/Horde/Plugins/Compute/HordeServer.Compute/Agents/IAgent.cs
+++ b/Engine/Source/Programs/Horde/Plugins/Compute/HordeServer.Compute/Agents/IAgent.cs
@@ -819,47 +819,11 @@ namespace HordeServer.Agents
 		/// </summary>
 		public LeaseId? ParentId { get; set; }
 
-		/// <summary>
-		/// Name of this lease
-		/// </summary>
-		public string Name { get; set; }
-
 		/// <summary>
 		/// The current state of the lease
 		/// </summary>
 		public LeaseState State { get; set; }
 
-		/// <summary>
-		/// The stream for this lease
-		/// </summary>
-		public StreamId? StreamId { get; set; }
-
-		/// <summary>
-		/// The pool for this lease
-		/// </summary>
-		public PoolId? PoolId { get; set; }
-
-		/// <summary>
-		/// Optional log for this lease
-		/// </summary>
-		public LogId? LogId { get; set; }
-
-		/// <summary>
-		/// Time at which the lease started
-		/// </summary>
-		[BsonRequired]
-		public DateTime StartTime { get; set; }
-
-		/// <summary>
-		/// Time at which the lease should be terminated
-		/// </summary>
-		public DateTime? ExpiryTime { get; set; }
-
-		/// <summary>
-		/// Flag indicating whether this lease has been accepted by the agent
-		/// </summary>
-		public bool Active { get; set; }
-
 		/// <summary>
 		/// Resources used by this lease
 		/// </summary>
@@ -881,7 +845,6 @@ namespace HordeServer.Agents
 		[BsonConstructor]
 		private AgentLease()
 		{
-			Name = String.Empty;
 		}
 
 		/// <summary>
@@ -889,27 +852,18 @@ namespace HordeServer.Agents
 		/// </summary>
 		/// <param name="id">Identifier for the lease</param>
 		/// <param name="parentId">The parent lease id</param>
-		/// <param name="name">Name of this lease</param>
-		/// <param name="streamId"></param>
-		/// <param name="poolId"></param>
-		/// <param name="logId">Unique id for the log</param>
 		/// <param name="state">State for the lease</param>
 		/// <param name="resources">Resources required for this lease</param>
 		/// <param name="exclusive">Whether to reserve the entire device</param>
 		/// <param name="payload">Encoded "any" protobuf describing the contents of the payload</param>
-		public AgentLease(LeaseId id, LeaseId? parentId, string name, StreamId? streamId, PoolId? poolId, LogId? logId, LeaseState state, IReadOnlyDictionary<string, int>? resources, bool exclusive, byte[]? payload)
+		public AgentLease(LeaseId id, LeaseId? parentId, LeaseState state, IReadOnlyDictionary<string, int>? resources, bool exclusive, byte[]? payload)
 		{
 			Id = id;
 			ParentId = parentId;
-			Name = name;
-			StreamId = streamId;
-			PoolId = poolId;
-			LogId = logId;
 			State = state;
 			Resources = resources;
 			Exclusive = exclusive;
 			Payload = payload;
-			StartTime = DateTime.UtcNow;
 		}
 
 		/// <summary>
@@ -918,7 +872,7 @@ namespace HordeServer.Agents
 		/// <param name="options"></param>
 		/// <returns></returns>
 		public AgentLease(CreateLeaseOptions options)
-			: this(options.Id, options.ParentId, options.Name, options.StreamId, options.PoolId, options.LogId, LeaseState.Pending, options.Resources, options.Exclusive, Any.Pack(options.Payload).ToByteArray())
+			: this(options.Id, options.ParentId, LeaseState.Pending, options.Resources, options.Exclusive, Any.Pack(options.Payload).ToByteArray())
 		{
 		}
 
