Horde: Add an interface for agent leases.

#rnx

diff --git a/Engine/Source/Programs/Horde/Plugins/Build/HordeServer.Build/Jobs/JobService.cs b/Engine/Source/Programs/Horde/Plugins/Build/HordeServer.Build/Jobs/JobService.cs
--- a/Engine/Source/Programs/Horde/Plugins/Build/HordeServer.Build/Jobs/JobService.cs
+++ b/Engine/Source/Programs/Horde/Plugins/Build/HordeServer.Build/Jobs/JobService.cs
@@ -477,7 +477,7 @@ namespace HordeServer.Jobs
 					break;
 				}
 
-				AgentLease lease = agent.Leases[index];
+				IAgentLease lease = agent.Leases[index];
 				if (lease.State != LeaseState.Active)
 				{
 					break;
diff --git a/Engine/Source/Programs/Horde/Plugins/Compute/HordeServer.Compute/Agents/AgentCollection.cs b/Engine/Source/Programs/Horde/Plugins/Compute/HordeServer.Compute/Agents/AgentCollection.cs
--- a/Engine/Source/Programs/Horde/Plugins/Compute/HordeServer.Compute/Agents/AgentCollection.cs
+++ b/Engine/Source/Programs/Horde/Plugins/Compute/HordeServer.Compute/Agents/AgentCollection.cs
@@ -61,7 +61,7 @@ namespace HordeServer.Agents
 			DateTime IAgent.LastConformTime => _document.LastConformTime;
 			int? IAgent.ConformAttemptCount => _document.ConformAttemptCount;
 			static readonly IReadOnlyList<AgentLease> s_emptyLeases = new List<AgentLease>();
-			IReadOnlyList<AgentLease> IAgent.Leases => _document.Leases ?? s_emptyLeases;
+			IReadOnlyList<IAgentLease> IAgent.Leases => _document.Leases ?? s_emptyLeases;
 			string IAgent.EnrollmentKey => _document.EnrollmentKey;
 			DateTime IAgent.UpdateTime => _document.UpdateTime;
 			uint IAgent.UpdateIndex => _document.UpdateIndex;
diff --git a/Engine/Source/Programs/Horde/Plugins/Compute/HordeServer.Compute/Agents/AgentService.cs b/Engine/Source/Programs/Horde/Plugins/Compute/HordeServer.Compute/Agents/AgentService.cs
--- a/Engine/Source/Programs/Horde/Plugins/Compute/HordeServer.Compute/Agents/AgentService.cs
+++ b/Engine/Source/Programs/Horde/Plugins/Compute/HordeServer.Compute/Agents/AgentService.cs
@@ -751,7 +751,7 @@ namespace HordeServer.Agents
 
 				// Flag for whether the leases array should be updated
 				bool updateLeases = false;
-				List<AgentLease> leases = new List<AgentLease>(agent.Leases);
+				List<AgentLease> leases = agent.Leases.ConvertAll(x => new AgentLease(x));
 
 				// Remove any completed leases from the agent
 				Dictionary<LeaseId, RpcLease> leaseIdToNewState = newLeases.ToDictionary(x => x.Id, x => x);
@@ -854,7 +854,7 @@ namespace HordeServer.Agents
 
 			// Save off the session id and current leases
 			SessionId sessionId = agent.SessionId.Value;
-			List<AgentLease> leases = new List<AgentLease>(agent.Leases);
+			IReadOnlyList<IAgentLease> leases = agent.Leases;
 
 			// Clear the current session
 			IAgent? newAgent = await agent.TryTerminateSessionAsync(cancellationToken);
@@ -863,7 +863,7 @@ namespace HordeServer.Agents
 				agent = newAgent;
 
 				// Remove any outstanding leases
-				foreach (AgentLease lease in leases)
+				foreach (IAgentLease lease in leases)
 				{
 					Agents.GetLogger(agent.Id).LogInformation("Removing lease {LeaseId} during session terminate...", lease.Id);
 					await RemoveLeaseAsync(agent, lease, finishTime, LeaseOutcome.Failed, null, cancellationToken);
@@ -949,7 +949,7 @@ namespace HordeServer.Agents
 		/// <param name="output">Output from executing the task</param>
 		/// <param name="cancellationToken">Cancellation token for the operation</param>
 		/// <returns>Async task</returns>
-		private async Task RemoveLeaseAsync(IAgent agent, AgentLease lease, DateTime utcNow, LeaseOutcome outcome, byte[]? output, CancellationToken cancellationToken = default)
+		private async Task RemoveLeaseAsync(IAgent agent, IAgentLease lease, DateTime utcNow, LeaseOutcome outcome, byte[]? output, CancellationToken cancellationToken = default)
 		{
 			// Make sure the lease is terminated correctly
 			if (lease.Payload == null)
@@ -958,7 +958,7 @@ namespace HordeServer.Agents
 			}
 			else
 			{
-				Any any = Any.Parser.ParseFrom(lease.Payload);
+				Any any = lease.Payload;
 				_logger.LogInformation("Removing lease {LeaseId} ({LeaseType})", lease.Id, any.TypeUrl);
 
 				foreach (ITaskSource taskSource in _taskSources)
diff --git a/Engine/Source/Programs/Horde/Plugins/Compute/HordeServer.Compute/Agents/IAgent.cs b/Engine/Source/Programs/Horde/Plugins/Compute/HordeServer.Compute/Agents/IAgent.cs
--- a/Engine/Source/Programs/Horde/Plugins/Compute/HordeServer.Compute/Agents/IAgent.cs
+++ b/Engine/Source/Programs/Horde/Plugins/Compute/HordeServer.Compute/Agents/IAgent.cs
@@ -166,7 +166,7 @@ namespace HordeServer.Agents
 		/// <summary>
 		/// Array of active leases.
 		/// </summary>
-		public IReadOnlyList<AgentLease> Leases { get; }
+		public IReadOnlyList<IAgentLease> Leases { get; }
 
 		/// <summary>
 		/// Key used to validate that a particular enrollment is still valid for this agent
@@ -802,43 +802,68 @@ namespace HordeServer.Agents
 			return Enumerable.SequenceEqual(lhs.View, rhs.View);
 		}
 	}
-
 	/// <summary>
 	/// Document describing an active lease
 	/// </summary>
-	public class AgentLease
+	public interface IAgentLease
 	{
 		/// <summary>
 		/// Name of this lease
 		/// </summary>
-		[BsonRequired]
-		public LeaseId Id { get; set; }
+		LeaseId Id { get; }
 
 		/// <summary>
 		/// The parent lease id
 		/// </summary>
-		public LeaseId? ParentId { get; set; }
+		LeaseId? ParentId { get; }
 
 		/// <summary>
 		/// The current state of the lease
 		/// </summary>
-		public LeaseState State { get; set; }
+		LeaseState State { get; }
 
 		/// <summary>
 		/// Resources used by this lease
 		/// </summary>
-		public IReadOnlyDictionary<string, int>? Resources { get; set; }
+		IReadOnlyDictionary<string, int>? Resources { get; }
 
 		/// <summary>
 		/// Whether the lease requires exclusive access to the agent
 		/// </summary>
-		public bool Exclusive { get; set; }
+		bool Exclusive { get; }
 
 		/// <summary>
 		/// For leases in the pending state, encodes an "any" protobuf containing the payload for the agent to execute the lease.
 		/// </summary>
+		Any? Payload { get; }
+	}
+
+	/// <summary>
+	/// Document describing an active lease
+	/// </summary>
+	public class AgentLease : IAgentLease
+	{
+		/// <inheritdoc/>
+		[BsonRequired]
+		public LeaseId Id { get; set; }
+
+		/// <inheritdoc/>
+		public LeaseId? ParentId { get; set; }
+
+		/// <inheritdoc/>
+		public LeaseState State { get; set; }
+
+		/// <inheritdoc/>
+		public IReadOnlyDictionary<string, int>? Resources { get; set; }
+
+		/// <inheritdoc/>
+		public bool Exclusive { get; set; }
+
+		/// <inheritdoc/>
 		public byte[]? Payload { get; set; }
 
+		Any? IAgentLease.Payload => (Payload != null) ? Any.Parser.ParseFrom(Payload) : null;
+
 		/// <summary>
 		/// Private constructor
 		/// </summary>
@@ -847,6 +872,14 @@ namespace HordeServer.Agents
 		{
 		}
 
+		/// <summary>
+		/// Constructor
+		/// </summary>
+		public AgentLease(IAgentLease other)
+			: this(other.Id, other.ParentId, other.State, other.Resources, other.Exclusive, other.Payload?.ToByteArray())
+		{
+		}
+
 		/// <summary>
 		/// Constructor
 		/// </summary>
@@ -875,17 +908,20 @@ namespace HordeServer.Agents
 			: this(options.Id, options.ParentId, LeaseState.Pending, options.Resources, options.Exclusive, Any.Pack(options.Payload).ToByteArray())
 		{
 		}
+	}
 
+	static class AgentLeaseExtensions
+	{
 		/// <summary>
 		/// Converts this lease to an RPC message
 		/// </summary>
 		/// <returns>RPC message</returns>
-		public RpcLease ToRpcMessage()
+		public static RpcLease ToRpcMessage(this IAgentLease agentLease)
 		{
 			RpcLease lease = new RpcLease();
-			lease.Id = Id;
-			lease.Payload = Google.Protobuf.WellKnownTypes.Any.Parser.ParseFrom(Payload);
-			lease.State = (RpcLeaseState)State;
+			lease.Id = agentLease.Id;
+			lease.Payload = agentLease.Payload;
+			lease.State = (RpcLeaseState)agentLease.State;
 			return lease;
 		}
 	}
diff --git a/Engine/Source/Programs/Horde/Plugins/Compute/HordeServer.Compute/Agents/Leases/LeasesController.cs b/Engine/Source/Programs/Horde/Plugins/Compute/HordeServer.Compute/Agents/Leases/LeasesController.cs
--- a/Engine/Source/Programs/Horde/Plugins/Compute/HordeServer.Compute/Agents/Leases/LeasesController.cs
+++ b/Engine/Source/Programs/Horde/Plugins/Compute/HordeServer.Compute/Agents/Leases/LeasesController.cs
@@ -238,7 +238,7 @@ namespace HordeServer.Agents.Leases
 				return NotFound(lease.AgentId);
 			}
 
-			AgentLease? agentLease = agent.Leases.FirstOrDefault(x => x.Id == leaseId);
+			IAgentLease? agentLease = agent.Leases.FirstOrDefault(x => x.Id == leaseId);
 			if (agentLease == null)
 			{
 				return NotFound(agent.Id, leaseId);
