Horde: Add agent properties for the pool that an agent belongs to.

#rnx

diff --git a/Engine/Source/Programs/Horde/Plugins/Compute/HordeServer.Compute/Agents/AgentCollection.cs b/Engine/Source/Programs/Horde/Plugins/Compute/HordeServer.Compute/Agents/AgentCollection.cs
--- a/Engine/Source/Programs/Horde/Plugins/Compute/HordeServer.Compute/Agents/AgentCollection.cs
+++ b/Engine/Source/Programs/Horde/Plugins/Compute/HordeServer.Compute/Agents/AgentCollection.cs
@@ -549,15 +549,11 @@ namespace HordeServer.Agents
 				updates.Add(updateBuilder.Set(x => x.SessionExpiresAt, options.SessionExpiresAt.Value));
 			}
 
-			List<string>? newProperties = agent.Properties;
+			List<string> newProperties = agent.Properties ?? new List<string>();
 			if (options.Capabilities != null)
 			{
 				options.Capabilities.Flatten(out newProperties, out Dictionary<string, int> newResources);
 
-				if (!(agent.Properties ?? Enumerable.Empty<string>()).SequenceEqual(newProperties, StringComparer.Ordinal))
-				{
-					updates.Add(updateBuilder.Set(x => x.Properties, newProperties));
-				}
 				if (!ResourcesEqual(newResources, agent.Resources))
 				{
 					updates.Add(updateBuilder.Set(x => x.Resources, newResources));
@@ -600,12 +596,19 @@ namespace HordeServer.Agents
 			}
 
 			// Update the pools
-			List<PoolId> pools = CreatePoolsList(options.DynamicPools ?? agent.DynamicPools, agent.ExplicitPools, newProperties ?? Enumerable.Empty<string>());
+			List<PoolId> pools = CreatePoolsList(options.DynamicPools ?? agent.DynamicPools, agent.ExplicitPools, newProperties);
 			if (!Enumerable.SequenceEqual(pools, agent.Pools))
 			{
 				updates.Add(updateBuilder.Set(x => x.Pools, pools));
 			}
 
+			SetPoolProperties(newProperties, pools);
+
+			if (!Enumerable.SequenceEqual(agent.Properties ?? Enumerable.Empty<string>(), newProperties, StringComparer.Ordinal))
+			{
+				updates.Add(updateBuilder.Set(x => x.Properties, newProperties));
+			}
+
 			// If there are no new updates, return immediately. This is important for preventing UpdateSession calls from returning immediately.
 			if (updates.Count == 0)
 			{
@@ -616,6 +619,13 @@ namespace HordeServer.Agents
 			return await TryUpdateAsync(agent, updateBuilder.Combine(updates), cancellationToken);
 		}
 
+		static void SetPoolProperties(List<string> properties, List<PoolId> pools)
+		{
+			properties.RemoveAll(x => x.StartsWith($"{KnownPropertyNames.Pool}=", StringComparison.OrdinalIgnoreCase));
+			properties.AddRange(pools.Select(x => $"{KnownPropertyNames.Pool}={x}"));
+			properties.Sort(StringComparer.OrdinalIgnoreCase);
+		}
+
 		static List<PoolId> CreatePoolsList(IEnumerable<PoolId> pools)
 			=> pools.Distinct().OrderBy(x => x.Id.Text).ToList();
 
@@ -692,7 +702,10 @@ namespace HordeServer.Agents
 		async Task<AgentDocument?> TryCreateSessionAsync(AgentDocument agent, CreateSessionOptions options, CancellationToken cancellationToken)
 		{
 			options.Capabilities.Flatten(out List<string> newProperties, out Dictionary<string, int> newResources);
+
 			List<PoolId> newDynamicPools = new(options.DynamicPools);
+			List<PoolId> newPools = CreatePoolsList(agent.ExplicitPools, newDynamicPools, newProperties);
+			SetPoolProperties(newProperties, newPools);
 
 			// Reset the agent to use the new session
 			UpdateDefinitionBuilder<AgentDocument> updateBuilder = Builders<AgentDocument>.Update;
@@ -705,7 +718,7 @@ namespace HordeServer.Agents
 			updates.Add(updateBuilder.Unset(x => x.Deleted));
 			updates.Add(updateBuilder.Set(x => x.Properties, newProperties));
 			updates.Add(updateBuilder.Set(x => x.Resources, newResources));
-			updates.Add(updateBuilder.Set(x => x.Pools, CreatePoolsList(agent.ExplicitPools, newDynamicPools, newProperties)));
+			updates.Add(updateBuilder.Set(x => x.Pools, newPools));
 			updates.Add(updateBuilder.Set(x => x.DynamicPools, newDynamicPools));
 			updates.Add(updateBuilder.Set(x => x.Version, options.Version));
 			updates.Add(updateBuilder.Unset(x => x.RequestRestart));
