Horde: Remove dedicated session collection. Sessions are now managed by AgentCollection to ensure consistency.

#rnx

diff --git a/Engine/Source/Programs/Horde/Plugins/Build/HordeServer.Build.Tests/Agents/AgentCollectionTests.cs b/Engine/Source/Programs/Horde/Plugins/Build/HordeServer.Build.Tests/Agents/AgentCollectionTests.cs
--- a/Engine/Source/Programs/Horde/Plugins/Build/HordeServer.Build.Tests/Agents/AgentCollectionTests.cs
+++ b/Engine/Source/Programs/Horde/Plugins/Build/HordeServer.Build.Tests/Agents/AgentCollectionTests.cs
@@ -6,7 +6,6 @@ using EpicGames.Horde.Agents.Pools;
 using Google.Protobuf.WellKnownTypes;
 using HordeCommon.Rpc.Messages;
 using HordeServer.Agents;
-using HordeServer.Agents.Sessions;
 
 namespace HordeServer.Tests.Agents;
 
@@ -57,7 +56,7 @@ public class AgentCollectionTests : BuildTestSetup
 		await _agent.TryCreateLeaseAsync(_lease2);
 		await UpdateAgentAsync();
 
-		await _agent.TryCreateSessionAsync(new CreateSessionOptions(SessionIdUtils.GenerateNewId(), DateTime.UtcNow, AgentStatus.Ok,
+		await _agent.TryCreateSessionAsync(new CreateSessionOptions(AgentStatus.Ok,
 			new RpcAgentCapabilities(), new List<PoolId>(), DateTime.UtcNow, null));
 
 		List<LeaseId> leases = await AgentCollection.FindActiveLeaseIdsAsync();
diff --git a/Engine/Source/Programs/Horde/Plugins/Build/HordeServer.Build.Tests/BuildTestSetup.cs b/Engine/Source/Programs/Horde/Plugins/Build/HordeServer.Build.Tests/BuildTestSetup.cs
--- a/Engine/Source/Programs/Horde/Plugins/Build/HordeServer.Build.Tests/BuildTestSetup.cs
+++ b/Engine/Source/Programs/Horde/Plugins/Build/HordeServer.Build.Tests/BuildTestSetup.cs
@@ -15,7 +15,6 @@ using HordeServer.Agents.Fleet;
 using HordeServer.Agents.Leases;
 using HordeServer.Agents.Pools;
 using HordeServer.Agents.Relay;
-using HordeServer.Agents.Sessions;
 using HordeServer.Agents.Telemetry;
 using HordeServer.Agents.Utilization;
 using HordeServer.Artifacts;
@@ -71,7 +70,6 @@ namespace HordeServer.Tests
 		public IPoolCollection PoolCollection => ServiceProvider.GetRequiredService<IPoolCollection>();
 		public ILeaseCollection LeaseCollection => ServiceProvider.GetRequiredService<ILeaseCollection>();
 		public ILogCollection LogCollection => ServiceProvider.GetRequiredService<ILogCollection>();
-		public ISessionCollection SessionCollection => ServiceProvider.GetRequiredService<ISessionCollection>();
 		public ITestDataCollection TestDataCollection => ServiceProvider.GetRequiredService<ITestDataCollection>();
 		public IUserCollection UserCollection => ServiceProvider.GetRequiredService<IUserCollection>();
 		public IDeviceCollection DeviceCollection => ServiceProvider.GetRequiredService<IDeviceCollection>();
@@ -141,7 +139,6 @@ namespace HordeServer.Tests
 			services.AddSingleton<INotificationTriggerCollection, NotificationTriggerCollection>();
 			services.AddSingleton<IPoolCollection, PoolCollection>();
 			services.AddSingleton<IBisectTaskCollection, BisectTaskCollection>();
-			services.AddSingleton<ISessionCollection, SessionCollection>();
 			services.AddSingleton<ISubscriptionCollection, SubscriptionCollection>();
 			services.AddSingleton<IStreamCollection, StreamCollection>();
 			services.AddSingleton<ITemplateCollection, TemplateCollection>();
@@ -169,7 +166,6 @@ namespace HordeServer.Tests
 			services.AddSingleton<AwsAutoScalingLifecycleService>();
 			services.AddSingleton<ArtifactExpirationService>();
 			services.AddSingleton<FleetService>();
-			services.AddSingleton<ConsistencyService>();
 			services.AddSingleton<RequestTrackerService>();
 			services.AddSingleton<GlobalsService>();
 			services.AddSingleton<JobTaskSource>();
diff --git a/Engine/Source/Programs/Horde/Plugins/Compute/HordeServer.Compute/Agents/AgentCollection.cs b/Engine/Source/Programs/Horde/Plugins/Compute/HordeServer.Compute/Agents/AgentCollection.cs
--- a/Engine/Source/Programs/Horde/Plugins/Compute/HordeServer.Compute/Agents/AgentCollection.cs
+++ b/Engine/Source/Programs/Horde/Plugins/Compute/HordeServer.Compute/Agents/AgentCollection.cs
@@ -11,8 +11,10 @@ using EpicGames.Horde.Agents.Sessions;
 using EpicGames.Redis;
 using Google.Protobuf.WellKnownTypes;
 using HordeCommon.Rpc.Tasks;
+using HordeServer.Agents.Sessions;
 using HordeServer.Auditing;
 using HordeServer.Server;
+using Microsoft.Extensions.Logging;
 using MongoDB.Bson;
 using MongoDB.Bson.Serialization.Attributes;
 using MongoDB.Driver;
@@ -72,6 +74,16 @@ namespace HordeServer.Agents
 				_document = document;
 			}
 
+			public async Task<ISession?> GetSessionAsync(SessionId sessionId, CancellationToken cancellationToken = default)
+			{
+				return await _collection.GetSessionAsync(_document.Id, sessionId, cancellationToken);
+			}
+
+			public async Task<IReadOnlyList<ISession>> FindSessionsAsync(DateTime? startTime, DateTime? finishTime, int index, int count, CancellationToken cancellationToken = default)
+			{
+				return await _collection.FindSessionsAsync(_document.Id, startTime, finishTime, index, count, cancellationToken);
+			}
+
 			public async Task<IAgent?> TryCreateLeaseAsync(CreateLeaseOptions options, CancellationToken cancellationToken = default)
 			{
 				AgentDocument? newDocument = await _collection.TryCreateLeaseAsync(_document, options, cancellationToken);
@@ -214,7 +226,40 @@ namespace HordeServer.Agents
 			}
 		}
 
-		readonly IMongoCollection<AgentDocument> _agents;
+		/// <summary>
+		/// Concrete implementation of ISession
+		/// </summary>
+		class SessionDocument : ISession
+		{
+			[BsonRequired, BsonId]
+			public SessionId Id { get; set; }
+
+			[BsonRequired]
+			public AgentId AgentId { get; set; }
+
+			public DateTime StartTime { get; set; }
+			public DateTime? FinishTime { get; set; }
+			public string Version { get; set; } = String.Empty;
+
+			[BsonConstructor]
+			private SessionDocument()
+			{
+			}
+
+			public SessionDocument(SessionId id, AgentId agentId, DateTime startTime, string? version)
+			{
+				Id = id;
+				AgentId = agentId;
+				StartTime = startTime;
+				if (version != null)
+				{
+					Version = version;
+				}
+			}
+		}
+
+		readonly IMongoCollection<AgentDocument> _agentCollection;
+		readonly IMongoCollection<SessionDocument> _sessionCollection;
 		readonly IAuditLog<AgentId> _auditLog;
 		readonly IRedisService _redisService;
 		readonly IClock _clock;
@@ -225,11 +270,16 @@ namespace HordeServer.Agents
 		/// </summary>
 		public AgentCollection(IMongoService mongoService, IRedisService redisService, IClock clock, IAuditLog<AgentId> auditLog)
 		{
-			List<MongoIndex<AgentDocument>> indexes = new List<MongoIndex<AgentDocument>>();
-			indexes.Add(keys => keys.Ascending(x => x.Deleted).Ascending(x => x.Id).Ascending(x => x.Pools));
-			indexes.Add(keys => keys.Ascending(x => x.SessionExpiresAt), sparse: true);
+			List<MongoIndex<AgentDocument>> agentIndexes = new List<MongoIndex<AgentDocument>>();
+			agentIndexes.Add(keys => keys.Ascending(x => x.Deleted).Ascending(x => x.Id).Ascending(x => x.Pools));
+			agentIndexes.Add(keys => keys.Ascending(x => x.SessionExpiresAt), sparse: true);
+			_agentCollection = mongoService.GetCollection<AgentDocument>("Agents", agentIndexes);
+
+			List<MongoIndex<SessionDocument>> sessionIndexes = new List<MongoIndex<SessionDocument>>();
+			sessionIndexes.Add(keys => keys.Ascending(x => x.AgentId).Ascending(x => x.StartTime).Ascending(x => x.FinishTime));
+			sessionIndexes.Add(keys => keys.Ascending(x => x.FinishTime));
+			_sessionCollection = mongoService.GetCollection<SessionDocument>("Sessions", sessionIndexes);
 
-			_agents = mongoService.GetCollection<AgentDocument>("Agents", indexes);
 			_redisService = redisService;
 			_clock = clock;
 			_updateEventChannel = new RedisChannel<AgentId>(RedisChannel.Literal("agents/notify"));
@@ -255,7 +305,7 @@ namespace HordeServer.Agents
 				{
 					break;
 				}
-				document = newDocument ?? await _agents.Find<AgentDocument>(x => x.Id == document.Id).FirstOrDefaultAsync(cancellationToken);
+				document = newDocument ?? await _agentCollection.Find<AgentDocument>(x => x.Id == document.Id).FirstOrDefaultAsync(cancellationToken);
 			}
 			return document;
 		}
@@ -291,7 +341,7 @@ namespace HordeServer.Agents
 		public async Task<IAgent> AddAsync(AgentId id, bool ephemeral, string enrollmentKey, CancellationToken cancellationToken)
 		{
 			AgentDocument agent = new AgentDocument(id, ephemeral, enrollmentKey);
-			await _agents.InsertOneAsync(agent, null, cancellationToken);
+			await _agentCollection.InsertOneAsync(agent, null, cancellationToken);
 			return CreateAgentObject(agent);
 		}
 
@@ -328,13 +378,14 @@ namespace HordeServer.Agents
 		/// <inheritdoc/>
 		public async Task ForceDeleteAsync(AgentId agentId, CancellationToken cancellationToken)
 		{
-			await _agents.DeleteOneAsync(x => x.Id == agentId, cancellationToken);
+			await _sessionCollection.DeleteManyAsync(x => x.AgentId == agentId, cancellationToken);
+			await _agentCollection.DeleteOneAsync(x => x.Id == agentId, cancellationToken);
 		}
 
 		/// <inheritdoc/>
 		public async Task<IAgent?> GetAsync(AgentId agentId, CancellationToken cancellationToken)
 		{
-			AgentDocument? document = await _agents.Find<AgentDocument>(x => x.Id == agentId).FirstOrDefaultAsync(cancellationToken);
+			AgentDocument? document = await _agentCollection.Find<AgentDocument>(x => x.Id == agentId).FirstOrDefaultAsync(cancellationToken);
 			document = await PostLoadAsync(document, cancellationToken);
 			return CreateAgentObject(document);
 		}
@@ -342,7 +393,7 @@ namespace HordeServer.Agents
 		/// <inheritdoc/>
 		public async Task<IReadOnlyList<IAgent>> GetManyAsync(List<AgentId> agentIds, CancellationToken cancellationToken)
 		{
-			List<AgentDocument> documents = await _agents.Find(p => agentIds.Contains(p.Id)).ToListAsync(cancellationToken);
+			List<AgentDocument> documents = await _agentCollection.Find(p => agentIds.Contains(p.Id)).ToListAsync(cancellationToken);
 			documents = await PostLoadAsync(documents, cancellationToken);
 			return documents.ConvertAll(x => CreateAgentObject(x));
 		}
@@ -383,7 +434,7 @@ namespace HordeServer.Agents
 				filter &= filterBuilder.Eq(x => x.Enabled, enabled.Value);
 			}
 
-			IMongoCollection<AgentDocument> collection = consistentRead ? _agents : _agents.WithReadPreference(ReadPreference.SecondaryPreferred);
+			IMongoCollection<AgentDocument> collection = consistentRead ? _agentCollection : _agentCollection.WithReadPreference(ReadPreference.SecondaryPreferred);
 			IFindFluent<AgentDocument, AgentDocument> search = collection.Find(filter);
 			if (index != null)
 			{
@@ -404,7 +455,7 @@ namespace HordeServer.Agents
 		{
 			FilterDefinition<AgentDocument> filter = Builders<AgentDocument>.Filter.Exists(x => x.SessionExpiresAt) & Builders<AgentDocument>.Filter.Lt(x => x.SessionExpiresAt, utcNow);
 
-			List<AgentDocument> documents = await _agents.Find(filter).Limit(maxAgents).ToListAsync(cancellationToken);
+			List<AgentDocument> documents = await _agentCollection.Find(filter).Limit(maxAgents).ToListAsync(cancellationToken);
 			documents = await PostLoadAsync(documents, cancellationToken);
 			return documents.ConvertAll(x => CreateAgentObject(x));
 		}
@@ -412,7 +463,7 @@ namespace HordeServer.Agents
 		/// <inheritdoc/>
 		public async Task<IReadOnlyList<IAgent>> FindDeletedAsync(CancellationToken cancellationToken)
 		{
-			List<AgentDocument> documents = await _agents.Find(x => x.Deleted).ToListAsync(cancellationToken);
+			List<AgentDocument> documents = await _agentCollection.Find(x => x.Deleted).ToListAsync(cancellationToken);
 			documents = await PostLoadAsync(documents, cancellationToken);
 			return documents.ConvertAll(x => CreateAgentObject(x));
 		}
@@ -431,6 +482,33 @@ namespace HordeServer.Agents
 			return childIds.Select(x => LeaseId.Parse(x.ToString())).ToList();
 		}
 
+		/// <inheritdoc/>
+		async Task<SessionDocument?> GetSessionAsync(AgentId agentId, SessionId sessionId, CancellationToken cancellationToken = default)
+		{
+			FilterDefinitionBuilder<SessionDocument> filterBuilder = Builders<SessionDocument>.Filter;
+
+			FilterDefinition<SessionDocument> filter = filterBuilder.Eq(x => x.Id, sessionId) & filterBuilder.Eq(x => x.AgentId, agentId);
+			return await _sessionCollection.Find(filter).FirstOrDefaultAsync(cancellationToken);
+		}
+
+		/// <inheritdoc/>
+		async Task<List<SessionDocument>> FindSessionsAsync(AgentId agentId, DateTime? startTime, DateTime? finishTime, int index, int count, CancellationToken cancellationToken = default)
+		{
+			FilterDefinitionBuilder<SessionDocument> filterBuilder = Builders<SessionDocument>.Filter;
+
+			FilterDefinition<SessionDocument> filter = filterBuilder.Eq(x => x.AgentId, agentId);
+			if (startTime != null)
+			{
+				filter &= filterBuilder.Gte(x => x.StartTime, startTime.Value);
+			}
+			if (finishTime != null)
+			{
+				filter &= filterBuilder.Or(filterBuilder.Eq(x => x.FinishTime, null), filterBuilder.Lte(x => x.FinishTime, finishTime.Value));
+			}
+
+			return await _sessionCollection.Find(filter).SortByDescending(x => x.StartTime).Skip(index).Limit(count).ToListAsync(cancellationToken);
+		}
+
 		/// <summary>
 		/// Update a single document
 		/// </summary>
@@ -446,7 +524,7 @@ namespace HordeServer.Agents
 			Expression<Func<AgentDocument, bool>> filter = x => x.Id == current.Id && x.UpdateIndex == prevUpdateIndex;
 			UpdateDefinition<AgentDocument> updateWithIndex = update.Set(x => x.UpdateIndex, current.UpdateIndex).Set(x => x.UpdateTime, current.UpdateTime);
 
-			return await _agents.FindOneAndUpdateAsync<AgentDocument>(filter, updateWithIndex, new FindOneAndUpdateOptions<AgentDocument, AgentDocument> { ReturnDocument = ReturnDocument.After }, cancellationToken);
+			return await _agentCollection.FindOneAndUpdateAsync<AgentDocument>(filter, updateWithIndex, new FindOneAndUpdateOptions<AgentDocument, AgentDocument> { ReturnDocument = ReturnDocument.After }, cancellationToken);
 		}
 
 		/// <inheritdoc/>
@@ -707,54 +785,75 @@ namespace HordeServer.Agents
 			List<PoolId> newPools = CreatePoolsList(agent.ExplicitPools, newDynamicPools, newProperties);
 			SetPoolProperties(newProperties, newPools);
 
-			// Reset the agent to use the new session
-			UpdateDefinitionBuilder<AgentDocument> updateBuilder = Builders<AgentDocument>.Update;
-
-			List<UpdateDefinition<AgentDocument>> updates = new List<UpdateDefinition<AgentDocument>>();
-			updates.Add(updateBuilder.Set(x => x.SessionId, options.SessionId));
-			updates.Add(updateBuilder.Set(x => x.SessionExpiresAt, options.SessionExpiresAt));
-			updates.Add(updateBuilder.Set(x => x.Status, options.Status));
-			updates.Add(updateBuilder.Unset(x => x.Leases));
-			updates.Add(updateBuilder.Unset(x => x.Deleted));
-			updates.Add(updateBuilder.Set(x => x.Properties, newProperties));
-			updates.Add(updateBuilder.Set(x => x.Resources, newResources));
-			updates.Add(updateBuilder.Set(x => x.Pools, newPools));
-			updates.Add(updateBuilder.Set(x => x.DynamicPools, newDynamicPools));
-			updates.Add(updateBuilder.Set(x => x.Version, options.Version));
-			updates.Add(updateBuilder.Unset(x => x.RequestRestart));
-			updates.Add(updateBuilder.Unset(x => x.RequestShutdown));
-			updates.Add(updateBuilder.Unset(x => x.RequestForceRestart));
-			updates.Add(updateBuilder.Set(x => x.LastShutdownReason, "Unexpected"));
+			// Create a new session document
+			SessionDocument newSession = new SessionDocument(SessionIdUtils.GenerateNewId(), agent.Id, _clock.UtcNow, options.Version);
+			await _sessionCollection.InsertOneAsync(newSession, null, cancellationToken);
 
-			if (String.Equals(options.Version, agent.LastUpgradeVersion, StringComparison.Ordinal))
+			try
 			{
-				updates.Add(updateBuilder.Unset(x => x.UpgradeAttemptCount));
-			}
+				// Reset the agent to use the new session
+				UpdateDefinitionBuilder<AgentDocument> updateBuilder = Builders<AgentDocument>.Update;
 
-			if (agent.Status != options.Status)
-			{
-				updates.Add(updateBuilder.Set(x => x.LastStatusChange, options.LastStatusChange));
-			}
+				List<UpdateDefinition<AgentDocument>> updates = new List<UpdateDefinition<AgentDocument>>();
+				updates.Add(updateBuilder.Set(x => x.SessionId, newSession.Id));
+				updates.Add(updateBuilder.Set(x => x.SessionExpiresAt, newSession.StartTime + AgentService.SessionExpiryTime));
+				updates.Add(updateBuilder.Set(x => x.Status, options.Status));
+				updates.Add(updateBuilder.Unset(x => x.Leases));
+				updates.Add(updateBuilder.Unset(x => x.Deleted));
+				updates.Add(updateBuilder.Set(x => x.Properties, newProperties));
+				updates.Add(updateBuilder.Set(x => x.Resources, newResources));
+				updates.Add(updateBuilder.Set(x => x.Pools, newPools));
+				updates.Add(updateBuilder.Set(x => x.DynamicPools, newDynamicPools));
+				updates.Add(updateBuilder.Set(x => x.Version, options.Version));
+				updates.Add(updateBuilder.Unset(x => x.RequestRestart));
+				updates.Add(updateBuilder.Unset(x => x.RequestShutdown));
+				updates.Add(updateBuilder.Unset(x => x.RequestForceRestart));
+				updates.Add(updateBuilder.Set(x => x.LastShutdownReason, "Unexpected"));
+
+				if (String.Equals(options.Version, agent.LastUpgradeVersion, StringComparison.Ordinal))
+				{
+					updates.Add(updateBuilder.Unset(x => x.UpgradeAttemptCount));
+				}
 
-			foreach (AgentLease agentLease in agent.Leases ?? Enumerable.Empty<AgentLease>())
+				if (agent.Status != options.Status)
+				{
+					updates.Add(updateBuilder.Set(x => x.LastStatusChange, options.LastStatusChange));
+				}
+
+				foreach (AgentLease agentLease in agent.Leases ?? Enumerable.Empty<AgentLease>())
+				{
+					await RemoveActiveLeaseAsync(agentLease);
+				}
+
+				// Apply the update
+				AgentDocument? newAgent = await TryUpdateAsync(agent, updateBuilder.Combine(updates), cancellationToken);
+				if (newAgent == null)
+				{
+					await _sessionCollection.DeleteOneAsync(x => x.Id == newSession.Id, CancellationToken.None);
+					return null;
+				}
+
+				return newAgent;
+			}
+			catch
 			{
-				await RemoveActiveLeaseAsync(agentLease);
+				await _sessionCollection.DeleteOneAsync(x => x.Id == newSession.Id, CancellationToken.None);
+				throw;
 			}
-
-			// Apply the update
-			return await TryUpdateAsync(agent, updateBuilder.Combine(updates), cancellationToken);
 		}
 
 		/// <inheritdoc/>
 		async Task<AgentDocument?> TryTerminateSessionAsync(AgentDocument agent, CancellationToken cancellationToken)
 		{
+			DateTime utcNow = _clock.UtcNow;
+
 			UpdateDefinition<AgentDocument> update = new BsonDocument();
 
 			update = update.Unset(x => x.SessionId);
 			update = update.Unset(x => x.SessionExpiresAt);
 			update = update.Unset(x => x.Leases);
 			update = update.Set(x => x.Status, AgentStatus.Stopped);
-			update = update.Set(x => x.LastStatusChange, _clock.UtcNow);
+			update = update.Set(x => x.LastStatusChange, utcNow);
 
 			if (agent.Ephemeral)
 			{
@@ -766,7 +865,17 @@ namespace HordeServer.Agents
 				await RemoveActiveLeaseAsync(agentLease);
 			}
 
-			return await TryUpdateAsync(agent, update, cancellationToken);
+			AgentDocument? newAgent = await TryUpdateAsync(agent, update, cancellationToken);
+			if (newAgent == null)
+			{
+				return null;
+			}
+
+			// Update the session document
+			GetLogger(newAgent.Id).LogInformation("Terminated session {SessionId}", agent.SessionId);
+			await _sessionCollection.UpdateOneAsync(x => x.Id == agent.SessionId, Builders<SessionDocument>.Update.Set(x => x.FinishTime, utcNow), cancellationToken: CancellationToken.None);
+
+			return newAgent;
 		}
 
 		private static string RedisKeyActiveLeaseIds() => $"agent/active-lease-id";
diff --git a/Engine/Source/Programs/Horde/Plugins/Compute/HordeServer.Compute/Agents/AgentService.cs b/Engine/Source/Programs/Horde/Plugins/Compute/HordeServer.Compute/Agents/AgentService.cs
--- a/Engine/Source/Programs/Horde/Plugins/Compute/HordeServer.Compute/Agents/AgentService.cs
+++ b/Engine/Source/Programs/Horde/Plugins/Compute/HordeServer.Compute/Agents/AgentService.cs
@@ -17,7 +17,6 @@ using HordeCommon.Rpc.Messages;
 using HordeServer.Acls;
 using HordeServer.Agents.Leases;
 using HordeServer.Agents.Pools;
-using HordeServer.Agents.Sessions;
 using HordeServer.Auditing;
 using HordeServer.Server;
 using HordeServer.Tasks;
@@ -70,7 +69,6 @@ namespace HordeServer.Agents
 		public IAgentCollection Agents { get; }
 
 		readonly ILeaseCollection _leases;
-		readonly ISessionCollection _sessions;
 		readonly IAclService _aclService;
 		readonly IDowntimeService _downtimeService;
 		readonly ITaskSource[] _taskSources;
@@ -108,7 +106,6 @@ namespace HordeServer.Agents
 		public AgentService(
 			IAgentCollection agents,
 			ILeaseCollection leases,
-			ISessionCollection sessions,
 			IAclService aclService,
 			IDowntimeService downtimeService,
 			IPoolCollection poolCollection,
@@ -122,7 +119,6 @@ namespace HordeServer.Agents
 		{
 			Agents = agents;
 			_leases = leases;
-			_sessions = sessions;
 			_aclService = aclService;
 			_downtimeService = downtimeService;
 			_cachedRates = new AsyncCachedValue<AgentRateTable?>(_ => redisService.GetDatabase().StringGetAsync(_agentRateTableData), TimeSpan.FromMinutes(2));
@@ -412,25 +408,18 @@ namespace HordeServer.Agents
 						await RemoveLeaseAsync(agent, lease, utcNow, LeaseOutcome.Failed, null, cancellationToken);
 					}
 
-					// Create a new session document
-					ISession newSession = await _sessions.AddAsync(SessionIdUtils.GenerateNewId(), agent.Id, _clock.UtcNow, version, cancellationToken);
-					DateTime sessionExpiresAt = utcNow + SessionExpiryTime;
-
 					// Get the new pools for the agent
 					List<PoolId> dynamicPools = await GetDynamicPoolsAsync(agent, cancellationToken);
 
 					// Reset the agent to use the new session
-					newAgent = await agent.TryCreateSessionAsync(new CreateSessionOptions(newSession.Id, sessionExpiresAt, status, capabilities, dynamicPools, lastStatusChange ?? utcNow, version), cancellationToken);
+					newAgent = await agent.TryCreateSessionAsync(new CreateSessionOptions(status, capabilities, dynamicPools, lastStatusChange ?? utcNow, version), cancellationToken);
 					if (newAgent != null)
 					{
 						LogPropertyChanges(agentLogger, agent.Properties, newAgent.Properties);
 						agent = newAgent;
-						agentLogger.LogInformation("Session {SessionId} started", newSession.Id);
+						agentLogger.LogInformation("Session {SessionId} started", newAgent.SessionId);
 						break;
 					}
-
-					// Remove the session we didn't use
-					await _sessions.DeleteAsync(newSession.Id, cancellationToken);
 				}
 
 				// Get the current agent state
@@ -853,7 +842,6 @@ namespace HordeServer.Agents
 			}
 
 			// Save off the session id and current leases
-			SessionId sessionId = agent.SessionId.Value;
 			IReadOnlyList<IAgentLease> leases = agent.Leases;
 
 			// Clear the current session
@@ -869,9 +857,6 @@ namespace HordeServer.Agents
 					await RemoveLeaseAsync(agent, lease, finishTime, LeaseOutcome.Failed, null, cancellationToken);
 				}
 
-				// Update the session document
-				Agents.GetLogger(agent.Id).LogInformation("Terminated session {SessionId}", sessionId);
-				await _sessions.UpdateAsync(sessionId, finishTime, cancellationToken);
 				return agent;
 			}
 			return null;
@@ -1004,30 +989,6 @@ namespace HordeServer.Agents
 #endif
 		}
 
-		/// <summary>
-		/// Gets information about a particular session
-		/// </summary>
-		/// <param name="sessionId">The unique session id</param>
-		/// <returns>The session information</returns>
-		public Task<ISession?> GetSessionAsync(SessionId sessionId)
-		{
-			return _sessions.GetAsync(sessionId);
-		}
-
-		/// <summary>
-		/// Find sessions for the given agent
-		/// </summary>
-		/// <param name="agentId">The unique agent id</param>
-		/// <param name="startTime">Start time to include in the search</param>
-		/// <param name="finishTime">Finish time to include in the search</param>
-		/// <param name="index">Index of the first result to return</param>
-		/// <param name="count">Number of results to return</param>
-		/// <returns>List of sessions matching the given criteria</returns>
-		public Task<List<ISession>> FindSessionsAsync(AgentId agentId, DateTime? startTime, DateTime? finishTime, int index, int count)
-		{
-			return _sessions.FindAsync(agentId, startTime, finishTime, index, count);
-		}
-
 		/// <summary>
 		/// Update the rate table for different agent types
 		/// </summary>
diff --git a/Engine/Source/Programs/Horde/Plugins/Compute/HordeServer.Compute/Agents/AgentsController.cs b/Engine/Source/Programs/Horde/Plugins/Compute/HordeServer.Compute/Agents/AgentsController.cs
--- a/Engine/Source/Programs/Horde/Plugins/Compute/HordeServer.Compute/Agents/AgentsController.cs
+++ b/Engine/Source/Programs/Horde/Plugins/Compute/HordeServer.Compute/Agents/AgentsController.cs
@@ -342,7 +342,7 @@ namespace HordeServer.Agents
 				return NotFound(agentId);
 			}
 
-			List<ISession> sessions = await _agentService.FindSessionsAsync(agentId, startTime?.UtcDateTime, finishTime?.UtcDateTime, index, count);
+			IReadOnlyList<ISession> sessions = await agent.FindSessionsAsync(startTime?.UtcDateTime, finishTime?.UtcDateTime, index, count);
 			return sessions.ConvertAll(x => CreateGetSessionResponse(x));
 		}
 
@@ -372,8 +372,8 @@ namespace HordeServer.Agents
 				return NotFound(agentId);
 			}
 
-			ISession? session = await _agentService.GetSessionAsync(sessionId);
-			if (session == null || session.AgentId != agentId)
+			ISession? session = await agent.GetSessionAsync(sessionId, HttpContext.RequestAborted);
+			if (session == null)
 			{
 				return NotFound();
 			}
diff --git a/Engine/Source/Programs/Horde/Plugins/Compute/HordeServer.Compute/Agents/IAgent.cs b/Engine/Source/Programs/Horde/Plugins/Compute/HordeServer.Compute/Agents/IAgent.cs
--- a/Engine/Source/Programs/Horde/Plugins/Compute/HordeServer.Compute/Agents/IAgent.cs
+++ b/Engine/Source/Programs/Horde/Plugins/Compute/HordeServer.Compute/Agents/IAgent.cs
@@ -16,6 +16,7 @@ using Google.Protobuf;
 using Google.Protobuf.WellKnownTypes;
 using HordeCommon.Rpc;
 using HordeCommon.Rpc.Messages;
+using HordeServer.Agents.Sessions;
 using MongoDB.Bson;
 using MongoDB.Bson.Serialization.Attributes;
 
@@ -183,6 +184,25 @@ namespace HordeServer.Agents
 		/// </summary>
 		public uint UpdateIndex { get; }
 
+		/// <summary>
+		/// Gets a sessions for the agent
+		/// </summary>
+		/// <param name="sessionId">Identifier for the session</param>
+		/// <param name="cancellationToken">Cancellation token for the operation</param>
+		/// <returns>List of sessions matching the given criteria</returns>
+		Task<ISession?> GetSessionAsync(SessionId sessionId, CancellationToken cancellationToken = default);
+
+		/// <summary>
+		/// Find sessions for the given agent
+		/// </summary>
+		/// <param name="startTime">Start time to include in the search</param>
+		/// <param name="finishTime">Finish time to include in the search</param>
+		/// <param name="index">Index of the first result to return</param>
+		/// <param name="count">Number of results to return</param>
+		/// <param name="cancellationToken">Cancellation token for the operation</param>
+		/// <returns>List of sessions matching the given criteria</returns>
+		Task<IReadOnlyList<ISession>> FindSessionsAsync(DateTime? startTime = null, DateTime? finishTime = null, int index = 0, int count = 10, CancellationToken cancellationToken = default);
+
 		/// <summary>
 		/// Resets an agent to use new settings
 		/// </summary>
@@ -277,14 +297,12 @@ namespace HordeServer.Agents
 	/// <summary>
 	/// Options for starting a new agent session
 	/// </summary>
-	/// <param name="SessionId">New session id</param>
-	/// <param name="SessionExpiresAt">Expiry time for the new session</param>
 	/// <param name="Status">Status of the agent</param>
 	/// <param name="Capabilities">Capabilities for the agent</param>
 	/// <param name="DynamicPools">New list of dynamic pools for the agent</param>
 	/// <param name="LastStatusChange">Time to force status change timestamp to</param>
 	/// <param name="Version">Current version of the agent software</param>
-	public record class CreateSessionOptions(SessionId SessionId, DateTime SessionExpiresAt, AgentStatus Status, RpcAgentCapabilities Capabilities, IReadOnlyList<PoolId> DynamicPools, DateTime LastStatusChange, string? Version);
+	public record class CreateSessionOptions(AgentStatus Status, RpcAgentCapabilities Capabilities, IReadOnlyList<PoolId> DynamicPools, DateTime LastStatusChange, string? Version);
 
 	/// <summary>
 	/// Options for updating a new agent session
diff --git a/Engine/Source/Programs/Horde/Plugins/Compute/HordeServer.Compute/Agents/Sessions/ISessionCollection.cs b/Engine/Source/Programs/Horde/Plugins/Compute/HordeServer.Compute/Agents/Sessions/ISessionCollection.cs
deleted file mode 100644
diff --git a/Engine/Source/Programs/Horde/Plugins/Compute/HordeServer.Compute/Agents/Sessions/SessionCollection.cs b/Engine/Source/Programs/Horde/Plugins/Compute/HordeServer.Compute/Agents/Sessions/SessionCollection.cs
deleted file mode 100644
diff --git a/Engine/Source/Programs/Horde/Plugins/Compute/HordeServer.Compute/ComputePlugin.cs b/Engine/Source/Programs/Horde/Plugins/Compute/HordeServer.Compute/ComputePlugin.cs
--- a/Engine/Source/Programs/Horde/Plugins/Compute/HordeServer.Compute/ComputePlugin.cs
+++ b/Engine/Source/Programs/Horde/Plugins/Compute/HordeServer.Compute/ComputePlugin.cs
@@ -13,7 +13,6 @@ using HordeServer.Agents.Fleet;
 using HordeServer.Agents.Leases;
 using HordeServer.Agents.Pools;
 using HordeServer.Agents.Relay;
-using HordeServer.Agents.Sessions;
 using HordeServer.Agents.Telemetry;
 using HordeServer.Aws;
 using HordeServer.Compute;
@@ -67,7 +66,6 @@ namespace HordeServer
 			services.AddSingleton<ILeaseCollection, LeaseCollection>();
 			services.AddSingleton<ILogCollection, LogCollection>();
 			services.AddSingleton<IPoolCollection, PoolCollection>();
-			services.AddSingleton<ISessionCollection, SessionCollection>();
 
 			services.AddSingleton<IAgentVersionProvider, AgentVersionProvider>();
 
@@ -92,7 +90,6 @@ namespace HordeServer
 			services.AddSingleton<AgentRelayService>();
 			services.AddHostedService(provider => provider.GetRequiredService<AgentRelayService>());
 
-			services.AddSingleton<ConsistencyService>();
 			services.AddSingleton<ComputeService>();
 			services.AddSingleton<EnrollmentService>();
 			services.AddSingleton<LogTailService>();
@@ -121,7 +118,6 @@ namespace HordeServer
 				if (_serverInfo.IsRunModeActive(RunMode.Worker))
 				{
 					services.AddHostedService(provider => provider.GetRequiredService<FleetService>());
-					services.AddHostedService(provider => provider.GetRequiredService<ConsistencyService>());
 					services.AddHostedService(provider => provider.GetRequiredService<ComputeService>());
 					services.AddHostedService(provider => provider.GetRequiredService<EnrollmentService>());
 				}
diff --git a/Engine/Source/Programs/Horde/Plugins/Compute/HordeServer.Compute/Server/ConsistencyService.cs b/Engine/Source/Programs/Horde/Plugins/Compute/HordeServer.Compute/Server/ConsistencyService.cs
deleted file mode 100644
