Horde: Fix bug with AWS spot interruption immediately aborting leases

With agent marked as busy, leases are terminated which is correct behavior in the case for say dev workstations.
For spot instances, workloads such as UBA should read the signal file at exit gracefully (within the spot termination grace period of course).

#rnx

diff --git a/Engine/Source/Programs/Horde/HordeAgent/Leases/LeaseManager.cs b/Engine/Source/Programs/Horde/HordeAgent/Leases/LeaseManager.cs
--- a/Engine/Source/Programs/Horde/HordeAgent/Leases/LeaseManager.cs
+++ b/Engine/Source/Programs/Horde/HordeAgent/Leases/LeaseManager.cs
@@ -246,6 +246,11 @@ namespace HordeAgent.Leases
 						stopping = true;
 					}
 				}
+				
+				if (_statusService.IsStopRequested)
+				{
+					stopping = true;
+				}
 
 				// Get the new agent status
 				bool busy = _statusService.IsBusy;
diff --git a/Engine/Source/Programs/Horde/HordeAgent/Services/AwsInstanceLifecycleService.cs b/Engine/Source/Programs/Horde/HordeAgent/Services/AwsInstanceLifecycleService.cs
--- a/Engine/Source/Programs/Horde/HordeAgent/Services/AwsInstanceLifecycleService.cs
+++ b/Engine/Source/Programs/Horde/HordeAgent/Services/AwsInstanceLifecycleService.cs
@@ -225,8 +225,8 @@ class AwsInstanceLifecycleService : BackgroundService
 
 	private async Task OnTerminationWarningAsync(Ec2TerminationInfo info, CancellationToken cancellationToken)
 	{
-		// Signal to server we are disabled, setting state to paused preventing new leases getting scheduled
-		_statusService.IsBusy = true;
+		// Request agent to be stopped. Agent's session loop will pick this up and notify the server, preventing new leases getting scheduled.
+		_statusService.IsStopRequested = true;
 
 		// Create and write the termination signal file, containing the time-to-live for the EC2 instance.
 		// Workloads executed by the agent that support this protocol can pick this up and prepare/clean up prior to termination
diff --git a/Engine/Source/Programs/Horde/HordeAgent/Services/StatusService.cs b/Engine/Source/Programs/Horde/HordeAgent/Services/StatusService.cs
--- a/Engine/Source/Programs/Horde/HordeAgent/Services/StatusService.cs
+++ b/Engine/Source/Programs/Horde/HordeAgent/Services/StatusService.cs
@@ -21,6 +21,7 @@ namespace HordeAgent.Services
 
 		private AgentStatusMessage _current;
 		private bool _isBusy;
+		private bool _isStopRequested;
 
 		readonly IOptionsMonitor<AgentSettings> _settings;
 		readonly BackgroundTask _task;
@@ -46,6 +47,22 @@ namespace HordeAgent.Services
 				}
 			}
 		}
+		
+		/// <summary>
+		/// Whether a stop of the agent is requested
+		/// </summary>
+		public bool IsStopRequested
+		{
+			get => _isStopRequested;
+			set
+			{
+				if (_isStopRequested != value)
+				{
+					_isStopRequested = value;
+					StatusChangedEvent.Set();
+				}
+			}
+		}
 
 		/// <summary>
 		/// Status was updated
