Horde: Improve lease termination handling during AWS spot interruptions

Only terminate session once a lease has finished.

diff --git a/Engine/Source/Programs/Horde/HordeAgent.Tests/Services/AwsInstanceLifecycleServiceTests.cs b/Engine/Source/Programs/Horde/HordeAgent.Tests/Services/AwsInstanceLifecycleServiceTests.cs
--- a/Engine/Source/Programs/Horde/HordeAgent.Tests/Services/AwsInstanceLifecycleServiceTests.cs
+++ b/Engine/Source/Programs/Horde/HordeAgent.Tests/Services/AwsInstanceLifecycleServiceTests.cs
@@ -1,6 +1,7 @@
 // Copyright Epic Games, Inc. All Rights Reserved.
 
 using System;
+using System.Collections.Generic;
 using System.IO;
 using System.Net;
 using System.Net.Http;
@@ -8,6 +9,7 @@ using System.Threading;
 using System.Threading.Tasks;
 using EpicGames.Core;
 using HordeAgent.Services;
+using HordeCommon.Rpc.Messages;
 using Microsoft.Extensions.Logging;
 using Microsoft.Extensions.Options;
 using Microsoft.VisualStudio.TestTools.UnitTesting;
@@ -16,6 +18,24 @@ using Moq.Protected;
 
 namespace HordeAgent.Tests.Services;
 
+/// <summary>
+/// Testing stub for IWorkerService
+/// </summary>
+public class WorkerServiceStub : IWorkerService
+{
+	public bool TerminateSession { get; set; }
+	
+	public List<RpcLease> GetActiveLeases()
+	{
+		throw new NotImplementedException();
+	}
+	
+	public void TerminateSessionAfterLease()
+	{
+		TerminateSession = true;
+	}
+}
+
 /// <summary>
 /// Fake implementation of the AWS EC2 metadata server (IMDS)
 /// </summary>
@@ -73,7 +93,7 @@ public class FakeAwsImds
 }
 
 [TestClass]
-public sealed class AwsInstanceLifecycleServiceTests : IAsyncDisposable, IDisposable
+public sealed class AwsInstanceLifecycleServiceTests : IDisposable
 {
 	// Stub for fulfilling IOptionsMonitor interface during testing
 	// Copied from HordeServerTests until a good way to share code between these is decided.
@@ -94,7 +114,7 @@ public sealed class AwsInstanceLifecycleServiceTests : IAsyncDisposable, IDispos
 			=> null;
 	}
 
-	private readonly StatusService _statusService;
+	private readonly WorkerServiceStub _workerService = new();
 	private readonly HttpClient _httpClient;
 	private readonly FakeAwsImds _fakeImds = new();
 	private readonly AwsInstanceLifecycleService _service;
@@ -113,9 +133,8 @@ public sealed class AwsInstanceLifecycleServiceTests : IAsyncDisposable, IDispos
 		AgentSettings settings = new() { WorkingDir = new DirectoryReference(tempDir) };
 		_terminationSignalFile = settings.GetTerminationSignalFile();
 
-		_statusService = new(new TestOptionsMonitor<AgentSettings>(settings), loggerFactory.CreateLogger<StatusService>());
 		_httpClient = _fakeImds.GetHttpClient();
-		_service = new AwsInstanceLifecycleService(_statusService, _httpClient, new OptionsWrapper<AgentSettings>(settings), loggerFactory.CreateLogger<AwsInstanceLifecycleService>());
+		_service = new AwsInstanceLifecycleService(_workerService, _httpClient, new OptionsWrapper<AgentSettings>(settings), loggerFactory.CreateLogger<AwsInstanceLifecycleService>());
 		_service._timeToLiveAsg = TimeSpan.FromMilliseconds(10);
 		_service._timeToLiveSpot = TimeSpan.FromMilliseconds(20);
 		_service._terminationBufferTime = TimeSpan.FromMilliseconds(2);
@@ -165,7 +184,9 @@ public sealed class AwsInstanceLifecycleServiceTests : IAsyncDisposable, IDispos
 		_fakeImds.InstanceLifeCycle = FakeAwsImds.Spot;
 		Assert.IsFalse(File.Exists(_terminationSignalFile.FullName));
 
+		Assert.IsFalse(_workerService.TerminateSession);
 		await _service.MonitorInstanceLifecycleAsync(CancellationToken.None);
+		Assert.IsTrue(_workerService.TerminateSession);
 
 		string data = await File.ReadAllTextAsync(_terminationSignalFile.FullName);
 		Assert.AreEqual("v1\n18\n2222000\nAWS EC2 Spot interruption\n", data); // 20 ms for spot, minus 2 ms for termination buffer
@@ -176,9 +197,4 @@ public sealed class AwsInstanceLifecycleServiceTests : IAsyncDisposable, IDispos
 		_httpClient.Dispose();
 		_service.Dispose();
 	}
-
-	public ValueTask DisposeAsync()
-	{
-		return _statusService.DisposeAsync();
-	}
 }
\ No newline at end of file
diff --git a/Engine/Source/Programs/Horde/HordeAgent/AgentApp.cs b/Engine/Source/Programs/Horde/HordeAgent/AgentApp.cs
--- a/Engine/Source/Programs/Horde/HordeAgent/AgentApp.cs
+++ b/Engine/Source/Programs/Horde/HordeAgent/AgentApp.cs
@@ -222,6 +222,7 @@ namespace HordeAgent
 			services.AddSingleton<CapabilitiesService>();
 			services.AddSingleton<ISessionFactory, SessionFactory>();
 			services.AddSingleton<WorkerService>();
+			services.AddSingleton<IWorkerService>(sp => sp.GetRequiredService<WorkerService>());
 			services.AddSingleton<LeaseLoggerFactory>();
 			services.AddHostedService(sp => sp.GetRequiredService<WorkerService>());
 
diff --git a/Engine/Source/Programs/Horde/HordeAgent/Leases/LeaseManager.cs b/Engine/Source/Programs/Horde/HordeAgent/Leases/LeaseManager.cs
--- a/Engine/Source/Programs/Horde/HordeAgent/Leases/LeaseManager.cs
+++ b/Engine/Source/Programs/Horde/HordeAgent/Leases/LeaseManager.cs
@@ -17,6 +17,32 @@ namespace HordeAgent.Leases
 	/// </summary>
 	class LeaseManager
 	{
+		/// <summary>
+		/// Number of leases completed
+		/// </summary>
+		public int NumLeasesCompleted { get; private set; }
+		
+		/// <summary>
+		/// Delegate for lease active events
+		/// </summary>
+		public delegate void LeaseActiveEvent(RpcLease lease);
+		
+		/// <summary>
+		/// Event triggered when a lease is accepted and set to active
+		/// </summary>
+		public event LeaseActiveEvent? OnLeaseActive;
+		
+		/// <summary>
+		/// List of pool IDs this session is a member of
+		/// </summary>
+		public IReadOnlyList<string> PoolIds { get; private set; } = new List<string>();
+		
+		/// <summary>
+		/// Whether to terminate session after at least one lease has finished
+		/// Useful for shutting down as gracefully as possible without interrupting the current lease executing.
+		/// </summary>
+		public bool TerminateSessionAfterLease { get; set; }
+		
 		/// <summary>
 		/// Object used for controlling access to the access tokens and active sessions list
 		/// </summary>
@@ -27,11 +53,6 @@ namespace HordeAgent.Leases
 		/// </summary>
 		readonly List<LeaseHandler> _activeLeases = new List<LeaseHandler>();
 
-		/// <summary>
-		/// Number of leases completed
-		/// </summary>
-		public int NumLeasesCompleted { get; private set; }
-
 		/// <summary>
 		/// Whether the agent is currently in an unhealthy state
 		/// </summary>
@@ -52,21 +73,6 @@ namespace HordeAgent.Leases
 		/// </summary>
 		int _updateSessionFailures;
 
-		/// <summary>
-		/// Delegate for lease active events
-		/// </summary>
-		public delegate void LeaseActiveEvent(RpcLease lease);
-
-		/// <summary>
-		/// Event triggered when a lease is accepted and set to active
-		/// </summary>
-		public event LeaseActiveEvent? OnLeaseActive;
-
-		/// <summary>
-		/// List of pool IDs this session is a member of
-		/// </summary>
-		public IReadOnlyList<string> PoolIds { get; private set; } = new List<string>();
-
 		/// <summary>
 		/// How long to wait before trying to reacquire a new connection
 		/// Exposed as internal to ease testing. Using a lower delay can speed up tests. 
@@ -107,12 +113,12 @@ namespace HordeAgent.Leases
 			}
 		}
 
-		public async Task<SessionResult> RunAsync(bool shutdownAfterFinishedLease, CancellationToken stoppingToken)
+		public async Task<SessionResult> RunAsync(CancellationToken stoppingToken)
 		{
 			SessionResult result;
 			try
 			{
-				result = await HandleSessionAsync(shutdownAfterFinishedLease, stoppingToken);
+				result = await HandleSessionAsync(stoppingToken);
 			}
 			catch (OperationCanceledException ex) when (stoppingToken.IsCancellationRequested)
 			{
@@ -192,7 +198,7 @@ namespace HordeAgent.Leases
 			}
 		}
 
-		async Task<SessionResult> HandleSessionAsync(bool shutdownAfterFinishedLease, CancellationToken stoppingToken)
+		async Task<SessionResult> HandleSessionAsync(CancellationToken stoppingToken)
 		{
 			HordeRpc.HordeRpcClient hordeRpc = await _session.HordeClient.CreateGrpcClientAsync<HordeRpc.HordeRpcClient>(stoppingToken);
 
@@ -246,11 +252,6 @@ namespace HordeAgent.Leases
 						stopping = true;
 					}
 				}
-				
-				if (_statusService.IsStopRequested)
-				{
-					stopping = true;
-				}
 
 				// Get the new agent status
 				bool busy = _statusService.IsBusy;
@@ -286,9 +287,9 @@ namespace HordeAgent.Leases
 						if (updateSessionResponse != null)
 						{
 							bool atLeastOneLeaseFinished = _activeLeases.Any(x => x.RpcLease.State is RpcLeaseState.Completed or RpcLeaseState.Cancelled);
-							if (atLeastOneLeaseFinished && shutdownAfterFinishedLease)
+							if (atLeastOneLeaseFinished && TerminateSessionAfterLease)
 							{
-								_logger.LogInformation("At least one lease executed. Requesting shutdown.");
+								_logger.LogInformation("At least one lease executed. Terminating session");
 								_sessionResult = new SessionResult(SessionOutcome.Terminate);
 							}
 
@@ -337,15 +338,16 @@ namespace HordeAgent.Leases
 							// Update the session result if we've transitioned to stopped
 							if (updateSessionResponse.Status == RpcAgentStatus.Stopped)
 							{
-								_logger.LogInformation("Agent status is stopped; returning from session update loop.");
-								return _sessionResult ?? new SessionResult(SessionOutcome.BackOff);
+								SessionResult result = _sessionResult ?? new SessionResult(SessionOutcome.BackOff);
+								_logger.LogInformation("Agent status is stopped; returning from session update loop with result {Result}", result.Outcome);
+								return result;
 							}
 						}
 
 						// If there's nothing still running and cancellation was requested, exit
 						if (_activeLeases.Count == 0 && _sessionResult != null)
 						{
-							_logger.LogInformation("No leases are active. Agent is stopping."); // TODO: Should not really hit this any more; server should report stopping state in lease update above.
+							_logger.LogInformation("No leases are active. Agent is stopping. Session result {Result}", _sessionResult.Outcome); // TODO: Should not really hit this any more; server should report stopping state in lease update above.
 							return _sessionResult;
 						}
 					}
diff --git a/Engine/Source/Programs/Horde/HordeAgent/Services/AwsInstanceLifecycleService.cs b/Engine/Source/Programs/Horde/HordeAgent/Services/AwsInstanceLifecycleService.cs
--- a/Engine/Source/Programs/Horde/HordeAgent/Services/AwsInstanceLifecycleService.cs
+++ b/Engine/Source/Programs/Horde/HordeAgent/Services/AwsInstanceLifecycleService.cs
@@ -61,7 +61,7 @@ class AwsInstanceLifecycleService : BackgroundService
 	public const string HttpClientName = "Horde.HttpAwsInstanceClient";
 
 	private const string BaseUri = "http://169.254.169.254/latest/meta-data";
-	private readonly StatusService _statusService;
+	private readonly IWorkerService _workerService;
 	private readonly HttpClient _httpClient;
 	private readonly FileReference _terminationSignalFile;
 	private readonly ILogger<AwsInstanceLifecycleService> _logger;
@@ -93,9 +93,9 @@ class AwsInstanceLifecycleService : BackgroundService
 	/// <summary>
 	/// Constructor
 	/// </summary>
-	public AwsInstanceLifecycleService(StatusService statusService, HttpClient httpClient, IOptions<AgentSettings> settings, ILogger<AwsInstanceLifecycleService> logger)
+	public AwsInstanceLifecycleService(IWorkerService workerService, HttpClient httpClient, IOptions<AgentSettings> settings, ILogger<AwsInstanceLifecycleService> logger)
 	{
-		_statusService = statusService;
+		_workerService = workerService;
 		_httpClient = httpClient;
 		_httpClient.Timeout = TimeSpan.FromSeconds(10);
 		_terminationWarningCallback = OnTerminationWarningAsync;
@@ -225,8 +225,9 @@ class AwsInstanceLifecycleService : BackgroundService
 
 	private async Task OnTerminationWarningAsync(Ec2TerminationInfo info, CancellationToken cancellationToken)
 	{
-		// Request agent to be stopped. Agent's session loop will pick this up and notify the server, preventing new leases getting scheduled.
-		_statusService.IsStopRequested = true;
+		// Request shutdown of entire machine while waiting for the current lease executing to finish.
+		// Agent's session loop will pick this up and notify the server, preventing new leases getting scheduled.
+		_workerService.TerminateSessionAfterLease();
 
 		// Create and write the termination signal file, containing the time-to-live for the EC2 instance.
 		// Workloads executed by the agent that support this protocol can pick this up and prepare/clean up prior to termination
diff --git a/Engine/Source/Programs/Horde/HordeAgent/Services/WorkerService.cs b/Engine/Source/Programs/Horde/HordeAgent/Services/WorkerService.cs
--- a/Engine/Source/Programs/Horde/HordeAgent/Services/WorkerService.cs
+++ b/Engine/Source/Programs/Horde/HordeAgent/Services/WorkerService.cs
@@ -10,10 +10,26 @@ using Microsoft.Extensions.Logging;
 
 namespace HordeAgent.Services
 {
+	/// <summary>
+	/// Interface for WorkerService. Mainly for extracting methods which simplifies testing and dependencies
+	/// </summary>
+	internal interface IWorkerService
+	{
+		/// <summary>
+		/// Gets all the current leases held by the agent
+		/// </summary>
+		public List<RpcLease> GetActiveLeases();
+		
+		/// <summary>
+		/// Stop the session loop after at least one lease has finished executing
+		/// </summary>
+		public void TerminateSessionAfterLease();
+	}
+	
 	/// <summary>
 	/// Implements the message handling loop for an agent. Runs asynchronously until disposed.
 	/// </summary>
-	class WorkerService : BackgroundService, IDisposable
+	class WorkerService : BackgroundService, IWorkerService
 	{
 		readonly ILogger _logger;
 		readonly ISessionFactory _sessionFactory;
@@ -41,14 +57,21 @@ namespace HordeAgent.Services
 			_logger = logger;
 			_serviceProvider = serviceProvider;
 		}
-
-		/// <summary>
-		/// Gets all the current leases held by the agent
-		/// </summary>
+		
+		/// <inheritdoc/>
 		public List<RpcLease> GetActiveLeases()
 		{
 			return _currentLeaseManager?.GetActiveLeases() ?? new List<RpcLease>();
 		}
+		
+		/// <inheritdoc/>
+		public void TerminateSessionAfterLease()
+		{
+			if (_currentLeaseManager != null)
+			{
+				_currentLeaseManager.TerminateSessionAfterLease = true;
+			}
+		}
 
 		/// <summary>
 		/// Executes the ServerTaskAsync method and swallows the exception for the task being cancelled. This allows waiting for it to terminate.
@@ -107,7 +130,7 @@ namespace HordeAgent.Services
 						await using (ISession session = await _sessionFactory.CreateAsync(stoppingToken))
 						{
 							_currentLeaseManager = new LeaseManager(session, _serviceProvider);
-							result = await _currentLeaseManager.RunAsync(false, stoppingToken);
+							result = await _currentLeaseManager.RunAsync(stoppingToken);
 						}
 
 						failureCount = 0;
