Horde: Require artifact types to be declared in config files. Also allow setting expiration policies per-project and per-stream, and handle expiration of artifact types in deleted streams or of a type which are no longer declared.

diff --git a/Engine/Source/Programs/Horde/Plugins/Build/HordeServer.Build.Tests/Artifacts/ArtifactTests.cs b/Engine/Source/Programs/Horde/Plugins/Build/HordeServer.Build.Tests/Artifacts/ArtifactTests.cs
--- a/Engine/Source/Programs/Horde/Plugins/Build/HordeServer.Build.Tests/Artifacts/ArtifactTests.cs
+++ b/Engine/Source/Programs/Horde/Plugins/Build/HordeServer.Build.Tests/Artifacts/ArtifactTests.cs
@@ -61,7 +61,7 @@ namespace HordeServer.Tests.Artifacts
 		public async Task ExpireArtifactAsync()
 		{
 			FakeClock clock = ServiceProvider.GetRequiredService<FakeClock>();
-			ArtifactExpirationService expirationService = ServiceProvider.GetRequiredService<ArtifactExpirationService>();
+			ArtifactCollection expirationService = ServiceProvider.GetRequiredService<ArtifactCollection>();
 
 			ArtifactType type = new ArtifactType("my-artifact");
 			UpdateConfig(config => config.Plugins.GetBuildConfig().ArtifactTypes.Add(new ArtifactTypeConfig { Type = type, KeepDays = 1 }));
@@ -95,7 +95,7 @@ namespace HordeServer.Tests.Artifacts
 			ArtifactType type = new ArtifactType("my-artifact");
 			UpdateConfig(config => config.Plugins.GetBuildConfig().ArtifactTypes.Add(new ArtifactTypeConfig { Type = type, KeepDays = 1 }));
 
-			ArtifactExpirationService expirationService = ServiceProvider.GetRequiredService<ArtifactExpirationService>();
+			ArtifactCollection expirationService = ServiceProvider.GetRequiredService<ArtifactCollection>();
 
 			await expirationService.StartAsync(CancellationToken.None);
 
@@ -135,8 +135,7 @@ namespace HordeServer.Tests.Artifacts
 			ArtifactType type = new ArtifactType("my-artifact");
 			UpdateConfig(config => config.Plugins.GetBuildConfig().ArtifactTypes.Add(new ArtifactTypeConfig { Type = type, KeepCount = 4 }));
 
-			ArtifactExpirationService expirationService = ServiceProvider.GetRequiredService<ArtifactExpirationService>();
-
+			ArtifactCollection expirationService = ServiceProvider.GetRequiredService<ArtifactCollection>();
 			await expirationService.StartAsync(CancellationToken.None);
 
 			StreamId fooStreamId = new StreamId("foo");
@@ -185,5 +184,53 @@ namespace HordeServer.Tests.Artifacts
 				Assert.AreEqual(0, artifacts.Count);
 			}
 		}
+
+		[TestMethod]
+		public async Task ExpireOrphanedAsync()
+		{
+			DateTime startTime = Clock.UtcNow;
+
+			ArtifactType type = new ArtifactType("my-artifact");
+			StreamId fooStreamId = new StreamId("foo");
+
+			ArtifactCollection artifactCollection = ServiceProvider.GetRequiredService<ArtifactCollection>();
+			await artifactCollection.StartAsync(CancellationToken.None);
+
+			// Orphan the artifact type
+			{
+				UpdateConfig(config => config.Plugins.GetBuildConfig().ArtifactTypes.Add(new ArtifactTypeConfig { Type = type, KeepCount = 4 }));
+				IArtifact artifact = await artifactCollection.AddAsync(new ArtifactName($"default"), type, null, fooStreamId, CommitIdWithOrder.FromPerforceChange(1), new string[] { "test1", "test2" }, Array.Empty<string>(), AclScopeName.Root);
+
+				await Clock.AdvanceAsync(TimeSpan.FromDays(30.0));
+				IArtifact? artifact2 = await artifactCollection.GetAsync(artifact.Id);
+				Assert.IsNotNull(artifact2);
+
+				UpdateConfig(config => config.Plugins.GetBuildConfig().ArtifactTypes.Clear());
+				IArtifact? artifact3 = await artifactCollection.GetAsync(artifact.Id);
+				Assert.IsNotNull(artifact3);
+
+				await Clock.AdvanceAsync(TimeSpan.FromDays(30.0));
+				IArtifact? artifact4 = await artifactCollection.GetAsync(artifact.Id);
+				Assert.IsNull(artifact4);
+			}
+
+			// Orphan the stream
+			{
+				UpdateConfig(config => config.Plugins.GetBuildConfig().ArtifactTypes.Add(new ArtifactTypeConfig { Type = type, KeepCount = 4 }));
+				IArtifact artifact = await artifactCollection.AddAsync(new ArtifactName($"default"), type, null, fooStreamId, CommitIdWithOrder.FromPerforceChange(1), new string[] { "test1", "test2" }, Array.Empty<string>(), AclScopeName.Root);
+
+				await Clock.AdvanceAsync(TimeSpan.FromDays(30.0));
+				IArtifact? artifact2 = await artifactCollection.GetAsync(artifact.Id);
+				Assert.IsNotNull(artifact2);
+
+				UpdateConfig(config => config.Plugins.GetBuildConfig().Projects.Clear());
+				IArtifact? artifact3 = await artifactCollection.GetAsync(artifact.Id);
+				Assert.IsNotNull(artifact3);
+
+				await Clock.AdvanceAsync(TimeSpan.FromDays(30.0));
+				IArtifact? artifact4 = await artifactCollection.GetAsync(artifact.Id);
+				Assert.IsNull(artifact4);
+			}
+		}
 	}
 }
diff --git a/Engine/Source/Programs/Horde/Plugins/Build/HordeServer.Build.Tests/BuildTestSetup.cs b/Engine/Source/Programs/Horde/Plugins/Build/HordeServer.Build.Tests/BuildTestSetup.cs
--- a/Engine/Source/Programs/Horde/Plugins/Build/HordeServer.Build.Tests/BuildTestSetup.cs
+++ b/Engine/Source/Programs/Horde/Plugins/Build/HordeServer.Build.Tests/BuildTestSetup.cs
@@ -127,7 +127,8 @@ namespace HordeServer.Tests
 
 			services.AddSingleton<IAgentCollection, AgentCollection>();
 			services.AddSingleton<IAgentTelemetryCollection, AgentTelemetryCollection>();
-			services.AddSingleton<IArtifactCollection, ArtifactCollection>();
+			services.AddSingleton<ArtifactCollection>();
+			services.AddSingleton<IArtifactCollection>(sp => sp.GetRequiredService<ArtifactCollection>());
 			services.AddSingleton<ICommitService, CommitService>();
 			services.AddSingleton<IGraphCollection, GraphCollection>();
 			services.AddSingleton<IIssueCollection, IssueCollection>();
@@ -164,7 +165,6 @@ namespace HordeServer.Tests
 			services.AddSingleton(provider => new Lazy<AgentService>(provider.GetRequiredService<AgentService>));
 			services.AddSingleton<AgentRelayService>();
 			services.AddSingleton<AwsAutoScalingLifecycleService>();
-			services.AddSingleton<ArtifactExpirationService>();
 			services.AddSingleton<FleetService>();
 			services.AddSingleton<RequestTrackerService>();
 			services.AddSingleton<GlobalsService>();
diff --git a/Engine/Source/Programs/Horde/Plugins/Build/HordeServer.Build/Artifacts/ArtifactCollection.cs b/Engine/Source/Programs/Horde/Plugins/Build/HordeServer.Build/Artifacts/ArtifactCollection.cs
--- a/Engine/Source/Programs/Horde/Plugins/Build/HordeServer.Build/Artifacts/ArtifactCollection.cs
+++ b/Engine/Source/Programs/Horde/Plugins/Build/HordeServer.Build/Artifacts/ArtifactCollection.cs
@@ -1,5 +1,6 @@
 // Copyright Epic Games, Inc. All Rights Reserved.
 
+using System.Diagnostics;
 using System.Runtime.CompilerServices;
 using EpicGames.Core;
 using EpicGames.Horde.Acls;
@@ -8,9 +9,14 @@ using EpicGames.Horde.Commits;
 using EpicGames.Horde.Storage;
 using EpicGames.Horde.Streams;
 using HordeServer.Commits;
+using HordeServer.Projects;
 using HordeServer.Server;
 using HordeServer.Storage;
+using HordeServer.Streams;
 using HordeServer.Utilities;
+using Microsoft.Extensions.Hosting;
+using Microsoft.Extensions.Logging;
+using Microsoft.Extensions.Options;
 using MongoDB.Bson;
 using MongoDB.Bson.Serialization.Attributes;
 using MongoDB.Driver;
@@ -20,12 +26,34 @@ namespace HordeServer.Artifacts
 	/// <summary>
 	/// Implementation of <see cref="IArtifactCollection"/>
 	/// </summary>
-	class ArtifactCollection : IArtifactCollection
+	class ArtifactCollection : IArtifactCollection, IHostedService, IAsyncDisposable
 	{
 		class Artifact : IArtifact
+		{
+			readonly ArtifactDocument _document;
+
+			public ArtifactDocument Document => _document;
+
+			ArtifactId IArtifact.Id => new ArtifactId(BinaryIdUtils.FromObjectId(_document.Id));
+			ArtifactName IArtifact.Name => _document.Name;
+			ArtifactType IArtifact.Type => _document.Type;
+			string? IArtifact.Description => _document.Description;
+			StreamId IArtifact.StreamId => _document.StreamId;
+			CommitIdWithOrder IArtifact.CommitId => _document.CommitId;
+			IReadOnlyList<string> IArtifact.Keys => _document.Keys;
+			IReadOnlyList<string> IArtifact.Metadata => _document.Metadata;
+			NamespaceId IArtifact.NamespaceId => _document.NamespaceId;
+			RefName IArtifact.RefName => _document.RefName;
+			DateTime IArtifact.CreatedAtUtc => (_document.CreatedAtUtc == default) ? _document.Id.CreationTime : _document.CreatedAtUtc;
+
+			public Artifact(ArtifactDocument document)
+				=> _document = document;
+		}
+
+		class ArtifactDocument
 		{
 			[BsonRequired, BsonId]
-			public ArtifactId Id { get; set; }
+			public ObjectId Id { get; set; }
 
 			[BsonElement("nam")]
 			public ArtifactName Name { get; set; }
@@ -48,20 +76,16 @@ namespace HordeServer.Artifacts
 			[BsonIgnore]
 			public CommitIdWithOrder CommitId
 			{
-				get => (CommitName != null)? new CommitIdWithOrder(CommitName, CommitOrder) : CommitIdWithOrder.FromPerforceChange(CommitOrder);
+				get => (CommitName != null) ? new CommitIdWithOrder(CommitName, CommitOrder) : CommitIdWithOrder.FromPerforceChange(CommitOrder);
 				set => (CommitName, CommitOrder) = (value.Name, value.Order);
 			}
 
 			[BsonElement("key")]
 			public List<string> Keys { get; set; } = new List<string>();
 
-			IReadOnlyList<string> IArtifact.Keys => Keys;
-
 			[BsonElement("met")]
 			public List<string> Metadata { get; set; } = new List<string>();
 
-			IReadOnlyList<string> IArtifact.Metadata => Metadata;
-
 			[BsonElement("ns")]
 			public NamespaceId NamespaceId { get; set; }
 
@@ -77,14 +101,12 @@ namespace HordeServer.Artifacts
 			[BsonElement("upd")]
 			public int UpdateIndex { get; set; }
 
-			DateTime IArtifact.CreatedAtUtc => (CreatedAtUtc == default) ? BinaryIdUtils.ToObjectId(Id.Id).CreationTime : CreatedAtUtc;
-
 			[BsonConstructor]
-			private Artifact()
+			private ArtifactDocument()
 			{
 			}
 
-			public Artifact(ArtifactId id, ArtifactName name, ArtifactType type, string? description, StreamId streamId, CommitIdWithOrder commitId, IEnumerable<string> keys, IEnumerable<string> metadata, NamespaceId namespaceId, RefName refName, DateTime createdAtUtc, AclScopeName scopeName)
+			public ArtifactDocument(ObjectId id, ArtifactName name, ArtifactType type, string? description, StreamId streamId, CommitIdWithOrder commitId, IEnumerable<string> keys, IEnumerable<string> metadata, NamespaceId namespaceId, RefName refName, DateTime createdAtUtc, AclScopeName scopeName)
 			{
 				Id = id;
 				Name = name;
@@ -101,25 +123,65 @@ namespace HordeServer.Artifacts
 			}
 		}
 
-		readonly IMongoCollection<Artifact> _artifacts;
+		class ArtifactExpiryDocument
+		{
+			public ObjectId Id { get; set; }
+
+			[BsonElement("str")]
+			public StreamId StreamId { get; set; }
+
+			[BsonElement("typ")]
+			public ArtifactType Type { get; set; }
+
+			[BsonElement("tim")]
+			public DateTime Time { get; set; }
+		}
+
+		readonly IMongoCollection<ArtifactDocument> _artifactCollection;
+		readonly IMongoCollection<ArtifactExpiryDocument> _artifactExpiryCollection;
 		readonly IClock _clock;
 		readonly ICommitService _commitService;
+		readonly IOptionsMonitor<BuildConfig> _buildConfig;
+		readonly StorageService _storageService;
+		readonly ITicker _ticker;
+		readonly ILogger _logger;
 
 		/// <summary>
 		/// Constructor
 		/// </summary>
-		public ArtifactCollection(IMongoService mongoService, IClock clock, ICommitService commitService)
+		public ArtifactCollection(IMongoService mongoService, IClock clock, ICommitService commitService, StorageService storageService, IOptionsMonitor<BuildConfig> buildConfig, ILogger<ArtifactCollection> logger)
 		{
-			List<MongoIndex<Artifact>> indexes = new List<MongoIndex<Artifact>>();
+			List<MongoIndex<ArtifactDocument>> indexes = new List<MongoIndex<ArtifactDocument>>();
 			indexes.Add(keys => keys.Ascending(x => x.Keys));
 			indexes.Add(keys => keys.Ascending(x => x.Type).Descending(x => x.Id));
 			indexes.Add(keys => keys.Ascending(x => x.StreamId).Descending(x => x.CommitOrder).Ascending(x => x.Name).Descending(x => x.Id));
-			_artifacts = mongoService.GetCollection<Artifact>("ArtifactsV2", indexes);
+			indexes.Add(keys => keys.Ascending(x => x.StreamId).Ascending(x => x.Type).Descending(x => x.Id).Ascending(x => x.Name));
+			_artifactCollection = mongoService.GetCollection<ArtifactDocument>("ArtifactsV2", indexes);
+
+			List<MongoIndex<ArtifactExpiryDocument>> expiryIndexes = new List<MongoIndex<ArtifactExpiryDocument>>();
+			expiryIndexes.Add(keys => keys.Ascending(x => x.StreamId).Ascending(x => x.Type), unique: true);
+			_artifactExpiryCollection = mongoService.GetCollection<ArtifactExpiryDocument>("ArtifactsV2.Expiry", expiryIndexes);
 
 			_clock = clock;
 			_commitService = commitService;
+			_buildConfig = buildConfig;
+			_storageService = storageService;
+			_ticker = clock.AddSharedTicker<ArtifactCollection>(TimeSpan.FromHours(1.0), ExpireArtifactsAsync, logger);
+			_logger = logger;
 		}
 
+		/// <inheritdoc/>
+		public ValueTask DisposeAsync()
+			=> _ticker.DisposeAsync();
+
+		/// <inheritdoc/>
+		public async Task StartAsync(CancellationToken cancellationToken)
+			=> await _ticker.StartAsync();
+
+		/// <inheritdoc/>
+		public async Task StopAsync(CancellationToken cancellationToken)
+			=> await _ticker.StopAsync();
+
 #pragma warning disable CA1308 // Expect ansi-only keys here
 		static string NormalizeKey(string key)
 			=> key.ToLowerInvariant();
@@ -131,7 +193,7 @@ namespace HordeServer.Artifacts
 		public static string GetArtifactPath(StreamId streamId, ArtifactType type) => $"{type}/{streamId}";
 
 		/// <inheritdoc/>
-		public async Task<IArtifact> AddAsync(ArtifactName name, ArtifactType type, string? description, StreamId streamId, CommitId commitId, IEnumerable<string> keys, IEnumerable<string> metadata, AclScopeName scopeName, CancellationToken cancellationToken)
+		public async Task<IArtifact> AddAsync(ArtifactName name, ArtifactType type, string? description, StreamId streamId, CommitId commitId, IEnumerable<string> keys, IEnumerable<string> metadata, AclScopeName scopeName, CancellationToken cancellationToken = default)
 		{
 			if (name.Id.IsEmpty)
 			{
@@ -142,91 +204,186 @@ namespace HordeServer.Artifacts
 				throw new ArgumentException($"Artifact type for '{name}' is not valid", nameof(type));
 			}
 
-			ArtifactId id = new ArtifactId(BinaryIdUtils.CreateNew());
+			StreamConfig streamConfig = _buildConfig.CurrentValue.GetStream(streamId);
+			if (!streamConfig.TryGetArtifactType(type, out _))
+			{
+				throw new ArtifactTypeNotFoundException(streamId, type);
+			}
+
+			// Create an entry in the expiry collection to ensure we can GC any unused items if the stream goes away
+			DateTime utcNow = _clock.UtcNow;
+			await AddExpiryRecordAsync(streamId, type, utcNow, cancellationToken);
+
+			// Create the artifact
+			ObjectId id = ObjectId.GenerateNewId();
 
 			NamespaceId namespaceId = Namespace.Artifacts;
 			RefName refName = new RefName($"{GetArtifactPath(streamId, type)}/{commitId}/{name}/{id}");
 
 			CommitIdWithOrder commitIdWithOrder = await _commitService.GetOrderedAsync(streamId, commitId, cancellationToken);
 
-			Artifact artifact = new Artifact(id, name, type, description, streamId, commitIdWithOrder, keys.Select(x => NormalizeKey(x)), metadata, namespaceId, refName, _clock.UtcNow, scopeName);
-			await _artifacts.InsertOneAsync(artifact, null, cancellationToken);
-			return artifact;
+			ArtifactDocument artifactDocument = new ArtifactDocument(id, name, type, description, streamId, commitIdWithOrder, keys.Select(x => NormalizeKey(x)), metadata, namespaceId, refName, _clock.UtcNow, scopeName);
+			await _artifactCollection.InsertOneAsync(artifactDocument, null, cancellationToken);
+
+			return new Artifact(artifactDocument);
 		}
 
-		/// <inheritdoc/>
-		public async Task DeleteAsync(IEnumerable<ArtifactId> ids, CancellationToken cancellationToken = default)
+		async Task AddExpiryRecordAsync(StreamId streamId, ArtifactType type, DateTime utcNow, CancellationToken cancellationToken = default)
 		{
-			FilterDefinition<Artifact> filter = Builders<Artifact>.Filter.In(x => x.Id, ids);
-			await _artifacts.DeleteManyAsync(filter, cancellationToken);
+			FilterDefinition<ArtifactExpiryDocument> streamFilter = Builders<ArtifactExpiryDocument>.Filter.Expr(x => x.StreamId == streamId && x.Type == type);
+			UpdateDefinition<ArtifactExpiryDocument> streamUpdate = Builders<ArtifactExpiryDocument>.Update.Set(x => x.Time, utcNow);
+			await _artifactExpiryCollection.UpdateOneAsync(streamFilter, streamUpdate, new UpdateOptions { IsUpsert = true }, cancellationToken);
 		}
-
 		/// <inheritdoc/>
 		public async IAsyncEnumerable<IArtifact> FindAsync(StreamId? streamId = null, CommitId? minCommitId = null, CommitId? maxCommitId = null, ArtifactName? name = null, ArtifactType? type = null, IEnumerable<string>? keys = null, int maxResults = 100, [EnumeratorCancellation] CancellationToken cancellationToken = default)
 		{
-			FilterDefinition<Artifact> filter = FilterDefinition<Artifact>.Empty;
+			FilterDefinition<ArtifactDocument> filter = FilterDefinition<ArtifactDocument>.Empty;
 			if (streamId != null)
 			{
-				filter &= Builders<Artifact>.Filter.Eq(x => x.StreamId, streamId.Value);
+				filter &= Builders<ArtifactDocument>.Filter.Eq(x => x.StreamId, streamId.Value);
 				if (minCommitId != null)
 				{
 					CommitIdWithOrder minCommitIdWithOrder = await _commitService.GetOrderedAsync(streamId.Value, minCommitId, cancellationToken);
-					filter &= Builders<Artifact>.Filter.Gte(x => x.CommitOrder, minCommitIdWithOrder.Order);
+					filter &= Builders<ArtifactDocument>.Filter.Gte(x => x.CommitOrder, minCommitIdWithOrder.Order);
 				}
 				if (maxCommitId != null)
 				{
 					CommitIdWithOrder maxCommitIdWithOrder = await _commitService.GetOrderedAsync(streamId.Value, maxCommitId, cancellationToken);
-					filter &= Builders<Artifact>.Filter.Lte(x => x.CommitOrder, maxCommitIdWithOrder.Order);
+					filter &= Builders<ArtifactDocument>.Filter.Lte(x => x.CommitOrder, maxCommitIdWithOrder.Order);
 				}
 			}
 			if (name != null)
 			{
-				filter &= Builders<Artifact>.Filter.Eq(x => x.Name, name.Value);
+				filter &= Builders<ArtifactDocument>.Filter.Eq(x => x.Name, name.Value);
 			}
 			if (type != null)
 			{
-				filter &= Builders<Artifact>.Filter.Eq(x => x.Type, type.Value);
+				filter &= Builders<ArtifactDocument>.Filter.Eq(x => x.Type, type.Value);
 			}
 			if (keys != null && keys.Any())
 			{
-				filter &= Builders<Artifact>.Filter.All(x => x.Keys, keys.Select(x => NormalizeKey(x)));
+				filter &= Builders<ArtifactDocument>.Filter.All(x => x.Keys, keys.Select(x => NormalizeKey(x)));
 			}
 
-			using (IAsyncCursor<Artifact> cursor = await _artifacts.Find(filter).SortByDescending(x => x.CommitOrder).ThenByDescending(x => x.Id).Limit(maxResults).ToCursorAsync(cancellationToken))
+			using (IAsyncCursor<ArtifactDocument> cursor = await _artifactCollection.Find(filter).SortByDescending(x => x.CommitOrder).ThenByDescending(x => x.Id).Limit(maxResults).ToCursorAsync(cancellationToken))
 			{
 				while (await cursor.MoveNextAsync(cancellationToken))
 				{
-					foreach (Artifact artifact in cursor.Current)
+					foreach (ArtifactDocument artifactDocument in cursor.Current)
 					{
-						yield return artifact;
+						yield return new Artifact(artifactDocument);
 					}
 				}
 			}
 		}
 
-		/// <inheritdoc/>
-		public async IAsyncEnumerable<IEnumerable<IArtifact>> FindExpiredAsync(ArtifactType type, DateTime? expireAtUtc, [EnumeratorCancellation] CancellationToken cancellationToken = default)
+		async ValueTask ExpireArtifactsAsync(CancellationToken cancellationToken)
+		{
+			_logger.LogInformation("Checking for expired artifacts...");
+			Stopwatch timer = Stopwatch.StartNew();
+
+			BuildConfig buildConfig = _buildConfig.CurrentValue;
+
+			DateTime utcNow = _clock.UtcNow;
+
+			// Expire any active streams listed in the config
+			foreach (ProjectConfig projectConfig in buildConfig.Projects)
+			{
+				foreach (StreamConfig streamConfig in projectConfig.Streams)
+				{
+					foreach (ArtifactTypeConfig artifactTypeConfig in streamConfig.GetAllArtifactTypes())
+					{
+						await AddExpiryRecordAsync(streamConfig.Id, artifactTypeConfig.Type, utcNow, cancellationToken);
+						await ExpireArtifactsForStreamAsync(streamConfig.Id, artifactTypeConfig, utcNow, cancellationToken);
+					}
+				}
+			}
+
+			// Find any orphaned configurations and expire those too
+			DateTime orphanTime = utcNow.AddDays(-5.0);
+
+			List<ArtifactExpiryDocument> expiryDocuments = await _artifactExpiryCollection.Find(x => x.Time < orphanTime).ToListAsync(cancellationToken);
+			foreach (ArtifactExpiryDocument expiryDocument in expiryDocuments)
+			{
+				_logger.LogInformation("Expiring artifacts from orphaned stream {StreamId} (type={Type})", expiryDocument.StreamId, expiryDocument.Type);
+				using (IAsyncCursor<ArtifactDocument> cursor = await _artifactCollection.Find(x => x.StreamId == expiryDocument.StreamId && x.Type == expiryDocument.Type).ToCursorAsync(cancellationToken))
+				{
+					while (await cursor.MoveNextAsync(cancellationToken))
+					{
+						await DeleteArtifactsAsync(cursor.Current, cancellationToken);
+					}
+				}
+				await _artifactExpiryCollection.DeleteOneAsync(x => x.Id == expiryDocument.Id && x.Time == expiryDocument.Time, cancellationToken);
+			}
+
+			_logger.LogInformation("Finished expiring artifacts in {TimeSecs}s.", (long)timer.Elapsed.TotalSeconds);
+		}
+
+		async Task ExpireArtifactsForStreamAsync(StreamId streamId, ArtifactTypeConfig artifactTypeConfig, DateTime utcNow, CancellationToken cancellationToken)
 		{
-			FilterDefinition<Artifact> filter = Builders<Artifact>.Filter.Eq(x => x.Type, type);
-			if (expireAtUtc != null)
+			if (artifactTypeConfig.KeepCount.HasValue)
 			{
-				filter &= Builders<Artifact>.Filter.Lt(x => x.Id, new ArtifactId(BinaryIdUtils.FromObjectId(ObjectId.GenerateNewId(expireAtUtc.Value))));
+				_logger.LogInformation("Removing {StreamId} {ArtifactType} artifacts except newest {Count}", streamId, artifactTypeConfig.Type, artifactTypeConfig.KeepCount.Value);
 			}
 
-			IFindFluent<Artifact, Artifact> query = _artifacts.Find(filter).SortByDescending(x => x.Id);
-			using (IAsyncCursor<Artifact> cursor = await query.ToCursorAsync(cancellationToken))
+			int count = 0;
+
+			FilterDefinition<ArtifactDocument> filter = Builders<ArtifactDocument>.Filter.Eq(x => x.StreamId, streamId) & Builders<ArtifactDocument>.Filter.Eq(x => x.Type, artifactTypeConfig.Type);
+			if (artifactTypeConfig.KeepDays.HasValue)
+			{
+				DateTime expireAtUtc = utcNow - TimeSpan.FromDays(artifactTypeConfig.KeepDays.Value);
+				filter &= Builders<ArtifactDocument>.Filter.Lt(x => x.Id, ObjectId.GenerateNewId(expireAtUtc));
+				_logger.LogInformation("Removing {StreamId} {ArtifactType} artifacts except newer than {Time}", streamId, artifactTypeConfig.Type, expireAtUtc);
+			}
+
+			IFindFluent<ArtifactDocument, ArtifactDocument> query = _artifactCollection.Find(filter).SortByDescending(x => x.Id);
+			using (IAsyncCursor<ArtifactDocument> cursor = await query.ToCursorAsync(cancellationToken))
 			{
 				while (await cursor.MoveNextAsync(cancellationToken))
 				{
-					yield return cursor.Current;
+					List<ArtifactDocument> deleteArtifacts = new List<ArtifactDocument>();
+					foreach (ArtifactDocument artifact in cursor.Current)
+					{
+						if (artifactTypeConfig.KeepCount == null || count >= artifactTypeConfig.KeepCount.Value)
+						{
+							deleteArtifacts.Add(artifact);
+						}
+						count++;
+					}
+					await DeleteArtifactsAsync(deleteArtifacts, cancellationToken);
 				}
 			}
 		}
 
+		async Task DeleteArtifactsAsync(IEnumerable<ArtifactDocument> deleteArtifacts, CancellationToken cancellationToken)
+		{
+			foreach (IGrouping<NamespaceId, ArtifactDocument> artifactGroup in deleteArtifacts.GroupBy(x => x.NamespaceId))
+			{
+				using IStorageClient storageClient = _storageService.CreateClient(artifactGroup.Key);
+				foreach (ArtifactDocument artifact in artifactGroup)
+				{
+					// Delete the ref allowing the storage service to expire this data
+					_logger.LogInformation("Expiring {StreamId} artifact {ArtifactId}, ref {RefName} (created {CreateTime})", artifact.StreamId, artifact.Id, artifact.RefName, artifact.CreatedAtUtc);
+					await storageClient.DeleteRefAsync(artifact.RefName, cancellationToken);
+				}
+
+				FilterDefinition<ArtifactDocument> filter = Builders<ArtifactDocument>.Filter.In(x => x.Id, artifactGroup.Select(x => x.Id));
+				await _artifactCollection.DeleteManyAsync(filter, cancellationToken);
+			}
+		}
+
 		/// <inheritdoc/>
-		public async Task<IArtifact?> GetAsync(ArtifactId artifactId, CancellationToken cancellationToken)
+		public async Task<IArtifact?> GetAsync(ArtifactId artifactId, CancellationToken cancellationToken = default)
 		{
-			return await _artifacts.Find(x => x.Id == artifactId).FirstOrDefaultAsync(cancellationToken);
+			ObjectId objectId = BinaryIdUtils.ToObjectId(artifactId.Id);
+
+			ArtifactDocument? document = await _artifactCollection.Find(x => x.Id == objectId).FirstOrDefaultAsync(cancellationToken);
+			if (document == null)
+			{
+				return null;
+			}
+
+			return new Artifact(document);
 		}
 	}
 }
diff --git a/Engine/Source/Programs/Horde/Plugins/Build/HordeServer.Build/Artifacts/ArtifactExpirationService.cs b/Engine/Source/Programs/Horde/Plugins/Build/HordeServer.Build/Artifacts/ArtifactExpirationService.cs
deleted file mode 100644
diff --git a/Engine/Source/Programs/Horde/Plugins/Build/HordeServer.Build/BuildConfig.cs b/Engine/Source/Programs/Horde/Plugins/Build/HordeServer.Build/BuildConfig.cs
--- a/Engine/Source/Programs/Horde/Plugins/Build/HordeServer.Build/BuildConfig.cs
+++ b/Engine/Source/Programs/Horde/Plugins/Build/HordeServer.Build/BuildConfig.cs
@@ -100,7 +100,7 @@ namespace HordeServer
 			foreach (ProjectConfig project in Projects)
 			{
 				_projectLookup.Add(project.Id, project);
-				project.PostLoad(project.Id, configOptions.ParentAcl);
+				project.PostLoad(project.Id, configOptions.ParentAcl, ArtifactTypes);
 
 				foreach (StreamConfig stream in project.Streams)
 				{
@@ -201,9 +201,9 @@ namespace HordeServer
 			return auth ?? Acl.Authorize(action, user);
 		}
 
-		static bool? AuthorizeArtifactType(IReadOnlyList<ArtifactTypeAclConfig> artifactTypes, ArtifactType artifactType, AclAction action, ClaimsPrincipal user)
+		static bool? AuthorizeArtifactType(IReadOnlyList<ArtifactTypeConfig> artifactTypes, ArtifactType artifactType, AclAction action, ClaimsPrincipal user)
 		{
-			ArtifactTypeAclConfig? config = artifactTypes.FirstOrDefault(x => x.Type == artifactType);
+			ArtifactTypeConfig? config = artifactTypes.FirstOrDefault(x => x.Type == artifactType);
 			return config?.Acl?.AuthorizeSingleScope(action, user);
 		}
 
@@ -321,26 +321,10 @@ namespace HordeServer
 		}
 	}
 
-	/// <summary>
-	/// ACL configuration for an artifact type
-	/// </summary>
-	public class ArtifactTypeAclConfig
-	{
-		/// <summary>
-		/// Name of the artifact type
-		/// </summary>
-		public ArtifactType Type { get; set; }
-
-		/// <summary>
-		/// Acl for the artifact type
-		/// </summary>
-		public AclConfig? Acl { get; set; }
-	}
-
 	/// <summary>
 	/// Configuration for an artifact
 	/// </summary>
-	public class ArtifactTypeConfig : ArtifactTypeAclConfig
+	public class ArtifactTypeConfig
 	{
 		/// <summary>
 		/// Legacy 'Name' property
@@ -352,6 +336,16 @@ namespace HordeServer
 			set => Type = value;
 		}
 
+		/// <summary>
+		/// Name of the artifact type
+		/// </summary>
+		public ArtifactType Type { get; set; }
+
+		/// <summary>
+		/// Acl for the artifact type
+		/// </summary>
+		public AclConfig? Acl { get; set; }
+
 		/// <summary>
 		/// Number of artifacts to retain
 		/// </summary>
diff --git a/Engine/Source/Programs/Horde/Plugins/Build/HordeServer.Build/BuildPlugin.cs b/Engine/Source/Programs/Horde/Plugins/Build/HordeServer.Build/BuildPlugin.cs
--- a/Engine/Source/Programs/Horde/Plugins/Build/HordeServer.Build/BuildPlugin.cs
+++ b/Engine/Source/Programs/Horde/Plugins/Build/HordeServer.Build/BuildPlugin.cs
@@ -69,7 +69,8 @@ namespace HordeServer
 			services.AddSingleton<IDefaultAclModifier, BuildAclModifier>();
 			services.AddSingleton<IPluginResponseFilter, BuildResponseFilter>();
 
-			services.AddSingleton<IArtifactCollection, ArtifactCollection>();
+			services.AddSingleton<ArtifactCollection>();
+			services.AddSingleton<IArtifactCollection>(sp => sp.GetRequiredService<ArtifactCollection>());
 			services.AddSingleton<IGraphCollection, GraphCollection>();
 			services.AddSingleton<IssueCollection>();
 			services.AddSingleton<IIssueCollection>(sp => sp.GetRequiredService<IssueCollection>());
@@ -88,7 +89,7 @@ namespace HordeServer
 			services.AddSingleton<ITemplateCollection, TemplateCollection>();
 
 			services.AddSingleton<IPoolSizeStrategyFactory, JobQueueStrategyFactory>();
-			services.AddHostedService<ArtifactExpirationService>();
+			services.AddHostedService(sp => sp.GetRequiredService<ArtifactCollection>());
 			services.AddSingleton<ICommitService, CommitService>();
 
 			services.AddSingleton<DeviceService>();
diff --git a/Engine/Source/Programs/Horde/Plugins/Build/HordeServer.Build/Projects/ProjectConfig.cs b/Engine/Source/Programs/Horde/Plugins/Build/HordeServer.Build/Projects/ProjectConfig.cs
--- a/Engine/Source/Programs/Horde/Plugins/Build/HordeServer.Build/Projects/ProjectConfig.cs
+++ b/Engine/Source/Programs/Horde/Plugins/Build/HordeServer.Build/Projects/ProjectConfig.cs
@@ -5,6 +5,7 @@ using System.Diagnostics;
 using System.Security.Claims;
 using EpicGames.Core;
 using EpicGames.Horde.Acls;
+using EpicGames.Horde.Artifacts;
 using EpicGames.Horde.Jobs;
 using EpicGames.Horde.Projects;
 using EpicGames.Horde.Telemetry;
@@ -104,13 +105,15 @@ namespace HordeServer.Projects
 		/// <summary>
 		/// Permissions for artifact types
 		/// </summary>
-		public List<ArtifactTypeAclConfig> ArtifactTypes { get; set; } = new List<ArtifactTypeAclConfig>();
+		public List<ArtifactTypeConfig> ArtifactTypes { get; set; } = new List<ArtifactTypeConfig>();
 
 		/// <summary>
 		/// Acl entries
 		/// </summary>
 		public AclConfig Acl { get; set; } = new AclConfig();
 
+		readonly Dictionary<ArtifactType, ArtifactTypeConfig> _artifactTypeLookup = new Dictionary<ArtifactType, ArtifactTypeConfig>();
+
 		/// <inheritdoc cref="AclConfig.Authorize(AclAction, ClaimsPrincipal)"/>
 		public bool Authorize(AclAction action, ClaimsPrincipal user)
 			=> Acl.Authorize(action, user);
@@ -120,15 +123,26 @@ namespace HordeServer.Projects
 		/// </summary>
 		/// <param name="id">Id of this project</param>
 		/// <param name="parentAcl">The owning global config object</param>
-		public void PostLoad(ProjectId id, AclConfig parentAcl)
+		/// <param name="parentArtifactTypes">Set of artifact types</param>
+		public void PostLoad(ProjectId id, AclConfig parentAcl, IEnumerable<ArtifactTypeConfig> parentArtifactTypes)
 		{
 			Id = id;
 			Acl.PostLoad(parentAcl, $"project:{Id}");
 			Acl.LegacyScopeNames = new AclScopeName[] { parentAcl.ScopeName.Append($"p:{Id}") };
 
+			_artifactTypeLookup.Clear();
+			foreach (ArtifactTypeConfig artifactTypeConfig in parentArtifactTypes)
+			{
+				_artifactTypeLookup[artifactTypeConfig.Type] = artifactTypeConfig;
+			}
+			foreach (ArtifactTypeConfig artifactTypeConfig in ArtifactTypes)
+			{
+				_artifactTypeLookup[artifactTypeConfig.Type] = artifactTypeConfig;
+			}
+
 			foreach (StreamConfig stream in Streams)
 			{
-				stream.PostLoad(stream.Id, this);
+				stream.PostLoad(stream.Id, this, _artifactTypeLookup.Values);
 			}
 		}
 	}
diff --git a/Engine/Source/Programs/Horde/Plugins/Build/HordeServer.Build/Streams/StreamConfig.cs b/Engine/Source/Programs/Horde/Plugins/Build/HordeServer.Build/Streams/StreamConfig.cs
--- a/Engine/Source/Programs/Horde/Plugins/Build/HordeServer.Build/Streams/StreamConfig.cs
+++ b/Engine/Source/Programs/Horde/Plugins/Build/HordeServer.Build/Streams/StreamConfig.cs
@@ -11,6 +11,7 @@ using EpicGames.Core;
 using EpicGames.Horde;
 using EpicGames.Horde.Acls;
 using EpicGames.Horde.Agents.Pools;
+using EpicGames.Horde.Artifacts;
 using EpicGames.Horde.Commits;
 using EpicGames.Horde.Issues;
 using EpicGames.Horde.Jobs;
@@ -244,7 +245,9 @@ namespace HordeServer.Streams
 		/// <summary>
 		/// Permissions for artifact types
 		/// </summary>
-		public List<ArtifactTypeAclConfig> ArtifactTypes { get; set; } = new List<ArtifactTypeAclConfig>();
+		public List<ArtifactTypeConfig> ArtifactTypes { get; set; } = new List<ArtifactTypeConfig>();
+
+		readonly Dictionary<ArtifactType, ArtifactTypeConfig> _artifactTypeLookup = new Dictionary<ArtifactType, ArtifactTypeConfig>();
 
 		/// <inheritdoc cref="AclConfig.Authorize(AclAction, ClaimsPrincipal)"/>
 		public bool Authorize(AclAction action, ClaimsPrincipal user)
@@ -255,7 +258,8 @@ namespace HordeServer.Streams
 		/// </summary>
 		/// <param name="id">The stream id</param>
 		/// <param name="projectConfig">Owning project</param>
-		public void PostLoad(StreamId id, ProjectConfig projectConfig)
+		/// <param name="parentArtifactTypes">Parent artifact types</param>
+		public void PostLoad(StreamId id, ProjectConfig projectConfig, IEnumerable<ArtifactTypeConfig> parentArtifactTypes)
 		{
 			Id = id;
 			ProjectConfig = projectConfig;
@@ -325,11 +329,36 @@ namespace HordeServer.Streams
 				}
 			}
 
+			// Build a set of all the artifact types
+			_artifactTypeLookup.Clear();
+			foreach (ArtifactTypeConfig artifactTypeConfig in parentArtifactTypes)
+			{
+				_artifactTypeLookup[artifactTypeConfig.Type] = artifactTypeConfig;
+			}
+			foreach (ArtifactTypeConfig artifactTypeConfig in ArtifactTypes)
+			{
+				_artifactTypeLookup[artifactTypeConfig.Type] = artifactTypeConfig;
+			}
+
 			// Compute a hash of this stream revision to make it easier to detect changes
 			byte[] streamData = JsonSerializer.SerializeToUtf8Bytes(this, JsonUtils.DefaultSerializerOptions);
 			Revision = IoHash.Compute(streamData).ToString();
 		}
 
+		/// <summary>
+		/// Get all artifact types for this stream
+		/// </summary>
+		public IEnumerable<ArtifactTypeConfig> GetAllArtifactTypes()
+			=> _artifactTypeLookup.Values;
+
+		/// <summary>
+		/// Finds an artifact type configuration block
+		/// </summary>
+		/// <param name="type">The artifact type</param>
+		/// <param name="artifactTypeConfig">Receives the configuration for the artifact type</param>
+		public bool TryGetArtifactType(ArtifactType type, [NotNullWhen(true)] out ArtifactTypeConfig? artifactTypeConfig)
+			=> _artifactTypeLookup.TryGetValue(type, out artifactTypeConfig);
+
 		/// <summary>
 		/// Enumerates all commit tags, including the default tags for code and content.
 		/// </summary>
diff --git a/Engine/Source/Programs/Shared/EpicGames.Horde/Artifacts/ArtifactHttpCollection.cs b/Engine/Source/Programs/Shared/EpicGames.Horde/Artifacts/ArtifactHttpCollection.cs
--- a/Engine/Source/Programs/Shared/EpicGames.Horde/Artifacts/ArtifactHttpCollection.cs
+++ b/Engine/Source/Programs/Shared/EpicGames.Horde/Artifacts/ArtifactHttpCollection.cs
@@ -47,12 +47,6 @@ namespace EpicGames.Horde.Artifacts
 			throw new NotImplementedException();
 		}
 
-		/// <inheritdoc/>
-		public Task DeleteAsync(IEnumerable<ArtifactId> ids, CancellationToken cancellationToken = default)
-		{
-			throw new NotImplementedException();
-		}
-
 		/// <inheritdoc/>
 		public async IAsyncEnumerable<IArtifact> FindAsync(StreamId? streamId = null, CommitId? minCommitId = null, CommitId? maxCommitId = null, ArtifactName? name = null, ArtifactType? type = null, IEnumerable<string>? keys = null, int maxResults = 100, [EnumeratorCancellation] CancellationToken cancellationToken = default)
 		{
@@ -65,12 +59,6 @@ namespace EpicGames.Horde.Artifacts
 			}
 		}
 
-		/// <inheritdoc/>
-		public IAsyncEnumerable<IEnumerable<IArtifact>> FindExpiredAsync(ArtifactType type, DateTime? expireAtUtc, CancellationToken cancellationToken = default)
-		{
-			throw new NotImplementedException();
-		}
-
 		/// <inheritdoc/>
 		public async Task<IArtifact?> GetAsync(ArtifactId artifactId, CancellationToken cancellationToken = default)
 		{
diff --git a/Engine/Source/Programs/Shared/EpicGames.Horde/Artifacts/ArtifactType.cs b/Engine/Source/Programs/Shared/EpicGames.Horde/Artifacts/ArtifactType.cs
--- a/Engine/Source/Programs/Shared/EpicGames.Horde/Artifacts/ArtifactType.cs
+++ b/Engine/Source/Programs/Shared/EpicGames.Horde/Artifacts/ArtifactType.cs
@@ -1,7 +1,9 @@
 // Copyright Epic Games, Inc. All Rights Reserved.
 
+using System;
 using System.ComponentModel;
 using EpicGames.Core;
+using EpicGames.Horde.Streams;
 
 namespace EpicGames.Horde.Artifacts
 {
@@ -60,4 +62,30 @@ namespace EpicGames.Horde.Artifacts
 		/// <inheritdoc/>
 		public override StringId ToStringId(ArtifactType value) => value.Id;
 	}
+
+	/// <summary>
+	/// Exception thrown to indicate that an artifact type does not exist
+	/// </summary>
+	public sealed class ArtifactTypeNotFoundException : Exception
+	{
+		/// <summary>
+		/// The stream containing the artifact
+		/// </summary>
+		public StreamId StreamId { get; }
+
+		/// <summary>
+		/// The missing artifact type
+		/// </summary>
+		public ArtifactType ArtifactType { get; }
+
+		/// <summary>
+		/// Constructor
+		/// </summary>
+		public ArtifactTypeNotFoundException(StreamId streamId, ArtifactType artifactType)
+			: base($"Artifact type {artifactType} not found in stream {streamId}")
+		{
+			StreamId = streamId;
+			ArtifactType = artifactType;
+		}
+	}
 }
diff --git a/Engine/Source/Programs/Shared/EpicGames.Horde/Artifacts/IArtifactCollection.cs b/Engine/Source/Programs/Shared/EpicGames.Horde/Artifacts/IArtifactCollection.cs
--- a/Engine/Source/Programs/Shared/EpicGames.Horde/Artifacts/IArtifactCollection.cs
+++ b/Engine/Source/Programs/Shared/EpicGames.Horde/Artifacts/IArtifactCollection.cs
@@ -1,6 +1,5 @@
 // Copyright Epic Games, Inc. All Rights Reserved.
 
-using System;
 using System.Collections.Generic;
 using System.Threading;
 using System.Threading.Tasks;
@@ -30,23 +29,6 @@ namespace EpicGames.Horde.Artifacts
 		/// <returns>The new log file document</returns>
 		Task<IArtifact> AddAsync(ArtifactName name, ArtifactType type, string? description, StreamId streamId, CommitId commitId, IEnumerable<string> keys, IEnumerable<string> metadata, AclScopeName scopeName, CancellationToken cancellationToken = default);
 
-		/// <summary>
-		/// Deletes artifacts
-		/// </summary>
-		/// <param name="ids">Ids to search for</param>
-		/// <param name="cancellationToken">Cancellation token for the operation</param>
-		/// <returns>Sequence of artifacts</returns>
-		Task DeleteAsync(IEnumerable<ArtifactId> ids, CancellationToken cancellationToken = default);
-
-		/// <summary>
-		/// Finds artifacts which are ready for expiry
-		/// </summary>
-		/// <param name="type">Type of artifacts to find</param>
-		/// <param name="expireAtUtc">Number of artifacts to keep</param>
-		/// <param name="cancellationToken">Cancellation token for the operation</param>
-		/// <returns>Sequence of artifacts</returns>
-		IAsyncEnumerable<IEnumerable<IArtifact>> FindExpiredAsync(ArtifactType type, DateTime? expireAtUtc, CancellationToken cancellationToken = default);
-
 		/// <summary>
 		/// Finds artifacts with the given keys.
 		/// </summary>
@@ -69,16 +51,4 @@ namespace EpicGames.Horde.Artifacts
 		/// <returns>The artifact document</returns>
 		Task<IArtifact?> GetAsync(ArtifactId artifactId, CancellationToken cancellationToken = default);
 	}
-
-	/// <summary>
-	/// Extension methods for <see cref="IArtifactCollection"/>
-	/// </summary>
-	public static class ArtifactCollectionExtensions
-	{
-		/// <inheritdoc cref="IArtifactCollection.DeleteAsync(IEnumerable{ArtifactId}, CancellationToken)"/>>
-		public static Task DeleteAsync(this IArtifactCollection artifactCollection, ArtifactId artifactId, CancellationToken cancellationToken = default)
-		{
-			return artifactCollection.DeleteAsync(new[] { artifactId }, cancellationToken);
-		}
-	}
 }
