Horde: Make log events wrap again
#rnx
#skipci

diff --git a/Engine/Source/Programs/Horde/HordeDashboard/src/components/LogView.tsx b/Engine/Source/Programs/Horde/HordeDashboard/src/components/LogView.tsx
--- a/Engine/Source/Programs/Horde/HordeDashboard/src/components/LogView.tsx
+++ b/Engine/Source/Programs/Horde/HordeDashboard/src/components/LogView.tsx
@@ -93,7 +93,7 @@ class LogHandler {
 
       const startLine = currentLine + 1;
 
-      return this.events.filter(e => {
+      const events = this.events.filter(e => {
          if (warnings && e.severity !== EventSeverity.Warning) {
             return false;
          }
@@ -101,7 +101,15 @@ class LogHandler {
             return false;
          }
          return true;
-      }).find(e => e.lineIndex >= startLine);
+      });
+
+      let event = events.find(e => e.lineIndex >= startLine);
+
+      if (!event && events.length) {
+         event = events[0];
+      }
+
+      return event;
    }
 
    getPrevLogEvent(warnings?: boolean) {
@@ -113,11 +121,7 @@ class LogHandler {
 
       const startLine = currentLine - 1;
 
-      if (startLine < 0) {
-         return undefined;
-      }
-
-      return this.events.filter(e => {
+      const events = this.events.filter(e => {
          if (warnings && e.severity !== EventSeverity.Warning) {
             return false;
          }
@@ -125,7 +129,22 @@ class LogHandler {
             return false;
          }
          return true;
-      }).reverse().find(e => e.lineIndex <= startLine);
+      }).reverse();
+
+      if (startLine < 0) {
+         if (events.length) {
+            return events[0];
+         }
+         return undefined;
+      }
+
+      let event = events.find(e => e.lineIndex <= startLine)
+
+      if (!event && events.length) {
+         event = events[0];
+      }
+
+      return event;
    }
 
    getLogEvent(line: number | undefined) {
