Horde: Artifact search improvements
#skipci

diff --git a/Engine/Source/Programs/Horde/HordeDashboard/src/backend/Backend.ts b/Engine/Source/Programs/Horde/HordeDashboard/src/backend/Backend.ts
--- a/Engine/Source/Programs/Horde/HordeDashboard/src/backend/Backend.ts
+++ b/Engine/Source/Programs/Horde/HordeDashboard/src/backend/Backend.ts
@@ -506,7 +506,7 @@ export class Backend {
 
         return new Promise<FindArtifactsResponse>((resolve, reject) => {
 
-            this.backend.get(`/api/v2/artifacts`, { params: { streamId: streamId, minChange: minChange, maxChange: maxChange, name: name, type: type } })
+            this.backend.get(`/api/v2/artifacts`, { params: { streamId: streamId, minChange: minChange, maxChange: maxChange, name: name, type: type, maxResults: 512 } })
                 .then(response => { resolve(response.data); })
                 .catch(reason => reject(reason));
         })
diff --git a/Engine/Source/Programs/Horde/HordeDashboard/src/base/utilities/streamUtils.ts b/Engine/Source/Programs/Horde/HordeDashboard/src/base/utilities/streamUtils.ts
new file mode 100644
--- /dev/null
+++ b/Engine/Source/Programs/Horde/HordeDashboard/src/base/utilities/streamUtils.ts
@@ -0,0 +1,21 @@
+import { projectStore } from "../../backend/ProjectStore";
+
+
+export const getActiveStreamId = (): string | undefined => {
+
+    const path = window.location.pathname?.split("/").filter(c => !!c.trim());
+    if (path.length < 2) {
+        return undefined;
+    }
+
+    if (path[0] !== "stream") {
+        return undefined;
+    }
+
+    if (![projectStore.streamById(path[1])]) {
+        return undefined;
+    }
+
+    return path[1];
+
+}
\ No newline at end of file
diff --git a/Engine/Source/Programs/Horde/HordeDashboard/src/components/StreamView.tsx b/Engine/Source/Programs/Horde/HordeDashboard/src/components/StreamView.tsx
--- a/Engine/Source/Programs/Horde/HordeDashboard/src/components/StreamView.tsx
+++ b/Engine/Source/Programs/Horde/HordeDashboard/src/components/StreamView.tsx
@@ -23,8 +23,6 @@ import { NewBuild } from './NewBuild';
 import { StreamSummary } from './StreamSummary';
 import { TopNav } from './TopNav';
 import { NewBuildV2 } from './build/NewBuildV2';
-import { getSiteConfig } from '../backend/Config';
-import { FindArtifactsModal } from './artifacts/ArtifactsSearch';
 
 export const SummaryPage: React.FC = () => {
 
@@ -345,7 +343,6 @@ const StreamViewInner: React.FC = observer(() => {
    const [showOthersPreflights, setShowOthersPreflights] = useState<boolean | undefined>(dashboard.showPreflights);
 
    const [shown, setShown] = useState(query.get("newbuild") ? true : false);
-   const [artifactSearchShown, setArtifactSearchShown] = useState(false);
    const [findJobsShown, setFindJobsShown] = useState(false);
 
    const { projectStore } = useBackend();
@@ -465,14 +462,6 @@ const StreamViewInner: React.FC = observer(() => {
                   </Stack>
                }
             },
-            {
-               key: "find_artifacts",
-               text: "Find Artifacts",
-               iconProps: { iconName: "Search" },
-               onClick: () => {
-                  setArtifactSearchShown(true);
-               }
-            },
             {
                key: 'show_other_preflights',
                text: 'Show preflights for all users',
@@ -503,7 +492,6 @@ const StreamViewInner: React.FC = observer(() => {
       <Stack className={hordeClasses.horde}>
          <TopNav />
          <Breadcrumbs items={crumbItems} title={crumbTitle} />
-         {artifactSearchShown && <FindArtifactsModal streamId={stream.id} onClose={() => setArtifactSearchShown(false)} />}
          {shown && newBuildVersion == "2" && <NewBuildV2 streamId={streamId!} jobKey={newBuildTab!} show={true} onClose={(newJobId) => {
             setShown(false);
             if (newJobId) {
diff --git a/Engine/Source/Programs/Horde/HordeDashboard/src/components/TopNav.tsx b/Engine/Source/Programs/Horde/HordeDashboard/src/components/TopNav.tsx
--- a/Engine/Source/Programs/Horde/HordeDashboard/src/components/TopNav.tsx
+++ b/Engine/Source/Programs/Horde/HordeDashboard/src/components/TopNav.tsx
@@ -12,7 +12,7 @@ import { PreviewChangesModal } from './PreviewChanges';
 import { VersionModal } from './VersionModal';
 import { getHordeTheme } from '../styles/theme';
 import { getHordeStyling } from '../styles/Styles';
-import { FindArtifactsModal } from './artifacts/ArtifactsSearch';
+import { ArtifactQueryState, FindArtifactsModal } from './artifacts/ArtifactsSearch';
 
 
 const getStyles = () => {
@@ -758,6 +758,7 @@ export const TopNav: React.FC<{ suppressServer?: boolean }> = observer(({ suppre
       <div style={{ backgroundColor: hordeTheme.horde.topNavBackground }}>
          {showArtifacts && <FindArtifactsModal onClose={() => {
             query.delete("showArtifacts");
+            ArtifactQueryState.clearSearch(query);
             const url = `${window.location.pathname}?` + query.toString();
             navigate(url, { replace: true })
 
diff --git a/Engine/Source/Programs/Horde/HordeDashboard/src/components/artifacts/ArtifactsSearch.tsx b/Engine/Source/Programs/Horde/HordeDashboard/src/components/artifacts/ArtifactsSearch.tsx
--- a/Engine/Source/Programs/Horde/HordeDashboard/src/components/artifacts/ArtifactsSearch.tsx
+++ b/Engine/Source/Programs/Horde/HordeDashboard/src/components/artifacts/ArtifactsSearch.tsx
@@ -1,35 +1,295 @@
-import { ComboBox, DefaultButton, DetailsList, DetailsListLayoutMode, IComboBox, IComboBoxOption, IContextualMenuItem, IContextualMenuProps, IDetailsListProps, ITextField, IconButton, Label, Modal, PrimaryButton, ScrollablePane, ScrollbarVisibility, SelectionMode, Spinner, SpinnerSize, Stack, Text, TextField } from "@fluentui/react";
-import { getHordeStyling } from "../../styles/Styles";
-import { IStreamChooser, StreamChooser } from "../projects/StreamChooser";
+import { ComboBox, DefaultButton, DetailsList, DetailsListLayoutMode, IComboBox, IComboBoxOption, IDetailsListProps, IconButton, Label, Modal, PrimaryButton, ScrollablePane, ScrollbarVisibility, SelectionMode, Spinner, SpinnerSize, Stack, Text, TextField } from "@fluentui/react";
+import { useConst } from '@fluentui/react-hooks';
 import React, { useState } from "react";
+import { NavigateFunction, useNavigate } from "react-router-dom";
 import backend from "../../backend";
 import { GetArtifactResponse } from "../../backend/Api";
 import dashboard from "../../backend/Dashboard";
-import { JobArtifactsModal } from "./ArtifactsModal";
 import { projectStore } from "../../backend/ProjectStore";
+import { getActiveStreamId } from "../../base/utilities/streamUtils";
+import { getHordeStyling } from "../../styles/Styles";
+import { StreamChooser, streamIdAll } from "../projects/StreamChooser";
+import { JobArtifactsModal } from "./ArtifactsModal";
+import { useQuery } from "../JobDetailCommon";
+
+type ArtifactSearchState = {
+   streamId?: string;
+   minChange?: string;
+   maxChange?: string;
+   name?: string;
+   typeKey?: string;
+   sort?: string;
+   browseArtifactId?: string;
+   browseJobId?: string;
+   browseStepId?: string;
+   browseType?: string;
+}
+
+export class ArtifactQueryState {
+
+   constructor(search: URLSearchParams) {
+      this.originalSearch = new URLSearchParams(search);
+      const state = this.state = this.fromSearch(search);
+
+      if (state.minChange || state.maxChange || state.name || state.typeKey) {
+         this.autoLoad = true;
+      }   
+   }
+
+   fromSearch(search: URLSearchParams): ArtifactSearchState {
+
+      const state: ArtifactSearchState = {};
+
+      state.streamId = search.get("artifactStreamId") ?? undefined;
+      state.minChange = search.get("artifactMinChange") ?? undefined;
+      state.maxChange = search.get("artifactMaxChange") ?? undefined;
+      state.name = search.get("artifactName") ?? undefined;
+      state.typeKey = search.get("artifactTypeKey") ?? undefined;
+      state.sort = search.get("artifactSort") ?? undefined;
+      state.browseArtifactId = search.get("artifactId") ?? undefined;
+      state.browseJobId = search.get("artifactJobId") ?? undefined;
+      state.browseStepId = search.get("artifactStepId") ?? undefined;
+      state.browseType = search.get("artifactType") ?? undefined;
+      return state;
+
+   }
+
+   static clearSearch(search: URLSearchParams) {
+      search.delete("artifactStreamId");
+      search.delete("artifactMinChange");
+      search.delete("artifactMaxChange");
+      search.delete("artifactName");
+      search.delete("artifactTypeKey");
+      search.delete("artifactSort");
+      search.delete("artifactId");
+      search.delete("artifactJobId");
+      search.delete("artifactStepId");
+      search.delete("artifactType");
+   }
+
+   reset(navigate: NavigateFunction) {
+      const search = new URLSearchParams(this.originalSearch);
+      ArtifactQueryState.clearSearch(search);
+      this.state = this.fromSearch(search);
+      const url = `${window.location.pathname}?` + search.toString();
+      navigate(url, { replace: true })
+   }
+
+   updateSearch(navigate?: NavigateFunction, replace: boolean = true) {
+
+      const search = new URLSearchParams(this.originalSearch);
+      ArtifactQueryState.clearSearch(search);
+      const state = this.state;
+
+      if (state.streamId) {
+         search.append("artifactStreamId", state.streamId)
+      }
+
+      if (state.minChange) {
+         search.append("artifactMinChange", state.minChange)
+      }
+
+      if (state.maxChange) {
+         search.append("artifactMaxChange", state.maxChange)
+      }
+
+      if (state.name) {
+         search.append("artifactName", state.name)
+      }
+
+      if (state.typeKey) {
+         search.append("artifactTypeKey", state.typeKey)
+      }
+
+      if (state.sort) {
+         search.append("artifactSort", state.sort)
+      }
+
+      if (state.browseArtifactId) {
+         search.append("artifactId", state.browseArtifactId)
+      }
+
+      if (state.browseJobId) {
+         search.append("artifactJobId", state.browseJobId)
+      }
 
+      if (state.browseStepId) {
+         search.append("artifactStepId", state.browseStepId)
+      }
+
+      if (state.browseType) {
+         search.append("artifactType", state.browseType)
+      }
+
+      if (navigate) {
+         const url = `${window.location.pathname}?` + search.toString();
+         navigate(url, { replace: replace })
+      }
+   }
+
+   get streamId(): string | undefined {
+      return this.state.streamId;
+   }
+
+   set streamId(streamId: string | undefined) {
+      this.state.streamId = streamId;
+   }
+
+   get name(): string | undefined {
+      return this.state.name;
+   }
+
+   set name(name: string | undefined) {
+      this.state.name = name;
+   }
+
+   get sort(): string | undefined {
+      return this.state.sort;
+   }
+
+   set sort(sort: string | undefined) {
+      this.state.sort = sort;
+   }
 
-const ArtifactsList: React.FC<{ artifacts?: GetArtifactResponse[] }> = ({ artifacts }) => {
+   get typeKey(): string | undefined {
+      return this.state.typeKey;
+   }
+
+   set typeKey(typeKey: string | undefined) {
+      this.state.typeKey = typeKey;
+   }
+
+   get minChangeList(): string | undefined {
+      return this.state.minChange;
+   }
+
+   set minChangeList(minChange: string | undefined) {
+      this.state.minChange = minChange;
+   }
+
+   get maxChangeList(): string | undefined {
+      return this.state.maxChange;
+   }
 
-   const [browse, setBrowse] = useState<{ jobId?: string, stepId?: string, artifactId?: string, artifact?: GetArtifactResponse }>({});
+   set maxChangeList(maxChange: string | undefined) {
+      this.state.maxChange = maxChange;
+   }
+
+   get browseArtifactId(): string | undefined {
+      return this.state.browseArtifactId;
+   }
+
+   set browseArtifactId(browseArtifactId: string | undefined) {
+      this.state.browseArtifactId = browseArtifactId;
+   }
+
+   get browseJobId(): string | undefined {
+      return this.state.browseJobId;
+   }
+
+   set browseJobId(browseJobId: string | undefined) {
+      this.state.browseJobId = browseJobId;
+   }
+
+   get browseStepId(): string | undefined {
+      return this.state.browseStepId;
+   }
+
+   set browseStepId(browseStepId: string | undefined) {
+      this.state.browseStepId = browseStepId;
+   }
+
+   get browseType(): string | undefined {
+      return this.state.browseType;
+   }
+
+   set browseType(browseType: string | undefined) {
+      this.state.browseType = browseType;
+   }
+
+   autoLoad = false;
+   private state: ArtifactSearchState = {}
+   private originalSearch: URLSearchParams;
+
+}
+
+const ArtifactBrowser: React.FC<{ state: ArtifactQueryState }> = ({ state }) => {
+
+   const navigate = useNavigate();
+   useQuery();
+
+   // sync with search, needed for browser hisgtory navigation
+   const search = new URLSearchParams(window.location.search);
+
+   const artifactId = search.get("artifactId") ?? undefined;
+   const jobId = search.get("artifactJobId") ?? undefined;
+   const stepId = search.get("artifactStepId") ?? undefined;
+   const type = search.get("artifactType") ?? undefined;
+
+   state.browseArtifactId = artifactId;
+   state.browseJobId = jobId;
+   state.browseStepId = stepId;
+   state.browseType = type;
+
+   if (!artifactId || !jobId || !stepId || !type) {
+      return null;
+   }
+
+   return <Stack>
+      <JobArtifactsModal jobId={jobId} stepId={stepId!} artifactId={artifactId} contextType={type} onClose={() => {
+         state.browseArtifactId = undefined;
+         state.browseJobId = undefined;
+         state.browseStepId = undefined;
+         state.browseType = undefined;
+         state.updateSearch(navigate)
+      }} />
+   </Stack>
+}
+
+const ArtifactsList: React.FC<{ state: ArtifactQueryState, artifacts?: GetArtifactResponse[], sortBy?: string }> = ({ state, artifacts, sortBy }) => {
+
+   const navigate = useNavigate();
 
    if (!artifacts?.length) {
       return null;
    }
 
+   sortBy = sortBy ?? "sort-name";
+
    const sorted = artifacts.filter(a => {
       return !!a.keys.find(k => k.startsWith("job:") && k.indexOf("/step:") !== -1)
    }).sort((a, b) => {
 
-      if (a.name !== b.name) {
-         return a.name.localeCompare(b.name);
-      }
-
       if (a.streamId !== b.streamId) {
          return a.streamId!.localeCompare(b.streamId!);
       }
 
-      return (a.change ?? 0) - (b.change ?? 0)
+      const changeA = a.change ?? 0;
+      const changeB = b.change ?? 0;
+
+      if (sortBy === "sort-name") {
+
+         if (a.name !== b.name) {
+            return a.name.localeCompare(b.name);
+         }
+
+         if (changeA !== changeB) {
+            return changeA - changeB;
+         }
+      } else {
+
+         if (changeA !== changeB) {
+            return changeA - changeB;
+         }
+
+         if (a.name !== b.name) {
+            return a.name.localeCompare(b.name);
+         }
+      }
+
+
+      return 0;
+
    })
 
    const renderRow: IDetailsListProps['onRenderRow'] = (props) => {
@@ -46,27 +306,31 @@ const ArtifactsList: React.FC<{ artifacts?: GetArtifactResponse[] }> = ({ artifa
          const stream = projectStore.streamById(item.streamId);
 
          return <Stack horizontal verticalAlign="center" verticalFill tokens={{ childrenGap: 24 }} styles={{ root: { backgroundColor: background, paddingLeft: 12, paddingRight: 12, paddingTop: 8, paddingBottom: 8 } }}>
-            <Stack style={{ width: 420, paddingLeft: 8 }}>
+            <Stack style={{ width: 410, paddingLeft: 8 }}>
                <Text style={{ fontWeight: 600 }}>{item.name}</Text>
             </Stack>
-            <Stack style={{ width: 140 }}>
+            <Stack style={{ width: 240 }}>
                <Text>{stream ? (stream.fullname ?? stream.name) : item.streamId}</Text>
             </Stack>
             <Stack style={{ width: 72 }}>
                <Text>{item.change}</Text>
             </Stack>
-            <Stack style={{ width: 80 }}>
+            <Stack style={{ width: 160 }}>
                <Text>{item.type}</Text>
             </Stack>
             <Stack>
-               <DefaultButton style={{ width: 90, paddingRight: 8 }} text="Browse" onClick={() => {
+               <DefaultButton style={{ width: 90 }} text="Browse" onClick={() => {
 
                   const key = item.keys.find(k => k.startsWith("job:") && k.indexOf("/step:") !== -1);
                   if (!key) {
                      return;
                   }
 
-                  setBrowse({ jobId: key.slice(4, 28), stepId: key.slice(-4), artifactId: item.id, artifact: item })
+                  state.browseArtifactId = item.id;
+                  state.browseJobId = key.slice(4, 28);
+                  state.browseStepId = key.slice(-4);
+                  state.browseType = item.type;
+                  state.updateSearch(navigate, false);
 
                }} />
             </Stack>
@@ -75,8 +339,7 @@ const ArtifactsList: React.FC<{ artifacts?: GetArtifactResponse[] }> = ({ artifa
       return null;
    };
 
-   return <Stack styles={{ root: { position: "relative", height: 590 } }}>
-      {!!browse.jobId && <JobArtifactsModal jobId={browse.jobId} stepId={browse.stepId!} artifactId={browse.artifactId} contextType={browse.artifact!.type} onClose={() => setBrowse({})} />}
+   return <Stack styles={{ root: { position: "relative", height: 590, marginRight: 8 } }}>
       <ScrollablePane scrollbarVisibility={ScrollbarVisibility.always}>
          <DetailsList
             isHeaderVisible={false}
@@ -103,6 +366,9 @@ const artifactTypes: IComboBoxOption[] = [
    {
       key: `step-all`,
       text: `All`
+   }, {
+      key: `packaged-build`,
+      text: `Packaged Build`
    }, {
       key: `step-saved`,
       text: `step-saved`
@@ -119,60 +385,75 @@ const artifactTypes: IComboBoxOption[] = [
    }
 ]
 
-export const FindArtifactsModal: React.FC<{ streamId?: string, onClose: () => void }> = ({ streamId, onClose }) => {
+const sortOptions: IComboBoxOption[] = [
+   {
+      key: `sort-name`,
+      text: `Name`
+   }, {
+      key: `sort-change`,
+      text: `Change`
+   }
+]
+
+let streamChooserId = 0;
+
+export const FindArtifactsModal: React.FC<{ onClose: () => void }> = ({ onClose }) => {
 
+   const navigate = useNavigate();
+   const searchState = useConst(new ArtifactQueryState(new URLSearchParams(window.location.search)));
    const [state, setState] = useState<{ searching?: boolean, artifacts?: GetArtifactResponse[] }>({});
-   const streamRef = React.useRef<IStreamChooser>(null);
-   const minChangeRef = React.useRef<ITextField>(null);
-   const maxChangeRef = React.useRef<ITextField>(null);
-   const nameRef = React.useRef<ITextField>(null);
-   const typeRef = React.useRef<IComboBox>(null);
 
    const { hordeClasses } = getHordeStyling();
 
+   if (!searchState.streamId) {
+
+      if (!searchState.streamId) {
+         const streamId = getActiveStreamId();
+         if (streamId) {
+            searchState.streamId = streamId;
+            searchState.updateSearch(navigate);
+         }
+      }
+   }
+
    const queryArtifacts = async () => {
 
       setState({ ...state, searching: true });
 
       try {
 
-         let minChange: number | undefined = parseInt(minChangeRef.current?.value?.trim() ?? "0");
+         let minChange: number | undefined = parseInt(searchState.minChangeList?.trim() ?? "0");
          if (!minChange || isNaN(minChange)) {
             minChange = undefined;
          }
 
-         let maxChange: number | undefined = parseInt(maxChangeRef.current?.value?.trim() ?? "0");
+         let maxChange: number | undefined = parseInt(searchState.maxChangeList?.trim() ?? "0");
          if (!maxChange || isNaN(maxChange)) {
             maxChange = undefined;
          }
 
-         let name: string | undefined = nameRef.current?.value?.trim();
+         let name: string | undefined = searchState.name?.trim();
          if (!name) {
             name = undefined;
          }
 
-         let type: string | undefined;
-
-         if (typeRef.current?.selectedOptions?.length) {
+         let type: string | undefined = searchState.typeKey?.trim();
 
-            type = (typeRef.current.selectedOptions[0].key as string)?.trim();
-            const text = (typeRef.current.selectedOptions[0].text as string)?.trim();
-            if (type === "step-all" || !text) {
-               type = undefined;
-            } else {
-               const existing = artifactTypes.find(t => t.key === type);
-               if (!existing) {
-                  artifactTypes.push({ key: text, text: type });
-               }
+         if (type === "step-all") {
+            type = undefined;
+         } else if (type) {
+            const existing = artifactTypes.find(t => t.key === type);
+            if (!existing) {
+               artifactTypes.push({ key: type, text: type });
             }
          }
 
 
-         let streamId = streamRef?.current?.streamId?.trim();
-         if (!streamId) {
+         let streamId = searchState?.streamId?.trim();
+         if (!streamId || streamId === streamIdAll) {
             streamId = undefined;
          }
-         
+
          const mongoId = /^[a-fA-F0-9]{24}$/i;
 
          let id: string | undefined;
@@ -193,13 +474,13 @@ export const FindArtifactsModal: React.FC<{ streamId?: string, onClose: () => vo
                }
             } catch (reason) {
                console.error(reason)
-            }            
+            }
 
          } else {
             const find = await backend.getArtifacts(streamId, minChange, maxChange, name, type);
             artifacts = find.artifacts;
          }
-         
+
 
          setState({
             searching: false, artifacts: artifacts
@@ -226,11 +507,30 @@ export const FindArtifactsModal: React.FC<{ streamId?: string, onClose: () => vo
 
    const filterTypes = [...artifactTypes].filter(t => !!t.text?.trim());
 
+   let typeText = "";
+   let key = searchState.typeKey ?? "step-all"
+   const option = artifactTypes.find(o => o.key === key);
+   if (option) {
+      typeText = option.text;
+   } else {
+      typeText = searchState.typeKey ?? "";
+   }
+
+   if (searchState.autoLoad) {
+      searchState.autoLoad = false;
+      setTimeout(() => {
+         queryArtifacts();
+      })
+   }
+
    return <Stack>
-      <Modal isOpen={true} isBlocking={true} topOffsetFixed={true} styles={{ main: { padding: 8, width: 1024, height: 820, hasBeenOpened: false, top: "80px", position: "absolute" } }} onDismiss={() => onClose()} className={hordeClasses.modal}>
+      <Modal isOpen={true} isBlocking={true} topOffsetFixed={true} styles={{ main: { padding: 8, width: 1200, height: 820, hasBeenOpened: false, top: "80px", position: "absolute" } }} onDismiss={() => {
+         onClose()
+      }} className={hordeClasses.modal}>
          {!!state.searching && <Searching />}
+         <ArtifactBrowser state={searchState} />
          <Stack className="horde-no-darktheme" styles={{ root: { paddingTop: 10, paddingRight: 12 } }}>
-            <Stack style={{ paddingLeft: 24, paddingRight: 24 }}>
+            <Stack style={{ paddingLeft: 24, paddingRight: 12 }}>
                <Stack tokens={{ childrenGap: 12 }} style={{ height: 800 }}>
                   <Stack horizontal verticalAlign="start">
                      <Stack style={{ paddingTop: 3 }}>
@@ -240,37 +540,107 @@ export const FindArtifactsModal: React.FC<{ streamId?: string, onClose: () => vo
                      <Stack horizontalAlign="end">
                         <IconButton
                            iconProps={{ iconName: 'Cancel' }}
-                           onClick={() => { onClose() }}
+                           onClick={() => {
+                              onClose()
+                           }}
                         />
                      </Stack>
                   </Stack>
                   <Stack horizontal verticalAlign="center" tokens={{ childrenGap: 20 }}>
-                     <Stack >
+                     <Stack key={`stream_chooser_${streamChooserId++}`}>
                         <Label>Stream</Label>
-                        <StreamChooser defaultStreamId={streamId} ref={streamRef} />
+                        <StreamChooser defaultStreamId={searchState.streamId} allowAll={true} onChange={(streamId) => {
+                           searchState.streamId = streamId;
+                           setState({ ...state });
+                        }} />
                      </Stack>
                      <Stack >
-                        <TextField key="min_change_option" componentRef={minChangeRef} style={{ width: 92 }} label="Min Changelist" />
+                        <TextField key="min_change_option" defaultValue={searchState.minChangeList} value={searchState.minChangeList} autoComplete="off" spellCheck={false} style={{ width: 92 }} label="Min Changelist" onChange={(ev, newValue) => {
+                           let change: number | undefined;
+                           newValue = newValue ?? "";
+                           change = parseInt(newValue);
+                           if (isNaN(change)) {
+                              change = undefined;
+                           }
+
+                           if (typeof (change) === "number") {
+                              searchState.minChangeList = change.toString();
+                           } else {
+                              searchState.minChangeList = undefined;
+                           }
+
+                           setState({ ...state });
+                        }} />
                      </Stack>
                      <Stack >
-                        <TextField key="max_change_option" componentRef={maxChangeRef} style={{ width: 92 }} label="Max Changelist" />
+                        <TextField key="max_change_option" defaultValue={searchState.maxChangeList} value={searchState.maxChangeList} autoComplete="off" spellCheck={false} style={{ width: 92 }} label="Max Changelist" onChange={(ev, newValue) => {
+                           let change: number | undefined;
+                           newValue = newValue ?? "";
+                           change = parseInt(newValue);
+                           if (isNaN(change)) {
+                              change = undefined;
+                           }
+
+                           if (typeof (change) === "number") {
+                              searchState.maxChangeList = change.toString();
+                           } else {
+                              searchState.maxChangeList = undefined;
+                           }
+                           setState({ ...state });
+
+                        }} />
                      </Stack>
                      <Stack >
-                        <TextField key="name_option" componentRef={nameRef} style={{ width: 220 }} label="Name / Artifact Id" spellCheck={false} autoComplete="off"/>
+                        <TextField key="name_option" defaultValue={searchState.name} value={searchState.name} style={{ width: 232 }} label="Name / Artifact Id" spellCheck={false} autoComplete="off" onChange={(ev, newValue) => {
+                           searchState.name = newValue;
+                           setState({ ...state });
+                        }} />
                      </Stack>
                      <Stack>
                         <Label>Artifact Type</Label>
-                        <ComboBox key="type_option" componentRef={typeRef} allowFreeform={true} autoComplete="off" spellCheck={false} style={{ width: 144, textAlign: "left" }} defaultSelectedKey="step-all" options={filterTypes} calloutProps={{ doNotLayer: true }} />
+                        <ComboBox key="type_option" allowFreeform={true} autoComplete="off" text={typeText} spellCheck={false} style={{ width: 144, textAlign: "left" }} selectedKey={searchState.typeKey ?? "step-all"} options={filterTypes} calloutProps={{ doNotLayer: true }} onChange={(event: React.FormEvent<IComboBox>, option?: IComboBoxOption, index?: number, value?: string) => {
+                           if (option) {
+                              searchState.typeKey = option.key as string;
+                           } else if (value) {
+
+                              searchState.typeKey = value;
+
+                           }
+
+                           setState({ ...state });
+
+
+                        }} />
+                     </Stack>
+                     <Stack>
+                        <Label>Sort By</Label>
+                        <ComboBox key="sort_option" style={{ width: 144, textAlign: "left" }} selectedKey={searchState.sort ?? "sort-name"} options={sortOptions} calloutProps={{ doNotLayer: true }} onChange={(event: React.FormEvent<IComboBox>, option?: IComboBoxOption, index?: number, value?: string) => {
+                           if (option) {
+                              searchState.sort = option.key as string;
+                              setState({ ...state })
+                           }
+                        }} />
                      </Stack>
                   </Stack>
-                  <Stack horizontal>
+                  <Stack horizontal style={{ paddingTop: 12, paddingRight: 8 }}>
                      <Stack grow />
-                     <Stack>
-                        <PrimaryButton disabled={!!state.searching} text="Find" onClick={() => (queryArtifacts())} />
+                     <Stack horizontal tokens={{ childrenGap: 24 }}>
+                        <Stack>
+                           <DefaultButton disabled={!!state.searching} text="Reset" onClick={() => {
+                              searchState.reset(navigate);
+                              setState({ ...state })
+                           }} />
+                        </Stack>
+                        <Stack>
+                           <PrimaryButton disabled={!!state.searching} text="Find" onClick={() => {
+                              searchState.updateSearch(navigate);
+                              queryArtifacts()
+                           }} />
+                        </Stack>
                      </Stack>
                   </Stack>
-                  <Stack styles={{ root: { paddingTop: 12 } }}>
-                     <ArtifactsList artifacts={state.artifacts} />
+                  <Stack key={`artifact_list_${streamChooserId++}`} styles={{ root: { paddingTop: 12 } }}>
+                     <ArtifactsList state={searchState} artifacts={state.artifacts} sortBy={searchState.sort} />
                   </Stack>
                </Stack>
             </Stack>
diff --git a/Engine/Source/Programs/Horde/HordeDashboard/src/components/projects/StreamChooser.tsx b/Engine/Source/Programs/Horde/HordeDashboard/src/components/projects/StreamChooser.tsx
--- a/Engine/Source/Programs/Horde/HordeDashboard/src/components/projects/StreamChooser.tsx
+++ b/Engine/Source/Programs/Horde/HordeDashboard/src/components/projects/StreamChooser.tsx
@@ -10,16 +10,20 @@ export interface IStreamChooser {
 
 export interface IStreamChooserProps {
    defaultStreamId?: string;
+   allowAll?: boolean;
+   onChange?: (streamId: string | undefined) => void;
 
 }
 
+export const streamIdAll = "stream_id_all";
+
 export const StreamChooser = React.forwardRef<IStreamChooser, IStreamChooserProps>((props, ref) => {
 
-   const [chosen, setChosen] = useState(props?.defaultStreamId ?? "");
+   const [chosen, setChosen] = useState(props?.defaultStreamId ?? streamIdAll);
 
    useImperativeHandle(ref, () => ({
       streamId: chosen,
-    }));   
+   }));
 
    const projectIds = projectStore.projects.sort((a, b) => a.name.localeCompare(b.name)).map(p => p.id);
 
@@ -50,7 +54,10 @@ export const StreamChooser = React.forwardRef<IStreamChooser, IStreamChooserProp
             key: s.id,
             text: s.fullname ?? s.name,
             onClick: () => {
-               setChosen(s.id)
+               setChosen(s.id);
+               if (props.onChange) {
+                  props.onChange(s.id);
+               }
             }
          };
 
@@ -70,6 +77,19 @@ export const StreamChooser = React.forwardRef<IStreamChooser, IStreamChooserProp
 
    });
 
+   if (props.allowAll) {
+      projectOptions.unshift({
+         key: `project_key_all`,
+         text: "All Streams",
+         onClick: () => {
+            setChosen(streamIdAll);
+            if (props.onChange) {
+               props.onChange(streamIdAll);
+            }
+         }
+      })
+   }
+
    const templateMenuProps: IContextualMenuProps = {
       shouldFocusOnMount: true,
       subMenuHoverDelay: 0,
@@ -79,9 +99,13 @@ export const StreamChooser = React.forwardRef<IStreamChooser, IStreamChooserProp
    let text = "Choose Stream";
 
    if (chosen) {
-      const stream = projectStore.streamById(chosen);
-      if (stream) {
-         text = stream.fullname ?? stream.name;
+      if (chosen === streamIdAll) {
+         text = "All Streams";
+      } else {
+         const stream = projectStore.streamById(chosen);
+         if (stream) {
+            text = stream.fullname ?? stream.name;
+         }
       }
    }
 
