Horde: Fix registration of bundled tools in installed builds.

diff --git a/Engine/Source/Programs/Horde/Plugins/Tools/HordeServer.Tools/Commands/Generate/ToolDataCommand.cs b/Engine/Source/Programs/Horde/Plugins/Tools/HordeServer.Tools/Commands/Generate/ToolDataCommand.cs
--- a/Engine/Source/Programs/Horde/Plugins/Tools/HordeServer.Tools/Commands/Generate/ToolDataCommand.cs
+++ b/Engine/Source/Programs/Horde/Plugins/Tools/HordeServer.Tools/Commands/Generate/ToolDataCommand.cs
@@ -102,7 +102,7 @@ namespace HordeServer.Commands.Generate
 
 				JsonObject hordeConfig = JsonConfigFile.FindOrAddNode(serverConfig.Root, "Horde", () => new JsonObject());
 				JsonObject pluginsConfig = JsonConfigFile.FindOrAddNode(hordeConfig, "Plugins", () => new JsonObject());
-				JsonObject toolsConfig = JsonConfigFile.FindOrAddNode(serverConfig.Root, "Tools", () => new JsonObject());
+				JsonObject toolsConfig = JsonConfigFile.FindOrAddNode(pluginsConfig, "Tools", () => new JsonObject());
 				JsonArray bundledTools = JsonConfigFile.FindOrAddNode(toolsConfig, nameof(ToolsServerConfig.BundledTools), () => new JsonArray());
 
 				JsonObject bundledTool = JsonConfigFile.FindOrAddElementByKey(bundledTools, nameof(BundledToolConfig.Id), Id);
