Horde: State redraw fixes for artifact search
#skipci

diff --git a/Engine/Source/Programs/Horde/HordeDashboard/src/components/artifacts/ArtifactsSearch.tsx b/Engine/Source/Programs/Horde/HordeDashboard/src/components/artifacts/ArtifactsSearch.tsx
--- a/Engine/Source/Programs/Horde/HordeDashboard/src/components/artifacts/ArtifactsSearch.tsx
+++ b/Engine/Source/Programs/Horde/HordeDashboard/src/components/artifacts/ArtifactsSearch.tsx
@@ -482,12 +482,14 @@ export const FindArtifactsModal: React.FC<{ onClose: () => void }> = ({ onClose
          }
 
 
+         streamChooserId++;
          setState({
             searching: false, artifacts: artifacts
          });
 
       } catch (reason) {
          console.error(reason);
+         streamChooserId++;
          setState({ searching: false });
       }
    }
@@ -523,7 +525,7 @@ export const FindArtifactsModal: React.FC<{ onClose: () => void }> = ({ onClose
       })
    }
 
-   return <Stack>
+   return <Stack key="artifact_search_modal">
       <Modal isOpen={true} isBlocking={true} topOffsetFixed={true} styles={{ main: { padding: 8, width: 1200, height: 820, hasBeenOpened: false, top: "80px", position: "absolute" } }} onDismiss={() => {
          onClose()
       }} className={hordeClasses.modal}>
@@ -547,7 +549,7 @@ export const FindArtifactsModal: React.FC<{ onClose: () => void }> = ({ onClose
                      </Stack>
                   </Stack>
                   <Stack horizontal verticalAlign="center" tokens={{ childrenGap: 20 }}>
-                     <Stack key={`stream_chooser_${streamChooserId++}`}>
+                     <Stack key={`stream_chooser_${streamChooserId}`}>
                         <Label>Stream</Label>
                         <StreamChooser defaultStreamId={searchState.streamId} allowAll={true} onChange={(streamId) => {
                            searchState.streamId = streamId;
@@ -555,7 +557,7 @@ export const FindArtifactsModal: React.FC<{ onClose: () => void }> = ({ onClose
                         }} />
                      </Stack>
                      <Stack >
-                        <TextField key="min_change_option" defaultValue={searchState.minChangeList} value={searchState.minChangeList} autoComplete="off" spellCheck={false} style={{ width: 92 }} label="Min Changelist" onChange={(ev, newValue) => {
+                        <TextField key={`min_change_option_${streamChooserId}`} defaultValue={searchState.minChangeList} value={searchState.minChangeList} autoComplete="off" spellCheck={false} style={{ width: 92 }} label="Min Changelist" onChange={(ev, newValue) => {
                            let change: number | undefined;
                            newValue = newValue ?? "";
                            change = parseInt(newValue);
@@ -573,7 +575,7 @@ export const FindArtifactsModal: React.FC<{ onClose: () => void }> = ({ onClose
                         }} />
                      </Stack>
                      <Stack >
-                        <TextField key="max_change_option" defaultValue={searchState.maxChangeList} value={searchState.maxChangeList} autoComplete="off" spellCheck={false} style={{ width: 92 }} label="Max Changelist" onChange={(ev, newValue) => {
+                        <TextField key={`max_change_option_${streamChooserId}`} defaultValue={searchState.maxChangeList} value={searchState.maxChangeList} autoComplete="off" spellCheck={false} style={{ width: 92 }} label="Max Changelist" onChange={(ev, newValue) => {
                            let change: number | undefined;
                            newValue = newValue ?? "";
                            change = parseInt(newValue);
@@ -591,14 +593,14 @@ export const FindArtifactsModal: React.FC<{ onClose: () => void }> = ({ onClose
                         }} />
                      </Stack>
                      <Stack >
-                        <TextField key="name_option" defaultValue={searchState.name} value={searchState.name} style={{ width: 232 }} label="Name / Artifact Id" spellCheck={false} autoComplete="off" onChange={(ev, newValue) => {
+                        <TextField key={`name_option_${streamChooserId}`} defaultValue={searchState.name} value={searchState.name} style={{ width: 232 }} label="Name / Artifact Id" spellCheck={false} autoComplete="off" onChange={(ev, newValue) => {
                            searchState.name = newValue;
                            setState({ ...state });
                         }} />
                      </Stack>
                      <Stack>
                         <Label>Artifact Type</Label>
-                        <ComboBox key="type_option" allowFreeform={true} autoComplete="off" text={typeText} spellCheck={false} style={{ width: 144, textAlign: "left" }} selectedKey={searchState.typeKey ?? "step-all"} options={filterTypes} calloutProps={{ doNotLayer: true }} onChange={(event: React.FormEvent<IComboBox>, option?: IComboBoxOption, index?: number, value?: string) => {
+                        <ComboBox key={`type_option_${streamChooserId}`} allowFreeform={true} autoComplete="off" text={typeText} spellCheck={false} style={{ width: 144, textAlign: "left" }} selectedKey={searchState.typeKey ?? "step-all"} options={filterTypes} calloutProps={{ doNotLayer: true }} onChange={(event: React.FormEvent<IComboBox>, option?: IComboBoxOption, index?: number, value?: string) => {
                            if (option) {
                               searchState.typeKey = option.key as string;
                            } else if (value) {
@@ -614,7 +616,7 @@ export const FindArtifactsModal: React.FC<{ onClose: () => void }> = ({ onClose
                      </Stack>
                      <Stack>
                         <Label>Sort By</Label>
-                        <ComboBox key="sort_option" style={{ width: 144, textAlign: "left" }} selectedKey={searchState.sort ?? "sort-name"} options={sortOptions} calloutProps={{ doNotLayer: true }} onChange={(event: React.FormEvent<IComboBox>, option?: IComboBoxOption, index?: number, value?: string) => {
+                        <ComboBox key={`sort_option_${streamChooserId}`} style={{ width: 144, textAlign: "left" }} selectedKey={searchState.sort ?? "sort-name"} options={sortOptions} calloutProps={{ doNotLayer: true }} onChange={(event: React.FormEvent<IComboBox>, option?: IComboBoxOption, index?: number, value?: string) => {
                            if (option) {
                               searchState.sort = option.key as string;
                               setState({ ...state })
@@ -627,6 +629,7 @@ export const FindArtifactsModal: React.FC<{ onClose: () => void }> = ({ onClose
                      <Stack horizontal tokens={{ childrenGap: 24 }}>
                         <Stack>
                            <DefaultButton disabled={!!state.searching} text="Reset" onClick={() => {
+                              streamChooserId++;
                               searchState.reset(navigate);
                               setState({ ...state })
                            }} />
@@ -639,7 +642,7 @@ export const FindArtifactsModal: React.FC<{ onClose: () => void }> = ({ onClose
                         </Stack>
                      </Stack>
                   </Stack>
-                  <Stack key={`artifact_list_${streamChooserId++}`} styles={{ root: { paddingTop: 12 } }}>
+                  <Stack key={`artifact_list_${streamChooserId}`} styles={{ root: { paddingTop: 12 } }}>
                      <ArtifactsList state={searchState} artifacts={state.artifacts} sortBy={searchState.sort} />
                   </Stack>
                </Stack>
