Horde: Disable dedupe of bundled tool data for the moment; the controller currently filters blob requests by the tool id, which breaks when blobs are read from another tool.

diff --git a/Engine/Source/Programs/Horde/Plugins/Tools/HordeServer.Tools/Commands/Generate/ToolDataCommand.cs b/Engine/Source/Programs/Horde/Plugins/Tools/HordeServer.Tools/Commands/Generate/ToolDataCommand.cs
--- a/Engine/Source/Programs/Horde/Plugins/Tools/HordeServer.Tools/Commands/Generate/ToolDataCommand.cs
+++ b/Engine/Source/Programs/Horde/Plugins/Tools/HordeServer.Tools/Commands/Generate/ToolDataCommand.cs
@@ -75,20 +75,14 @@ namespace HordeServer.Commands.Generate
 				using (IStorageClient client = BundleStorageClient.CreateFromDirectory(bundleDir, bundleCache, logger))
 				{
 					IHashedBlobRef<DirectoryNode> dirNodeRef;
-					await using (DedupeBlobWriter writer = client.CreateDedupeBlobWriter(refName))
+					await using (IBlobWriter writer = client.CreateBlobWriter(refName))
 					{
-						logger.LogInformation("Populating cache with existing refs...");
-						await PopulateCacheAsync(client, writer, bundleDir, CancellationToken.None);
-
 						logger.LogInformation("");
 						logger.LogInformation("Writing tool data for {ToolId}", refName);
 						DirectoryNode dirNode = new DirectoryNode();
 						await dirNode.AddFilesAsync(InputDir.ToDirectoryInfo(), writer);
 						dirNodeRef = await writer.WriteBlobAsync(dirNode);
 
-						logger.LogInformation("");
-						writer.GetStats().Print(logger);
-
 						logger.LogInformation("");
 					}
 					await client.WriteRefAsync(refName, dirNodeRef);
diff --git a/Engine/Source/Programs/Horde/Plugins/Tools/HordeServer.Tools/Tools/ToolCollection.cs b/Engine/Source/Programs/Horde/Plugins/Tools/HordeServer.Tools/Tools/ToolCollection.cs
--- a/Engine/Source/Programs/Horde/Plugins/Tools/HordeServer.Tools/Tools/ToolCollection.cs
+++ b/Engine/Source/Programs/Horde/Plugins/Tools/HordeServer.Tools/Tools/ToolCollection.cs
@@ -461,7 +461,7 @@ namespace HordeServer.Tools
 		{
 			if (toolConfig is BundledToolConfig bundledConfig)
 			{
-				return BundleStorageClient.CreateFromDirectory(DirectoryReference.Combine(_serverInfo.AppDir, bundledConfig.DataDir ?? $"Tools"), _cache, _logger);
+				return BundleStorageClient.CreateFromDirectory(DirectoryReference.Combine(_serverInfo.AppDir, bundledConfig.DataDir ?? "Tools"), _cache, _logger);
 			}
 			else
 			{
@@ -478,7 +478,7 @@ namespace HordeServer.Tools
 		{
 			if (toolConfig is BundledToolConfig bundledConfig)
 			{
-				return new FileStorageBackend(_fileObjectStoreFactory.CreateStore(DirectoryReference.Combine(_serverInfo.AppDir, bundledConfig.DataDir ?? $"tools/{toolConfig.Id}")), _logger);
+				return new FileStorageBackend(_fileObjectStoreFactory.CreateStore(DirectoryReference.Combine(_serverInfo.AppDir, bundledConfig.DataDir ?? "Tools")), _logger);
 			}
 			else
 			{
