Horde: Change default pool size strategy to JobQueue

#rnx

diff --git a/Engine/Source/Programs/Horde/Plugins/Compute/HordeServer.Compute/StaticComputeConfig.cs b/Engine/Source/Programs/Horde/Plugins/Compute/HordeServer.Compute/StaticComputeConfig.cs
--- a/Engine/Source/Programs/Horde/Plugins/Compute/HordeServer.Compute/StaticComputeConfig.cs
+++ b/Engine/Source/Programs/Horde/Plugins/Compute/HordeServer.Compute/StaticComputeConfig.cs
@@ -47,7 +47,7 @@ namespace HordeServer
 		/// <summary>
 		/// Default agent pool sizing strategy for pools that doesn't have one explicitly configured
 		/// </summary>
-		public PoolSizeStrategy DefaultAgentPoolSizeStrategy { get; set; } = PoolSizeStrategy.LeaseUtilization;
+		public PoolSizeStrategy DefaultAgentPoolSizeStrategy { get; set; } = PoolSizeStrategy.JobQueue;
 
 		/// <summary>
 		/// Scale-out cooldown for auto-scaling agent pools (in seconds). Can be overridden by per-pool settings.
