Horde: Skip expiration of artifacts if no retention policy is specified. Previous behavior was to purge all that weren't explicitly kept.

diff --git a/Engine/Source/Programs/Horde/Plugins/Build/HordeServer.Build/Artifacts/ArtifactCollection.cs b/Engine/Source/Programs/Horde/Plugins/Build/HordeServer.Build/Artifacts/ArtifactCollection.cs
--- a/Engine/Source/Programs/Horde/Plugins/Build/HordeServer.Build/Artifacts/ArtifactCollection.cs
+++ b/Engine/Source/Programs/Horde/Plugins/Build/HordeServer.Build/Artifacts/ArtifactCollection.cs
@@ -321,6 +321,12 @@ namespace HordeServer.Artifacts
 
 		async Task ExpireArtifactsForStreamAsync(StreamId streamId, ArtifactTypeConfig artifactTypeConfig, DateTime utcNow, CancellationToken cancellationToken)
 		{
+			if (artifactTypeConfig.KeepCount == null && artifactTypeConfig.KeepDays == null)
+			{
+				_logger.LogInformation("No expiration policy set for {StreamId} {ArtifactType}; keeping all artifacts.", streamId, artifactTypeConfig.Type);
+				return;
+			}
+
 			if (artifactTypeConfig.KeepCount.HasValue)
 			{
 				_logger.LogInformation("Removing {StreamId} {ArtifactType} artifacts except newest {Count}", streamId, artifactTypeConfig.Type, artifactTypeConfig.KeepCount.Value);
