Horde: Updating schema documentation

#rnx

diff --git a/Engine/Source/Programs/Horde/Docs/Config/Schema/Projects.md b/Engine/Source/Programs/Horde/Docs/Config/Schema/Projects.md
--- a/Engine/Source/Programs/Horde/Docs/Config/Schema/Projects.md
+++ b/Engine/Source/Programs/Horde/Docs/Config/Schema/Projects.md
@@ -20,7 +20,7 @@ Name | Description
 `workspaceTypes` | `string` `->` [WorkspaceConfig](#workspaceconfig)<br>Default workspace types for streams These are added to the list of each stream's workspace types.
 `telemetryStoreId` | `string`<br>Telemetry store for Horde data for this project
 `streams` | [StreamConfig](Streams.md)`[]`<br>List of streams
-`artifactTypes` | [ArtifactTypeAclConfig](#artifacttypeaclconfig)`[]`<br>Permissions for artifact types
+`artifactTypes` | [ArtifactTypeConfig](#artifacttypeconfig)`[]`<br>Permissions for artifact types
 `acl` | [AclConfig](#aclconfig)<br>Acl entries
 
 ## ConfigInclude
@@ -214,14 +214,17 @@ Name | Description
 `minScratchSpace` | `integer`<br>Minimum disk space that must be available *after* syncing this workspace (in megabytes) If not available, the job will be aborted.
 `conformDiskFreeSpace` | `integer`<br>Threshold for when to trigger an automatic conform of agent. Measured in megabytes free on disk. Set to null or 0 to disable.
 
-## ArtifactTypeAclConfig
+## ArtifactTypeConfig
 
-ACL configuration for an artifact type
+Configuration for an artifact
 
 Name | Description
 ---- | -----------
+`name` | `string`<br>Legacy 'Name' property
 `type` | `string`<br>Name of the artifact type
 `acl` | [AclConfig](#aclconfig)<br>Acl for the artifact type
+`keepCount` | `integer`<br>Number of artifacts to retain
+`keepDays` | `integer`<br>Number of days to retain artifacts of this type
 
 ## AclConfig
 
diff --git a/Engine/Source/Programs/Horde/Docs/Config/Schema/Streams.md b/Engine/Source/Programs/Horde/Docs/Config/Schema/Streams.md
--- a/Engine/Source/Programs/Horde/Docs/Config/Schema/Streams.md
+++ b/Engine/Source/Programs/Horde/Docs/Config/Schema/Streams.md
@@ -34,7 +34,7 @@ Name | Description
 `replicators` | [ReplicatorConfig](#replicatorconfig)`[]`<br>Configuration for workers to replicate commit data into Horde Storage.
 `workflows` | [WorkflowConfig](#workflowconfig)`[]`<br>Workflows for dealing with new issues
 `tokens` | [TokenConfig](#tokenconfig)`[]`<br>Tokens to create for each job step
-`artifactTypes` | [ArtifactTypeAclConfig](#artifacttypeaclconfig)`[]`<br>Permissions for artifact types
+`artifactTypes` | [ArtifactTypeConfig](#artifacttypeconfig)`[]`<br>Permissions for artifact types
 
 ## ConfigInclude
 
@@ -488,11 +488,14 @@ Name | Description
 `defaultComponentId` | `string`<br>Default component id for issues using workflow
 `defaultIssueTypeId` | `string`<br>Default issue type id for issues using workflow
 
-## ArtifactTypeAclConfig
+## ArtifactTypeConfig
 
-ACL configuration for an artifact type
+Configuration for an artifact
 
 Name | Description
 ---- | -----------
+`name` | `string`<br>Legacy 'Name' property
 `type` | `string`<br>Name of the artifact type
 `acl` | [AclConfig](#aclconfig)<br>Acl for the artifact type
+`keepCount` | `integer`<br>Number of artifacts to retain
+`keepDays` | `integer`<br>Number of days to retain artifacts of this type
