Horde: Improve logging of port mappings for relay

#rnx

diff --git a/Engine/Source/Programs/Horde/Plugins/Build/HordeServer.Build.Tests/Agents/Relay/AgentRelayServiceTests.cs b/Engine/Source/Programs/Horde/Plugins/Build/HordeServer.Build.Tests/Agents/Relay/AgentRelayServiceTests.cs
--- a/Engine/Source/Programs/Horde/Plugins/Build/HordeServer.Build.Tests/Agents/Relay/AgentRelayServiceTests.cs
+++ b/Engine/Source/Programs/Horde/Plugins/Build/HordeServer.Build.Tests/Agents/Relay/AgentRelayServiceTests.cs
@@ -6,7 +6,6 @@ using EpicGames.Horde.Compute;
 using Grpc.Core;
 using Horde.Common.Rpc;
 using HordeServer.Agents.Relay;
-//using HordeServer.Tests.Server;
 using HordeServer.Utilities;
 using Microsoft.Extensions.Logging.Abstractions;
 
@@ -51,7 +50,7 @@ public class AgentRelayServiceTests : BuildTestSetup
 
 	public AgentRelayServiceTests()
 	{
-		_service = new AgentRelayService(GetRedisServiceSingleton(), Clock, NullLogger<AgentRelayService>.Instance);
+		_service = new AgentRelayService(GetRedisServiceSingleton(), Clock, Tracer, NullLogger<AgentRelayService>.Instance);
 	}
 
 	public override async ValueTask DisposeAsync()
diff --git a/Engine/Source/Programs/Horde/Plugins/Compute/HordeServer.Compute/Agents/Relay/AgentRelayService.cs b/Engine/Source/Programs/Horde/Plugins/Compute/HordeServer.Compute/Agents/Relay/AgentRelayService.cs
--- a/Engine/Source/Programs/Horde/Plugins/Compute/HordeServer.Compute/Agents/Relay/AgentRelayService.cs
+++ b/Engine/Source/Programs/Horde/Plugins/Compute/HordeServer.Compute/Agents/Relay/AgentRelayService.cs
@@ -15,7 +15,10 @@ using HordeServer.Server;
 using HordeServer.Utilities;
 using Microsoft.Extensions.Hosting;
 using Microsoft.Extensions.Logging;
+using OpenTelemetry.Trace;
 using StackExchange.Redis;
+using Status = Grpc.Core.Status;
+using StatusCode = Grpc.Core.StatusCode;
 
 namespace HordeServer.Agents.Relay;
 
@@ -78,6 +81,7 @@ public sealed class AgentRelayService : RelayRpc.RelayRpcBase, IHostedService, I
 
 	private readonly IRedisService _redis;
 	private readonly IClock _clock;
+	private readonly Tracer _tracer;
 	private readonly ILogger<AgentRelayService> _logger;
 	private readonly object _lock = new();
 	private TaskCompletionSource _onPortMappingUpdated = new();
@@ -94,11 +98,13 @@ public sealed class AgentRelayService : RelayRpc.RelayRpcBase, IHostedService, I
 	/// </summary>
 	/// <param name="redis"></param>
 	/// <param name="clock"></param>
+	/// <param name="tracer"></param>
 	/// <param name="logger"></param>
-	public AgentRelayService(IRedisService redis, IClock clock, ILogger<AgentRelayService> logger)
+	public AgentRelayService(IRedisService redis, IClock clock, Tracer tracer, ILogger<AgentRelayService> logger)
 	{
 		_redis = redis;
 		_clock = clock;
+		_tracer = tracer;
 		_logger = logger;
 		_updateTaskQueue = new AsyncTaskQueue(logger);
 		_agentExpirationTimeout = _longPollTimeout + TimeSpan.FromSeconds(5);
@@ -186,6 +192,9 @@ public sealed class AgentRelayService : RelayRpc.RelayRpcBase, IHostedService, I
 	/// <exception cref="Exception"></exception>
 	public async Task<PortMapping> AddPortMappingAsync(ClusterId clusterId, LeaseId leaseId, IPAddress? clientIp, IPAddress agentIp, IList<Port> ports, int numRetries = 10)
 	{
+		using TelemetrySpan span = _tracer.StartActiveSpan($"{nameof(AgentRelayService)}.{nameof(RemovePortMappingAsync)}");
+		span.SetAttribute("clusterId", clusterId.ToString());
+		span.SetAttribute("leaseId", leaseId.ToString());
 		IDatabase redis = _redis.GetDatabase();
 
 		// TODO: Randomize start position to spread out port use
@@ -228,13 +237,17 @@ public sealed class AgentRelayService : RelayRpc.RelayRpcBase, IHostedService, I
 			_ = transaction.StringIncrementAsync(KeyPortMappingRevision(clusterId));
 			_ = transaction.HashSetAsync(KeyClusters(), clusterId, _clock.UtcNow.ToFileTimeUtc());
 			bool isSuccessful = await transaction.ExecuteAsync();
+			
+			AddPortMappingMetadataToSpan(span, newPortMapping);
 			if (isSuccessful)
 			{
 				await PublishUpdateEventAsync(clusterId.ToString());
+				span.SetAttribute("added", true);
 				return newPortMapping;
 			}
 		}
-
+		
+		span.SetAttribute("added", false);
 		throw new Exception("Unable to find an available port range");
 	}
 
@@ -246,18 +259,26 @@ public sealed class AgentRelayService : RelayRpc.RelayRpcBase, IHostedService, I
 	/// <returns>True if successful</returns>
 	public async Task<bool> RemovePortMappingAsync(ClusterId clusterId, LeaseId leaseId)
 	{
+		using TelemetrySpan span = _tracer.StartActiveSpan($"{nameof(AgentRelayService)}.{nameof(RemovePortMappingAsync)}");
+		span.SetAttribute("clusterId", clusterId.ToString());
+		span.SetAttribute("leaseId", leaseId.ToString());
+		
 		PortMapping? portMapping = await GetPortMappingAsync(clusterId, leaseId);
 		if (portMapping == null)
 		{
+			span.SetAttribute("removed", false);
 			return false;
 		}
+		
+		AddPortMappingMetadataToSpan(span, portMapping);
 
 		RedisValue[] ports = portMapping.Ports.Select(x => new RedisValue(Convert.ToString(x.RelayPort))).ToArray();
 		ITransaction transaction = _redis.GetDatabase().CreateTransaction();
 		_ = transaction.SetRemoveAsync(KeyUsedPorts(clusterId), ports);
 		_ = transaction.HashDeleteAsync(KeyPortMappings(clusterId), leaseId.ToString());
 		_ = transaction.StringIncrementAsync(KeyPortMappingRevision(clusterId));
-
+		
+		span.SetAttribute("removed", true);
 		return await transaction.ExecuteAsync();
 	}
 
@@ -394,6 +415,17 @@ public sealed class AgentRelayService : RelayRpc.RelayRpcBase, IHostedService, I
 		_longPollTimeout = TimeSpan.FromMilliseconds(longPollMs);
 		_agentExpirationTimeout = TimeSpan.FromMilliseconds(expirationMs);
 	}
+	
+	private static void AddPortMappingMetadataToSpan(TelemetrySpan span, PortMapping portMapping)
+	{
+		span.SetAttribute("leaseId", portMapping.LeaseId);
+		span.SetAttribute("agentIp", portMapping.AgentIp);
+		span.SetAttribute("createdAt", portMapping.CreatedAt.ToDateTimeOffset().ToUnixTimeMilliseconds());
+		span.SetAttribute("relayPorts", portMapping.Ports.Select(x => x.RelayPort).ToArray());
+		span.SetAttribute("agentPorts", portMapping.Ports.Select(x => x.AgentPort).ToArray());
+		span.SetAttribute("protocols", portMapping.Ports.Select(x => x.Protocol.ToString()).ToArray());
+		span.SetAttribute("allowedSourceIps", portMapping.AllowedSourceIps.ToArray());
+	}
 
 	/// <inheritdoc/>
 	public override async Task GetPortMappings(GetPortMappingsRequest request, IServerStreamWriter<GetPortMappingsResponse> responseStream, ServerCallContext context)
diff --git a/Engine/Source/Programs/Horde/Plugins/Compute/HordeServer.Compute/Compute/ComputeService.cs b/Engine/Source/Programs/Horde/Plugins/Compute/HordeServer.Compute/Compute/ComputeService.cs
--- a/Engine/Source/Programs/Horde/Plugins/Compute/HordeServer.Compute/Compute/ComputeService.cs
+++ b/Engine/Source/Programs/Horde/Plugins/Compute/HordeServer.Compute/Compute/ComputeService.cs
@@ -760,6 +760,7 @@ namespace HordeServer.Compute
 						if (agentPort == pmPort.AgentPort)
 						{
 							ports[portId] = new ComputeResourcePort(pmPort.RelayPort, pmPort.AgentPort);
+							_logger.LogInformation("Port {PortId} mapped {RelayIp}:{RelayPort} -> {AgentPort} for lease ID {LeaseId}", portId, relayIp, pmPort.RelayPort, pmPort.AgentPort, leaseId);
 						}
 					}
 				}
