Horde: Remove IDisposable interface from IStorageClient. We want to be able to pass IBlobRef objects around without having to know anything about where they're stored.

#rnx

diff --git a/Engine/Source/Programs/AutomationTool/BuildGraph/Tasks/CreateArtifactTask.cs b/Engine/Source/Programs/AutomationTool/BuildGraph/Tasks/CreateArtifactTask.cs
--- a/Engine/Source/Programs/AutomationTool/BuildGraph/Tasks/CreateArtifactTask.cs
+++ b/Engine/Source/Programs/AutomationTool/BuildGraph/Tasks/CreateArtifactTask.cs
@@ -127,21 +127,20 @@ namespace AutomationTool.Tasks
 			Stopwatch timer = Stopwatch.StartNew();
 
 			HttpStorageClientFactory httpStorageClientFactory = serviceProvider.GetRequiredService<HttpStorageClientFactory>();
-			using (IStorageClient client = httpStorageClientFactory.CreateClient(response.NamespaceId, response.Token))
+			IStorageClient client = httpStorageClientFactory.CreateClient(response.NamespaceId, response.Token);
+
+			await using (IBlobWriter writer = client.CreateBlobWriter(response.RefName))
 			{
-				await using (IBlobWriter writer = client.CreateBlobWriter(response.RefName))
-				{
-					DirectoryReference baseDir = ResolveDirectory(_parameters.BaseDir);
-					List<FileInfo> files = ResolveFilespec(baseDir, _parameters.Files, tagNameToFileSet).Select(x => x.ToFileInfo()).ToList();
+				DirectoryReference baseDir = ResolveDirectory(_parameters.BaseDir);
+				List<FileInfo> files = ResolveFilespec(baseDir, _parameters.Files, tagNameToFileSet).Select(x => x.ToFileInfo()).ToList();
 
-					int totalCount = files.Count;
-					long totalSize = files.Sum(x => x.Length);
+				int totalCount = files.Count;
+				long totalSize = files.Sum(x => x.Length);
 
-					IHashedBlobRef<DirectoryNode> outputNodeRef = await writer.WriteFilesAsync(baseDir.ToDirectoryInfo(), files, progress: new UpdateStatsLogger(totalCount, totalSize, Logger));
-					await writer.FlushAsync();
+				IHashedBlobRef<DirectoryNode> outputNodeRef = await writer.WriteFilesAsync(baseDir.ToDirectoryInfo(), files, progress: new UpdateStatsLogger(totalCount, totalSize, Logger));
+				await writer.FlushAsync();
 
-					await client.WriteRefAsync(response.RefName, outputNodeRef);
-				}
+				await client.WriteRefAsync(response.RefName, outputNodeRef);
 			}
 
 			Logger.LogInformation("Completed in {Time:n1}s", timer.Elapsed.TotalSeconds);
diff --git a/Engine/Source/Programs/AutomationTool/BuildGraph/Tasks/DeployToolTask.cs b/Engine/Source/Programs/AutomationTool/BuildGraph/Tasks/DeployToolTask.cs
--- a/Engine/Source/Programs/AutomationTool/BuildGraph/Tasks/DeployToolTask.cs
+++ b/Engine/Source/Programs/AutomationTool/BuildGraph/Tasks/DeployToolTask.cs
@@ -161,7 +161,7 @@ namespace AutomationTool.Tasks
 
 			IHashedBlobRef handle;
 
-			using IStorageClient storageClient = hordeClient.CreateStorageClient(toolId);
+			IStorageClient storageClient = hordeClient.CreateStorageClient(toolId);
 			await using (IBlobWriter blobWriter = storageClient.CreateBlobWriter(serializerOptions: serializerOptions))
 			{
 				DirectoryNode sandbox = new DirectoryNode();
diff --git a/Engine/Source/Programs/Horde/Drivers/JobDriver/Execution/JobExecutor.cs b/Engine/Source/Programs/Horde/Drivers/JobDriver/Execution/JobExecutor.cs
--- a/Engine/Source/Programs/Horde/Drivers/JobDriver/Execution/JobExecutor.cs
+++ b/Engine/Source/Programs/Horde/Drivers/JobDriver/Execution/JobExecutor.cs
@@ -822,7 +822,7 @@ namespace JobDriver.Execution
 				logger.LogInformation("Creating output artifact {ArtifactId} '{ArtifactName}' ({ArtifactType}) with ref {RefName} in namespace {NamespaceId}", artifactId, artifactName, artifactType, artifact.RefName, artifact.NamespaceId);
 
 				// Write the data
-				using IStorageClient storage = CreateStorageClient(new NamespaceId(artifact.NamespaceId), artifact.Token);
+				IStorageClient storage = CreateStorageClient(new NamespaceId(artifact.NamespaceId), artifact.Token);
 
 				Stopwatch timer = Stopwatch.StartNew();
 
@@ -1078,7 +1078,7 @@ namespace JobDriver.Execution
 
 				logger.LogInformation("Reading preprocessed script from {NamespaceId}:{RefName}", namespaceId, refName);
 
-				using IStorageClient storage = CreateStorageClient(namespaceId, artifact.Token);
+				IStorageClient storage = CreateStorageClient(namespaceId, artifact.Token);
 
 				DirectoryNode node = await storage.ReadRefTargetAsync<DirectoryNode>(refName, cancellationToken: cancellationToken);
 				DirectoryNode? buildGraphDir = await node.TryOpenDirectoryAsync(BuildGraphTempStorageDir, cancellationToken: cancellationToken);
@@ -1139,7 +1139,7 @@ namespace JobDriver.Execution
 				ArtifactId artifactId = ArtifactId.Parse(artifact.Id);
 				Logger.LogInformation("Created artifact {ArtifactId} '{ArtifactName}' ({ArtifactType}) with ref {RefName} ({Link})", artifactId, name, type, artifact.RefName, $"{HordeClient.ServerUrl}/api/v1/storage/{artifact.NamespaceId}/refs/{artifact.RefName}");
 
-				using IStorageClient storage = CreateStorageClient(new NamespaceId(artifact.NamespaceId), artifact.Token);
+				IStorageClient storage = CreateStorageClient(new NamespaceId(artifact.NamespaceId), artifact.Token);
 
 				IHashedBlobRef<DirectoryNode> rootRef;
 				await using (IBlobWriter blobWriter = storage.CreateBlobWriter(new RefName(artifact.RefName)))
@@ -1341,7 +1341,7 @@ namespace JobDriver.Execution
 				ArtifactId artifactId = ArtifactId.Parse(artifact.Id);
 				logger.LogInformation("Created artifact {ArtifactId} '{ArtifactName}' ({ArtifactType}) with ref {RefName} ({RefUrl})", artifactId, artifactRequest.Name, ArtifactType.StepOutput, artifact.RefName, $"{HordeClient.ServerUrl.ToString().TrimEnd('/')}/api/v1/storage/{artifact.NamespaceId}/refs/{artifact.RefName}");
 
-				using IStorageClient storage = CreateStorageClient(new NamespaceId(artifact.NamespaceId), artifact.Token);
+				IStorageClient storage = CreateStorageClient(new NamespaceId(artifact.NamespaceId), artifact.Token);
 
 				// Upload the data
 				Stopwatch timer = Stopwatch.StartNew();
@@ -1421,7 +1421,7 @@ namespace JobDriver.Execution
 					ArtifactId artifactId = ArtifactId.Parse(artifact.Id);
 					logger.LogInformation("Created artifact {ArtifactId} '{ArtifactName}' ({ArtifactType}) with ref {RefName} ({RefUrl})", artifactId, artifactRequest.Name, ArtifactType.StepOutput, artifact.RefName, $"{HordeClient.ServerUrl.ToString().TrimEnd('/')}/api/v1/storage/{artifact.NamespaceId}/refs/{artifact.RefName}");
 
-					using IStorageClient storage = CreateStorageClient(new NamespaceId(artifact.NamespaceId), artifact.Token);
+					IStorageClient storage = CreateStorageClient(new NamespaceId(artifact.NamespaceId), artifact.Token);
 
 					// Upload the data
 					Stopwatch timer = Stopwatch.StartNew();
diff --git a/Engine/Source/Programs/Horde/Drivers/JobDriver/Utility/TempStorage.cs b/Engine/Source/Programs/Horde/Drivers/JobDriver/Utility/TempStorage.cs
--- a/Engine/Source/Programs/Horde/Drivers/JobDriver/Utility/TempStorage.cs
+++ b/Engine/Source/Programs/Horde/Drivers/JobDriver/Utility/TempStorage.cs
@@ -632,7 +632,7 @@ namespace JobDriver.Utility
 
 				logger.LogInformation("Reading tag \"{NodeName}\":\"{TagName}\" from temp storage (artifact: {ArtifactId} '{ArtifactName}' ({ArtifactType}), ns: {NamespaceId}, ref: {RefName}, localFile: {LocalFile})", nodeName, tagName, artifact.Id, artifactName, artifactType, namespaceId, refName, localFileListLocation);
 
-				using IStorageClient storageClient = hordeClient.CreateStorageClient(namespaceId, artifact.Token);
+				IStorageClient storageClient = hordeClient.CreateStorageClient(namespaceId, artifact.Token);
 				DirectoryNode node = await storageClient.ReadRefTargetAsync<DirectoryNode>(artifact.RefName, cancellationToken: cancellationToken);
 
 				FileEntry fileEntry = node.GetFileEntry(localFileListLocation.GetFileName());
@@ -770,7 +770,7 @@ namespace JobDriver.Utility
 
 				logger.LogInformation("Reading block \"{NodeName}\":\"{BlockName}\" from temp storage (artifact: {ArtifactId} '{ArtifactName}' ({ArtifactType}), ns: {NamespaceId}, ref: {RefName}, local: {LocalFile}, blockdir: {BlockDir})", nodeName, blockName, artifact.Id, artifactName, artifactType, namespaceId, refName, localManifestFile, blockDirectoryName);
 
-				using IStorageClient storageClient = hordeClient.CreateStorageClient(namespaceId, artifact.Token);
+				IStorageClient storageClient = hordeClient.CreateStorageClient(namespaceId, artifact.Token);
 				DirectoryNode node = await storageClient.ReadRefTargetAsync<DirectoryNode>(refName, cancellationToken: cancellationToken);
 
 				DirectoryEntry? rootDirEntry;
diff --git a/Engine/Source/Programs/Horde/HordeCmd/Commands/Artifacts/ArtifactDownload.cs b/Engine/Source/Programs/Horde/HordeCmd/Commands/Artifacts/ArtifactDownload.cs
--- a/Engine/Source/Programs/Horde/HordeCmd/Commands/Artifacts/ArtifactDownload.cs
+++ b/Engine/Source/Programs/Horde/HordeCmd/Commands/Artifacts/ArtifactDownload.cs
@@ -50,7 +50,7 @@ namespace Horde.Commands.Artifacts
 				FileUtils.ForceDeleteDirectoryContents(OutputDir);
 			}
 
-			using IStorageClient store = _hordeClient.CreateStorageClient(artifact.Id);
+			IStorageClient store = _hordeClient.CreateStorageClient(artifact.Id);
 
 			Stopwatch timer = Stopwatch.StartNew();
 
diff --git a/Engine/Source/Programs/Horde/HordeCmd/Commands/Bundles/BundleCreate.cs b/Engine/Source/Programs/Horde/HordeCmd/Commands/Bundles/BundleCreate.cs
--- a/Engine/Source/Programs/Horde/HordeCmd/Commands/Bundles/BundleCreate.cs
+++ b/Engine/Source/Programs/Horde/HordeCmd/Commands/Bundles/BundleCreate.cs
@@ -49,12 +49,13 @@ namespace Horde.Commands.Bundles
 
 			if (File != null)
 			{
-				using IStorageClient store = BundleStorageClient.CreateFromDirectory(File.Directory, BundleCache, logger);
+				using MemoryMappedFileCache memoryMappedFileCache = new MemoryMappedFileCache();
+				IStorageClient store = BundleStorageClient.CreateFromDirectory(File.Directory, BundleCache, memoryMappedFileCache, logger);
 				return await ExecuteInternalAsync(store, logger);
 			}
 			else if (!String.IsNullOrEmpty(Ref))
 			{
-				using IStorageClient store = CreateStorageClient();
+				IStorageClient store = CreateStorageClient();
 				return await ExecuteInternalAsync(store, logger);
 			}
 			else
diff --git a/Engine/Source/Programs/Horde/HordeCmd/Commands/Bundles/BundleExtract.cs b/Engine/Source/Programs/Horde/HordeCmd/Commands/Bundles/BundleExtract.cs
--- a/Engine/Source/Programs/Horde/HordeCmd/Commands/Bundles/BundleExtract.cs
+++ b/Engine/Source/Programs/Horde/HordeCmd/Commands/Bundles/BundleExtract.cs
@@ -54,19 +54,20 @@ namespace Horde.Commands.Bundles
 
 			if (File != null)
 			{
-				using IStorageClient store = BundleStorageClient.CreateFromDirectory(File.Directory, BundleCache, logger);
+				using MemoryMappedFileCache memoryMappedFileCache = new MemoryMappedFileCache();
+				IStorageClient store = BundleStorageClient.CreateFromDirectory(File.Directory, BundleCache, memoryMappedFileCache, logger);
 				IBlobRef<DirectoryNode> handle = store.CreateBlobRef<DirectoryNode>(await FileStorageBackend.ReadRefAsync(File));
 				await ExecuteInternalAsync(store, handle, logger);
 			}
 			else if (Ref != null)
 			{
-				using IStorageClient store = CreateStorageClient();
+				IStorageClient store = CreateStorageClient();
 				IBlobRef<DirectoryNode> handle = await store.ReadRefAsync<DirectoryNode>(new RefName(Ref));
 				await ExecuteInternalAsync(store, handle, logger);
 			}
 			else if (Node != null)
 			{
-				using IStorageClient store = CreateStorageClient();
+				IStorageClient store = CreateStorageClient();
 				IBlobRef<DirectoryNode> handle = store.CreateBlobRef<DirectoryNode>(new BlobLocator(Node));
 				await ExecuteInternalAsync(store, handle, logger);
 			}
diff --git a/Engine/Source/Programs/Horde/HordeCmd/Commands/Compute/ComputeRun.cs b/Engine/Source/Programs/Horde/HordeCmd/Commands/Compute/ComputeRun.cs
--- a/Engine/Source/Programs/Horde/HordeCmd/Commands/Compute/ComputeRun.cs
+++ b/Engine/Source/Programs/Horde/HordeCmd/Commands/Compute/ComputeRun.cs
@@ -130,7 +130,7 @@ namespace Horde.Commands.Compute
 			JsonComputeTask task = JsonSerializer.Deserialize<JsonComputeTask>(data, new JsonSerializerOptions { AllowTrailingCommas = true, PropertyNameCaseInsensitive = true, PropertyNamingPolicy = JsonNamingPolicy.CamelCase })!;
 
 			// Create a sandbox from the data to be uploaded
-			using BundleStorageClient storage = BundleStorageClient.CreateInMemory(logger);
+			BundleStorageClient storage = BundleStorageClient.CreateInMemory(logger);
 			BlobLocator sandbox = await CreateSandboxAsync(TaskFile, storage, cancellationToken);
 
 			// Open a socket and upload the sandbox
diff --git a/Engine/Source/Programs/Horde/HordeCmd/Commands/Tools/ToolUpload.cs b/Engine/Source/Programs/Horde/HordeCmd/Commands/Tools/ToolUpload.cs
--- a/Engine/Source/Programs/Horde/HordeCmd/Commands/Tools/ToolUpload.cs
+++ b/Engine/Source/Programs/Horde/HordeCmd/Commands/Tools/ToolUpload.cs
@@ -38,7 +38,7 @@ namespace Horde.Commands
 		/// <inheritdoc/>
 		public override async Task<int> ExecuteAsync(ILogger logger)
 		{
-			using IStorageClient storageClient = _httpStorageClientFactory.CreateClientWithPath($"api/v1/tools/{ToolId}");
+			IStorageClient storageClient = _httpStorageClientFactory.CreateClientWithPath($"api/v1/tools/{ToolId}");
 
 			IHashedBlobRef<DirectoryNode> target;
 			await using (IBlobWriter writer = storageClient.CreateBlobWriter())
diff --git a/Engine/Source/Programs/Horde/HordeCmd/Commands/Vcs/VcsBranch.cs b/Engine/Source/Programs/Horde/HordeCmd/Commands/Vcs/VcsBranch.cs
--- a/Engine/Source/Programs/Horde/HordeCmd/Commands/Vcs/VcsBranch.cs
+++ b/Engine/Source/Programs/Horde/HordeCmd/Commands/Vcs/VcsBranch.cs
@@ -23,7 +23,7 @@ namespace Horde.Commands.Vcs
 
 			WorkspaceState workspaceState = await ReadStateAsync(rootDir);
 
-			using IStorageClient store = CreateStorageClient();
+			IStorageClient store = CreateStorageClient();
 
 			RefName branchName = new RefName(Name);
 			if (await store.RefExistsAsync(branchName))
diff --git a/Engine/Source/Programs/Horde/HordeCmd/Commands/Vcs/VcsCheckout.cs b/Engine/Source/Programs/Horde/HordeCmd/Commands/Vcs/VcsCheckout.cs
--- a/Engine/Source/Programs/Horde/HordeCmd/Commands/Vcs/VcsCheckout.cs
+++ b/Engine/Source/Programs/Horde/HordeCmd/Commands/Vcs/VcsCheckout.cs
@@ -51,7 +51,7 @@ namespace Horde.Commands.Vcs
 
 			RefName branchName = (Branch != null) ? new RefName(Branch) : workspaceState.Branch;
 
-			using IStorageClient store = CreateStorageClient();
+			IStorageClient store = CreateStorageClient();
 
 			CommitNode? tip = await GetCommitAsync(store, branchName, Change);
 			if (tip == null)
diff --git a/Engine/Source/Programs/Horde/HordeCmd/Commands/Vcs/VcsClone.cs b/Engine/Source/Programs/Horde/HordeCmd/Commands/Vcs/VcsClone.cs
--- a/Engine/Source/Programs/Horde/HordeCmd/Commands/Vcs/VcsClone.cs
+++ b/Engine/Source/Programs/Horde/HordeCmd/Commands/Vcs/VcsClone.cs
@@ -25,7 +25,7 @@ namespace Horde.Commands.Vcs
 
 			RefName branchName = new RefName(Branch ?? "ue5-main");
 
-			using IStorageClient storageClient = CreateStorageClient();
+			IStorageClient storageClient = CreateStorageClient();
 
 			CommitNode? tip = await GetCommitAsync(storageClient, branchName, Change);
 			if (tip == null)
diff --git a/Engine/Source/Programs/Horde/HordeCmd/Commands/Vcs/VcsLog.cs b/Engine/Source/Programs/Horde/HordeCmd/Commands/Vcs/VcsLog.cs
--- a/Engine/Source/Programs/Horde/HordeCmd/Commands/Vcs/VcsLog.cs
+++ b/Engine/Source/Programs/Horde/HordeCmd/Commands/Vcs/VcsLog.cs
@@ -25,7 +25,7 @@ namespace Horde.Commands.Vcs
 
 			WorkspaceState workspaceState = await ReadStateAsync(rootDir);
 
-			using IStorageClient store = CreateStorageClient();
+			IStorageClient store = CreateStorageClient();
 
 			using MemoryCache cache = new MemoryCache(new MemoryCacheOptions());
 
diff --git a/Engine/Source/Programs/Horde/HordeCmd/Commands/Workspace/WorkspaceSync.cs b/Engine/Source/Programs/Horde/HordeCmd/Commands/Workspace/WorkspaceSync.cs
--- a/Engine/Source/Programs/Horde/HordeCmd/Commands/Workspace/WorkspaceSync.cs
+++ b/Engine/Source/Programs/Horde/HordeCmd/Commands/Workspace/WorkspaceSync.cs
@@ -50,19 +50,20 @@ namespace Horde.Commands.Workspace
 		{
 			if (File != null)
 			{
-				using IStorageClient store = BundleStorageClient.CreateFromDirectory(File.Directory, BundleCache, logger);
+				using MemoryMappedFileCache memoryMappedFileCache = new MemoryMappedFileCache();
+				IStorageClient store = BundleStorageClient.CreateFromDirectory(File.Directory, BundleCache, memoryMappedFileCache, logger);
 				IBlobRef handle = store.CreateBlobRef(await FileStorageBackend.ReadRefAsync(File));
 				return await ExecuteInternalAsync(store, handle, logger);
 			}
 			else if (Ref != null)
 			{
-				using IStorageClient store = CreateStorageClient();
+				IStorageClient store = CreateStorageClient();
 				IBlobRef handle = await store.ReadRefAsync(new RefName(Ref));
 				return await ExecuteInternalAsync(store, handle, logger);
 			}
 			else if (Node != null)
 			{
-				using IStorageClient store = CreateStorageClient();
+				IStorageClient store = CreateStorageClient();
 				IBlobRef handle = store.CreateBlobRef(new BlobLocator(Node));
 				return await ExecuteInternalAsync(store, handle, logger);
 			}
diff --git a/Engine/Source/Programs/Horde/HordeServer.Tests.Shared/ServerTestSetup.cs b/Engine/Source/Programs/Horde/HordeServer.Tests.Shared/ServerTestSetup.cs
--- a/Engine/Source/Programs/Horde/HordeServer.Tests.Shared/ServerTestSetup.cs
+++ b/Engine/Source/Programs/Horde/HordeServer.Tests.Shared/ServerTestSetup.cs
@@ -98,6 +98,8 @@ namespace HordeServer.Tests
 			services.AddSingleton<IServerInfo>(serverInfo);
 			services.AddSingleton<IPluginCollection>(_pluginCollection);
 
+			services.AddSingleton<MemoryMappedFileCache>();
+
 			services.AddSingleton<IAccountCollection, AccountCollection>();
 
 			services.AddLogging(builder => { builder.AddConsole().SetMinimumLevel(LogLevel.Debug); });
diff --git a/Engine/Source/Programs/Horde/HordeServer/Startup.cs b/Engine/Source/Programs/Horde/HordeServer/Startup.cs
--- a/Engine/Source/Programs/Horde/HordeServer/Startup.cs
+++ b/Engine/Source/Programs/Horde/HordeServer/Startup.cs
@@ -257,6 +257,7 @@ namespace HordeServer
 		// This method gets called *multiple times* by the runtime. Use this method to add services to the container.
 		public void ConfigureServices(IServiceCollection services)
 		{
+			services.AddSingleton<MemoryMappedFileCache>();
 			services.AddSingleton<JsonSchemaCache>();
 
 			// IOptionsMonitor pattern for live updating of configuration settings
diff --git a/Engine/Source/Programs/Horde/Plugins/Build/HordeServer.Build.Tests/Artifacts/UnsyncCacheTests.cs b/Engine/Source/Programs/Horde/Plugins/Build/HordeServer.Build.Tests/Artifacts/UnsyncCacheTests.cs
--- a/Engine/Source/Programs/Horde/Plugins/Build/HordeServer.Build.Tests/Artifacts/UnsyncCacheTests.cs
+++ b/Engine/Source/Programs/Horde/Plugins/Build/HordeServer.Build.Tests/Artifacts/UnsyncCacheTests.cs
@@ -21,7 +21,7 @@ namespace HordeServer.Tests.Artifacts
 			NamespaceId ns = new NamespaceId("test");
 			byte[] source = Enumerable.Range(0, 1024 * 1024).Select(x => (byte)(x % 257)).ToArray();
 
-			using BundleStorageClient storageClient = BundleStorageClient.CreateInMemory(NullLogger.Instance);
+			BundleStorageClient storageClient = BundleStorageClient.CreateInMemory(NullLogger.Instance);
 
 			Dictionary<IoHash, ReadOnlyMemory<byte>> chunks = new Dictionary<IoHash, ReadOnlyMemory<byte>>();
 
diff --git a/Engine/Source/Programs/Horde/Plugins/Build/HordeServer.Build.Tests/Issues/IssueServiceTests.cs b/Engine/Source/Programs/Horde/Plugins/Build/HordeServer.Build.Tests/Issues/IssueServiceTests.cs
--- a/Engine/Source/Programs/Horde/Plugins/Build/HordeServer.Build.Tests/Issues/IssueServiceTests.cs
+++ b/Engine/Source/Programs/Horde/Plugins/Build/HordeServer.Build.Tests/Issues/IssueServiceTests.cs
@@ -70,8 +70,6 @@ namespace HordeServer.Tests.Issues
 					IHashedBlobRef<LogNode> handle = await _builder.FlushAsync(writer, true, CancellationToken.None);
 					await _storageClient.WriteRefAsync(_log.RefName, handle);
 				}
-
-				_storageClient.Dispose();
 			}
 
 			public IDisposable? BeginScope<TState>(TState state) where TState : notnull => _scopeCollection.BeginScope(state);
@@ -349,16 +347,14 @@ namespace HordeServer.Tests.Issues
 
 			ILog log = (await LogCollection.GetAsync(logId, CancellationToken.None))!;
 
-			using (IStorageClient storageClient = StorageService.CreateClient(Namespace.Logs))
+			IStorageClient storageClient = StorageService.CreateClient(Namespace.Logs);
+			await using (IBlobWriter writer = storageClient.CreateBlobWriter())
 			{
-				await using (IBlobWriter writer = storageClient.CreateBlobWriter())
-				{
-					LogBuilder builder = new LogBuilder(LogFormat.Json, NullLogger.Instance);
-					builder.WriteData(data);
+				LogBuilder builder = new LogBuilder(LogFormat.Json, NullLogger.Instance);
+				builder.WriteData(data);
 
-					IHashedBlobRef<LogNode> handle = await builder.FlushAsync(writer, true, CancellationToken.None);
-					await storageClient.WriteRefAsync(log!.RefName, handle);
-				}
+				IHashedBlobRef<LogNode> handle = await builder.FlushAsync(writer, true, CancellationToken.None);
+				await storageClient.WriteRefAsync(log!.RefName, handle);
 			}
 
 			await log.AddEventsAsync(new List<NewLogEventData> { new NewLogEventData { LineIndex = 0, LineCount = 1, Severity = severity } }, CancellationToken.None);
@@ -394,7 +390,7 @@ namespace HordeServer.Tests.Issues
 			LogId logId = job.Batches[batchIdx].Steps[stepIdx].LogId!.Value;
 			ILog log = (await LogCollection.GetAsync(logId))!;
 
-			using IStorageClient storageClient = StorageService.CreateClient(Namespace.Logs);
+			IStorageClient storageClient = StorageService.CreateClient(Namespace.Logs);
 			await using (TestJsonLogger logger = new TestJsonLogger(log, storageClient))
 			{
 				PerforceMetadataLogger perforceLogger = new PerforceMetadataLogger(logger);
diff --git a/Engine/Source/Programs/Horde/Plugins/Build/HordeServer.Build/Artifacts/ArtifactCollection.cs b/Engine/Source/Programs/Horde/Plugins/Build/HordeServer.Build/Artifacts/ArtifactCollection.cs
--- a/Engine/Source/Programs/Horde/Plugins/Build/HordeServer.Build/Artifacts/ArtifactCollection.cs
+++ b/Engine/Source/Programs/Horde/Plugins/Build/HordeServer.Build/Artifacts/ArtifactCollection.cs
@@ -365,7 +365,7 @@ namespace HordeServer.Artifacts
 		{
 			foreach (IGrouping<NamespaceId, ArtifactDocument> artifactGroup in deleteArtifacts.GroupBy(x => x.NamespaceId))
 			{
-				using IStorageClient storageClient = _storageService.CreateClient(artifactGroup.Key);
+				IStorageClient storageClient = _storageService.CreateClient(artifactGroup.Key);
 				foreach (ArtifactDocument artifact in artifactGroup)
 				{
 					// Delete the ref allowing the storage service to expire this data
diff --git a/Engine/Source/Programs/Horde/Plugins/Build/HordeServer.Build/Artifacts/ArtifactsController.cs b/Engine/Source/Programs/Horde/Plugins/Build/HordeServer.Build/Artifacts/ArtifactsController.cs
--- a/Engine/Source/Programs/Horde/Plugins/Build/HordeServer.Build/Artifacts/ArtifactsController.cs
+++ b/Engine/Source/Programs/Horde/Plugins/Build/HordeServer.Build/Artifacts/ArtifactsController.cs
@@ -309,7 +309,7 @@ namespace HordeServer.Artifacts
 				return Forbid(ArtifactAclAction.ReadArtifact, artifact.StreamId);
 			}
 
-			using IStorageClient storageClient = _storageService.CreateClient(artifact.NamespaceId);
+			IStorageClient storageClient = _storageService.CreateClient(artifact.NamespaceId);
 
 			DirectoryNode directoryNode;
 			try
@@ -476,7 +476,7 @@ namespace HordeServer.Artifacts
 				return Forbid(ArtifactAclAction.ReadArtifact, artifact.StreamId);
 			}
 
-			using IStorageClient storageClient = _storageService.CreateClient(artifact.NamespaceId);
+			IStorageClient storageClient = _storageService.CreateClient(artifact.NamespaceId);
 			DirectoryNode directory = await storageClient.ReadRefTargetAsync<DirectoryNode>(artifact.RefName, DateTime.UtcNow.AddHours(1.0), cancellationToken: cancellationToken);
 
 			FileEntry? fileEntry = await directory.GetFileEntryByPathAsync(path, cancellationToken: cancellationToken);
@@ -572,18 +572,10 @@ namespace HordeServer.Artifacts
 
 #pragma warning disable CA2000
 			IStorageClient storageClient = _storageService.CreateClient(artifact.NamespaceId);
-			try
-			{
-				IHashedBlobRef<DirectoryNode> directory = await storageClient.ReadRefAsync<DirectoryNode>(artifact.RefName, DateTime.UtcNow.AddHours(1.0), cancellationToken: cancellationToken);
+			IHashedBlobRef<DirectoryNode> directory = await storageClient.ReadRefAsync<DirectoryNode>(artifact.RefName, DateTime.UtcNow.AddHours(1.0), cancellationToken: cancellationToken);
 
-				Stream stream = directory.AsZipStream(filter).WrapOwnership(storageClient);
-				return new FileStreamResult(stream, "application/zip") { FileDownloadName = $"{artifact.RefName}.zip" };
-			}
-			catch
-			{
-				storageClient.Dispose();
-				throw;
-			}
+			Stream stream = directory.AsZipStream(filter);
+			return new FileStreamResult(stream, "application/zip") { FileDownloadName = $"{artifact.RefName}.zip" };
 #pragma warning restore CA2000
 		}
 
diff --git a/Engine/Source/Programs/Horde/Plugins/Build/HordeServer.Build/Artifacts/UnsyncCache.cs b/Engine/Source/Programs/Horde/Plugins/Build/HordeServer.Build/Artifacts/UnsyncCache.cs
--- a/Engine/Source/Programs/Horde/Plugins/Build/HordeServer.Build/Artifacts/UnsyncCache.cs
+++ b/Engine/Source/Programs/Horde/Plugins/Build/HordeServer.Build/Artifacts/UnsyncCache.cs
@@ -44,23 +44,18 @@ namespace HordeServer.Artifacts
 	/// </summary>
 	public sealed class UnsyncCache : IDisposable
 	{
-		class ArtifactInfo : IDisposable
+		class ArtifactInfo
 		{
-			public IStorageClient StorageClient { get; }
 			public UnsyncManifest Manifest { get; }
 			public ReadOnlyMemory<byte> ManifestData { get; }
 			public FrozenDictionary<IoHash, IHashedBlobRef<LeafChunkedDataNode>> Blobs { get; }
 
-			public ArtifactInfo(IStorageClient storageClient, UnsyncManifest manifest, ReadOnlyMemory<byte> manifestData, FrozenDictionary<IoHash, IHashedBlobRef<LeafChunkedDataNode>> blobs)
+			public ArtifactInfo(UnsyncManifest manifest, ReadOnlyMemory<byte> manifestData, FrozenDictionary<IoHash, IHashedBlobRef<LeafChunkedDataNode>> blobs)
 			{
-				StorageClient = storageClient;
 				Manifest = manifest;
 				ManifestData = manifestData;
 				Blobs = blobs;
 			}
-
-			public void Dispose()
-				=> StorageClient.Dispose();
 		}
 
 		readonly IStorageService _storageService;
@@ -211,40 +206,30 @@ namespace HordeServer.Artifacts
 
 		async Task<ArtifactInfo?> ReadArtifactAsync(IArtifact artifact, CancellationToken cancellationToken)
 		{
-			IStorageClient? storageClient = null;
-			try
-			{
-				Stopwatch timer = Stopwatch.StartNew();
-				storageClient = _storageService.CreateClient(artifact.NamespaceId);
+			Stopwatch timer = Stopwatch.StartNew();
+			IStorageClient storageClient = _storageService.CreateClient(artifact.NamespaceId);
 
-				IHashedBlobRef<DirectoryNode>? target = await storageClient.TryReadRefAsync<DirectoryNode>(artifact.RefName, cancellationToken: cancellationToken);
-				if (target == null)
-				{
-					return null;
-				}
+			IHashedBlobRef<DirectoryNode>? target = await storageClient.TryReadRefAsync<DirectoryNode>(artifact.RefName, cancellationToken: cancellationToken);
+			if (target == null)
+			{
+				return null;
+			}
 
-				List<UnsyncFile> files = new List<UnsyncFile>();
-				await FindFilesAsync(new Utf8StringBuilder(), target, files, cancellationToken);
+			List<UnsyncFile> files = new List<UnsyncFile>();
+			await FindFilesAsync(new Utf8StringBuilder(), target, files, cancellationToken);
 
-				Dictionary<IoHash, IHashedBlobRef<LeafChunkedDataNode>> blocks = new Dictionary<IoHash, IHashedBlobRef<LeafChunkedDataNode>>();
-				foreach (UnsyncBlock block in files.SelectMany(x => x.Blocks))
-				{
-					blocks[block.Blob.Hash] = block.Blob;
-				}
+			Dictionary<IoHash, IHashedBlobRef<LeafChunkedDataNode>> blocks = new Dictionary<IoHash, IHashedBlobRef<LeafChunkedDataNode>>();
+			foreach (UnsyncBlock block in files.SelectMany(x => x.Blocks))
+			{
+				blocks[block.Blob.Hash] = block.Blob;
+			}
 
-				_logger.LogDebug("Generated Unsync manifest for artifact {ArtifactId} in {Time:n1}ms", artifact.Id, timer.ElapsedMilliseconds);
+			_logger.LogDebug("Generated Unsync manifest for artifact {ArtifactId} in {Time:n1}ms", artifact.Id, timer.ElapsedMilliseconds);
 
-				UnsyncManifest manifest = new UnsyncManifest(files);
-				ReadOnlyMemory<byte> manifestData = SerializeManifest(manifest);
+			UnsyncManifest manifest = new UnsyncManifest(files);
+			ReadOnlyMemory<byte> manifestData = SerializeManifest(manifest);
 
-				ArtifactInfo artifactInfo = new ArtifactInfo(storageClient, manifest, manifestData, blocks.ToFrozenDictionary());
-				storageClient = null;
-				return artifactInfo;
-			}
-			finally
-			{
-				storageClient?.Dispose();
-			}
+			return new ArtifactInfo(manifest, manifestData, blocks.ToFrozenDictionary());
 		}
 
 		static async Task FindFilesAsync(Utf8StringBuilder path, IHashedBlobRef<DirectoryNode> directoryNodeRef, List<UnsyncFile> files, CancellationToken cancellationToken)
diff --git a/Engine/Source/Programs/Horde/Plugins/Build/HordeServer.Build/Replicators/PerforceReplicator.cs b/Engine/Source/Programs/Horde/Plugins/Build/HordeServer.Build/Replicators/PerforceReplicator.cs
--- a/Engine/Source/Programs/Horde/Plugins/Build/HordeServer.Build/Replicators/PerforceReplicator.cs
+++ b/Engine/Source/Programs/Horde/Plugins/Build/HordeServer.Build/Replicators/PerforceReplicator.cs
@@ -343,7 +343,7 @@ namespace HordeServer.Replicators
 			RefName refName = GetRefName(replicator.Id);
 			RefName incRefName = GetIncrementalRefName(replicator.Id);
 
-			using IStorageClient store = _storageService.CreateClient(Namespace.Perforce);
+			IStorageClient store = _storageService.CreateClient(Namespace.Perforce);
 
 			ICommitCollection commits = _perforceService.GetCommits(streamConfig);
 
@@ -422,7 +422,7 @@ namespace HordeServer.Replicators
 
 		async Task<IReplicator> WriteInternalAsync(IReplicator replicator, int change, BuildConfig buildConfig, StreamConfig streamConfig, PerforceReplicationOptions options, BlobSerializerOptions blobOptions, CancellationToken cancellationToken = default)
 		{
-			using IStorageClient store = _storageService.CreateClient(Namespace.Perforce);
+			IStorageClient store = _storageService.CreateClient(Namespace.Perforce);
 
 			// Find the parent node
 			RefName refName = GetRefName(replicator.Id);
diff --git a/Engine/Source/Programs/Horde/Plugins/Compute/HordeServer.Compute.Tests/Logs/LogIndexingTests.cs b/Engine/Source/Programs/Horde/Plugins/Compute/HordeServer.Compute.Tests/Logs/LogIndexingTests.cs
--- a/Engine/Source/Programs/Horde/Plugins/Compute/HordeServer.Compute.Tests/Logs/LogIndexingTests.cs
+++ b/Engine/Source/Programs/Horde/Plugins/Compute/HordeServer.Compute.Tests/Logs/LogIndexingTests.cs
@@ -30,7 +30,6 @@ namespace HordeServer.Tests.Logs
 		public async ValueTask DisposeAsync()
 		{
 			await _blobWriter.DisposeAsync();
-			_storageClient.Dispose();
 		}
 
 		public async Task<ILog> WriteDataAsync(ReadOnlyMemory<byte> data)
diff --git a/Engine/Source/Programs/Horde/Plugins/Compute/HordeServer.Compute/Logs/LogCollection.cs b/Engine/Source/Programs/Horde/Plugins/Compute/HordeServer.Compute/Logs/LogCollection.cs
--- a/Engine/Source/Programs/Horde/Plugins/Compute/HordeServer.Compute/Logs/LogCollection.cs
+++ b/Engine/Source/Programs/Horde/Plugins/Compute/HordeServer.Compute/Logs/LogCollection.cs
@@ -306,7 +306,7 @@ namespace HordeServer.Logs
 
 			if (logDocument.UseNewStorageBackend)
 			{
-				using IStorageClient storageClient = _storageService.CreateClient(logDocument.NamespaceId);
+				IStorageClient storageClient = _storageService.CreateClient(logDocument.NamespaceId);
 				await storageClient.DeleteRefAsync(logDocument.RefName, cancellationToken);
 			}
 		}
@@ -327,7 +327,7 @@ namespace HordeServer.Logs
 		{
 			List<Utf8String> lines = new List<Utf8String>();
 
-			using IStorageClient storageClient = _storageService.CreateClient(log.NamespaceId);
+			IStorageClient storageClient = _storageService.CreateClient(log.NamespaceId);
 
 			int maxIndex = index + count;
 			bool complete = log.Complete;
@@ -392,7 +392,7 @@ namespace HordeServer.Logs
 		/// <inheritdoc/>
 		async Task<(int, long)> GetLineOffsetAsync(LogDocument log, int lineIdx, CancellationToken cancellationToken)
 		{
-			using IStorageClient storageClient = _storageService.CreateClient(log.NamespaceId);
+			IStorageClient storageClient = _storageService.CreateClient(log.NamespaceId);
 
 			LogNode? root = await storageClient.TryReadRefTargetAsync<LogNode>(log.RefName, cancellationToken: cancellationToken);
 			if (root == null)
@@ -422,7 +422,7 @@ namespace HordeServer.Logs
 		/// <inheritdoc/>
 		async Task<Stream> OpenRawStreamAsync(LogDocument log, long offset, long length, CancellationToken cancellationToken)
 		{
-			using IStorageClient storageClient = _storageService.CreateClient(log.NamespaceId);
+			IStorageClient storageClient = _storageService.CreateClient(log.NamespaceId);
 
 			LogNode? root = await storageClient.TryReadRefTargetAsync<LogNode>(log.RefName, cancellationToken: cancellationToken);
 			if (root == null || root.TextChunkRefs.Count == 0)
@@ -480,7 +480,7 @@ namespace HordeServer.Logs
 		async IAsyncEnumerable<int> SearchLogDataInternalNewAsync(LogDocument log, string text, int firstLine, SearchStats searchStats, [EnumeratorCancellation] CancellationToken cancellationToken)
 		{
 			SearchTerm searchText = new SearchTerm(text);
-			using IStorageClient storageClient = _storageService.CreateClient(log.NamespaceId);
+			IStorageClient storageClient = _storageService.CreateClient(log.NamespaceId);
 
 			// Search the index
 			if (log.LineCount > 0)
diff --git a/Engine/Source/Programs/Horde/Plugins/Compute/HordeServer.Compute/Logs/LogRpcService.cs b/Engine/Source/Programs/Horde/Plugins/Compute/HordeServer.Compute/Logs/LogRpcService.cs
--- a/Engine/Source/Programs/Horde/Plugins/Compute/HordeServer.Compute/Logs/LogRpcService.cs
+++ b/Engine/Source/Programs/Horde/Plugins/Compute/HordeServer.Compute/Logs/LogRpcService.cs
@@ -51,7 +51,7 @@ namespace HordeServer.Logs
 				throw new StructuredRpcException(StatusCode.PermissionDenied, "Access denied");
 			}
 
-			using IStorageClient store = _storageService.CreateClient(Namespace.Logs);
+			IStorageClient store = _storageService.CreateClient(Namespace.Logs);
 
 			_logger.LogInformation("Updating {LogId} to node {RefTarget} (lines: {LineCount}, complete: {Complete})", request.LogId, request.TargetLocator, request.LineCount, request.Complete);
 
diff --git a/Engine/Source/Programs/Horde/Plugins/Ddc/HordeServer.Ddc/Ddc/BlobService.cs b/Engine/Source/Programs/Horde/Plugins/Ddc/HordeServer.Ddc/Ddc/BlobService.cs
--- a/Engine/Source/Programs/Horde/Plugins/Ddc/HordeServer.Ddc/Ddc/BlobService.cs
+++ b/Engine/Source/Programs/Horde/Plugins/Ddc/HordeServer.Ddc/Ddc/BlobService.cs
@@ -60,13 +60,13 @@ namespace HordeServer.Ddc
 
 		public async Task<bool> ExistsAsync(NamespaceId ns, BlobId blob, List<string>? storageLayers = null, CancellationToken cancellationToken = default)
 		{
-			using IStorageClient storageClient = _storageClientFactory.CreateClient(ns);
+			IStorageClient storageClient = _storageClientFactory.CreateClient(ns);
 			return await storageClient.FindAliasAsync(GetAlias(blob), cancellationToken) != null;
 		}
 
 		public async Task DeleteObjectAsync(NamespaceId ns, BlobId blob, CancellationToken cancellationToken)
 		{
-			using IStorageClient storageClient = _storageClientFactory.CreateClient(ns);
+			IStorageClient storageClient = _storageClientFactory.CreateClient(ns);
 			string aliasName = GetAlias(blob);
 
 			BlobAlias[] aliases = await storageClient.FindAliasesAsync(aliasName, cancellationToken: cancellationToken);
@@ -78,7 +78,7 @@ namespace HordeServer.Ddc
 
 		public async Task<BlobId[]> FilterOutKnownBlobsAsync(NamespaceId ns, IEnumerable<BlobId> blobIds, CancellationToken cancellationToken)
 		{
-			using IStorageClient storageClient = _storageClientFactory.CreateClient(ns);
+			IStorageClient storageClient = _storageClientFactory.CreateClient(ns);
 
 			List<BlobId> unknownBlobIds = new List<BlobId>();
 			foreach (BlobId blobId in blobIds)
@@ -123,7 +123,7 @@ namespace HordeServer.Ddc
 
 		public async Task<BlobContents> GetObjectAsync(NamespaceId ns, BlobId blob, List<string>? storageLayers, bool supportsRedirectUri, bool allowOndemandReplication = true, CancellationToken cancellationToken = default)
 		{
-			using IStorageClient storageClient = _storageClientFactory.CreateClient(ns);
+			IStorageClient storageClient = _storageClientFactory.CreateClient(ns);
 
 			BlobAlias? alias = await storageClient.FindAliasAsync(GetAlias(blob), cancellationToken);
 			if (alias == null)
@@ -174,7 +174,7 @@ namespace HordeServer.Ddc
 
 		public async Task<BlobId> PutObjectKnownHashAsync(NamespaceId ns, IBufferedPayload content, BlobId identifier, CancellationToken cancellationToken)
 		{
-			using IStorageClient storageClient = _storageClientFactory.CreateClient(ns);
+			IStorageClient storageClient = _storageClientFactory.CreateClient(ns);
 
 			IHashedBlobRef blobRef;
 			await using (IBlobWriter writer = storageClient.CreateBlobWriter())
diff --git a/Engine/Source/Programs/Horde/Plugins/Ddc/HordeServer.Ddc/Ddc/ContentIdStore.cs b/Engine/Source/Programs/Horde/Plugins/Ddc/HordeServer.Ddc/Ddc/ContentIdStore.cs
--- a/Engine/Source/Programs/Horde/Plugins/Ddc/HordeServer.Ddc/Ddc/ContentIdStore.cs
+++ b/Engine/Source/Programs/Horde/Plugins/Ddc/HordeServer.Ddc/Ddc/ContentIdStore.cs
@@ -19,7 +19,7 @@ namespace HordeServer.Ddc
 
 		public async Task<BlobId[]?> ResolveAsync(NamespaceId ns, ContentId contentId, bool mustBeContentId = false, CancellationToken cancellationToken = default)
 		{
-			using IStorageClient storageClient = _storageClientFactory.CreateClient(ns);
+			IStorageClient storageClient = _storageClientFactory.CreateClient(ns);
 
 			BlobAlias? blobAlias = await storageClient.FindAliasAsync(GetAlias(contentId), cancellationToken);
 			if (blobAlias == null && !mustBeContentId)
@@ -37,7 +37,7 @@ namespace HordeServer.Ddc
 
 		public async Task PutAsync(NamespaceId ns, ContentId contentId, BlobId blobId, int contentWeight, CancellationToken cancellationToken)
 		{
-			using IStorageClient storageClient = _storageClientFactory.CreateClient(ns);
+			IStorageClient storageClient = _storageClientFactory.CreateClient(ns);
 
 			BlobAlias? blobAlias = await storageClient.FindAliasAsync(GetAlias(blobId), cancellationToken);
 			if (blobAlias == null)
diff --git a/Engine/Source/Programs/Horde/Plugins/Ddc/HordeServer.Ddc/Ddc/RefService.cs b/Engine/Source/Programs/Horde/Plugins/Ddc/HordeServer.Ddc/Ddc/RefService.cs
--- a/Engine/Source/Programs/Horde/Plugins/Ddc/HordeServer.Ddc/Ddc/RefService.cs
+++ b/Engine/Source/Programs/Horde/Plugins/Ddc/HordeServer.Ddc/Ddc/RefService.cs
@@ -23,19 +23,19 @@ namespace HordeServer.Ddc
 
 		public async Task<bool> DeleteAsync(NamespaceId ns, BucketId bucket, RefId key, CancellationToken cancellationToken)
 		{
-			using IStorageClient storageClient = _storageClientFactory.CreateClient(ns);
+			IStorageClient storageClient = _storageClientFactory.CreateClient(ns);
 			return await storageClient.DeleteRefAsync(GetRefName(bucket, key), cancellationToken);
 		}
 
 		public async Task<bool> ExistsAsync(NamespaceId ns, BucketId bucket, RefId key, CancellationToken cancellationToken)
 		{
-			using IStorageClient storageClient = _storageClientFactory.CreateClient(ns);
+			IStorageClient storageClient = _storageClientFactory.CreateClient(ns);
 			return await storageClient.RefExistsAsync(GetRefName(bucket, key), cancellationToken: cancellationToken);
 		}
 
 		public async Task<(ContentId[], BlobId[])> FinalizeAsync(NamespaceId ns, BucketId bucket, RefId key, BlobId blobHash, CancellationToken cancellationToken)
 		{
-			using IStorageClient storageClient = _storageClientFactory.CreateClient(ns);
+			IStorageClient storageClient = _storageClientFactory.CreateClient(ns);
 
 			BlobAlias? blobAlias = await storageClient.FindAliasAsync(BlobService.GetAlias(blobHash), cancellationToken);
 			if (blobAlias == null)
@@ -125,7 +125,7 @@ namespace HordeServer.Ddc
 
 		public async Task<(RefRecord, BlobContents?)> GetAsync(NamespaceId ns, BucketId bucket, RefId key, string[] fields, bool doLastAccessTracking, CancellationToken cancellationToken)
 		{
-			using IStorageClient storageClient = _storageClientFactory.CreateClient(ns);
+			IStorageClient storageClient = _storageClientFactory.CreateClient(ns);
 
 			DdcRefNode? node = await storageClient.TryReadRefTargetAsync<DdcRefNode>(GetRefName(bucket, key), cancellationToken: cancellationToken);
 			if (node == null)
@@ -142,7 +142,7 @@ namespace HordeServer.Ddc
 
 		public async Task<List<BlobId>> GetReferencedBlobsAsync(NamespaceId ns, BucketId bucket, RefId key, CancellationToken cancellationToken)
 		{
-			using IStorageClient storageClient = _storageClientFactory.CreateClient(ns);
+			IStorageClient storageClient = _storageClientFactory.CreateClient(ns);
 
 			DdcRefNode? node = await storageClient.TryReadRefTargetAsync<DdcRefNode>(GetRefName(bucket, key), cancellationToken: cancellationToken);
 			if (node == null)
diff --git a/Engine/Source/Programs/Horde/Plugins/Ddc/HordeServer.Ddc/Ddc/RequestHelper.cs b/Engine/Source/Programs/Horde/Plugins/Ddc/HordeServer.Ddc/Ddc/RequestHelper.cs
--- a/Engine/Source/Programs/Horde/Plugins/Ddc/HordeServer.Ddc/Ddc/RequestHelper.cs
+++ b/Engine/Source/Programs/Horde/Plugins/Ddc/HordeServer.Ddc/Ddc/RequestHelper.cs
@@ -28,7 +28,7 @@ namespace HordeServer.Ddc
 
 		public Task<ActionResult?> HasAccessToNamespaceAsync(ClaimsPrincipal user, HttpRequest request, NamespaceId ns, AclAction[] aclActions)
 		{
-			using IStorageClient? storageClient = _storageClientFactory.TryCreateClient(ns);
+			IStorageClient? storageClient = _storageClientFactory.TryCreateClient(ns);
 			if (storageClient == null)
 			{
 				return Task.FromResult<ActionResult?>(new ForbidResult());
diff --git a/Engine/Source/Programs/Horde/Plugins/Storage/HordeServer.Storage.Tests/Storage/BlobStoreTests.cs b/Engine/Source/Programs/Horde/Plugins/Storage/HordeServer.Storage.Tests/Storage/BlobStoreTests.cs
--- a/Engine/Source/Programs/Horde/Plugins/Storage/HordeServer.Storage.Tests/Storage/BlobStoreTests.cs
+++ b/Engine/Source/Programs/Horde/Plugins/Storage/HordeServer.Storage.Tests/Storage/BlobStoreTests.cs
@@ -94,7 +94,7 @@ namespace HordeServer.Tests.Storage
 		[TestMethod]
 		public async Task LeafTestAsync()
 		{
-			using IStorageClient store = CreateStorageClient();
+			IStorageClient store = CreateStorageClient();
 
 			byte[] input = CreateTestData(256, 0);
 
@@ -109,7 +109,7 @@ namespace HordeServer.Tests.Storage
 		[TestMethod]
 		public async Task ReferenceTestAsync()
 		{
-			using IStorageClient store = CreateStorageClient();
+			IStorageClient store = CreateStorageClient();
 
 			byte[] input1 = CreateTestData(256, 1);
 			HashedBlobRefValue locator1 = await WriteBlobAsync(store, new Blob(input1, Array.Empty<HashedBlobRefValue>()));
@@ -141,7 +141,7 @@ namespace HordeServer.Tests.Storage
 		[TestMethod]
 		public async Task RefExpiryTestAsync()
 		{
-			using IStorageClient store = CreateStorageClient();
+			IStorageClient store = CreateStorageClient();
 
 			Blob blob1 = new Blob(new byte[] { 1, 2, 3 }, Array.Empty<HashedBlobRefValue>());
 			HashedBlobRefValue target = await WriteBlobAsync(store, blob1);
diff --git a/Engine/Source/Programs/Horde/Plugins/Storage/HordeServer.Storage.Tests/Storage/GcServiceTests.cs b/Engine/Source/Programs/Horde/Plugins/Storage/HordeServer.Storage.Tests/Storage/GcServiceTests.cs
--- a/Engine/Source/Programs/Horde/Plugins/Storage/HordeServer.Storage.Tests/Storage/GcServiceTests.cs
+++ b/Engine/Source/Programs/Horde/Plugins/Storage/HordeServer.Storage.Tests/Storage/GcServiceTests.cs
@@ -31,7 +31,7 @@ namespace HordeServer.Tests.Storage
 			globalConfig.Plugins.Add(new PluginName("storage"), storageConfig);
 			SetConfig(globalConfig);
 
-			using IStorageClient store = StorageService.CreateClient(new NamespaceId("default"));
+			IStorageClient store = StorageService.CreateClient(new NamespaceId("default"));
 
 			Random random = new Random(0);
 			IHashedBlobRef[] blobs = await CreateTestDataAsync(store, 30, 50, 30, 5, random);
diff --git a/Engine/Source/Programs/Horde/Plugins/Storage/HordeServer.Storage.Tests/Storage/StorageServiceTests.cs b/Engine/Source/Programs/Horde/Plugins/Storage/HordeServer.Storage.Tests/Storage/StorageServiceTests.cs
--- a/Engine/Source/Programs/Horde/Plugins/Storage/HordeServer.Storage.Tests/Storage/StorageServiceTests.cs
+++ b/Engine/Source/Programs/Horde/Plugins/Storage/HordeServer.Storage.Tests/Storage/StorageServiceTests.cs
@@ -21,7 +21,7 @@ namespace HordeServer.Tests.Storage
 		public async Task BlobCollectionTestAsync()
 		{
 			StorageService storageService = ServiceProvider.GetRequiredService<StorageService>();
-			using IStorageClient client = storageService.CreateClient(new NamespaceId("memory"));
+			IStorageClient client = storageService.CreateClient(new NamespaceId("memory"));
 
 			BlobType type1 = new BlobType(Guid.Parse("{11C2D886-4164-3349-D1E9-6F943D2ED10B}"), 0);
 			byte[] data1 = new byte[] { 1, 2, 3 };
diff --git a/Engine/Source/Programs/Horde/Plugins/Storage/HordeServer.Storage/Storage/StorageController.cs b/Engine/Source/Programs/Horde/Plugins/Storage/HordeServer.Storage/Storage/StorageController.cs
--- a/Engine/Source/Programs/Horde/Plugins/Storage/HordeServer.Storage/Storage/StorageController.cs
+++ b/Engine/Source/Programs/Horde/Plugins/Storage/HordeServer.Storage/Storage/StorageController.cs
@@ -270,7 +270,7 @@ namespace HordeServer.Storage
 		/// </summary>
 		public static async Task<ActionResult<ReadRefResponse>> ReadRefInternalAsync(IStorageClientFactory storageService, NamespaceId namespaceId, RefName refName, IHeaderDictionary headers, CancellationToken cancellationToken)
 		{
-			using IStorageClient client = storageService.CreateClient(namespaceId);
+			IStorageClient client = storageService.CreateClient(namespaceId);
 
 			RefCacheTime cacheTime = new RefCacheTime();
 			foreach (string? entry in headers.CacheControl)
@@ -393,7 +393,7 @@ namespace HordeServer.Storage
 				locator = new BlobLocator(locator, String.Join("&", fragments));
 			}
 
-			using IStorageClient storageClient = _storageService.CreateClient(namespaceId);
+			IStorageClient storageClient = _storageService.CreateClient(namespaceId);
 
 			object content;
 
diff --git a/Engine/Source/Programs/Horde/Plugins/Storage/HordeServer.Storage/Storage/StorageService.cs b/Engine/Source/Programs/Horde/Plugins/Storage/HordeServer.Storage/Storage/StorageService.cs
--- a/Engine/Source/Programs/Horde/Plugins/Storage/HordeServer.Storage/Storage/StorageService.cs
+++ b/Engine/Source/Programs/Horde/Plugins/Storage/HordeServer.Storage/Storage/StorageService.cs
@@ -674,70 +674,61 @@ namespace HordeServer.Storage
 			State state = CreateState(_storageConfig.CurrentValue);
 
 			Dictionary<NamespaceId, BundleStorageClient> cachedClients = new();
-			try
-			{
-				long ingestedCount = 0;
 
-				// Compute missing import info, by searching for blobs with an ObjectId timestamp after the last import compute cycle
-				ObjectId latestInfoId = ObjectId.GenerateNewId(ingestTimeUtc);
-				while (!gcState.Reset)
+			long ingestedCount = 0;
+
+			// Compute missing import info, by searching for blobs with an ObjectId timestamp after the last import compute cycle
+			ObjectId latestInfoId = ObjectId.GenerateNewId(ingestTimeUtc);
+			while (!gcState.Reset)
+			{
+				// Fetch the next batch of blobs
+				List<BlobInfo> current = await _blobCollection.Find(x => x.Id > gcState.LastImportBlobInfoId && x.Id < latestInfoId).SortBy(x => x.Id).Limit(500).ToListAsync(cancellationToken);
+				if (current.Count == 0)
 				{
-					// Fetch the next batch of blobs
-					List<BlobInfo> current = await _blobCollection.Find(x => x.Id > gcState.LastImportBlobInfoId && x.Id < latestInfoId).SortBy(x => x.Id).Limit(500).ToListAsync(cancellationToken);
-					if (current.Count == 0)
-					{
-						break;
-					}
+					break;
+				}
 
-					// Find imports, and add a check record for each new blob
-					using TelemetrySpan innerSpan = _tracer.StartActiveSpan($"{nameof(StorageService)}.{nameof(TickBlobsAsync)}.Batch");
-					innerSpan.SetAttribute("First", current[0].Id.ToString());
-					innerSpan.SetAttribute("Last", current[^1].Id.ToString());
-					innerSpan.SetAttribute("Count", current.Count);
+				// Find imports, and add a check record for each new blob
+				using TelemetrySpan innerSpan = _tracer.StartActiveSpan($"{nameof(StorageService)}.{nameof(TickBlobsAsync)}.Batch");
+				innerSpan.SetAttribute("First", current[0].Id.ToString());
+				innerSpan.SetAttribute("Last", current[^1].Id.ToString());
+				innerSpan.SetAttribute("Count", current.Count);
 
-					_logger.LogDebug("Ticking {NumBlobs} blobs from {FirstId} to {LastId}", current.Count, current[0].Id, current[^1].Id);
+				_logger.LogDebug("Ticking {NumBlobs} blobs from {FirstId} to {LastId}", current.Count, current[0].Id, current[^1].Id);
 
-					foreach (BlobInfo blobInfo in current)
+				foreach (BlobInfo blobInfo in current)
+				{
+					NamespaceInfo? namespaceInfo;
+					if (state.Namespaces.TryGetValue(blobInfo.NamespaceId, out namespaceInfo))
 					{
-						NamespaceInfo? namespaceInfo;
-						if (state.Namespaces.TryGetValue(blobInfo.NamespaceId, out namespaceInfo))
+						BundleStorageClient? storageClient;
+						if (!cachedClients.TryGetValue(namespaceInfo.Id, out storageClient))
 						{
-							BundleStorageClient? storageClient;
-							if (!cachedClients.TryGetValue(namespaceInfo.Id, out storageClient))
-							{
-								storageClient = new BundleStorageClient(namespaceInfo.Backend, _bundleCache, null, _logger);
-								cachedClients.Add(namespaceInfo.Id, storageClient);
-							}
+							storageClient = new BundleStorageClient(namespaceInfo.Backend, _bundleCache, null, _logger);
+							cachedClients.Add(namespaceInfo.Id, storageClient);
+						}
 
-							try
-							{
-								await TickBlobAsync(storageClient, blobInfo, cancellationToken);
-							}
-							catch (ObjectNotFoundException ex)
-							{
-								_logger.LogInformation(ex, "Unable to read references for {NamespaceId} blob {BlobId}: {Message}", blobInfo.NamespaceId, blobInfo.Id, ex.Message);
-							}
-							catch (Exception ex)
-							{
-								_logger.LogWarning(ex, "Unable to read references for {NamespaceId} blob {BlobId} (key: {ObjectKey}): {Message}", blobInfo.NamespaceId, blobInfo.Id, GetObjectKey(blobInfo.Locator), ex.Message);
-							}
+						try
+						{
+							await TickBlobAsync(storageClient, blobInfo, cancellationToken);
+						}
+						catch (ObjectNotFoundException ex)
+						{
+							_logger.LogInformation(ex, "Unable to read references for {NamespaceId} blob {BlobId}: {Message}", blobInfo.NamespaceId, blobInfo.Id, ex.Message);
+						}
+						catch (Exception ex)
+						{
+							_logger.LogWarning(ex, "Unable to read references for {NamespaceId} blob {BlobId} (key: {ObjectKey}): {Message}", blobInfo.NamespaceId, blobInfo.Id, GetObjectKey(blobInfo.Locator), ex.Message);
 						}
 					}
-
-					// Update the last imported blob id
-					gcState = await _gcState.UpdateAsync(state => state.LastImportBlobInfoId = current[^1].Id, cancellationToken);
-					ingestedCount += current.Count;
 				}
 
-				_logger.LogInformation("Added {NumBlobs} blobs for GC (upper time: {Time})", ingestedCount, ingestTimeUtc);
-			}
-			finally
-			{
-				foreach (BundleStorageClient client in cachedClients.Values)
-				{
-					client.Dispose();
-				}
+				// Update the last imported blob id
+				gcState = await _gcState.UpdateAsync(state => state.LastImportBlobInfoId = current[^1].Id, cancellationToken);
+				ingestedCount += current.Count;
 			}
+
+			_logger.LogInformation("Added {NumBlobs} blobs for GC (upper time: {Time})", ingestedCount, ingestTimeUtc);
 		}
 
 		async Task TickBlobAsync(BundleStorageClient storageClient, BlobInfo blobInfo, CancellationToken cancellationToken)
@@ -1109,7 +1100,7 @@ namespace HordeServer.Storage
 
 		async Task TickGcForNamespaceAsync(NamespaceInfo namespaceInfo, ObjectId lastImportBlobInfoId, DateTime utcNow, bool deleteObjects, CancellationToken cancellationToken)
 		{
-			using IStorageClient client = this.CreateClient(namespaceInfo.Id);
+			IStorageClient client = this.CreateClient(namespaceInfo.Id);
 
 			Stopwatch timer = Stopwatch.StartNew();
 			_logger.LogInformation("Running garbage collection for namespace {NamespaceId}...", namespaceInfo.Id);
diff --git a/Engine/Source/Programs/Horde/Plugins/Tools/HordeServer.Tools/Commands/Generate/ToolDataCommand.cs b/Engine/Source/Programs/Horde/Plugins/Tools/HordeServer.Tools/Commands/Generate/ToolDataCommand.cs
--- a/Engine/Source/Programs/Horde/Plugins/Tools/HordeServer.Tools/Commands/Generate/ToolDataCommand.cs
+++ b/Engine/Source/Programs/Horde/Plugins/Tools/HordeServer.Tools/Commands/Generate/ToolDataCommand.cs
@@ -71,22 +71,22 @@ namespace HordeServer.Commands.Generate
 
 			RefName refName = new RefName(Id);
 			await using (BundleCache bundleCache = new BundleCache())
+			using (MemoryMappedFileCache memoryMappedFileCache = new MemoryMappedFileCache())
 			{
-				using (IStorageClient client = BundleStorageClient.CreateFromDirectory(bundleDir, bundleCache, logger))
+				IStorageClient client = BundleStorageClient.CreateFromDirectory(bundleDir, bundleCache, memoryMappedFileCache, logger);
+
+				IHashedBlobRef<DirectoryNode> dirNodeRef;
+				await using (IBlobWriter writer = client.CreateBlobWriter(refName))
 				{
-					IHashedBlobRef<DirectoryNode> dirNodeRef;
-					await using (IBlobWriter writer = client.CreateBlobWriter(refName))
-					{
-						logger.LogInformation("");
-						logger.LogInformation("Writing tool data for {ToolId}", refName);
-						DirectoryNode dirNode = new DirectoryNode();
-						await dirNode.AddFilesAsync(InputDir.ToDirectoryInfo(), writer);
-						dirNodeRef = await writer.WriteBlobAsync(dirNode);
+					logger.LogInformation("");
+					logger.LogInformation("Writing tool data for {ToolId}", refName);
+					DirectoryNode dirNode = new DirectoryNode();
+					await dirNode.AddFilesAsync(InputDir.ToDirectoryInfo(), writer);
+					dirNodeRef = await writer.WriteBlobAsync(dirNode);
 
-						logger.LogInformation("");
-					}
-					await client.WriteRefAsync(refName, dirNodeRef);
+					logger.LogInformation("");
 				}
+				await client.WriteRefAsync(refName, dirNodeRef);
 			}
 
 			// Update the server config to include the bundled tool
diff --git a/Engine/Source/Programs/Horde/Plugins/Tools/HordeServer.Tools/Tools/ToolCollection.cs b/Engine/Source/Programs/Horde/Plugins/Tools/HordeServer.Tools/Tools/ToolCollection.cs
--- a/Engine/Source/Programs/Horde/Plugins/Tools/HordeServer.Tools/Tools/ToolCollection.cs
+++ b/Engine/Source/Programs/Horde/Plugins/Tools/HordeServer.Tools/Tools/ToolCollection.cs
@@ -194,7 +194,8 @@ namespace HordeServer.Tools
 		private readonly StorageService _storageService;
 		private readonly FileObjectStoreFactory _fileObjectStoreFactory;
 		private readonly IClock _clock;
-		private readonly BundleCache _cache;
+		private readonly BundleCache _bundleCache;
+		private readonly MemoryMappedFileCache _memoryMappedFileCache;
 		private readonly IOptionsMonitor<ToolsConfig> _toolsConfig;
 		private readonly IOptions<ToolsServerConfig> _toolsServerConfig;
 		private readonly ILogger _logger;
@@ -202,7 +203,7 @@ namespace HordeServer.Tools
 		/// <summary>
 		/// Constructor
 		/// </summary>
-		public ToolCollection(IMongoService mongoService, IServerInfo serverInfo, StorageService storageService, BundleCache cache, FileObjectStoreFactory fileObjectStoreFactory, IClock clock, IOptionsMonitor<ToolsConfig> toolsConfig, IOptions<ToolsServerConfig> toolsServerConfig, ILogger<ToolCollection> logger)
+		public ToolCollection(IMongoService mongoService, IServerInfo serverInfo, StorageService storageService, BundleCache bundleCache, MemoryMappedFileCache memoryMappedFileCache, FileObjectStoreFactory fileObjectStoreFactory, IClock clock, IOptionsMonitor<ToolsConfig> toolsConfig, IOptions<ToolsServerConfig> toolsServerConfig, ILogger<ToolCollection> logger)
 		{
 			_serverInfo = serverInfo;
 			_tools = mongoService.GetCollection<ToolDocument>("Tools");
@@ -211,7 +212,8 @@ namespace HordeServer.Tools
 			_clock = clock;
 			_toolsConfig = toolsConfig;
 			_toolsServerConfig = toolsServerConfig;
-			_cache = cache;
+			_bundleCache = bundleCache;
+			_memoryMappedFileCache = memoryMappedFileCache;
 			_logger = logger;
 		}
 
@@ -326,7 +328,7 @@ namespace HordeServer.Tools
 		{
 			ToolDeploymentId deploymentId = new ToolDeploymentId(BinaryIdUtils.CreateNew());
 
-			using IStorageClient client = _storageService.CreateClient(toolConfig.NamespaceId);
+			IStorageClient client = _storageService.CreateClient(toolConfig.NamespaceId);
 
 			IHashedBlobRef<DirectoryNode> nodeRef;
 			await using (IBlobWriter writer = client.CreateBlobWriter($"{tool.Id}/{deploymentId}"))
@@ -343,7 +345,7 @@ namespace HordeServer.Tools
 		{
 			ToolDeploymentId deploymentId = new ToolDeploymentId(BinaryIdUtils.CreateNew());
 
-			using IStorageClient client = _storageService.CreateClient(toolConfig.NamespaceId);
+			IStorageClient client = _storageService.CreateClient(toolConfig.NamespaceId);
 			return await CreateDeploymentInternalAsync(tool, toolConfig, deploymentId, options, client, client.CreateBlobRef(target), cancellationToken);
 		}
 
@@ -401,7 +403,7 @@ namespace HordeServer.Tools
 				}
 
 				ToolDeploymentDocument removeDeployment = tool.Deployments[0];
-				using IStorageClient client = _storageService.CreateClient(removeDeployment.NamespaceId);
+				IStorageClient client = _storageService.CreateClient(removeDeployment.NamespaceId);
 				await client.DeleteRefAsync(removeDeployment.RefName, cancellationToken);
 			}
 
@@ -461,7 +463,7 @@ namespace HordeServer.Tools
 		{
 			if (toolConfig is BundledToolConfig bundledConfig)
 			{
-				return BundleStorageClient.CreateFromDirectory(DirectoryReference.Combine(_serverInfo.AppDir, bundledConfig.DataDir ?? "Tools"), _cache, _logger);
+				return BundleStorageClient.CreateFromDirectory(DirectoryReference.Combine(_serverInfo.AppDir, bundledConfig.DataDir ?? "Tools"), _bundleCache, _memoryMappedFileCache, _logger);
 			}
 			else
 			{
@@ -497,16 +499,9 @@ namespace HordeServer.Tools
 		{
 #pragma warning disable CA2000
 			IStorageClient client = CreateStorageClient(tool);
-			try
-			{
-				IHashedBlobRef<DirectoryNode> nodeRef = await client.ReadRefAsync<DirectoryNode>(deployment.RefName, DateTime.UtcNow - TimeSpan.FromDays(2.0), cancellationToken: cancellationToken);
-				return nodeRef.AsZipStream().WrapOwnership(client);
-			}
-			catch
-			{
-				client.Dispose();
-				throw;
-			}
+
+			IHashedBlobRef<DirectoryNode> nodeRef = await client.ReadRefAsync<DirectoryNode>(deployment.RefName, DateTime.UtcNow - TimeSpan.FromDays(2.0), cancellationToken: cancellationToken);
+			return nodeRef.AsZipStream();
 #pragma warning restore CA2000
 		}
 
diff --git a/Engine/Source/Programs/Horde/Plugins/Tools/HordeServer.Tools/Tools/ToolsController.cs b/Engine/Source/Programs/Horde/Plugins/Tools/HordeServer.Tools/Tools/ToolsController.cs
--- a/Engine/Source/Programs/Horde/Plugins/Tools/HordeServer.Tools/Tools/ToolsController.cs
+++ b/Engine/Source/Programs/Horde/Plugins/Tools/HordeServer.Tools/Tools/ToolsController.cs
@@ -318,44 +318,36 @@ namespace HordeServer.Tools
 			}
 
 			IStorageClient client = tool.CreateStorageClient();
-			try
-			{
-				IHashedBlobRef<DirectoryNode> nodeRef = await client.ReadRefAsync<DirectoryNode>(deployment.RefName, DateTime.UtcNow - TimeSpan.FromDays(2.0), cancellationToken: cancellationToken);
+			IHashedBlobRef<DirectoryNode> nodeRef = await client.ReadRefAsync<DirectoryNode>(deployment.RefName, DateTime.UtcNow - TimeSpan.FromDays(2.0), cancellationToken: cancellationToken);
 
-				// If we weren't specifically asked for a zip, see if this download is a single file. If it is, allow downloading it directory.
-				if (action != GetToolAction.Zip)
+			// If we weren't specifically asked for a zip, see if this download is a single file. If it is, allow downloading it directory.
+			if (action != GetToolAction.Zip)
+			{
+				DirectoryNode node = await nodeRef.ReadBlobAsync(cancellationToken);
+				if (node.Directories.Count == 0 && node.Files.Count == 1)
 				{
-					DirectoryNode node = await nodeRef.ReadBlobAsync(cancellationToken);
-					if (node.Directories.Count == 0 && node.Files.Count == 1)
-					{
-						FileEntry entry = node.Files.First();
+					FileEntry entry = node.Files.First();
 
-						string? contentType;
-						if (!new FileExtensionContentTypeProvider().TryGetContentType(entry.Name.ToString(), out contentType))
-						{
-							contentType = "application/octet-stream";
-						}
+					string? contentType;
+					if (!new FileExtensionContentTypeProvider().TryGetContentType(entry.Name.ToString(), out contentType))
+					{
+						contentType = "application/octet-stream";
+					}
 
-						Response.Headers.ContentLength = entry.Length;
+					Response.Headers.ContentLength = entry.Length;
 
-						Stream fileStream = entry.OpenAsStream().WrapOwnership(client);
-						return new FileStreamResult(fileStream, contentType) { FileDownloadName = entry.Name.ToString() };
-					}
+					Stream fileStream = entry.OpenAsStream();
+					return new FileStreamResult(fileStream, contentType) { FileDownloadName = entry.Name.ToString() };
 				}
-
-				Stream stream = nodeRef.AsZipStream().WrapOwnership(client);
-				return new FileStreamResult(stream, "application/zip") { FileDownloadName = $"{tool.Id}-{deployment.Version}.zip" };
-			}
-			catch
-			{
-				client.Dispose();
-				throw;
 			}
+
+			Stream stream = nodeRef.AsZipStream();
+			return new FileStreamResult(stream, "application/zip") { FileDownloadName = $"{tool.Id}-{deployment.Version}.zip" };
 		}
 
 		static async Task<GetToolDeploymentResponse> GetDeploymentInfoResponseAsync(ITool tool, IToolDeployment deployment, CancellationToken cancellationToken)
 		{
-			using IStorageClient client = tool.CreateStorageClient();
+			IStorageClient client = tool.CreateStorageClient();
 			IBlobRef rootHandle = await client.ReadRefAsync(deployment.RefName, cancellationToken: cancellationToken);
 
 			return new GetToolDeploymentResponse(deployment.Id, deployment.Version, deployment.State, deployment.Progress, deployment.StartedAt, deployment.Duration, deployment.RefName, rootHandle.GetLocator());
diff --git a/Engine/Source/Programs/Horde/Samples/RemoteClient/RemoteClientApp.cs b/Engine/Source/Programs/Horde/Samples/RemoteClient/RemoteClientApp.cs
--- a/Engine/Source/Programs/Horde/Samples/RemoteClient/RemoteClientApp.cs
+++ b/Engine/Source/Programs/Horde/Samples/RemoteClient/RemoteClientApp.cs
@@ -110,7 +110,7 @@ namespace RemoteClient
 				await channel.ForkAsync(BackgroundChannelId, 4 * 1024 * 1024, default);
 
 				// Upload the sandbox to the primary channel.
-				using BundleStorageClient storage =  BundleStorageClient.CreateInMemory(logger);
+				BundleStorageClient storage = BundleStorageClient.CreateInMemory(logger);
 
 				await using (IBlobWriter writer = storage.CreateBlobWriter())
 				{
diff --git a/Engine/Source/Programs/Shared/EpicGames.Horde.Tests/BlobPipelineTests.cs b/Engine/Source/Programs/Shared/EpicGames.Horde.Tests/BlobPipelineTests.cs
--- a/Engine/Source/Programs/Shared/EpicGames.Horde.Tests/BlobPipelineTests.cs
+++ b/Engine/Source/Programs/Shared/EpicGames.Horde.Tests/BlobPipelineTests.cs
@@ -20,7 +20,7 @@ namespace EpicGames.Horde.Tests
 		[TestMethod]
 		public async Task TestAsync()
 		{
-			using BundleStorageClient client = BundleStorageClient.CreateInMemory(NullLogger.Instance);
+			BundleStorageClient client = BundleStorageClient.CreateInMemory(NullLogger.Instance);
 
 			List<IHashedBlobRef> handles = new List<IHashedBlobRef>();
 			await using (IBlobWriter writer = client.CreateBlobWriter())
diff --git a/Engine/Source/Programs/Shared/EpicGames.Horde.Tests/BundleStoreTests.cs b/Engine/Source/Programs/Shared/EpicGames.Horde.Tests/BundleStoreTests.cs
--- a/Engine/Source/Programs/Shared/EpicGames.Horde.Tests/BundleStoreTests.cs
+++ b/Engine/Source/Programs/Shared/EpicGames.Horde.Tests/BundleStoreTests.cs
@@ -61,7 +61,7 @@ namespace EpicGames.Horde.Tests
 		static async Task<byte[]> CreateBundleNormalAsync()
 		{
 			MemoryStorageBackend memoryStore = new MemoryStorageBackend();
-			using BundleStorageClient store = new BundleStorageClient(memoryStore, BundleCache.None, new BundleOptions { MaxVersion = BundleVersion.ImportHashes, CompressionFormat = BundleCompressionFormat.None }, NullLogger.Instance);
+			BundleStorageClient store = new BundleStorageClient(memoryStore, BundleCache.None, new BundleOptions { MaxVersion = BundleVersion.ImportHashes, CompressionFormat = BundleCompressionFormat.None }, NullLogger.Instance);
 			await using IBlobWriter writer = store.CreateBlobWriter();
 
 			TextNode node = new TextNode("Hello world");
@@ -101,7 +101,7 @@ namespace EpicGames.Horde.Tests
 			MemoryStorageBackend blobStore = new MemoryStorageBackend();
 
 			BundleOptions bundleOptions = new BundleOptions { MaxBlobSize = 1024 * 1024 };
-			using BundleStorageClient bundleStore = new BundleStorageClient(blobStore, BundleCache.None, bundleOptions, NullLogger.Instance);
+			BundleStorageClient bundleStore = new BundleStorageClient(blobStore, BundleCache.None, bundleOptions, NullLogger.Instance);
 
 			await TestTreeAsync(bundleStore);
 
@@ -115,7 +115,7 @@ namespace EpicGames.Horde.Tests
 			MemoryStorageBackend blobStore = new MemoryStorageBackend();
 
 			BundleOptions bundleOptions = new BundleOptions { MaxBlobSize = 1 };
-			using BundleStorageClient bundleStore = new BundleStorageClient(blobStore, BundleCache.None, bundleOptions, NullLogger.Instance);
+			BundleStorageClient bundleStore = new BundleStorageClient(blobStore, BundleCache.None, bundleOptions, NullLogger.Instance);
 
 			await TestTreeAsync(bundleStore);
 
@@ -219,7 +219,7 @@ namespace EpicGames.Horde.Tests
 		[TestMethod]
 		public async Task SimpleNodeAsync()
 		{
-			using KeyValueStorageClient store = KeyValueStorageClient.CreateInMemory();
+			KeyValueStorageClient store = KeyValueStorageClient.CreateInMemory();
 
 			RefName refName = new RefName("test");
 
@@ -238,7 +238,7 @@ namespace EpicGames.Horde.Tests
 		[TestMethod]
 		public async Task DirectoryNodesAsync()
 		{
-			using KeyValueStorageClient store = KeyValueStorageClient.CreateInMemory();
+			KeyValueStorageClient store = KeyValueStorageClient.CreateInMemory();
 
 			// Generate a tree
 			{
@@ -285,7 +285,7 @@ namespace EpicGames.Horde.Tests
 		public async Task FileNodesAsync()
 		{
 			using IMemoryCache cache = new MemoryCache(new MemoryCacheOptions());
-			using BundleStorageClient store = BundleStorageClient.CreateInMemory(NullLogger.Instance);
+			BundleStorageClient store = BundleStorageClient.CreateInMemory(NullLogger.Instance);
 			BundleReader reader = new BundleReader(store, BundleCache.None, NullLogger.Instance);
 
 			// Generate a tree
@@ -334,7 +334,7 @@ namespace EpicGames.Horde.Tests
 		public async Task StreamTestAsync()
 		{
 			BundleOptions bundleOptions = new BundleOptions { MaxBlobSize = 1024 };
-			using BundleStorageClient store = BundleStorageClient.CreateInMemory(bundleOptions, NullLogger.Instance);
+			BundleStorageClient store = BundleStorageClient.CreateInMemory(bundleOptions, NullLogger.Instance);
 
 			const int Length = 4096;
 
@@ -374,7 +374,7 @@ namespace EpicGames.Horde.Tests
 		public async Task LargeFileTestAsync()
 		{
 			BundleOptions bundleOptions = new BundleOptions { MaxBlobSize = 1024 };
-			using BundleStorageClient store = BundleStorageClient.CreateInMemory(bundleOptions, NullLogger.Instance);
+			BundleStorageClient store = BundleStorageClient.CreateInMemory(bundleOptions, NullLogger.Instance);
 
 			const int Length = 1024;
 			const int Copies = 4096;
diff --git a/Engine/Source/Programs/Shared/EpicGames.Horde.Tests/BundleTests.cs b/Engine/Source/Programs/Shared/EpicGames.Horde.Tests/BundleTests.cs
--- a/Engine/Source/Programs/Shared/EpicGames.Horde.Tests/BundleTests.cs
+++ b/Engine/Source/Programs/Shared/EpicGames.Horde.Tests/BundleTests.cs
@@ -52,7 +52,7 @@ namespace EpicGames.Horde.Tests
 		[TestMethod]
 		public async Task BasicChunkingTestsAsync()
 		{
-			using BundleStorageClient storage = BundleStorageClient.CreateInMemory(NullLogger.Instance);
+			BundleStorageClient storage = BundleStorageClient.CreateInMemory(NullLogger.Instance);
 
 			RefName refName = new RefName("test");
 			await using IBlobWriter writer = storage.CreateBlobWriter(refName);
@@ -175,7 +175,7 @@ namespace EpicGames.Horde.Tests
 		[TestMethod]
 		public async Task BasicTestDirectoryAsync()
 		{
-			using BundleStorageClient store = BundleStorageClient.CreateInMemory(NullLogger.Instance);
+			BundleStorageClient store = BundleStorageClient.CreateInMemory(NullLogger.Instance);
 
 			IHashedBlobRef<DirectoryNode> rootRef;
 			await using (IBlobWriter writer = store.CreateBlobWriter())
@@ -235,7 +235,7 @@ namespace EpicGames.Horde.Tests
 			BundleOptions bundleOptions = new BundleOptions();
 			bundleOptions.MaxBlobSize = 1;
 
-			using BundleStorageClient storage = BundleStorageClient.CreateInMemory(bundleOptions, NullLogger.Instance);
+			BundleStorageClient storage = BundleStorageClient.CreateInMemory(bundleOptions, NullLogger.Instance);
 
 			await using (IBlobWriter writer = new DedupeBlobWriter(storage.CreateBlobWriter()))
 			{
@@ -263,7 +263,7 @@ namespace EpicGames.Horde.Tests
 			BundleOptions bundleOptions = new BundleOptions();
 			bundleOptions.MaxBlobSize = 1;
 
-			using BundleStorageClient storage = BundleStorageClient.CreateInMemory(bundleOptions, NullLogger.Instance);
+			BundleStorageClient storage = BundleStorageClient.CreateInMemory(bundleOptions, NullLogger.Instance);
 
 			RefName refName = new RefName("ref");
 
diff --git a/Engine/Source/Programs/Shared/EpicGames.Horde.Tests/ChunkingTests.cs b/Engine/Source/Programs/Shared/EpicGames.Horde.Tests/ChunkingTests.cs
--- a/Engine/Source/Programs/Shared/EpicGames.Horde.Tests/ChunkingTests.cs
+++ b/Engine/Source/Programs/Shared/EpicGames.Horde.Tests/ChunkingTests.cs
@@ -52,7 +52,7 @@ namespace EpicGames.Horde.Tests
 		[TestMethod]
 		public async Task EmptyNodeTestAsync()
 		{
-			using KeyValueStorageClient store = KeyValueStorageClient.CreateInMemory();
+			KeyValueStorageClient store = KeyValueStorageClient.CreateInMemory();
 
 			const string RefName = "hello";
 			await using (IBlobWriter writer = store.CreateBlobWriter(RefName))
@@ -97,7 +97,7 @@ namespace EpicGames.Horde.Tests
 		{
 			using MemoryCache cache = new MemoryCache(new MemoryCacheOptions());
 
-			using KeyValueStorageClient store = KeyValueStorageClient.CreateInMemory();
+			KeyValueStorageClient store = KeyValueStorageClient.CreateInMemory();
 
 			await using IBlobWriter writer = store.CreateBlobWriter();
 
@@ -176,7 +176,7 @@ namespace EpicGames.Horde.Tests
 			BlobSerializerOptions serializerOptions = new BlobSerializerOptions();
 			serializerOptions.Converters.Add(new InteriorChunkedDataNodeConverter(HordeApiVersion.Initial)); // Does not include length fields in interior nodes
 
-			using KeyValueStorageClient store = KeyValueStorageClient.CreateInMemory();
+			KeyValueStorageClient store = KeyValueStorageClient.CreateInMemory();
 
 			byte[] data = Encoding.UTF8.GetBytes("hello world");
 
diff --git a/Engine/Source/Programs/Shared/EpicGames.Horde.Tests/ComputeSocketTests.cs b/Engine/Source/Programs/Shared/EpicGames.Horde.Tests/ComputeSocketTests.cs
--- a/Engine/Source/Programs/Shared/EpicGames.Horde.Tests/ComputeSocketTests.cs
+++ b/Engine/Source/Programs/Shared/EpicGames.Horde.Tests/ComputeSocketTests.cs
@@ -175,7 +175,7 @@ namespace EpicGames.Horde.Tests
 						await channel2.CloseAsync(cancellationToken);
 					}
 
-					using BundleStorageClient storage = BundleStorageClient.CreateInMemory(NullLogger.Instance);
+					BundleStorageClient storage = BundleStorageClient.CreateInMemory(NullLogger.Instance);
 					await using (IBlobWriter blobWriter = storage.CreateBlobWriter())
 					{
 						FileReference file = FileReference.Combine(tempDir, "subdir/hello.txt");
diff --git a/Engine/Source/Programs/Shared/EpicGames.Horde.Tests/LogTests.cs b/Engine/Source/Programs/Shared/EpicGames.Horde.Tests/LogTests.cs
--- a/Engine/Source/Programs/Shared/EpicGames.Horde.Tests/LogTests.cs
+++ b/Engine/Source/Programs/Shared/EpicGames.Horde.Tests/LogTests.cs
@@ -40,7 +40,7 @@ namespace EpicGames.Horde.Tests
 		[TestMethod]
 		public async Task IndexTestsAsync()
 		{
-			using KeyValueStorageClient store = KeyValueStorageClient.CreateInMemory();
+			KeyValueStorageClient store = KeyValueStorageClient.CreateInMemory();
 
 			// Write the test data to the log file in blocks
 			LogBuilder builder = new LogBuilder(LogFormat.Text, 1, 1, NullLogger.Instance);
@@ -101,7 +101,7 @@ namespace EpicGames.Horde.Tests
 		[TestMethod]
 		public async Task PartialTokenTestsAsync()
 		{
-			using IStorageClient store = KeyValueStorageClient.CreateInMemory();
+			IStorageClient store = KeyValueStorageClient.CreateInMemory();
 
 			// Generate the test data
 			string[] lines =
diff --git a/Engine/Source/Programs/Shared/EpicGames.Horde.Tests/ServerLogPacketBuilderTests.cs b/Engine/Source/Programs/Shared/EpicGames.Horde.Tests/ServerLogPacketBuilderTests.cs
--- a/Engine/Source/Programs/Shared/EpicGames.Horde.Tests/ServerLogPacketBuilderTests.cs
+++ b/Engine/Source/Programs/Shared/EpicGames.Horde.Tests/ServerLogPacketBuilderTests.cs
@@ -239,10 +239,6 @@ namespace EpicGames.Horde.Tests
 
 			public ValueTask DisposeAsync()
 			{
-				foreach (BundleStorageClient storageClient in StorageClients.Values)
-				{
-					storageClient.Dispose();
-				}
 				StorageClients.Clear();
 				return default;
 			}
@@ -287,7 +283,7 @@ namespace EpicGames.Horde.Tests
 			}
 
 			// Read the log
-			using IStorageClient storageClient = hordeClient.CreateStorageClient(logId);
+			IStorageClient storageClient = hordeClient.CreateStorageClient(logId);
 			LogNode file = await storageClient.CreateBlobRef(hordeClient.LogRpc.Logs[logId]).ReadBlobAsync<LogNode>();
 
 			// Check the index text
diff --git a/Engine/Source/Programs/Shared/EpicGames.Horde.Tests/StorageClientTests.cs b/Engine/Source/Programs/Shared/EpicGames.Horde.Tests/StorageClientTests.cs
--- a/Engine/Source/Programs/Shared/EpicGames.Horde.Tests/StorageClientTests.cs
+++ b/Engine/Source/Programs/Shared/EpicGames.Horde.Tests/StorageClientTests.cs
@@ -62,14 +62,14 @@ namespace EpicGames.Horde.Tests
 		[TestMethod]
 		public async Task TestBasicAsync()
 		{
-			using KeyValueStorageClient store = KeyValueStorageClient.CreateInMemory();
+			KeyValueStorageClient store = KeyValueStorageClient.CreateInMemory();
 			await TestBasicAsync(store);
 		}
 
 		[TestMethod]
 		public async Task TestBasicBundleV2Async()
 		{
-			using BundleStorageClient storeV2 = BundleStorageClient.CreateInMemory(NullLogger.Instance);
+			BundleStorageClient storeV2 = BundleStorageClient.CreateInMemory(NullLogger.Instance);
 			await TestBasicAsync(storeV2);
 		}
 
@@ -91,7 +91,7 @@ namespace EpicGames.Horde.Tests
 		{
 			await using BundleCache cache = new BundleCache();
 
-			using BundleStorageClient store = BundleStorageClient.CreateInMemory(NullLogger.Instance);
+			BundleStorageClient store = BundleStorageClient.CreateInMemory(NullLogger.Instance);
 
 			IHashedBlobRef<TestNode> nodeRef2;
 			await using (IBlobWriter writer = store.CreateBlobWriter())
@@ -116,7 +116,7 @@ namespace EpicGames.Horde.Tests
 			await using BundleCache cache = new BundleCache();
 
 			BundleOptions bundleOptions = new BundleOptions { MinCompressionPacketSize = 100, MaxBlobSize = 1024 * 1024, MaxVersion = BundleVersion.LatestV2 };
-			using BundleStorageClient store = BundleStorageClient.CreateInMemory(bundleOptions, NullLogger.Instance);
+			BundleStorageClient store = BundleStorageClient.CreateInMemory(bundleOptions, NullLogger.Instance);
 
 			await using IBlobWriter writer = store.CreateBlobWriter();
 			IHashedBlobRef<TestNode> nodeRef1 = await writer.WriteBlobAsync(new TestNode(123) { Padding = new byte[1024] });
diff --git a/Engine/Source/Programs/Shared/EpicGames.Horde/Compute/AgentMessageHandler.cs b/Engine/Source/Programs/Shared/EpicGames.Horde/Compute/AgentMessageHandler.cs
--- a/Engine/Source/Programs/Shared/EpicGames.Horde/Compute/AgentMessageHandler.cs
+++ b/Engine/Source/Programs/Shared/EpicGames.Horde/Compute/AgentMessageHandler.cs
@@ -191,7 +191,7 @@ namespace EpicGames.Horde.Compute
 			await using BundleCache cache = new BundleCache(new BundleCacheOptions { HeaderCacheSize = 10 * 1024 * 1024, PacketCacheSize = 128 * 1024 * 1024 });
 
 			BundleOptions bundleOptions = ComputeProtocolUtilities.GetBundleOptions(channel.Protocol);
-			using BundleStorageClient store = new BundleStorageClient(innerStore, cache, bundleOptions, _logger);
+			BundleStorageClient store = new BundleStorageClient(innerStore, cache, bundleOptions, _logger);
 
 			IBlobRef<DirectoryNode> directoryRef = store.CreateBlobRef<DirectoryNode>(locator, options);
 
diff --git a/Engine/Source/Programs/Shared/EpicGames.Horde/Logs/ServerLogger.cs b/Engine/Source/Programs/Shared/EpicGames.Horde/Logs/ServerLogger.cs
--- a/Engine/Source/Programs/Shared/EpicGames.Horde/Logs/ServerLogger.cs
+++ b/Engine/Source/Programs/Shared/EpicGames.Horde/Logs/ServerLogger.cs
@@ -122,11 +122,6 @@ namespace EpicGames.Horde.Logs
 			{
 				await _writer.DisposeAsync();
 			}
-
-			if (_store != null)
-			{
-				_store.Dispose();
-			}
 		}
 
 		/// <summary>
diff --git a/Engine/Source/Programs/Shared/EpicGames.Horde/Storage/Bundles/BundleStorageClient.cs b/Engine/Source/Programs/Shared/EpicGames.Horde/Storage/Bundles/BundleStorageClient.cs
--- a/Engine/Source/Programs/Shared/EpicGames.Horde/Storage/Bundles/BundleStorageClient.cs
+++ b/Engine/Source/Programs/Shared/EpicGames.Horde/Storage/Bundles/BundleStorageClient.cs
@@ -25,7 +25,6 @@ namespace EpicGames.Horde.Storage.Bundles
 		readonly BundleOptions _options;
 		readonly PacketReaderStats _packetReaderStats = new PacketReaderStats();
 		readonly Bundles.V1.BundleReader _bundleReader;
-		readonly IDisposable? _ownedResources;
 
 		internal Bundles.V1.BundleReader BundleReader => _bundleReader;
 
@@ -52,19 +51,12 @@ namespace EpicGames.Horde.Storage.Bundles
 		/// <summary>
 		/// Constructor
 		/// </summary>
-		public BundleStorageClient(IStorageBackend backend, BundleCache cache, BundleOptions? options, ILogger logger, IDisposable? ownedResources = null)
+		public BundleStorageClient(IStorageBackend backend, BundleCache cache, BundleOptions? options, ILogger logger)
 		{
 			_backend = backend;
 			_cache = cache;
 			_options = options ?? BundleOptions.Default;
 			_bundleReader = new Bundles.V1.BundleReader(this, cache, logger);
-			_ownedResources = ownedResources;
-		}
-
-		/// <inheritdoc/>
-		public void Dispose()
-		{
-			_ownedResources?.Dispose();
 		}
 
 		/// <summary>
@@ -85,17 +77,16 @@ namespace EpicGames.Horde.Storage.Bundles
 		/// <summary>
 		/// Creates a bundle storage client around a directory on the filesystem
 		/// </summary>
-		public static BundleStorageClient CreateFromDirectory(DirectoryReference rootDir, BundleCache cache, ILogger logger)
-			=> CreateFromDirectory(rootDir, cache, null, logger);
+		public static BundleStorageClient CreateFromDirectory(DirectoryReference rootDir, BundleCache cache, MemoryMappedFileCache memoryMappedFileCache, ILogger logger)
+			=> CreateFromDirectory(rootDir, cache, memoryMappedFileCache, null, logger);
 
 		/// <summary>
 		/// Creates a bundle storage client around a directory on the filesystem
 		/// </summary>
-		public static BundleStorageClient CreateFromDirectory(DirectoryReference rootDir, BundleCache cache, BundleOptions? options, ILogger logger)
+		public static BundleStorageClient CreateFromDirectory(DirectoryReference rootDir, BundleCache cache, MemoryMappedFileCache memoryMappedFileCache, BundleOptions? options, ILogger logger)
 		{
-			MemoryMappedFileCache memoryMappedFileCache = new MemoryMappedFileCache();
 			FileStorageBackend backend = new FileStorageBackend(new FileObjectStore(rootDir, memoryMappedFileCache), logger);
-			return new BundleStorageClient(backend, cache, options, logger, memoryMappedFileCache);
+			return new BundleStorageClient(backend, cache, options, logger);
 		}
 
 		/// <summary>
diff --git a/Engine/Source/Programs/Shared/EpicGames.Horde/Storage/Clients/KeyValueStorageClient.cs b/Engine/Source/Programs/Shared/EpicGames.Horde/Storage/Clients/KeyValueStorageClient.cs
--- a/Engine/Source/Programs/Shared/EpicGames.Horde/Storage/Clients/KeyValueStorageClient.cs
+++ b/Engine/Source/Programs/Shared/EpicGames.Horde/Storage/Clients/KeyValueStorageClient.cs
@@ -127,11 +127,6 @@ namespace EpicGames.Horde.Storage.Clients
 			_backend = backend;
 		}
 
-		/// <inheritdoc/>
-		public void Dispose()
-		{
-		}
-
 		/// <summary>
 		/// Create an in-memory storage client
 		/// </summary>
diff --git a/Engine/Source/Programs/Shared/EpicGames.Horde/Storage/IStorageClient.cs b/Engine/Source/Programs/Shared/EpicGames.Horde/Storage/IStorageClient.cs
--- a/Engine/Source/Programs/Shared/EpicGames.Horde/Storage/IStorageClient.cs
+++ b/Engine/Source/Programs/Shared/EpicGames.Horde/Storage/IStorageClient.cs
@@ -21,7 +21,7 @@ namespace EpicGames.Horde.Storage
 	/// <summary>
 	/// Interface for the storage system.
 	/// </summary>
-	public interface IStorageClient : IDisposable
+	public interface IStorageClient
 	{
 		#region Blobs
 
diff --git a/Engine/Source/Programs/UnrealBuildTool/Artifacts/ActionArtifactCache.cs b/Engine/Source/Programs/UnrealBuildTool/Artifacts/ActionArtifactCache.cs
--- a/Engine/Source/Programs/UnrealBuildTool/Artifacts/ActionArtifactCache.cs
+++ b/Engine/Source/Programs/UnrealBuildTool/Artifacts/ActionArtifactCache.cs
@@ -164,11 +164,12 @@ namespace UnrealBuildTool.Artifacts
 		/// </summary>
 		/// <param name="directory">Directory for the cache</param>
 		/// <param name="cppDependencyCache">Previously created dependency cache</param>
+		/// <param name="memoryMappedFileCache">Cache for memory mapped files</param>
 		/// <param name="logger">Logging device</param>
 		/// <returns>Action artifact cache object</returns>
-		public static IActionArtifactCache CreateHordeFileCache(DirectoryReference directory, CppDependencyCache cppDependencyCache, ILogger logger)
+		public static IActionArtifactCache CreateHordeFileCache(DirectoryReference directory, CppDependencyCache cppDependencyCache, MemoryMappedFileCache memoryMappedFileCache, ILogger logger)
 		{
-			IArtifactCache artifactCache = HordeStorageArtifactCache.CreateFileCache(directory, NullLogger.Instance, false);
+			IArtifactCache artifactCache = HordeStorageArtifactCache.CreateFileCache(directory, memoryMappedFileCache, NullLogger.Instance, false);
 			return new ActionArtifactCache(artifactCache, cppDependencyCache, logger);
 		}
 
diff --git a/Engine/Source/Programs/UnrealBuildTool/Artifacts/HordeStorageArtifactCache.cs b/Engine/Source/Programs/UnrealBuildTool/Artifacts/HordeStorageArtifactCache.cs
--- a/Engine/Source/Programs/UnrealBuildTool/Artifacts/HordeStorageArtifactCache.cs
+++ b/Engine/Source/Programs/UnrealBuildTool/Artifacts/HordeStorageArtifactCache.cs
@@ -202,7 +202,6 @@ namespace UnrealBuildTool.Artifacts
 		/// <inheritdoc/>
 		public void Dispose()
 		{
-			_store?.Dispose();
 			_pendingWritesFlushTask?.Dispose();
 			_pendingWritesFlushTask = null;
 			_semaphore.Dispose();
@@ -225,13 +224,14 @@ namespace UnrealBuildTool.Artifacts
 		/// Create a file based cache
 		/// </summary>
 		/// <param name="directory">Destination directory</param>
+		/// <param name="memoryMappedFileCache">Cache for memory mapped files</param>
 		/// <param name="logger">Logging object</param>
 		/// <param name="cleanDirectory">If true, clean the directory</param>
 		/// <returns>Storage client instance</returns>
-		public static IArtifactCache CreateFileCache(DirectoryReference directory, ILogger logger, bool cleanDirectory)
+		public static IArtifactCache CreateFileCache(DirectoryReference directory, MemoryMappedFileCache memoryMappedFileCache, ILogger logger, bool cleanDirectory)
 		{
 			HordeStorageArtifactCache cache = new(null);
-			cache._readyTask = Task.Run(() => cache.InitFileCache(directory, logger, cleanDirectory));
+			cache._readyTask = Task.Run(() => cache.InitFileCache(directory, memoryMappedFileCache, logger, cleanDirectory));
 			return cache;
 		}
 
@@ -492,10 +492,11 @@ namespace UnrealBuildTool.Artifacts
 		/// Initialize a file based cache
 		/// </summary>
 		/// <param name="directory">Destination directory</param>
+		/// <param name="memoryMappedFileCache">Cache for memory mapped files</param>
 		/// <param name="logger">Logger</param>
 		/// <param name="cleanDirectory">If true, clean the directory</param>
 		/// <returns>Cache state</returns>
-		private ArtifactCacheState InitFileCache(DirectoryReference directory, ILogger logger, bool cleanDirectory)
+		private ArtifactCacheState InitFileCache(DirectoryReference directory, MemoryMappedFileCache memoryMappedFileCache, ILogger logger, bool cleanDirectory)
 		{
 			try
 			{
@@ -511,7 +512,7 @@ namespace UnrealBuildTool.Artifacts
 				}
 				Directory.CreateDirectory(directory.FullName);
 
-				_store = BundleStorageClient.CreateFromDirectory(directory, BundleCache.None, logger);
+				_store = BundleStorageClient.CreateFromDirectory(directory, BundleCache.None, memoryMappedFileCache, logger);
 
 				State = ArtifactCacheState.Available;
 				return State;
diff --git a/Engine/Source/Programs/UnrealBuildTool/Executors/UnrealBuildAccelerator/UBAAgentCoordinatorHorde.cs b/Engine/Source/Programs/UnrealBuildTool/Executors/UnrealBuildAccelerator/UBAAgentCoordinatorHorde.cs
--- a/Engine/Source/Programs/UnrealBuildTool/Executors/UnrealBuildAccelerator/UBAAgentCoordinatorHorde.cs
+++ b/Engine/Source/Programs/UnrealBuildTool/Executors/UnrealBuildAccelerator/UBAAgentCoordinatorHorde.cs
@@ -124,8 +124,6 @@ namespace UnrealBuildTool
 
 			await _serviceProvider.DisposeAsync();
 			_cancellationTokenSource.Dispose();
-
-			_storage.Dispose();
 		}
 
 		public async Task InitAsync(bool useSentry, CancellationToken cancellationToken)
diff --git a/Engine/Source/Programs/UnrealBuildTool/Modes/BuildMode.cs b/Engine/Source/Programs/UnrealBuildTool/Modes/BuildMode.cs
--- a/Engine/Source/Programs/UnrealBuildTool/Modes/BuildMode.cs
+++ b/Engine/Source/Programs/UnrealBuildTool/Modes/BuildMode.cs
@@ -415,6 +415,8 @@ namespace UnrealBuildTool
 			await BuildAsync(TargetMakefiles.ToArray(), TargetDescriptors, BuildConfiguration, Options, WriteOutdatedActionsFile, Logger, ActionTypeFilter);
 		}
 
+		static Lazy<MemoryMappedFileCache> s_memoryMappedFileCache = new Lazy<MemoryMappedFileCache>(() => new MemoryMappedFileCache());
+
 		/// <summary>
 		/// Build a list of targets with a given set of makefiles.
 		/// </summary>
@@ -553,7 +555,7 @@ namespace UnrealBuildTool
 
 					Logger.LogInformation("Experimental artifact system using '{Directory}' in {Mode} mode", artifactDirectory.FullName, mode);
 
-					actionArtifactCache = ActionArtifactCache.CreateHordeFileCache(artifactDirectory, CppDependencies, Logger);
+					actionArtifactCache = ActionArtifactCache.CreateHordeFileCache(artifactDirectory, CppDependencies, s_memoryMappedFileCache.Value, Logger);
 					actionArtifactCache.EnableReads = BuildConfiguration.bArtifactRead;
 					actionArtifactCache.EnableWrites = BuildConfiguration.bArtifactWrites;
 					actionArtifactCache.LogCacheMisses = BuildConfiguration.bLogArtifactCacheMisses;
diff --git a/Engine/Source/Programs/UnrealGameSync/UnrealGameSync/ArtifactDownload.cs b/Engine/Source/Programs/UnrealGameSync/UnrealGameSync/ArtifactDownload.cs
--- a/Engine/Source/Programs/UnrealGameSync/UnrealGameSync/ArtifactDownload.cs
+++ b/Engine/Source/Programs/UnrealGameSync/UnrealGameSync/ArtifactDownload.cs
@@ -55,7 +55,7 @@ namespace UnrealGameSync
 				await using ServiceProvider serviceProvider = services.BuildServiceProvider();
 
 				IHordeClient hordeClient = serviceProvider.GetRequiredService<IHordeClient>();
-				using IStorageClient storageClient = hordeClient.CreateStorageClient(baseUri.AbsolutePath);
+				IStorageClient storageClient = hordeClient.CreateStorageClient(baseUri.AbsolutePath);
 
 				progress.Report("Connecting to server...");
 
diff --git a/Engine/Source/Programs/UnrealGameSync/UnrealGameSync/Launcher.cs b/Engine/Source/Programs/UnrealGameSync/UnrealGameSync/Launcher.cs
--- a/Engine/Source/Programs/UnrealGameSync/UnrealGameSync/Launcher.cs
+++ b/Engine/Source/Programs/UnrealGameSync/UnrealGameSync/Launcher.cs
@@ -186,7 +186,7 @@ namespace UnrealGameSync
 						await SafeDeleteDirectoryContentsWithRetryAsync(applicationFolder, cancellationToken);
 
 						// Download and extract the data
-						using IStorageClient storageClient = hordeClient.CreateStorageClient(toolId);
+						IStorageClient storageClient = hordeClient.CreateStorageClient(toolId);
 						IBlobRef<DirectoryNode> handle = storageClient.CreateBlobRef<DirectoryNode>(deployment.Locator);
 
 						DirectoryNode rootNode = await handle.ReadBlobAsync(cancellationToken: cancellationToken);
