Horde: Add the storage system base path to ref responses, allowing us to construct a storage client from a URL.

#rnx

diff --git a/Engine/Source/Programs/Horde/Plugins/Build/HordeServer.Build/Artifacts/ArtifactsController.cs b/Engine/Source/Programs/Horde/Plugins/Build/HordeServer.Build/Artifacts/ArtifactsController.cs
--- a/Engine/Source/Programs/Horde/Plugins/Build/HordeServer.Build/Artifacts/ArtifactsController.cs
+++ b/Engine/Source/Programs/Horde/Plugins/Build/HordeServer.Build/Artifacts/ArtifactsController.cs
@@ -282,7 +282,7 @@ namespace HordeServer.Artifacts
 				return Forbid(ArtifactAclAction.ReadArtifact, artifact.StreamId);
 			}
 
-			return await StorageController.ReadRefInternalAsync(_storageService, artifact.NamespaceId, artifact.RefName, Request.Headers, cancellationToken);
+			return await StorageController.ReadRefInternalAsync(_storageService, $"/api/v2/artifacts/{id}", artifact.NamespaceId, artifact.RefName, Request.Headers, cancellationToken);
 		}
 
 		/// <summary>
diff --git a/Engine/Source/Programs/Horde/Plugins/Storage/HordeServer.Storage/Storage/StorageController.cs b/Engine/Source/Programs/Horde/Plugins/Storage/HordeServer.Storage/Storage/StorageController.cs
--- a/Engine/Source/Programs/Horde/Plugins/Storage/HordeServer.Storage/Storage/StorageController.cs
+++ b/Engine/Source/Programs/Horde/Plugins/Storage/HordeServer.Storage/Storage/StorageController.cs
@@ -262,13 +262,13 @@ namespace HordeServer.Storage
 				return Forbid(StorageAclAction.ReadRefs, namespaceId);
 			}
 
-			return await ReadRefInternalAsync(_storageService, namespaceId, refName, Request.Headers, cancellationToken);
+			return await ReadRefInternalAsync(_storageService, $"/api/v1/storage/{namespaceId}", namespaceId, refName, Request.Headers, cancellationToken);
 		}
 
 		/// <summary>
 		/// Reads a ref from storage, without performing namespace access checks.
 		/// </summary>
-		public static async Task<ActionResult<ReadRefResponse>> ReadRefInternalAsync(IStorageClientFactory storageService, NamespaceId namespaceId, RefName refName, IHeaderDictionary headers, CancellationToken cancellationToken)
+		public static async Task<ActionResult<ReadRefResponse>> ReadRefInternalAsync(IStorageClientFactory storageService, string basePath, NamespaceId namespaceId, RefName refName, IHeaderDictionary headers, CancellationToken cancellationToken)
 		{
 			IStorageClient client = storageService.CreateClient(namespaceId);
 
@@ -287,7 +287,7 @@ namespace HordeServer.Storage
 				return new NotFoundResult();
 			}
 
-			return new ReadRefResponse { Hash = target.Hash, Target = target.GetLocator(), Link = GetNodeLink(namespaceId, target) };
+			return new ReadRefResponse { Hash = target.Hash, Target = target.GetLocator(), Link = GetNodeLink(namespaceId, target), BasePath = basePath };
 		}
 
 		/// <summary>
diff --git a/Engine/Source/Programs/Shared/EpicGames.Horde/Storage/StorageMessages.cs b/Engine/Source/Programs/Shared/EpicGames.Horde/Storage/StorageMessages.cs
--- a/Engine/Source/Programs/Shared/EpicGames.Horde/Storage/StorageMessages.cs
+++ b/Engine/Source/Programs/Shared/EpicGames.Horde/Storage/StorageMessages.cs
@@ -136,5 +136,10 @@ namespace EpicGames.Horde.Storage
 		/// Link to information about the target node
 		/// </summary>
 		public string Link { get; set; } = String.Empty;
+
+		/// <summary>
+		/// Base path for this storage backend
+		/// </summary>
+		public string BasePath { get; set; } = String.Empty;
 	}
 }
