Horde: Enumerate types for artifact search and validation

diff --git a/Engine/Source/Programs/Horde/HordeDashboard/src/backend/Api.ts b/Engine/Source/Programs/Horde/HordeDashboard/src/backend/Api.ts
--- a/Engine/Source/Programs/Horde/HordeDashboard/src/backend/Api.ts
+++ b/Engine/Source/Programs/Horde/HordeDashboard/src/backend/Api.ts
@@ -453,6 +453,9 @@ export type GetDashboardConfigResponse = {
 	/** Categories to display on the pools page */
 	poolCategories: GetDashboardPoolCategoryResponse[];
 
+	/** Categories to display on the pools page */
+	artifactTypes: string[];
+
 	/** Telemetry views */
 	telemetryViews: GetTelemetryViewResponse[];
 }
diff --git a/Engine/Source/Programs/Horde/HordeServer/Dashboard/DashboardController.cs b/Engine/Source/Programs/Horde/HordeServer/Dashboard/DashboardController.cs
--- a/Engine/Source/Programs/Horde/HordeServer/Dashboard/DashboardController.cs
+++ b/Engine/Source/Programs/Horde/HordeServer/Dashboard/DashboardController.cs
@@ -158,6 +158,8 @@ namespace HordeServer.Dashboard
 				dashboardConfigResponse.PoolCategories.Add(new GetDashboardPoolCategoryResponse { Name = category.Name, Condition = category.Condition });
 			}
 
+			dashboardConfigResponse.ArtifactTypes = _globalConfig.Value.Plugins.GetBuildConfig().ArtifactTypes.Select(a => a.Type.ToString()).ToList();
+
 			foreach (IPluginResponseFilter responseFilter in _responseFilters)
 			{
 				responseFilter.Apply(HttpContext, dashboardConfigResponse);
diff --git a/Engine/Source/Programs/Shared/EpicGames.Horde/Dashboard/DashboardMessages.cs b/Engine/Source/Programs/Shared/EpicGames.Horde/Dashboard/DashboardMessages.cs
--- a/Engine/Source/Programs/Shared/EpicGames.Horde/Dashboard/DashboardMessages.cs
+++ b/Engine/Source/Programs/Shared/EpicGames.Horde/Dashboard/DashboardMessages.cs
@@ -60,6 +60,11 @@ namespace EpicGames.Horde.Dashboard
 		/// </summary>
 		public List<GetDashboardPoolCategoryResponse> PoolCategories { get; set; } = new List<GetDashboardPoolCategoryResponse>();
 
+		/// <summary>
+		/// Configured artifact types
+		/// </summary>
+		public List<string> ArtifactTypes { get; set; } = new List<string>();
+
 	}
 
 	/// <summary>
