Horde: Move artifact type config into its own file.

#rnx

diff --git a/Engine/Source/Programs/Horde/Plugins/Build/HordeServer.Build/Artifacts/ArtifactTypeConfig.cs b/Engine/Source/Programs/Horde/Plugins/Build/HordeServer.Build/Artifacts/ArtifactTypeConfig.cs
new file mode 100644
--- /dev/null
+++ b/Engine/Source/Programs/Horde/Plugins/Build/HordeServer.Build/Artifacts/ArtifactTypeConfig.cs
@@ -0,0 +1,43 @@
+// Copyright Epic Games, Inc. All Rights Reserved.
+
+using EpicGames.Horde.Artifacts;
+using HordeServer.Acls;
+
+namespace HordeServer.Artifacts
+{
+	/// <summary>
+	/// Configuration for an artifact
+	/// </summary>
+	public class ArtifactTypeConfig
+	{
+		/// <summary>
+		/// Legacy 'Name' property
+		/// </summary>
+		[Obsolete("Use Type instead")]
+		public ArtifactType Name
+		{
+			get => Type;
+			set => Type = value;
+		}
+
+		/// <summary>
+		/// Name of the artifact type
+		/// </summary>
+		public ArtifactType Type { get; set; }
+
+		/// <summary>
+		/// Acl for the artifact type
+		/// </summary>
+		public AclConfig? Acl { get; set; }
+
+		/// <summary>
+		/// Number of artifacts to retain
+		/// </summary>
+		public int? KeepCount { get; set; }
+
+		/// <summary>
+		/// Number of days to retain artifacts of this type
+		/// </summary>
+		public int? KeepDays { get; set; }
+	}
+}
diff --git a/Engine/Source/Programs/Horde/Plugins/Build/HordeServer.Build/BuildConfig.cs b/Engine/Source/Programs/Horde/Plugins/Build/HordeServer.Build/BuildConfig.cs
--- a/Engine/Source/Programs/Horde/Plugins/Build/HordeServer.Build/BuildConfig.cs
+++ b/Engine/Source/Programs/Horde/Plugins/Build/HordeServer.Build/BuildConfig.cs
@@ -16,6 +16,7 @@ using EpicGames.Perforce;
 using HordeServer.Acls;
 using HordeServer.Agents;
 using HordeServer.Agents.Pools;
+using HordeServer.Artifacts;
 using HordeServer.Commits;
 using HordeServer.Devices;
 using HordeServer.Perforce;
@@ -321,42 +322,6 @@ namespace HordeServer
 		}
 	}
 
-	/// <summary>
-	/// Configuration for an artifact
-	/// </summary>
-	public class ArtifactTypeConfig
-	{
-		/// <summary>
-		/// Legacy 'Name' property
-		/// </summary>
-		[Obsolete("Use Type instead")]
-		public ArtifactType Name
-		{
-			get => Type;
-			set => Type = value;
-		}
-
-		/// <summary>
-		/// Name of the artifact type
-		/// </summary>
-		public ArtifactType Type { get; set; }
-
-		/// <summary>
-		/// Acl for the artifact type
-		/// </summary>
-		public AclConfig? Acl { get; set; }
-
-		/// <summary>
-		/// Number of artifacts to retain
-		/// </summary>
-		public int? KeepCount { get; set; }
-
-		/// <summary>
-		/// Number of days to retain artifacts of this type
-		/// </summary>
-		public int? KeepDays { get; set; }
-	}
-
 	/// <summary>
 	/// Path to a platform and stream to use for syncing AutoSDK
 	/// </summary>
diff --git a/Engine/Source/Programs/Horde/Plugins/Build/HordeServer.Build/Projects/ProjectConfig.cs b/Engine/Source/Programs/Horde/Plugins/Build/HordeServer.Build/Projects/ProjectConfig.cs
--- a/Engine/Source/Programs/Horde/Plugins/Build/HordeServer.Build/Projects/ProjectConfig.cs
+++ b/Engine/Source/Programs/Horde/Plugins/Build/HordeServer.Build/Projects/ProjectConfig.cs
@@ -11,6 +11,7 @@ using EpicGames.Horde.Projects;
 using EpicGames.Horde.Telemetry;
 using HordeServer.Acls;
 using HordeServer.Agents.Pools;
+using HordeServer.Artifacts;
 using HordeServer.Configuration;
 using HordeServer.Streams;
 using HordeServer.Utilities;
diff --git a/Engine/Source/Programs/Horde/Plugins/Build/HordeServer.Build/Streams/StreamConfig.cs b/Engine/Source/Programs/Horde/Plugins/Build/HordeServer.Build/Streams/StreamConfig.cs
--- a/Engine/Source/Programs/Horde/Plugins/Build/HordeServer.Build/Streams/StreamConfig.cs
+++ b/Engine/Source/Programs/Horde/Plugins/Build/HordeServer.Build/Streams/StreamConfig.cs
@@ -22,6 +22,7 @@ using EpicGames.Horde.Telemetry;
 using EpicGames.Perforce;
 using Horde.Common.Rpc;
 using HordeServer.Acls;
+using HordeServer.Artifacts;
 using HordeServer.Configuration;
 using HordeServer.Issues;
 using HordeServer.Jobs.Graphs;
diff --git a/Engine/Source/Programs/Horde/Plugins/Tools/HordeServer.Tools/Commands/Generate/ToolDataCommand.cs b/Engine/Source/Programs/Horde/Plugins/Tools/HordeServer.Tools/Commands/Generate/ToolDataCommand.cs
--- a/Engine/Source/Programs/Horde/Plugins/Tools/HordeServer.Tools/Commands/Generate/ToolDataCommand.cs
+++ b/Engine/Source/Programs/Horde/Plugins/Tools/HordeServer.Tools/Commands/Generate/ToolDataCommand.cs
@@ -4,7 +4,6 @@ using System.ComponentModel;
 using System.Text.Json.Nodes;
 using EpicGames.Core;
 using EpicGames.Horde.Storage;
-using EpicGames.Horde.Storage.Backends;
 using EpicGames.Horde.Storage.Bundles;
 using EpicGames.Horde.Storage.Nodes;
 using HordeServer.Tools;
