Horde: Present artifact types on dashboard search
#skipci

diff --git a/Engine/Source/Programs/Horde/HordeDashboard/src/backend/Dashboard.ts b/Engine/Source/Programs/Horde/HordeDashboard/src/backend/Dashboard.ts
--- a/Engine/Source/Programs/Horde/HordeDashboard/src/backend/Dashboard.ts
+++ b/Engine/Source/Programs/Horde/HordeDashboard/src/backend/Dashboard.ts
@@ -138,7 +138,7 @@ export class Dashboard {
     get authMethod(): AuthMethod | undefined {
         return this.config?.authMethod;
     }
- 
+
     get p4user(): string {
         const claims = this.claims;
         const user = claims.filter(c => c.type.endsWith("/perforce-user"));
@@ -199,7 +199,11 @@ export class Dashboard {
     }
 
     get poolCategories(): GetDashboardPoolCategoryResponse[] {
-        return this.config?.poolCategories ?? [];        
+        return this.config?.poolCategories ?? [];
+    }
+
+    get artifactTypes(): string[] {
+        return this.config?.artifactTypes ?? [];
     }
 
     get deviceProblemCooldownMinutes(): number {
@@ -305,18 +309,18 @@ export class Dashboard {
         return this.preferences.get(DashboardPreference.DisplayUTC) === 'true';
 
     }
-    
+
     static get userPrefersDarkTheme(): boolean {
-        
+
         try {
             if (window.matchMedia && window.matchMedia('(prefers-color-scheme: dark)').matches) {
                 return true;
-            }    
+            }
         } catch (reason) {
             if (!Dashboard.hasLoggedDarkThemePref) {
                 Dashboard.hasLoggedDarkThemePref = true;
                 console.error(reason);
-            }            
+            }
         }
 
         return false;
@@ -326,16 +330,16 @@ export class Dashboard {
 
     get darktheme(): boolean {
 
-        if (!this.available) {            
+        if (!this.available) {
 
             let local = localStorage.getItem("horde_darktheme");
             if (local === "true") {
                 return true;
-            } 
+            }
             if (local === "false") {
                 return false;
-            } 
-    
+            }
+
             return Dashboard.userPrefersDarkTheme;
         }
 
@@ -546,7 +550,7 @@ export class Dashboard {
                     this.config = await backend.getDashboardConfig();
                 } catch (reason) {
                     console.error("Error getting dashboard config, defaults used: " + reason);
-                    this.config = { agentCategories: [], poolCategories: [],telemetryViews: []};
+                    this.config = { agentCategories: [], poolCategories: [], telemetryViews: [], artifactTypes: [] };
                 }
             }
 
diff --git a/Engine/Source/Programs/Horde/HordeDashboard/src/components/artifacts/ArtifactsSearch.tsx b/Engine/Source/Programs/Horde/HordeDashboard/src/components/artifacts/ArtifactsSearch.tsx
--- a/Engine/Source/Programs/Horde/HordeDashboard/src/components/artifacts/ArtifactsSearch.tsx
+++ b/Engine/Source/Programs/Horde/HordeDashboard/src/components/artifacts/ArtifactsSearch.tsx
@@ -33,7 +33,7 @@ export class ArtifactQueryState {
 
       if (state.minChange || state.maxChange || state.name || state.typeKey) {
          this.autoLoad = true;
-      }   
+      }
    }
 
    fromSearch(search: URLSearchParams): ArtifactSearchState {
@@ -362,29 +362,6 @@ const ArtifactsList: React.FC<{ state: ArtifactQueryState, artifacts?: GetArtifa
 
 }
 
-const artifactTypes: IComboBoxOption[] = [
-   {
-      key: `step-all`,
-      text: `All`
-   }, {
-      key: `packaged-build`,
-      text: `Packaged Build`
-   }, {
-      key: `step-saved`,
-      text: `step-saved`
-   }, {
-      key: `step-output`,
-      text: `step-output`
-   }, {
-      key: `step-trace`,
-      text: `step-trace`
-   },
-   {
-      key: `step-testdata`,
-      text: `step-testdata`
-   }
-]
-
 const sortOptions: IComboBoxOption[] = [
    {
       key: `sort-name`,
@@ -416,6 +393,18 @@ export const FindArtifactsModal: React.FC<{ onClose: () => void }> = ({ onClose
       }
    }
 
+   const artifactTypes: IComboBoxOption[] = [
+      {
+         key: `step-all`,
+         text: `All`
+      }
+   ]
+
+   dashboard.artifactTypes.sort((a, b) => a.localeCompare(b)).forEach(t => {
+      artifactTypes.push({ key: t, text: t });
+   });
+
+
    const queryArtifacts = async () => {
 
       setState({ ...state, searching: true });
