Horde: Allow configuring the namespace used for a particular artifact type.

diff --git a/Engine/Source/Programs/Horde/Plugins/Build/HordeServer.Build/Artifacts/ArtifactCollection.cs b/Engine/Source/Programs/Horde/Plugins/Build/HordeServer.Build/Artifacts/ArtifactCollection.cs
--- a/Engine/Source/Programs/Horde/Plugins/Build/HordeServer.Build/Artifacts/ArtifactCollection.cs
+++ b/Engine/Source/Programs/Horde/Plugins/Build/HordeServer.Build/Artifacts/ArtifactCollection.cs
@@ -205,7 +205,7 @@ namespace HordeServer.Artifacts
 			}
 
 			StreamConfig streamConfig = _buildConfig.CurrentValue.GetStream(streamId);
-			if (!streamConfig.TryGetArtifactType(type, out _))
+			if (!streamConfig.TryGetArtifactType(type, out ArtifactTypeConfig? artifactTypeConfig))
 			{
 				throw new ArtifactTypeNotFoundException(streamId, type);
 			}
@@ -217,7 +217,7 @@ namespace HordeServer.Artifacts
 			// Create the artifact
 			ObjectId id = ObjectId.GenerateNewId();
 
-			NamespaceId namespaceId = Namespace.Artifacts;
+			NamespaceId namespaceId = artifactTypeConfig.NamespaceId;
 			RefName refName = new RefName($"{GetArtifactPath(streamId, type)}/{commitId}/{name}/{id}");
 
 			CommitIdWithOrder commitIdWithOrder = await _commitService.GetOrderedAsync(streamId, commitId, cancellationToken);
diff --git a/Engine/Source/Programs/Horde/Plugins/Build/HordeServer.Build/Artifacts/ArtifactTypeConfig.cs b/Engine/Source/Programs/Horde/Plugins/Build/HordeServer.Build/Artifacts/ArtifactTypeConfig.cs
--- a/Engine/Source/Programs/Horde/Plugins/Build/HordeServer.Build/Artifacts/ArtifactTypeConfig.cs
+++ b/Engine/Source/Programs/Horde/Plugins/Build/HordeServer.Build/Artifacts/ArtifactTypeConfig.cs
@@ -1,6 +1,7 @@
 // Copyright Epic Games, Inc. All Rights Reserved.
 
 using EpicGames.Horde.Artifacts;
+using EpicGames.Horde.Storage;
 using HordeServer.Acls;
 
 namespace HordeServer.Artifacts
@@ -39,5 +40,10 @@ namespace HordeServer.Artifacts
 		/// Number of days to retain artifacts of this type
 		/// </summary>
 		public int? KeepDays { get; set; }
+
+		/// <summary>
+		/// Storage namespace to use for this artifact types
+		/// </summary>
+		public NamespaceId NamespaceId { get; set; } = new NamespaceId("horde-artifacts");
 	}
 }
diff --git a/Engine/Source/Programs/Horde/Plugins/Build/HordeServer.Build/Jobs/JobTaskSource.cs b/Engine/Source/Programs/Horde/Plugins/Build/HordeServer.Build/Jobs/JobTaskSource.cs
--- a/Engine/Source/Programs/Horde/Plugins/Build/HordeServer.Build/Jobs/JobTaskSource.cs
+++ b/Engine/Source/Programs/Horde/Plugins/Build/HordeServer.Build/Jobs/JobTaskSource.cs
@@ -11,7 +11,6 @@ using EpicGames.Horde.Agents.Pools;
 using EpicGames.Horde.Common;
 using EpicGames.Horde.Jobs;
 using EpicGames.Horde.Logs;
-using EpicGames.Horde.Storage;
 using EpicGames.Horde.Streams;
 using HordeCommon;
 using HordeCommon.Rpc.Messages;
@@ -23,7 +22,6 @@ using HordeServer.Jobs.Bisect;
 using HordeServer.Jobs.Graphs;
 using HordeServer.Logs;
 using HordeServer.Perforce;
-using HordeServer.Storage;
 using HordeServer.Streams;
 using HordeServer.Tasks;
 using HordeServer.Ugs;
@@ -883,17 +881,11 @@ namespace HordeServer.Jobs
 
 			// Get the global settings
 			BuildConfig buildConfig = _buildConfig.CurrentValue;
-			NamespaceId namespaceId = Namespace.Artifacts;
 
 			// Create a bearer token for the job executor
 			List<AclClaimConfig> claims = new List<AclClaimConfig>();
 			claims.Add(HordeClaims.AgentRoleClaim);
 			claims.Add(new AclClaimConfig(HordeClaimTypes.Lease, leaseId.ToString()));
-
-			string storagePrefix = $"{job.StreamId}/{job.CommitId}-{job.Id}";
-			claims.Add(new AclClaimConfig(HordeClaimTypes.ReadNamespace, $"{namespaceId}:{storagePrefix}"));
-			claims.Add(new AclClaimConfig(HordeClaimTypes.WriteNamespace, $"{namespaceId}:{storagePrefix}"));
-
 			claims.AddRange(job.Claims);
 
 			// Encode the payload
@@ -903,8 +895,6 @@ namespace HordeServer.Jobs
 			task.LogId = logId.ToString();
 			task.JobName = leaseName.ToString();
 			task.JobOptions = (job.JobOptions != null) ? GetRpcJobOptions(job.JobOptions) : null;
-			task.NamespaceId = namespaceId.ToString();
-			task.StoragePrefix = storagePrefix;
 			task.Token = await _aclService.IssueBearerTokenAsync(claims, null, cancellationToken);
 
 			List<RpcAgentWorkspace> workspaces = new();
diff --git a/Engine/Source/Programs/Horde/Plugins/Storage/HordeServer.Storage/StorageConfig.cs b/Engine/Source/Programs/Horde/Plugins/Storage/HordeServer.Storage/StorageConfig.cs
--- a/Engine/Source/Programs/Horde/Plugins/Storage/HordeServer.Storage/StorageConfig.cs
+++ b/Engine/Source/Programs/Horde/Plugins/Storage/HordeServer.Storage/StorageConfig.cs
@@ -25,11 +25,6 @@ namespace HordeServer.Storage
 	/// </summary>
 	public static class Namespace
 	{
-		/// <summary>
-		/// Storage of artifacts 
-		/// </summary>
-		public static NamespaceId Artifacts { get; } = new NamespaceId("horde-artifacts");
-
 		/// <summary>
 		/// Replicated Perforce data
 		/// </summary>
