Horde: Allow configuring the namespace for storing tool data.

diff --git a/Engine/Source/Programs/Horde/Plugins/Storage/HordeServer.Storage/StorageConfig.cs b/Engine/Source/Programs/Horde/Plugins/Storage/HordeServer.Storage/StorageConfig.cs
--- a/Engine/Source/Programs/Horde/Plugins/Storage/HordeServer.Storage/StorageConfig.cs
+++ b/Engine/Source/Programs/Horde/Plugins/Storage/HordeServer.Storage/StorageConfig.cs
@@ -34,11 +34,6 @@ namespace HordeServer.Storage
 		/// Log data
 		/// </summary>
 		public static NamespaceId Logs { get; } = new NamespaceId("horde-logs");
-
-		/// <summary>
-		/// Storage of tool data
-		/// </summary>
-		public static NamespaceId Tools { get; } = new NamespaceId("horde-tools");
 	}
 
 	/// <summary>
diff --git a/Engine/Source/Programs/Horde/Plugins/Tools/HordeServer.Tools.Tests/Tools/ToolTests.cs b/Engine/Source/Programs/Horde/Plugins/Tools/HordeServer.Tools.Tests/Tools/ToolTests.cs
--- a/Engine/Source/Programs/Horde/Plugins/Tools/HordeServer.Tools.Tests/Tools/ToolTests.cs
+++ b/Engine/Source/Programs/Horde/Plugins/Tools/HordeServer.Tools.Tests/Tools/ToolTests.cs
@@ -25,7 +25,7 @@ namespace HordeServer.Tests.Tools
 			storageConfig.Backends.Clear();
 			storageConfig.Backends.Add(new BackendConfig { Id = new BackendId("tools-backend"), Type = StorageBackendType.Memory });
 			storageConfig.Namespaces.Clear();
-			storageConfig.Namespaces.Add(new NamespaceConfig { Id = Namespace.Tools, Backend = new BackendId("tools-backend") });
+			storageConfig.Namespaces.Add(new NamespaceConfig { Id = ToolConfig.DefaultNamespaceId, Backend = new BackendId("tools-backend") });
 
 			ToolsConfig toolsConfig = new ToolsConfig();
 			toolsConfig.Tools.Add(new ToolConfig(_toolId) { Name = "UnrealGameSync", Description = "Tool for syncing content from source control" });
diff --git a/Engine/Source/Programs/Horde/Plugins/Tools/HordeServer.Tools.Tests/Tools/ToolsControllerTests.cs b/Engine/Source/Programs/Horde/Plugins/Tools/HordeServer.Tools.Tests/Tools/ToolsControllerTests.cs
--- a/Engine/Source/Programs/Horde/Plugins/Tools/HordeServer.Tools.Tests/Tools/ToolsControllerTests.cs
+++ b/Engine/Source/Programs/Horde/Plugins/Tools/HordeServer.Tools.Tests/Tools/ToolsControllerTests.cs
@@ -36,7 +36,7 @@ public class ToolsControllerTests
 		storageConfig.Backends.Clear();
 		storageConfig.Backends.Add(new BackendConfig { Id = new BackendId("tools-backend"), Type = StorageBackendType.Memory });
 		storageConfig.Namespaces.Clear();
-		storageConfig.Namespaces.Add(new NamespaceConfig { Id = Namespace.Tools, Backend = new BackendId("tools-backend") });
+		storageConfig.Namespaces.Add(new NamespaceConfig { Id = ToolConfig.DefaultNamespaceId, Backend = new BackendId("tools-backend") });
 
 		ToolsConfig toolsConfig = new ToolsConfig();
 		toolsConfig.Tools.Add(new ToolConfig(toolId) { Name = "Foo", Description = "This is foo", Acl = new AclConfig() { Entries = [aclEntryConfig] }, Public = false });
diff --git a/Engine/Source/Programs/Horde/Plugins/Tools/HordeServer.Tools/Tools/ToolCollection.cs b/Engine/Source/Programs/Horde/Plugins/Tools/HordeServer.Tools/Tools/ToolCollection.cs
--- a/Engine/Source/Programs/Horde/Plugins/Tools/HordeServer.Tools/Tools/ToolCollection.cs
+++ b/Engine/Source/Programs/Horde/Plugins/Tools/HordeServer.Tools/Tools/ToolCollection.cs
@@ -62,7 +62,7 @@ namespace HordeServer.Tools
 			public TimeSpan Duration { get; set; }
 
 			[BsonElement("ns")]
-			public NamespaceId NamespaceId { get; set; } = Namespace.Tools;
+			public NamespaceId NamespaceId { get; set; } = ToolConfig.DefaultNamespaceId;
 
 			[BsonElement("ref")]
 			public RefName RefName { get; set; }
@@ -467,7 +467,7 @@ namespace HordeServer.Tools
 			}
 			else
 			{
-				return _storageService.CreateClient(Namespace.Tools);
+				return _storageService.CreateClient(toolConfig.NamespaceId);
 			}
 		}
 
@@ -484,7 +484,7 @@ namespace HordeServer.Tools
 			}
 			else
 			{
-				return _storageService.CreateBackend(Namespace.Tools);
+				return _storageService.CreateBackend(toolConfig.NamespaceId);
 			}
 		}
 
diff --git a/Engine/Source/Programs/Horde/Plugins/Tools/HordeServer.Tools/Tools/ToolConfig.cs b/Engine/Source/Programs/Horde/Plugins/Tools/HordeServer.Tools/Tools/ToolConfig.cs
--- a/Engine/Source/Programs/Horde/Plugins/Tools/HordeServer.Tools/Tools/ToolConfig.cs
+++ b/Engine/Source/Programs/Horde/Plugins/Tools/HordeServer.Tools/Tools/ToolConfig.cs
@@ -19,6 +19,11 @@ namespace HordeServer.Tools
 	[DebuggerDisplay("{Id}")]
 	public class ToolConfig
 	{
+		/// <summary>
+		/// The default namespace for tool data
+		/// </summary>
+		public static NamespaceId DefaultNamespaceId { get; } = new NamespaceId("horde-tools");
+
 		/// <summary>
 		/// Unique identifier for the tool
 		/// </summary>
@@ -69,7 +74,7 @@ namespace HordeServer.Tools
 		/// <summary>
 		/// Default namespace for new deployments of this tool
 		/// </summary>
-		public NamespaceId NamespaceId { get; set; } = Namespace.Tools;
+		public NamespaceId NamespaceId { get; set; } = DefaultNamespaceId;
 
 		/// <summary>
 		/// Permissions for the tool
