Horde: Allow configuring the namespace for writing replicated Perforce data.

diff --git a/Engine/Source/Programs/Horde/Plugins/Build/HordeServer.Build/Commands/Test/ReplicationCommand.cs b/Engine/Source/Programs/Horde/Plugins/Build/HordeServer.Build/Commands/Test/ReplicationCommand.cs
--- a/Engine/Source/Programs/Horde/Plugins/Build/HordeServer.Build/Commands/Test/ReplicationCommand.cs
+++ b/Engine/Source/Programs/Horde/Plugins/Build/HordeServer.Build/Commands/Test/ReplicationCommand.cs
@@ -64,6 +64,13 @@ namespace HordeServer.Commands.Test
 			IReplicatorCollection replicatorCollection = serviceProvider.GetRequiredService<IReplicatorCollection>();
 
 			ReplicatorId id = new ReplicatorId(new StreamId(StreamId), new StreamReplicatorId(ReplicatorId));
+
+			ReplicatorConfig? replicatorConfig;
+			if (!streamConfig.TryGetReplicator(new StreamReplicatorId(ReplicatorId), out replicatorConfig))
+			{
+				throw new FatalErrorException($"Replicator '{ReplicatorId}' not found");
+			}
+
 			IReplicator replicator = await replicatorCollection.GetOrAddAsync(id);
 
 			while (replicator.Pause || replicator.Reset != Reset || replicator.Clean != Clean || replicator.NextChange != Change)
@@ -80,9 +87,7 @@ namespace HordeServer.Commands.Test
 				replicator = await replicatorCollection.GetOrAddAsync(id);
 			}
 
-			PerforceReplicationOptions options = new PerforceReplicationOptions();
-			await perforceReplicator.RunOnceAsync(replicator, buildConfig, streamConfig, options, default);
-
+			await perforceReplicator.RunOnceAsync(replicator, buildConfig, streamConfig, replicatorConfig, default);
 			return 0;
 		}
 	}
diff --git a/Engine/Source/Programs/Horde/Plugins/Build/HordeServer.Build/Replicators/PerforceReplicator.cs b/Engine/Source/Programs/Horde/Plugins/Build/HordeServer.Build/Replicators/PerforceReplicator.cs
--- a/Engine/Source/Programs/Horde/Plugins/Build/HordeServer.Build/Replicators/PerforceReplicator.cs
+++ b/Engine/Source/Programs/Horde/Plugins/Build/HordeServer.Build/Replicators/PerforceReplicator.cs
@@ -20,32 +20,6 @@ using Microsoft.Extensions.Logging;
 
 namespace HordeServer.Replicators
 {
-	/// <summary>
-	/// Options for replicating commits
-	/// </summary>
-	public class PerforceReplicationOptions
-	{
-		/// <summary>
-		/// Whether to replicate content, or just metadata
-		/// </summary>
-		public bool IncludeContent { get; set; }
-
-		/// <summary>
-		/// 
-		/// </summary>
-		public BundleOptions TreeOptions { get; set; } = new BundleOptions();
-
-		/// <summary>
-		/// 
-		/// </summary>
-		public ChunkingOptions ChunkingOptions { get; set; } = new ChunkingOptions();
-
-		/// <summary>
-		/// 
-		/// </summary>
-		public RefOptions RefOptions { get; set; } = new RefOptions();
-	}
-
 	/// <summary>
 	/// Exception thrown to indicate that the state of a replicator has been modified externally
 	/// </summary>
@@ -318,7 +292,7 @@ namespace HordeServer.Replicators
 		/// <summary>
 		/// Runs a replication loop for a stream
 		/// </summary>
-		public async Task RunAsync(ReplicatorId replicatorId, BuildConfig buildConfig, StreamConfig streamConfig, PerforceReplicationOptions options, CancellationToken cancellationToken = default)
+		public async Task RunAsync(ReplicatorId replicatorId, BuildConfig buildConfig, StreamConfig streamConfig, ReplicatorConfig replicatorConfig, CancellationToken cancellationToken = default)
 		{
 			_logger.LogInformation("Starting replication background task for {ReplicatorId}", replicatorId);
 
@@ -331,19 +305,19 @@ namespace HordeServer.Replicators
 
 			while (!cancellationToken.IsCancellationRequested)
 			{
-				replicator = await RunOnceAsync(replicator, buildConfig, streamConfig, options, cancellationToken);
+				replicator = await RunOnceAsync(replicator, buildConfig, streamConfig, replicatorConfig, cancellationToken);
 			}
 		}
 
 		/// <summary>
 		/// Runs the replicator for a single change
 		/// </summary>
-		public async Task<IReplicator> RunOnceAsync(IReplicator replicator, BuildConfig buildConfig, StreamConfig streamConfig, PerforceReplicationOptions replicatorOptions, CancellationToken cancellationToken)
+		public async Task<IReplicator> RunOnceAsync(IReplicator replicator, BuildConfig buildConfig, StreamConfig streamConfig, ReplicatorConfig replicatorConfig, CancellationToken cancellationToken)
 		{
 			RefName refName = GetRefName(replicator.Id);
 			RefName incRefName = GetIncrementalRefName(replicator.Id);
 
-			IStorageClient store = _storageService.CreateClient(Namespace.Perforce);
+			IStorageClient store = _storageService.CreateClient(replicatorConfig.NamespaceId);
 
 			ICommitCollection commits = _perforceService.GetCommits(streamConfig);
 
@@ -394,7 +368,7 @@ namespace HordeServer.Replicators
 			BlobSerializerOptions blobOptions = new BlobSerializerOptions();
 			try
 			{
-				replicator = await WriteInternalAsync(replicator, change, buildConfig, streamConfig, replicatorOptions, blobOptions, cancellationToken);
+				replicator = await WriteInternalAsync(replicator, change, buildConfig, streamConfig, replicatorConfig, blobOptions, cancellationToken);
 				return replicator;
 			}
 			catch (OperationCanceledException ex)
@@ -420,9 +394,10 @@ namespace HordeServer.Replicators
 			}
 		}
 
-		async Task<IReplicator> WriteInternalAsync(IReplicator replicator, int change, BuildConfig buildConfig, StreamConfig streamConfig, PerforceReplicationOptions options, BlobSerializerOptions blobOptions, CancellationToken cancellationToken = default)
+		async Task<IReplicator> WriteInternalAsync(IReplicator replicator, int change, BuildConfig buildConfig, StreamConfig streamConfig, ReplicatorConfig replicatorConfig, BlobSerializerOptions blobOptions, CancellationToken cancellationToken = default)
 		{
-			IStorageClient store = _storageService.CreateClient(Namespace.Perforce);
+			ChunkingOptions chunkingOptions = new();
+			IStorageClient store = _storageService.CreateClient(replicatorConfig.NamespaceId);
 
 			// Find the parent node
 			RefName refName = GetRefName(replicator.Id);
@@ -609,7 +584,7 @@ namespace HordeServer.Replicators
 				Stopwatch processTimer = new Stopwatch();
 				Stopwatch gcTimer = new Stopwatch();
 
-				using FileWriter fileWriter = new FileWriter(leafNodeWriter, options.ChunkingOptions.LeafOptions, _logger);
+				using FileWriter fileWriter = new FileWriter(leafNodeWriter, chunkingOptions.LeafOptions, _logger);
 				await foreach (PerforceResponse response in perforce.StreamCommandAsync("sync", Array.Empty<string>(), syncPaths, null, typeof(SyncRecord), true, default))
 				{
 					PerforceError? error = response.Error;
@@ -730,7 +705,7 @@ namespace HordeServer.Replicators
 
 				await leafNodeWriter.FlushAsync(cancellationToken);
 
-				await rootUpdate.WriteInteriorNodesAsync(interiorChunkWriter, options.ChunkingOptions.InteriorOptions, cancellationToken);
+				await rootUpdate.WriteInteriorNodesAsync(interiorChunkWriter, chunkingOptions.InteriorOptions, cancellationToken);
 				await interiorChunkWriter.FlushAsync(cancellationToken);
 
 				await root.UpdateAsync(rootUpdate, directoryWriter, cancellationToken);
@@ -765,7 +740,7 @@ namespace HordeServer.Replicators
 			IHashedBlobRef<RedirectNode<CommitNode>> redirectNodeRef = await commitWriter.WriteBlobAsync(redirectNode, cancellationToken);
 
 			await commitWriter.FlushAsync(cancellationToken);
-			await store.WriteRefAsync(refName, redirectNodeRef, options.RefOptions, cancellationToken: cancellationToken);
+			await store.WriteRefAsync(refName, redirectNodeRef, cancellationToken: cancellationToken);
 
 			// Update the replicator state
 			UpdateReplicatorOptions completeUpdateOptions = new UpdateReplicatorOptions
diff --git a/Engine/Source/Programs/Horde/Plugins/Build/HordeServer.Build/Replicators/ReplicationService.cs b/Engine/Source/Programs/Horde/Plugins/Build/HordeServer.Build/Replicators/ReplicationService.cs
--- a/Engine/Source/Programs/Horde/Plugins/Build/HordeServer.Build/Replicators/ReplicationService.cs
+++ b/Engine/Source/Programs/Horde/Plugins/Build/HordeServer.Build/Replicators/ReplicationService.cs
@@ -84,7 +84,7 @@ namespace HordeServer.Replicators
 								if (replicatorConfig.Enabled)
 								{
 									ReplicatorId replicatorId = new ReplicatorId(streamConfig.Id, replicatorConfig.Id);
-									replicatorIdToTask.Add(replicatorId, BackgroundTask.StartNew(ctx => RunReplicationGuardedAsync(replicatorId, buildConfig, streamConfig, ctx)));
+									replicatorIdToTask.Add(replicatorId, BackgroundTask.StartNew(ctx => RunReplicationGuardedAsync(replicatorId, buildConfig, streamConfig, replicatorConfig, ctx)));
 									removeReplicators.Remove(replicatorId);
 								}
 							}
@@ -107,7 +107,7 @@ namespace HordeServer.Replicators
 			}
 		}
 
-		async Task RunReplicationGuardedAsync(ReplicatorId replicatorId, BuildConfig buildConfig, StreamConfig streamConfig, CancellationToken cancellationToken)
+		async Task RunReplicationGuardedAsync(ReplicatorId replicatorId, BuildConfig buildConfig, StreamConfig streamConfig, ReplicatorConfig replicatorConfig, CancellationToken cancellationToken)
 		{
 			_logger.LogInformation("Started background task for replication of {ReplicatorId}", replicatorId);
 
@@ -121,8 +121,7 @@ namespace HordeServer.Replicators
 				.AddRetry(new RetryStrategyOptions { MaxRetryAttempts = int.MaxValue, OnRetry = OnRetry })
 				.Build();
 
-			PerforceReplicationOptions replicationOptions = new PerforceReplicationOptions();
-			await pipeline.ExecuteAsync(async ctx => await _replicator.RunAsync(replicatorId, buildConfig, streamConfig, replicationOptions, ctx), cancellationToken);
+			await pipeline.ExecuteAsync(async ctx => await _replicator.RunAsync(replicatorId, buildConfig, streamConfig, replicatorConfig, ctx), cancellationToken);
 		}
 	}
 }
diff --git a/Engine/Source/Programs/Horde/Plugins/Build/HordeServer.Build/Replicators/ReplicatorConfig.cs b/Engine/Source/Programs/Horde/Plugins/Build/HordeServer.Build/Replicators/ReplicatorConfig.cs
--- a/Engine/Source/Programs/Horde/Plugins/Build/HordeServer.Build/Replicators/ReplicatorConfig.cs
+++ b/Engine/Source/Programs/Horde/Plugins/Build/HordeServer.Build/Replicators/ReplicatorConfig.cs
@@ -1,6 +1,7 @@
 // Copyright Epic Games, Inc. All Rights Reserved.
 
 using EpicGames.Horde.Replicators;
+using EpicGames.Horde.Storage;
 
 namespace HordeServer.Replicators
 {
@@ -28,5 +29,15 @@ namespace HordeServer.Replicators
 		/// Maximum change number to replicate
 		/// </summary>
 		public int? MaxChange { get; set; }
+
+		/// <summary>
+		/// Whether to include content in the replication, or just metadata
+		/// </summary>
+		public bool IncludeContent { get; set; } = true;
+
+		/// <summary>
+		/// Namespace to replicate data to
+		/// </summary>
+		public NamespaceId NamespaceId { get; set; } = new NamespaceId("horde-perforce");
 	}
 }
diff --git a/Engine/Source/Programs/Horde/Plugins/Storage/HordeServer.Storage/StorageConfig.cs b/Engine/Source/Programs/Horde/Plugins/Storage/HordeServer.Storage/StorageConfig.cs
--- a/Engine/Source/Programs/Horde/Plugins/Storage/HordeServer.Storage/StorageConfig.cs
+++ b/Engine/Source/Programs/Horde/Plugins/Storage/HordeServer.Storage/StorageConfig.cs
@@ -25,11 +25,6 @@ namespace HordeServer.Storage
 	/// </summary>
 	public static class Namespace
 	{
-		/// <summary>
-		/// Replicated Perforce data
-		/// </summary>
-		public static NamespaceId Perforce { get; } = new NamespaceId("horde-perforce");
-
 		/// <summary>
 		/// Log data
 		/// </summary>
