Horde: Updating schema documentation

#rnx

diff --git a/Engine/Source/Programs/Horde/Docs/Config/Schema/Globals.md b/Engine/Source/Programs/Horde/Docs/Config/Schema/Globals.md
--- a/Engine/Source/Programs/Horde/Docs/Config/Schema/Globals.md
+++ b/Engine/Source/Programs/Horde/Docs/Config/Schema/Globals.md
@@ -175,6 +175,7 @@ Name | Description
 `acl` | [AclConfig](#aclconfig)<br>Acl for the artifact type
 `keepCount` | `integer`<br>Number of artifacts to retain
 `keepDays` | `integer`<br>Number of days to retain artifacts of this type
+`namespaceId` | `string`<br>Storage namespace to use for this artifact types
 
 ## AclConfig
 
diff --git a/Engine/Source/Programs/Horde/Docs/Config/Schema/Projects.md b/Engine/Source/Programs/Horde/Docs/Config/Schema/Projects.md
--- a/Engine/Source/Programs/Horde/Docs/Config/Schema/Projects.md
+++ b/Engine/Source/Programs/Horde/Docs/Config/Schema/Projects.md
@@ -225,6 +225,7 @@ Name | Description
 `acl` | [AclConfig](#aclconfig)<br>Acl for the artifact type
 `keepCount` | `integer`<br>Number of artifacts to retain
 `keepDays` | `integer`<br>Number of days to retain artifacts of this type
+`namespaceId` | `string`<br>Storage namespace to use for this artifact types
 
 ## AclConfig
 
diff --git a/Engine/Source/Programs/Horde/Docs/Config/Schema/Streams.md b/Engine/Source/Programs/Horde/Docs/Config/Schema/Streams.md
--- a/Engine/Source/Programs/Horde/Docs/Config/Schema/Streams.md
+++ b/Engine/Source/Programs/Horde/Docs/Config/Schema/Streams.md
@@ -449,6 +449,8 @@ Name | Description
 `enabled` | `boolean`<br>Whether the replicator is enabled
 `minChange` | `integer`<br>Minimum change number to replicate
 `maxChange` | `integer`<br>Maximum change number to replicate
+`includeContent` | `boolean`<br>Whether to include content in the replication, or just metadata
+`namespaceId` | `string`<br>Namespace to replicate data to
 
 ## WorkflowConfig
 
@@ -499,3 +501,4 @@ Name | Description
 `acl` | [AclConfig](#aclconfig)<br>Acl for the artifact type
 `keepCount` | `integer`<br>Number of artifacts to retain
 `keepDays` | `integer`<br>Number of days to retain artifacts of this type
+`namespaceId` | `string`<br>Storage namespace to use for this artifact types
