Horde: Check for invalid AWS region names in config

diff --git a/Engine/Source/Programs/Horde/Plugins/Compute/HordeServer.Compute/ComputePlugin.cs b/Engine/Source/Programs/Horde/Plugins/Compute/HordeServer.Compute/ComputePlugin.cs
--- a/Engine/Source/Programs/Horde/Plugins/Compute/HordeServer.Compute/ComputePlugin.cs
+++ b/Engine/Source/Programs/Horde/Plugins/Compute/HordeServer.Compute/ComputePlugin.cs
@@ -139,7 +139,12 @@ namespace HordeServer
 				}
 				foreach (string awsRegionStr in _staticComputeConfig.AwsRegions)
 				{
-					regions.Add(RegionEndpoint.GetBySystemName(awsRegionStr));
+					RegionEndpoint? region = RegionEndpoint.EnumerableAllRegions.FirstOrDefault(x => x.SystemName == awsRegionStr);
+					if (region == null)
+					{
+						throw new Exception("Invalid AWS region: " + awsRegionStr);
+					}
+					regions.Add(region);
 				}
 				
 				services.AddAWSService<IAmazonAutoScaling>();
