Horde: Log issue collection locks as milliseconds for easier log parsing

diff --git a/Engine/Source/Programs/Horde/Plugins/Build/HordeServer.Build/Issues/IssueCollection.cs b/Engine/Source/Programs/Horde/Plugins/Build/HordeServer.Build/Issues/IssueCollection.cs
--- a/Engine/Source/Programs/Horde/Plugins/Build/HordeServer.Build/Issues/IssueCollection.cs
+++ b/Engine/Source/Programs/Horde/Plugins/Build/HordeServer.Build/Issues/IssueCollection.cs
@@ -739,7 +739,7 @@ namespace HordeServer.Issues
 				{
 					if (timer.Elapsed > nextNotifyTime)
 					{
-						_logger.LogWarning("Waiting on lock over issue collection for {TimeSpan}", timer.Elapsed);
+						_logger.LogWarning("Waiting on lock over issue collection for {TimeSpan} ms", timer.ElapsedMilliseconds);
 						nextNotifyTime *= 2;
 					}
 					await Task.Delay(TimeSpan.FromMilliseconds(100));
