Horde: Randomize port assignment for relay port mappings

There's a worry of re-using the same port leads to potential races with how nftables/conntrack cleans up entries. In previous impl, lowest available port number was always used which led to higher contention.

diff --git a/Engine/Source/Programs/Horde/Plugins/Build/HordeServer.Build.Tests/Agents/Relay/AgentRelayServiceTests.cs b/Engine/Source/Programs/Horde/Plugins/Build/HordeServer.Build.Tests/Agents/Relay/AgentRelayServiceTests.cs
--- a/Engine/Source/Programs/Horde/Plugins/Build/HordeServer.Build.Tests/Agents/Relay/AgentRelayServiceTests.cs
+++ b/Engine/Source/Programs/Horde/Plugins/Build/HordeServer.Build.Tests/Agents/Relay/AgentRelayServiceTests.cs
@@ -51,6 +51,7 @@ public class AgentRelayServiceTests : BuildTestSetup
 	public AgentRelayServiceTests()
 	{
 		_service = new AgentRelayService(GetRedisServiceSingleton(), Clock, Tracer, NullLogger<AgentRelayService>.Instance);
+		_service.SetRandomSeed(1);
 	}
 
 	public override async ValueTask DisposeAsync()
@@ -244,19 +245,21 @@ public class AgentRelayServiceTests : BuildTestSetup
 	{
 		PortMapping newPm = await AddPortMappingAsync("cluster1", _pm2.LeaseId, null, _pm2.AgentIp, _pm2.Ports);
 		Assert.AreEqual(2, newPm.Ports.Count);
-		Assert.AreEqual(10000, newPm.Ports[0].RelayPort);
-		Assert.AreEqual(10001, newPm.Ports[1].RelayPort);
+		Assert.AreEqual(12214, newPm.Ports[0].RelayPort);
+		Assert.AreEqual(14973, newPm.Ports[1].RelayPort);
 	}
+	
+	// TODO: add a test checking for exception "Unable to find an available port range"
 
 	[TestMethod]
 	public async Task PortAssignment_TwoMappings_Async()
 	{
-		_service.SetMinMaxPorts(1000, 1003);
+		_service.SetMinMaxPorts(1000, 1020);
 		PortMapping newPm1 = await AddPortMappingAsync("cluster1", _pm1.LeaseId, null, _pm1.AgentIp, _pm1.Ports);
 		PortMapping newPm2 = await AddPortMappingAsync("cluster1", _pm2.LeaseId, null, _pm2.AgentIp, _pm2.Ports);
-		Assert.AreEqual(1000, newPm1.Ports[0].RelayPort);
-		Assert.AreEqual(1001, newPm2.Ports[0].RelayPort);
-		Assert.AreEqual(1002, newPm2.Ports[1].RelayPort);
+		Assert.AreEqual(1005, newPm1.Ports[0].RelayPort);
+		Assert.AreEqual(1002, newPm2.Ports[0].RelayPort);
+		Assert.AreEqual(1009, newPm2.Ports[1].RelayPort);
 	}
 
 	[TestMethod]
@@ -265,7 +268,7 @@ public class AgentRelayServiceTests : BuildTestSetup
 		_service.SetMinMaxPorts(1000, 1003);
 		await AddPortMappingAsync("cluster1", _pm1.LeaseId, null, _pm1.AgentIp, _pm1.Ports);
 		await AddPortMappingAsync("cluster1", _pm3.LeaseId, null, _pm3.AgentIp, _pm3.Ports);
-		await Assert.ThrowsExceptionAsync<Exception>(() => AddPortMappingAsync("cluster1", _pm2.LeaseId, null, _pm2.AgentIp, _pm2.Ports));
+		await Assert.ThrowsExceptionAsync<ArgumentException>(() => AddPortMappingAsync("cluster1", _pm2.LeaseId, null, _pm2.AgentIp, _pm2.Ports));
 	}
 
 	[TestMethod]
@@ -274,7 +277,7 @@ public class AgentRelayServiceTests : BuildTestSetup
 		_service.SetMinMaxPorts(1000, 1003);
 		await AddPortMappingAsync("cluster1", _pm1.LeaseId, null, _pm1.AgentIp, _pm1.Ports);
 		await AddPortMappingAsync("cluster1", _pm3.LeaseId, null, _pm3.AgentIp, _pm3.Ports);
-		await Assert.ThrowsExceptionAsync<Exception>(() => AddPortMappingAsync("cluster1", _pm2.LeaseId, null, _pm2.AgentIp, _pm2.Ports));
+		await Assert.ThrowsExceptionAsync<ArgumentException>(() => AddPortMappingAsync("cluster1", _pm2.LeaseId, null, _pm2.AgentIp, _pm2.Ports));
 		Assert.IsTrue(await _service.RemovePortMappingAsync(_cluster1, LeaseId.Parse(_pm3.LeaseId)));
 		await AddPortMappingAsync("cluster1", _pm2.LeaseId, null, _pm2.AgentIp, _pm2.Ports);
 	}
@@ -306,7 +309,39 @@ public class AgentRelayServiceTests : BuildTestSetup
 		CollectionAssert.AreEquivalent(Expected(), Actual(2, 0, 1, 2, 3, 4, 5));
 		CollectionAssert.AreEquivalent(Expected(1, 2), Actual(2, 0, 3, 4, 5));
 	}
-
+	
+	[TestMethod]
+	public void FindAvailablePorts()
+	{
+		List<int> Actual(int numPorts, params int[] usedPorts)
+		{
+			Random random = new (1);
+			return AgentRelayService.FindAvailablePorts(random, usedPorts.ToHashSet(), numPorts, 0, 9).ToList();
+		}
+		
+		List<int> Expected(params int[] ports)
+		{
+			return ports.ToHashSet().ToList();
+		}
+		
+		void AssertPorts(List<int> expected, List<int> actual)
+		{
+			// Console.WriteLine("{0,-25} {1,-25}", "Expected:", "Actual:");
+			// Console.WriteLine("{0,-25} {1,-25}", String.Join(' ', expected), String.Join(' ', actual));
+			CollectionAssert.AreEquivalent(expected, actual);
+		}
+		
+		AssertPorts(Expected(), Actual(0));
+		AssertPorts(Expected(2, 1, 4), Actual(3));
+		AssertPorts(Expected(2, 1, 4, 7, 6, 3, 9, 0, 8, 5), Actual(10));
+		AssertPorts(Expected(4, 7, 6, 3), Actual(4, [2, 1, 9]));
+		AssertPorts(Expected(4, 7, 6, 3, 0, 8, 5), Actual(7, [2, 1, 9]));
+		Assert.ThrowsException<ArgumentException>(() => Actual(8, [2, 1, 9]));
+		
+		Exception e1 = Assert.ThrowsException<ArgumentException>(() => Actual(11));
+		Assert.IsTrue(e1.Message.Contains("exceeds", StringComparison.InvariantCulture));
+	}
+	
 	private Task<PortMapping> AddPortMappingAsync(string clusterId, string leaseId, string? clientIp, string agentIp, IList<Port> ports, int numRetries = 10)
 	{
 		IPAddress? clientIpObj = clientIp == null ? null : IPAddress.Parse(clientIp);
diff --git a/Engine/Source/Programs/Horde/Plugins/Compute/HordeServer.Compute/Agents/Relay/AgentRelayService.cs b/Engine/Source/Programs/Horde/Plugins/Compute/HordeServer.Compute/Agents/Relay/AgentRelayService.cs
--- a/Engine/Source/Programs/Horde/Plugins/Compute/HordeServer.Compute/Agents/Relay/AgentRelayService.cs
+++ b/Engine/Source/Programs/Horde/Plugins/Compute/HordeServer.Compute/Agents/Relay/AgentRelayService.cs
@@ -91,7 +91,8 @@ public sealed class AgentRelayService : RelayRpc.RelayRpcBase, IHostedService, I
 	private readonly AsyncTaskQueue _updateTaskQueue;
 	private IAsyncDisposable? _redisSubscription;
 	private int _minPort = 10000;
-	private int _maxPort = 50000;
+	private int _maxPort = 30000; // Linux ephemeral port range is usually 32768-60999 so staying below 30000 seems like a sane default
+	private Random _random = new ();
 
 	/// <summary>
 	/// Constructor
@@ -197,7 +198,6 @@ public sealed class AgentRelayService : RelayRpc.RelayRpcBase, IHostedService, I
 		span.SetAttribute("leaseId", leaseId.ToString());
 		IDatabase redis = _redis.GetDatabase();
 
-		// TODO: Randomize start position to spread out port use
 		// TODO: Add pinging to and do not consider stale agent relays
 		// TODO: Check failed deserialization from Redis
 
@@ -205,7 +205,7 @@ public sealed class AgentRelayService : RelayRpc.RelayRpcBase, IHostedService, I
 		{
 			HashSet<int> usedPorts = (await redis.SetMembersAsync(KeyUsedPorts(clusterId))).Select(x => (int)x).ToHashSet();
 			int numPorts = ports.Count;
-			IReadOnlySet<int> portRange = FindAvailablePortRange(usedPorts, numPorts, _minPort, _maxPort);
+			IReadOnlySet<int> portRange = FindAvailablePorts(_random, usedPorts, numPorts, _minPort, _maxPort);
 			if (portRange.Count == 0)
 			{
 				throw new Exception("No ports are available");
@@ -403,6 +403,53 @@ public sealed class AgentRelayService : RelayRpc.RelayRpcBase, IHostedService, I
 
 		return ImmutableSortedSet<int>.Empty;
 	}
+	
+	/// <summary>
+	/// Find random ports available
+	/// </summary>
+	/// <param name="random">Random generator</param>
+	/// <param name="usedPorts">Ports already in use</param>
+	/// <param name="numPorts">Number of ports to find</param>
+	/// <param name="minPort">Min port number</param>
+	/// <param name="maxPort">Max port number</param>
+	/// <returns>A list of random but available ports</returns>
+	/// <exception cref="Exception">If port assignment does not complete within max attempts</exception>
+	internal static IReadOnlySet<int> FindAvailablePorts(Random random, IReadOnlySet<int> usedPorts, int numPorts, int minPort, int maxPort)
+	{
+		const int MaxAttempts = 10000;
+		
+		if (minPort > maxPort)
+		{
+			throw new ArgumentException("minPort must be less than or equal to maxPort");
+		}
+		
+		int availablePortRange = maxPort - minPort + 1;
+		if (numPorts > availablePortRange - usedPorts.Count)
+		{
+			throw new ArgumentException("The number of requested ports exceeds the available port range");
+		}
+		
+		HashSet<int> acquiredPorts = new(numPorts);
+		int numAttempts = 0;
+		while (acquiredPorts.Count < numPorts)
+		{
+			numAttempts++;
+			if (numAttempts >= MaxAttempts)
+			{
+				throw new Exception($"Failed to acquire {numPorts} random ports after {numAttempts} attempts. Acquired ports: {acquiredPorts.Count}");
+			}
+			
+			int port = random.Next(minPort, maxPort + 1);
+			if (usedPorts.Contains(port) || acquiredPorts.Contains(port))
+			{
+				continue; // Try assigning a new random port
+			}
+			
+			acquiredPorts.Add(port);
+		}
+		
+		return acquiredPorts;
+	}
 
 	internal void SetMinMaxPorts(int minPort, int maxPort)
 	{
@@ -416,6 +463,11 @@ public sealed class AgentRelayService : RelayRpc.RelayRpcBase, IHostedService, I
 		_agentExpirationTimeout = TimeSpan.FromMilliseconds(expirationMs);
 	}
 	
+	internal void SetRandomSeed(int seed)
+	{
+		_random = new Random(seed);
+	}
+	
 	private static void AddPortMappingMetadataToSpan(TelemetrySpan span, PortMapping portMapping)
 	{
 		span.SetAttribute("leaseId", portMapping.LeaseId);
