Horde: Generate issues for completed job steps asynchronously.

diff --git a/Engine/Source/Programs/Horde/Plugins/Build/HordeServer.Build/Issues/IssueService.cs b/Engine/Source/Programs/Horde/Plugins/Build/HordeServer.Build/Issues/IssueService.cs
--- a/Engine/Source/Programs/Horde/Plugins/Build/HordeServer.Build/Issues/IssueService.cs
+++ b/Engine/Source/Programs/Horde/Plugins/Build/HordeServer.Build/Issues/IssueService.cs
@@ -4,6 +4,7 @@ using System.Globalization;
 using System.Reflection;
 using System.Text;
 using System.Text.RegularExpressions;
+using System.Threading.Channels;
 using EpicGames.Core;
 using EpicGames.Horde.Commits;
 using EpicGames.Horde.Issues;
@@ -16,6 +17,7 @@ using HordeServer.Commits;
 using HordeServer.Jobs;
 using HordeServer.Jobs.Graphs;
 using HordeServer.Logs;
+using HordeServer.Server;
 using HordeServer.Streams;
 using HordeServer.Users;
 using Microsoft.Extensions.DependencyInjection;
@@ -192,6 +194,9 @@ namespace HordeServer.Issues
 		/// </summary>
 		const int MaxChanges = 1000;
 
+		record class CompleteStep(JobId JobId, JobStepBatchId BatchId, JobStepId StepId);
+
+		readonly JobService _jobService;
 		readonly IJobStepRefCollection _jobStepRefs;
 		readonly IIssueCollection _issueCollection;
 		readonly ICommitService _commitService;
@@ -200,6 +205,8 @@ namespace HordeServer.Issues
 		readonly ILogCollection _logCollection;
 		readonly IClock _clock;
 		readonly ITicker _ticker;
+		readonly Channel<CompleteStep> _completeSteps;
+		readonly BackgroundTask _completeStepUpdateTask;
 
 		/// <summary>
 		/// Accessor for the issue collection
@@ -249,7 +256,7 @@ namespace HordeServer.Issues
 		/// <summary>
 		/// Constructor
 		/// </summary>
-		public IssueService(IIssueCollection issueCollection, ICommitService commitService, IJobStepRefCollection jobStepRefs, IStreamCollection streams, IUserCollection userCollection, ILogCollection logCollection, IClock clock, IOptionsMonitor<BuildConfig> buildConfig, Tracer tracer, ILogger<IssueService> logger)
+		public IssueService(JobService jobService, IIssueCollection issueCollection, ICommitService commitService, IJobStepRefCollection jobStepRefs, IStreamCollection streams, IUserCollection userCollection, ILogCollection logCollection, IClock clock, IOptionsMonitor<BuildConfig> buildConfig, Tracer tracer, ILogger<IssueService> logger)
 		{
 			Type[] issueTypes = Assembly.GetExecutingAssembly().GetTypes().Where(x => !x.IsAbstract && typeof(IIssue).IsAssignableFrom(x)).ToArray();
 			foreach (Type issueType in issueTypes)
@@ -258,6 +265,7 @@ namespace HordeServer.Issues
 			}
 
 			// Get all the collections
+			_jobService = jobService;
 			_issueCollection = issueCollection;
 			_commitService = commitService;
 			_jobStepRefs = jobStepRefs;
@@ -266,6 +274,8 @@ namespace HordeServer.Issues
 			_logCollection = logCollection;
 			_clock = clock;
 			_ticker = clock.AddTicker<IssueService>(TimeSpan.FromMinutes(1.0), TickAsync, logger);
+			_completeSteps = Channel.CreateUnbounded<CompleteStep>(new UnboundedChannelOptions { SingleReader = true });
+			_completeStepUpdateTask = new BackgroundTask(UpdateCompleteStepsAsync);
 			_buildConfig = buildConfig;
 			_tracer = tracer;
 			_logger = logger;
@@ -284,15 +294,38 @@ namespace HordeServer.Issues
 		}
 
 		/// <inheritdoc/>
-		public Task StartAsync(CancellationToken cancellationToken) => _ticker.StartAsync();
+		public async Task StartAsync(CancellationToken cancellationToken)
+		{
+			_jobService.OnJobStepComplete += OnJobStepComplete;
+
+			_completeStepUpdateTask.Start();
+			await _ticker.StartAsync();
+		}
 
 		/// <inheritdoc/>
-		public Task StopAsync(CancellationToken cancellationToken) => _ticker.StopAsync();
+		public async Task StopAsync(CancellationToken cancellationToken)
+		{
+			await _ticker.StopAsync();
+
+			_completeSteps.Writer.TryComplete();
+			await _completeStepUpdateTask.WaitAsync(cancellationToken);
+
+			_jobService.OnJobStepComplete -= OnJobStepComplete;
+		}
 
 		/// <inheritdoc/>
 		public async ValueTask DisposeAsync()
 		{
 			await _ticker.DisposeAsync();
+			await _completeStepUpdateTask.DisposeAsync();
+		}
+
+		void OnJobStepComplete(IJob job, IGraph graph, JobStepBatchId batchId, JobStepId stepId)
+		{
+			if (job.UpdateIssues)
+			{
+				_completeSteps.Writer.TryWrite(new CompleteStep(job.Id, batchId, stepId));
+			}
 		}
 
 		/// <summary>
@@ -537,16 +570,30 @@ namespace HordeServer.Issues
 			return true;
 		}
 
-		/// <summary>
-		/// Marks a step as complete
-		/// </summary>
-		/// <param name="job">The job to update</param>
-		/// <param name="graph">Graph for the job</param>
-		/// <param name="batchId">Unique id of the batch</param>
-		/// <param name="stepId">Unique id of the step</param>
-		/// <param name="cancellationToken">Cancellation token for the operation</param>
-		/// <returns>Async task</returns>
-		public async Task UpdateCompleteStepAsync(IJob job, IGraph graph, JobStepBatchId batchId, JobStepId stepId, CancellationToken cancellationToken = default)
+		async Task UpdateCompleteStepsAsync(CancellationToken cancellationToken)
+		{
+			while (await _completeSteps.Reader.WaitToReadAsync(cancellationToken))
+			{
+				CompleteStep? completeStep;
+				if (_completeSteps.Reader.TryRead(out completeStep))
+				{
+					IJob? job = await _jobService.GetJobAsync(completeStep.JobId, cancellationToken);
+					if (job != null)
+					{
+						try
+						{
+							await UpdateCompleteStepAsync(job, job.Graph, completeStep.BatchId, completeStep.StepId, CancellationToken.None);
+						}
+						catch (Exception ex)
+						{
+							_logger.LogError(ex, "Exception while updating issues for {JobId}:{BatchId}:{StepId} - {Message}", completeStep.JobId, completeStep.BatchId, completeStep.StepId, ex.Message);
+						}
+					}
+				}
+			}
+		}
+
+		internal async Task UpdateCompleteStepAsync(IJob job, IGraph graph, JobStepBatchId batchId, JobStepId stepId, CancellationToken cancellationToken = default)
 		{
 			using TelemetrySpan span = _tracer.StartActiveSpan($"{nameof(IssueService)}.{nameof(UpdateCompleteStepAsync)}");
 			span.SetAttribute("jobId", job.Id.ToString());
@@ -1128,7 +1175,7 @@ namespace HordeServer.Issues
 				// If a fix changelist has been specified and it's not valid, clear it out
 				if (newVerifiedAt == null && newResolvedAt != null)
 				{
-					if(await HasFixFailedAsync(spans, issue.FixCommitId, newResolvedAt.Value, newStreams, cancellationToken))
+					if (await HasFixFailedAsync(spans, issue.FixCommitId, newResolvedAt.Value, newStreams, cancellationToken))
 					{
 						newStreams.ForEach(x => x.ContainsFix = null);
 						newResolvedAt = null;
diff --git a/Engine/Source/Programs/Horde/Plugins/Build/HordeServer.Build/Jobs/JobService.cs b/Engine/Source/Programs/Horde/Plugins/Build/HordeServer.Build/Jobs/JobService.cs
--- a/Engine/Source/Programs/Horde/Plugins/Build/HordeServer.Build/Jobs/JobService.cs
+++ b/Engine/Source/Programs/Horde/Plugins/Build/HordeServer.Build/Jobs/JobService.cs
@@ -15,7 +15,6 @@ using EpicGames.Horde.Streams;
 using EpicGames.Horde.Users;
 using HordeServer.Agents;
 using HordeServer.Commits;
-using HordeServer.Issues;
 using HordeServer.Jobs.Bisect;
 using HordeServer.Jobs.Graphs;
 using HordeServer.Jobs.Templates;
@@ -81,7 +80,6 @@ namespace HordeServer.Jobs
 		readonly JobTaskSource _jobTaskSource;
 		readonly IStreamCollection _streamCollection;
 		readonly ITemplateCollection _templateCollection;
-		readonly IssueService? _issueService;
 		readonly IPerforceService _perforceService;
 		readonly IOptionsMonitor<BuildConfig> _buildConfig;
 		readonly Tracer _tracer;
@@ -117,7 +115,7 @@ namespace HordeServer.Jobs
 		/// <summary>
 		/// Constructor
 		/// </summary>
-		public JobService(IJobCollection jobs, IGraphCollection graphs, IAgentCollection agents, IJobStepRefCollection jobStepRefs, IBisectTaskCollection bisectTasks, IJobTimingCollection jobTimings, IUserCollection userCollection, INotificationTriggerCollection triggerCollection, JobTaskSource jobTaskSource, IStreamCollection streamCollection, ITemplateCollection templateCollection, IssueService issueService, IPerforceService perforceService, Tracer tracer, IOptionsMonitor<BuildConfig> buildConfig, ILogger<JobService> logger)
+		public JobService(IJobCollection jobs, IGraphCollection graphs, IAgentCollection agents, IJobStepRefCollection jobStepRefs, IBisectTaskCollection bisectTasks, IJobTimingCollection jobTimings, IUserCollection userCollection, INotificationTriggerCollection triggerCollection, JobTaskSource jobTaskSource, IStreamCollection streamCollection, ITemplateCollection templateCollection, IPerforceService perforceService, Tracer tracer, IOptionsMonitor<BuildConfig> buildConfig, ILogger<JobService> logger)
 		{
 			_jobs = jobs;
 			_graphs = graphs;
@@ -130,7 +128,6 @@ namespace HordeServer.Jobs
 			_jobTaskSource = jobTaskSource;
 			_streamCollection = streamCollection;
 			_templateCollection = templateCollection;
-			_issueService = issueService;
 			_perforceService = perforceService;
 			_buildConfig = buildConfig;
 			_tracer = tracer;
@@ -976,6 +973,17 @@ namespace HordeServer.Jobs
 						}
 					}
 
+					// Update the jobstep ref if it completed
+					if (newState == JobStepState.Running || newState == JobStepState.Completed || newState == JobStepState.Aborted)
+					{
+						using TelemetrySpan _ = _tracer.StartActiveSpan($"{nameof(JobService)}.{nameof(TryUpdateStepAsync)}.UpdateJobStepRef");
+						if (job.TryGetBatch(batchId, out IJobStepBatch? batch) && batch.TryGetStep(stepId, out IJobStep? step) && step.StartTimeUtc != null)
+						{
+							await _jobStepRefs.UpdateAsync(job, batch, step, graph);
+							await _bisectTasks.UpdateAsync(job, batch, step, graph, cancellationToken: cancellationToken);
+						}
+					}
+
 					// Notify subscribers
 					if (newState == JobStepState.Completed)
 					{
@@ -1005,34 +1013,6 @@ namespace HordeServer.Jobs
 							}
 						}
 					}
-
-					// Update the jobstep ref if it completed
-					if (newState == JobStepState.Running || newState == JobStepState.Completed || newState == JobStepState.Aborted)
-					{
-						using TelemetrySpan _ = _tracer.StartActiveSpan($"{nameof(JobService)}.{nameof(TryUpdateStepAsync)}.UpdateJobStepRef");
-						if (job.TryGetBatch(batchId, out IJobStepBatch? batch) && batch.TryGetStep(stepId, out IJobStep? step) && step.StartTimeUtc != null)
-						{
-							await _jobStepRefs.UpdateAsync(job, batch, step, graph);
-							await _bisectTasks.UpdateAsync(job, batch, step, graph, cancellationToken: cancellationToken);
-						}
-					}
-
-					// Update any issues that depend on this step
-					if (newState == JobStepState.Completed && job.UpdateIssues)
-					{
-						if (_issueService != null)
-						{
-							using TelemetrySpan _ = _tracer.StartActiveSpan($"{nameof(JobService)}.{nameof(TryUpdateStepAsync)}.UpdateIssuesV2");
-							try
-							{
-								await _issueService.UpdateCompleteStepAsync(job, graph, batchId, stepId, cancellationToken);
-							}
-							catch (Exception ex)
-							{
-								_logger.LogError(ex, "Exception while updating issue service for job {JobId}:{BatchId}:{StepId}: {Message}", job.Id, batchId, stepId, ex.Message);
-							}
-						}
-					}
 				}
 
 				using (TelemetrySpan _ = _tracer.StartActiveSpan($"{nameof(JobService)}.{nameof(TryUpdateStepAsync)}.UpdateQueuedJobs"))
diff --git a/Engine/Source/Programs/Shared/EpicGames.Redis/Utility/RedisEvent.cs b/Engine/Source/Programs/Shared/EpicGames.Redis/Utility/RedisEvent.cs
--- a/Engine/Source/Programs/Shared/EpicGames.Redis/Utility/RedisEvent.cs
+++ b/Engine/Source/Programs/Shared/EpicGames.Redis/Utility/RedisEvent.cs
@@ -1,6 +1,7 @@
 // Copyright Epic Games, Inc. All Rights Reserved.
 
 using System;
+using System.Threading;
 using System.Threading.Tasks;
 using EpicGames.Core;
 using StackExchange.Redis;
@@ -38,12 +39,20 @@ namespace EpicGames.Redis
 		/// Create a new async event using the given channel name
 		/// </summary>
 		/// <param name="multiplexer">Multiplexer for the </param>
-		/// <param name="channel"></param>
+		/// <param name="channel">Channel for posting event updates</param>
+		/// <param name="cancellationToken">Cancellation token for the operation</param>
 		/// <returns></returns>
-		public static async Task<RedisEvent> CreateAsync(IConnectionMultiplexer multiplexer, RedisChannel channel)
+		public static async Task<RedisEvent> CreateAsync(IConnectionMultiplexer multiplexer, RedisChannel channel, CancellationToken cancellationToken = default)
 		{
 			AsyncEvent asyncEvent = new AsyncEvent();
+
 			RedisSubscription subscription = await multiplexer.SubscribeAsync(channel, x => asyncEvent.Pulse());
+			if (cancellationToken.IsCancellationRequested)
+			{
+				await subscription.DisposeAsync();
+				cancellationToken.ThrowIfCancellationRequested();
+			}
+
 			return new RedisEvent(multiplexer, channel, asyncEvent, subscription);
 		}
 
diff --git a/Engine/Source/Programs/Shared/EpicGames.Redis/Utility/RedisQueue.cs b/Engine/Source/Programs/Shared/EpicGames.Redis/Utility/RedisQueue.cs
--- a/Engine/Source/Programs/Shared/EpicGames.Redis/Utility/RedisQueue.cs
+++ b/Engine/Source/Programs/Shared/EpicGames.Redis/Utility/RedisQueue.cs
@@ -18,9 +18,10 @@ namespace EpicGames.Redis.Utility
 		/// <param name="multiplexer">Multiplexer for the connection</param>
 		/// <param name="queueKey">Key for the queue</param>
 		/// <param name="eventChannel">Channel to use for posting update notifications</param>
-		public static async Task<RedisQueue<T>> CreateAsync<T>(IConnectionMultiplexer multiplexer, RedisKey queueKey, RedisChannel eventChannel)
+		/// <param name="cancellationToken">Cancellation token for the operation</param>
+		public static async Task<RedisQueue<T>> CreateAsync<T>(IConnectionMultiplexer multiplexer, RedisKey queueKey, RedisChannel eventChannel, CancellationToken cancellationToken = default)
 		{
-			RedisEvent asyncEvent = await RedisEvent.CreateAsync(multiplexer, eventChannel);
+			RedisEvent asyncEvent = await RedisEvent.CreateAsync(multiplexer, eventChannel, cancellationToken);
 			return new RedisQueue<T>(multiplexer, asyncEvent, new RedisListKey<T>(queueKey));
 		}
 	}
@@ -61,8 +62,10 @@ namespace EpicGames.Redis.Utility
 		/// <summary>
 		/// Attempt to pop an item from the front of the queue. Returns the default value for the item if the queue is empty.
 		/// </summary>
-		public async Task<T> TryPopAsync()
+		public async Task<T> TryPopAsync(CancellationToken cancellationToken = default)
 		{
+			_ = cancellationToken; // Don't currently support cancellation due to potential of items being dropped
+
 			IDatabase database = _multiplexer.GetDatabase();
 			return await database.ListLeftPopAsync(_listKey);
 		}
