Horde: Decouple task source notifications for leases completing from the agent service.

#rnx

diff --git a/Engine/Source/Programs/Horde/Plugins/Compute/HordeServer.Compute/Agents/AgentService.cs b/Engine/Source/Programs/Horde/Plugins/Compute/HordeServer.Compute/Agents/AgentService.cs
--- a/Engine/Source/Programs/Horde/Plugins/Compute/HordeServer.Compute/Agents/AgentService.cs
+++ b/Engine/Source/Programs/Horde/Plugins/Compute/HordeServer.Compute/Agents/AgentService.cs
@@ -936,26 +936,6 @@ namespace HordeServer.Agents
 		/// <returns>Async task</returns>
 		private async Task RemoveLeaseAsync(IAgent agent, IAgentLease lease, DateTime utcNow, LeaseOutcome outcome, byte[]? output, CancellationToken cancellationToken = default)
 		{
-			// Make sure the lease is terminated correctly
-			if (lease.Payload == null)
-			{
-				_logger.LogWarning("Removing lease {LeaseId} (no payload)", lease.Id);
-			}
-			else
-			{
-				Any any = lease.Payload;
-				_logger.LogInformation("Removing lease {LeaseId} ({LeaseType})", lease.Id, any.TypeUrl);
-
-				foreach (ITaskSource taskSource in _taskSources)
-				{
-					if (any.Is(taskSource.Descriptor))
-					{
-						await taskSource.OnLeaseFinishedAsync(agent, lease.Id, any, outcome, output, Agents.GetLogger(agent.Id), cancellationToken);
-						break;
-					}
-				}
-			}
-
 			// Figure out what time the lease finished
 			DateTime finishTime = utcNow;
 			if (agent.SessionExpiresAt.HasValue && agent.SessionExpiresAt.Value < finishTime)
diff --git a/Engine/Source/Programs/Horde/Plugins/Compute/HordeServer.Compute/Agents/Leases/LeaseCollection.cs b/Engine/Source/Programs/Horde/Plugins/Compute/HordeServer.Compute/Agents/Leases/LeaseCollection.cs
--- a/Engine/Source/Programs/Horde/Plugins/Compute/HordeServer.Compute/Agents/Leases/LeaseCollection.cs
+++ b/Engine/Source/Programs/Horde/Plugins/Compute/HordeServer.Compute/Agents/Leases/LeaseCollection.cs
@@ -14,6 +14,11 @@ using OpenTelemetry.Trace;
 
 namespace HordeServer.Agents.Leases
 {
+	/// <summary>
+	/// Notification that a lease has been updated
+	/// </summary>
+	public delegate void OnLeaseCompleteDelegate(ILease lease);
+
 	/// <summary>
 	/// Collection of lease documents
 	/// </summary>
@@ -72,6 +77,9 @@ namespace HordeServer.Agents.Leases
 		readonly Tracer _tracer;
 		readonly MongoIndex<LeaseDocument> _finishTimeStartTimeCompoundIndex;
 
+		/// <inheritdoc/>
+		public event OnLeaseCompleteDelegate? OnLeaseComplete;
+
 		/// <summary>
 		/// Constructor
 		/// </summary>
@@ -198,8 +206,10 @@ namespace HordeServer.Agents.Leases
 				update = update.Set(x => x.Output, output);
 			}
 
-			UpdateResult result = await _leases.UpdateOneAsync(filter, update, cancellationToken: cancellationToken);
-			return result.ModifiedCount > 0;
+			LeaseDocument? document = await _leases.FindOneAndUpdateAsync(filter, update, new FindOneAndUpdateOptions<LeaseDocument, LeaseDocument> { ReturnDocument = ReturnDocument.After }, cancellationToken: cancellationToken);
+			OnLeaseComplete?.Invoke(document);
+
+			return document != null;
 		}
 	}
 }
diff --git a/Engine/Source/Programs/Horde/Plugins/Compute/HordeServer.Compute/ComputePlugin.cs b/Engine/Source/Programs/Horde/Plugins/Compute/HordeServer.Compute/ComputePlugin.cs
--- a/Engine/Source/Programs/Horde/Plugins/Compute/HordeServer.Compute/ComputePlugin.cs
+++ b/Engine/Source/Programs/Horde/Plugins/Compute/HordeServer.Compute/ComputePlugin.cs
@@ -63,7 +63,10 @@ namespace HordeServer
 			services.AddSingleton<IPluginResponseFilter, ComputeResponseFilter>();
 
 			services.AddSingleton<IAgentCollection, AgentCollection>();
-			services.AddSingleton<ILeaseCollection, LeaseCollection>();
+
+			services.AddSingleton<LeaseCollection>();
+			services.AddSingleton<ILeaseCollection>(sp => sp.GetRequiredService<LeaseCollection>());
+
 			services.AddSingleton<ILogCollection, LogCollection>();
 			services.AddSingleton<IPoolCollection, PoolCollection>();
 
@@ -101,6 +104,9 @@ namespace HordeServer
 			// Always run agent service too; need to be able to listen to Redis for events on any server.
 			services.AddHostedService(provider => provider.GetRequiredService<AgentService>());
 
+			// Notifications to task sources about lease completion events
+			services.AddHostedService<TaskSourceNotificationService>();
+
 			// Create the agent telemetry collection, and register the hosted service so we can flush from any server.
 			services.AddSingleton<AgentTelemetryCollection>();
 			services.AddSingleton<IAgentTelemetryCollection>(sp => sp.GetRequiredService<AgentTelemetryCollection>());
diff --git a/Engine/Source/Programs/Horde/Plugins/Compute/HordeServer.Compute/Tasks/ITaskSource.cs b/Engine/Source/Programs/Horde/Plugins/Compute/HordeServer.Compute/Tasks/ITaskSource.cs
--- a/Engine/Source/Programs/Horde/Plugins/Compute/HordeServer.Compute/Tasks/ITaskSource.cs
+++ b/Engine/Source/Programs/Horde/Plugins/Compute/HordeServer.Compute/Tasks/ITaskSource.cs
@@ -10,6 +10,8 @@ using Google.Protobuf;
 using Google.Protobuf.Reflection;
 using Google.Protobuf.WellKnownTypes;
 using HordeServer.Agents;
+using HordeServer.Agents.Leases;
+using Microsoft.Extensions.Hosting;
 using Microsoft.Extensions.Logging;
 
 namespace HordeServer.Tasks
@@ -260,4 +262,69 @@ namespace HordeServer.Tasks
 			return Task.FromResult<CreateLeaseOptions?>(lease);
 		}
 	}
+
+	/// <summary>
+	/// Service which dispatches lease complete notifications to task sources
+	/// </summary>
+	class TaskSourceNotificationService : IHostedService, IAsyncDisposable
+	{
+		readonly IAgentCollection _agentCollection;
+		readonly LeaseCollection _leaseCollection;
+		readonly ITaskSource[] _taskSources;
+		readonly AsyncQueue _asyncQueue;
+		readonly ILogger _logger;
+
+		public TaskSourceNotificationService(IAgentCollection agentCollection, LeaseCollection leaseCollection, IEnumerable<ITaskSource> taskSources, ILogger<TaskSourceNotificationService> logger)
+		{
+			_agentCollection = agentCollection;
+			_leaseCollection = leaseCollection;
+			_taskSources = taskSources.ToArray();
+			_asyncQueue = new AsyncQueue(1);
+			_logger = logger;
+		}
+
+		/// <inheritdoc/>
+		public Task StartAsync(CancellationToken cancellationToken)
+		{
+			_leaseCollection.OnLeaseComplete += OnLeaseComplete;
+			return Task.CompletedTask;
+		}
+
+		/// <inheritdoc/>
+		public async Task StopAsync(CancellationToken cancellationToken)
+		{
+			_leaseCollection.OnLeaseComplete -= OnLeaseComplete;
+			await _asyncQueue.StopAsync(cancellationToken);
+		}
+
+		/// <inheritdoc/>
+		public async ValueTask DisposeAsync()
+			=> await _asyncQueue.DisposeAsync();
+
+		void OnLeaseComplete(ILease lease)
+		{
+			if(lease.Payload.Length > 0)
+			{
+				Any payload = Any.Parser.ParseFrom(lease.Payload.Span);
+				foreach (ITaskSource taskSource in _taskSources)
+				{
+					if (payload.Is(taskSource.Descriptor))
+					{
+						_asyncQueue.Enqueue(ctx => HandleCompleteLeaseAsync(taskSource, lease, payload, ctx));
+						break;
+					}
+				}
+			}
+		}
+
+		async Task HandleCompleteLeaseAsync(ITaskSource taskSource, ILease lease, Any payload, CancellationToken cancellationToken)
+		{
+			IAgent? agent = await _agentCollection.GetAsync(lease.AgentId, cancellationToken);
+			if (agent != null)
+			{
+				_logger.LogInformation("Removing lease {LeaseId} ({LeaseType})", lease.Id, payload.TypeUrl);
+				await taskSource.OnLeaseFinishedAsync(agent, lease.Id, payload, lease.Outcome, lease.Output, _agentCollection.GetLogger(agent.Id), cancellationToken);
+			}
+		}
+	}
 }
diff --git a/Engine/Source/Programs/Shared/EpicGames.Core/AsyncQueue.cs b/Engine/Source/Programs/Shared/EpicGames.Core/AsyncQueue.cs
new file mode 100644
--- /dev/null
+++ b/Engine/Source/Programs/Shared/EpicGames.Core/AsyncQueue.cs
@@ -0,0 +1,80 @@
+// Copyright Epic Games, Inc. All Rights Reserved.
+
+using System;
+using System.Collections.Generic;
+using System.Threading;
+using System.Threading.Channels;
+using System.Threading.Tasks;
+using System.Xml;
+
+namespace EpicGames.Core
+{
+	/// <summary>
+	/// Manages a set of tasks run in the background on a fixed number of tasks. Similar to <see cref="AsyncThreadPoolWorkQueue"/>.
+	/// </summary>
+	public sealed class AsyncQueue : IAsyncDisposable
+	{
+		readonly List<Task> _workerTasks = new List<Task>();
+		readonly Channel<Func<CancellationToken, Task>> _queuedTasks = Channel.CreateUnbounded<Func<CancellationToken, Task>>();
+		readonly CancellationTokenSource _cancellationSource = new CancellationTokenSource();
+
+		/// <summary>
+		/// Constructor
+		/// </summary>
+		/// <param name="numWorkers">Number of concurrent workers executing tasks</param>
+		public AsyncQueue(int numWorkers)
+		{
+			for (int idx = 0; idx < numWorkers; idx++)
+			{
+				_workerTasks.Add(Task.Run(() => RunWorkerAsync(), _cancellationSource.Token));
+			}
+		}
+
+		/// <summary>
+		/// Finish executing all the tasks and wait for them to complete
+		/// </summary>
+		/// <param name="cancellationToken">Cancellation token for the operation</param>
+		public async Task StopAsync(CancellationToken cancellationToken)
+		{
+			_queuedTasks.Writer.TryComplete();
+			await Task.WhenAll(_workerTasks).WaitAsync(cancellationToken);
+		}
+
+		/// <inheritdoc/>
+		public async ValueTask DisposeAsync()
+		{
+			_cancellationSource.Cancel();
+			await StopAsync(CancellationToken.None).ConfigureAwait(false);
+		}
+
+		/// <summary>
+		/// Enqueue a task to be executed
+		/// </summary>
+		/// <param name="task"></param>
+		public void Enqueue(Func<CancellationToken, Task> task)
+			=> _queuedTasks.Writer.TryWrite(task);
+
+		/// <summary>
+		/// Worker executing tasks. Any exception encountered will be bubbled up as all calls are awaited.
+		/// </summary>
+		private async Task RunWorkerAsync()
+		{
+			try
+			{
+				await foreach (Func<CancellationToken, Task> task in _queuedTasks.Reader.ReadAllAsync(_cancellationSource.Token))
+				{
+					try
+					{
+						await task(_cancellationSource.Token);
+					}
+					catch
+					{
+					}
+				}
+			}
+			catch (OperationCanceledException)
+			{
+			}
+		}
+	}
+}
