Horde: Add ability to override agent's compute IP used for incoming connections

diff --git a/Engine/Source/Programs/Horde/HordeAgent/AgentSettings.cs b/Engine/Source/Programs/Horde/HordeAgent/AgentSettings.cs
--- a/Engine/Source/Programs/Horde/HordeAgent/AgentSettings.cs
+++ b/Engine/Source/Programs/Horde/HordeAgent/AgentSettings.cs
@@ -167,9 +167,14 @@ namespace HordeAgent
 		/// Option to use a local storage client rather than connecting through the server. Primarily for convenience when debugging / iterating locally.
 		/// </summary>
 		public bool UseLocalStorageClient { get; set; }
-
+		
+		/// <summary>
+		/// Incoming IP for listening for compute work. If not set, it will be automatically resolved.
+		/// </summary>
+		public string? ComputeIp { get; set; } = null;
+		
 		/// <summary>
-		/// Incoming port for listening for compute work. Needs to be tied with a lease.
+		/// Incoming port for listening for compute work. Needs to be tied with a lease. Set port to 0 to disable incoming compute requests.
 		/// </summary>
 		public int ComputePort { get; set; } = 7000;
 
diff --git a/Engine/Source/Programs/Horde/HordeAgent/Services/CapabilitiesService.cs b/Engine/Source/Programs/Horde/HordeAgent/Services/CapabilitiesService.cs
--- a/Engine/Source/Programs/Horde/HordeAgent/Services/CapabilitiesService.cs
+++ b/Engine/Source/Programs/Horde/HordeAgent/Services/CapabilitiesService.cs
@@ -407,11 +407,20 @@ namespace HordeAgent.Services
 			{
 				logger.LogWarning(ex, "Unable to get local IP addresses");
 			}
-
-			IPAddress? ip = await GetLocalIpAddressAsync(_settings.GetCurrentServerProfile().Url.Host);
+			
+			IPAddress? ip;
+			if (_settings.ComputeIp != null)
+			{
+				ip = IPAddress.Parse(_settings.ComputeIp);
+			}
+			else
+			{
+				ip = await GetLocalIpAddressAsync(_settings.GetCurrentServerProfile().Url.Host);
+			}
+			 
 			if (ip == null)
 			{
-				logger.LogWarning("Unable to get local IP address");
+				logger.LogWarning("Unable to get IP address for incoming compute requests");
 			}
 
 			// Add the compute configuration
