Horde: Update release notes.

#rnx

diff --git a/Engine/Source/Programs/Horde/Docs/ReleaseNotes.md b/Engine/Source/Programs/Horde/Docs/ReleaseNotes.md
--- a/Engine/Source/Programs/Horde/Docs/ReleaseNotes.md
+++ b/Engine/Source/Programs/Horde/Docs/ReleaseNotes.md
@@ -1,5 +1,45 @@
 # Release Notes
 
+## 2024-08-19
+
+* Add ability to override agent's compute IP used for incoming connections (35621653)
+* Generate issues for completed job steps asynchronously. (35601740)
+* Randomize port assignment for relay port mappings
+  There's a worry of re-using the same port leads to potential races with how nftables/conntrack cleans up entries. In previous impl, lowest available port number was always used which led to higher contention. (35594027)
+* Add verse module information to the UHT manifest (35590700)
+* Log issue collection locks as milliseconds for easier log parsing (35590536)
+* Check for invalid AWS region names in config (35589406)
+* Allow configuring the namespace for writing replicated Perforce data. (35570127)
+* Allow configuring the namespace for storing tool data. (35568869)
+* Allow configuring the namespace used for a particular artifact type. (35568444)
+* Present artifact types on dashboard search (35567790)
+* Enumerate types for artifact search and validation (35565664)
+* Skip expiration of artifacts if no retention policy is specified. Previous behavior was to purge all that weren't explicitly kept. (35543239)
+* Add ushell to Horde installer. (35541567)
+* Disable dedupe of bundled tool data for the moment; the controller currently filters blob requests by the tool id, which breaks when blobs are read from another tool. (35541517)
+* State redraw fixes for artifact search (35539278)
+* Fix registration of bundled tools in installed builds. (35538391)
+* Artifact search improvements (35536348)
+* Require artifact types to be declared in config files. Also allow setting expiration policies per-project and per-stream, and handle expiration of artifact types in deleted streams or of a type which are no longer declared. (35531333)
+* Improve lease termination handling during AWS spot interruptions
+  Only terminate session once a lease has finished. (35529116)
+* Fix syntax warning for CommitId (35529102)
+* Ensure resolved issues are excluded from queries when resolved by timeout. (35449290)
+* EpicGames.Core: Fix rendering of message templates containing width and format specifiers. (35446849)
+* Added more info to the "Unable to upload data to redirected URL" Horde Storage error message. (35425594)
+* Small UX Change to Horde Analytics dashboard
+  Removed the EditorPerformance category
+  Renamed the Editor Iteration category to Editor (35395963)
+* Refactor server Horde Analytics telemetry plugin dashboard stuff (35311253)
+* EpicGames.Redis: Add extension methods to allow creating conditions from typed Redis keys. (35250201)
+* [UBT]
+  * Made cache record types readonly to fix warning on farm (35194279)
+* [UBT]
+  * Fixed so cached build steps can contain log lines. This to preserve warnings etc (35192000)
+* EpicGames.Redis: Support for serializing regular Protobuf messages as well as types annotated with Protobuf.Net attributes. (35185136)
+* EpicGames.Redis: Add wrapper types for Redis data structures that combine a database reference and key. (35184619)
+* EpicGames.Redis: Rename files containing key classes to *Key.cs. (35156897)
+
 ## 2024-07-29
 
 * Fix changelist number not being returned in artifact responses. (35140759)
