Horde: Adjusting AutoConform agent scan conditions

Only include offline, enabled, and non-ephermial agents - reduces risks for possible looping conform requests if online

#rnx

diff --git a/Engine/Source/Programs/Horde/Plugins/Build/HordeServer.Build/Agents/Pools/PoolUpdateService.cs b/Engine/Source/Programs/Horde/Plugins/Build/HordeServer.Build/Agents/Pools/PoolUpdateService.cs
--- a/Engine/Source/Programs/Horde/Plugins/Build/HordeServer.Build/Agents/Pools/PoolUpdateService.cs
+++ b/Engine/Source/Programs/Horde/Plugins/Build/HordeServer.Build/Agents/Pools/PoolUpdateService.cs
@@ -132,7 +132,7 @@ namespace HordeServer.Agents.Pools
 		{
 			using TelemetrySpan span = _tracer.StartActiveSpan($"{nameof(PoolUpdateService)}.{nameof(AutoConformAgentsAsync)}");
 			List<IAgent> agents = (await _agentService.GetCachedAgentsAsync(cancellationToken))
-				.Where(x => x is { Status: AgentStatus.Ok, Enabled: true, RequestShutdown: false }).ToList();
+				.Where(x => x is { Status: AgentStatus.Stopped, Deleted: false, Enabled: true, Ephemeral: false }).ToList();
 
 			long MegabytesToBytes(long v) => v * 1024 * 1024;
 			foreach (IAgent agent in agents)
