Horde: Fix agent auto conform tests

diff --git a/Engine/Source/Programs/Horde/Plugins/Build/HordeServer.Build.Tests/Agents/Pools/PoolUpdateServiceTest.cs b/Engine/Source/Programs/Horde/Plugins/Build/HordeServer.Build.Tests/Agents/Pools/PoolUpdateServiceTest.cs
--- a/Engine/Source/Programs/Horde/Plugins/Build/HordeServer.Build.Tests/Agents/Pools/PoolUpdateServiceTest.cs
+++ b/Engine/Source/Programs/Horde/Plugins/Build/HordeServer.Build.Tests/Agents/Pools/PoolUpdateServiceTest.cs
@@ -119,7 +119,6 @@ public class PoolUpdateServiceTest : BuildTestSetup
 	[DataRow(true, 200, 50)]
 	[DataRow(true, 200, null)]
 	[DataRow(true, null, 0, 50, 300)]
-	[Ignore]
 	public async Task AutoConformAgentsAsync(bool conformRequested, params int?[] autoConformThresholdsM)
 	{
 		// Arrange
@@ -136,6 +135,9 @@ public class PoolUpdateServiceTest : BuildTestSetup
 	private async Task<IAgent> CreateAutoConformAgentAsync(int freeDiskSpaceMb, params int?[] autoConformThresholdsMb)
 	{
 		IEnumerable<AgentWorkspaceInfo> workspaces = autoConformThresholdsMb.Select(x => new AgentWorkspaceInfo(null, null, "someWorkspace", "//Some/Stream", null, false, null, null, x));
-		return await CreateAgentAsync(_pool, properties: [$"{KnownPropertyNames.DiskFreeSpace}={freeDiskSpaceMb * 1024 * 1024}"], workspaces: workspaces.ToList());
+		IAgent? agent = await CreateAgentAsync(_pool, properties: [$"{KnownPropertyNames.DiskFreeSpace}={freeDiskSpaceMb * 1024 * 1024}"], workspaces: workspaces.ToList());
+		agent = await agent.TryTerminateSessionAsync(); // Make agent status = stopped
+		Assert.IsNotNull(agent);
+		return agent;
 	}
 }
