Horde: Fix null reference exception when lease updates fail.

#rnx

diff --git a/Engine/Source/Programs/Horde/Plugins/Compute/HordeServer.Compute/Agents/Leases/LeaseCollection.cs b/Engine/Source/Programs/Horde/Plugins/Compute/HordeServer.Compute/Agents/Leases/LeaseCollection.cs
--- a/Engine/Source/Programs/Horde/Plugins/Compute/HordeServer.Compute/Agents/Leases/LeaseCollection.cs
+++ b/Engine/Source/Programs/Horde/Plugins/Compute/HordeServer.Compute/Agents/Leases/LeaseCollection.cs
@@ -207,9 +207,13 @@ namespace HordeServer.Agents.Leases
 			}
 
 			LeaseDocument? document = await _leases.FindOneAndUpdateAsync(filter, update, new FindOneAndUpdateOptions<LeaseDocument, LeaseDocument> { ReturnDocument = ReturnDocument.After }, cancellationToken: cancellationToken);
-			OnLeaseComplete?.Invoke(document);
+			if (document == null)
+			{
+				return false;
+			}
 
-			return document != null;
+			OnLeaseComplete?.Invoke(document);
+			return true;
 		}
 	}
 }
