Horde: Re-enable ComputeService relay port test with a deterministic seed

Ports are assigned at random with CL 35594027

diff --git a/Engine/Source/Programs/Horde/Plugins/Build/HordeServer.Build.Tests/Compute/ComputeServiceTest.cs b/Engine/Source/Programs/Horde/Plugins/Build/HordeServer.Build.Tests/Compute/ComputeServiceTest.cs
--- a/Engine/Source/Programs/Horde/Plugins/Build/HordeServer.Build.Tests/Compute/ComputeServiceTest.cs
+++ b/Engine/Source/Programs/Horde/Plugins/Build/HordeServer.Build.Tests/Compute/ComputeServiceTest.cs
@@ -366,15 +366,14 @@ namespace HordeServer.Tests.Compute
 		}
 
 		[TestMethod]
-		[Ignore]
 		public async Task Connection_Relay_PortsAreMapped_Async()
 		{
 			ComputeResource? cr = await AllocateAsync(ConnectionMode.Relay, ports: new Dictionary<string, int> { { "myOtherPort", 13000 }, { "myPort", 12000 } });
 			Assert.AreEqual(ConnectionMode.Relay, cr!.ConnectionMode);
 			Assert.AreEqual(3, cr.Ports.Count);
-			Assert.AreEqual(new ComputeResourcePort(10000, 5000), cr.Ports[ConnectionMetadataPort.ComputeId]);
-			Assert.AreEqual(new ComputeResourcePort(10002, 12000), cr.Ports["myPort"]);
-			Assert.AreEqual(new ComputeResourcePort(10004, 13000), cr.Ports["myOtherPort"]);
+			Assert.AreEqual(new ComputeResourcePort(12214, 5000), cr.Ports[ConnectionMetadataPort.ComputeId]);
+			Assert.AreEqual(new ComputeResourcePort(18656, 12000), cr.Ports["myPort"]);
+			Assert.AreEqual(new ComputeResourcePort(23151, 13000), cr.Ports["myOtherPort"]);
 		}
 
 		private async Task<ComputeService> CreateComputeServiceAsync(string? tunnelAddress, string pool, ComputeClusterConfig ccc)
@@ -387,6 +386,7 @@ namespace HordeServer.Tests.Compute
 			await CreateAgentAsync(new PoolId(pool), properties: ["ComputeIp=11.0.0.1", "ComputePort=5000"]);
 			computeConfig.CurrentValue.Clusters = [ccc];
 			computeConfig.CurrentValue.PostLoad(new PluginConfigOptions(ConfigVersion.Latest, Enumerable.Empty<IPluginConfig>(), new AclConfig()));
+			AgentRelayService.SetRandomSeed(1);
 			return cs;
 		}
 
