Horde: Dashboard pluginification and telemetry plugin
#siubmittool bypass

diff --git a/Engine/Source/Programs/Horde/HordeDashboard/package.json b/Engine/Source/Programs/Horde/HordeDashboard/package.json
--- a/Engine/Source/Programs/Horde/HordeDashboard/package.json
+++ b/Engine/Source/Programs/Horde/HordeDashboard/package.json
@@ -2,27 +2,27 @@
   "name": "horde-dashboard",
   "private": true,
   "version": "0.1.3",
-  "type": "module",  
+  "type": "module",
   "scripts": {
     "dev": "vite",
-    "build": "tsc && vite build",    
+    "build": "tsc && vite build",
     "lint": "eslint . --ext ts,tsx --report-unused-disable-directives --max-warnings 0",
     "preview": "vite preview",
     "license-check": "license-checker-rseidelsohn --production --excludePackages horde-dashboard --depth 0 --limitAttributes licenses --csv --customPath ./license-format.json --onlyAllow \"Apache-2.0;MIT;ISC\""
   },
   "dependencies": {
-    "react": "^18.2.0",
-    "react-dom": "^18.2.0",
     "@blueprintjs/icons": "^3.13.0",
     "@datadog/browser-logs": "^4.30.0",
     "@fluentui/react": "^8.106.2",
     "@fluentui/react-charting": "5.2.0",
     "d3": "^7.6.1",
-    "markdown-to-jsx": "^7.1.8",
+    "markdown-to-jsx": "7.4.7",
     "mobx": "^6.12.2",
     "mobx-react-lite": "^4.0.7",
     "moment-timezone": "^0.5.40",
+    "react": "^18.2.0",
     "react-custom-scrollbars": "^4.2.1",
+    "react-dom": "^18.2.0",
     "react-flame-graph": "^1.4.0",
     "react-highlighter": "^0.4.3",
     "react-router-dom": "6.6.2",
@@ -31,23 +31,23 @@
     "react-text-marquee": "^1.0.1"
   },
   "devDependencies": {
+    "@types/d3": "^7.4.0",
     "@types/react": "^18.2.66",
+    "@types/react-custom-scrollbars": "^4.0.10",
     "@types/react-dom": "^18.2.22",
+    "@types/react-highlighter": "^0.3.5",
+    "@types/react-router-dom": "^5.3.3",
+    "@types/react-sparklines": "^1.7.5",
     "@typescript-eslint/eslint-plugin": "^7.2.0",
     "@typescript-eslint/parser": "^7.2.0",
     "@vitejs/plugin-react": "^4.2.1",
     "eslint": "^8.57.0",
     "eslint-plugin-react-hooks": "^4.6.0",
     "eslint-plugin-react-refresh": "^0.4.6",
+    "license-checker-rseidelsohn": "^4.3.0",
+    "rimraf": "^3.0.0",
     "typescript": "^5.4.5",
     "vite": "^5.2.0",
-    "@types/d3": "^7.4.0",
-    "@types/markdown-to-jsx": "7.0.1",
-    "@types/react-custom-scrollbars": "^4.0.10",
-    "@types/react-highlighter": "^0.3.5",
-    "@types/react-router-dom": "^5.3.3",
-    "@types/react-sparklines": "^1.7.5",
-    "license-checker-rseidelsohn": "^4.3.0",
-    "rimraf": "^3.0.0"
+    "vite-tsconfig-paths": "^4.3.2"
   }
-}
\ No newline at end of file
+}
diff --git a/Engine/Source/Programs/Horde/HordeDashboard/plugins/analytics/api.ts b/Engine/Source/Programs/Horde/HordeDashboard/plugins/analytics/api.ts
new file mode 100644
--- /dev/null
+++ b/Engine/Source/Programs/Horde/HordeDashboard/plugins/analytics/api.ts
@@ -0,0 +1,163 @@
+import backend from "horde/backend";
+
+export type MetricsQuery = {
+	id: string[];
+	minTime?: string;
+	maxTime?: string;
+	group?: string;
+	results?: number;
+}
+
+/// Metrics matching a particular query
+export type GetTelemetryMetricsResponse = {
+
+	metricId: string;
+
+	groupBy: string;
+
+	/// Metrics matching the search terms	
+	metrics: GetTelemetryMetricResponse[];
+}
+
+/// Information about a particular metric
+export type GetTelemetryMetricResponse = {
+
+	/// Start time for the sample	
+	time: Date;
+
+	/// Name of the group	
+	group?: string;
+
+	/// Value for the metric	
+	value: number;
+
+	// Added locally in the dashboard
+	// GetTelemetryMetricsResponse id
+	id: string;
+
+	// added locally by dashboard
+	key: string;
+
+	// added locally by dashboard
+	keyElements: string[];
+
+	threshold?: number;
+
+	// calculated on dashboard, group name => value
+	groupValues?: Record<string, string>;
+}
+
+export type TelemetryDisplayType = "Time" | "Ratio" | "Value";
+export type TelemetryGraphType = "Line" | "Indicator";
+
+/// Metric attached to a telemetry chart	
+export type GetTelemetryChartMetricResponse = {
+
+	/// Associated metric id	
+	metricId: string;
+
+	/// The threshold for KPI values	
+	threshold?: number;
+
+	/// The metric alias for display purposes	
+	alias?: string;
+
+}
+
+/// Telemetry chart configuraton
+export type GetTelemetryChartResponse = {
+
+	/// The name of the chart, will be displayed on the dashboard	
+	name: string;
+
+	/// The unit to display	
+	display: TelemetryDisplayType;
+
+	/// The graph type 	
+	graph: TelemetryGraphType;
+
+	/// List of configured metrics	
+	metrics: GetTelemetryChartMetricResponse[];
+
+	/// The min unit value for clamping chart	
+	min?: number;
+
+	/// The max unit value for clamping chart	
+	max?: number;
+}
+
+/// A chart categody, will be displayed on the dashbord under an associated pivot
+export type GetTelemetryCategoryResponse = {
+
+	/// The name of the category
+	name: string;
+
+	/// The charts contained within the category
+	charts: GetTelemetryChartResponse[];
+}
+
+/// A telemetry view variable used for filtering the charting data
+export type GetTelemetryVariableResponse = {
+	/// The name of the variable for display purposes
+	name: string;
+
+	/// The associated data group attached to the variable 
+	group: string;
+
+	/// default values to select
+	defaults: string[];
+
+	/// Populated on dashboard
+	values: string[];
+}
+
+/// A telemetry view of related metrics, divided into categofies
+export type GetTelemetryViewResponse = {
+
+	/// Identifier for the view
+	id: string;
+
+	/// The name of the view
+	name: string;
+
+	/// The telemetry store id the view uses
+	telemetryStoreId: string;
+
+	///  The variables used to filter the view data
+	variables: GetTelemetryVariableResponse[];
+
+	/// The categories contained within the view
+	categories: GetTelemetryCategoryResponse[];
+}
+
+export async function getViews(): Promise<GetTelemetryViewResponse[]> {
+	return new Promise<GetTelemetryViewResponse[]>((resolve, reject) => {
+		backend.fetch.get(`api/v1/telemetry/views`).then((response) => {
+			const result = response.data as GetTelemetryViewResponse[];
+			resolve(result);
+		}).catch(reason => { reject(reason); });
+	});
+}
+
+export async function getMetrics(telemetryStoreId: string, query: MetricsQuery): Promise<GetTelemetryMetricsResponse[]> {
+	query.id = query.id.map(id => encodeURIComponent(id));
+	return new Promise<GetTelemetryMetricsResponse[]>((resolve, reject) => {		
+		backend.fetch.get(`api/v1/telemetry/${telemetryStoreId}/metrics`, { params: query }).then((response) => {
+			const result = response.data as GetTelemetryMetricsResponse[];
+			result?.forEach(r => {
+				r?.metrics.forEach(m => {
+					if (m.time) {
+						m.time = new Date(m.time);
+					} else {
+						console.warn("Metrics missing time property");
+						m.time = new Date();
+					}
+				})
+			})
+			resolve(result);
+		}).catch(reason => { reject(reason); });
+	});
+}
+
+
+
diff --git a/Engine/Source/Programs/Horde/HordeDashboard/plugins/analytics/index.tsx b/Engine/Source/Programs/Horde/HordeDashboard/plugins/analytics/index.tsx
new file mode 100644
--- /dev/null
+++ b/Engine/Source/Programs/Horde/HordeDashboard/plugins/analytics/index.tsx
@@ -0,0 +1,13 @@
+import { registerHordePlugin, MountType } from "..";
+import { TelemetryView } from "./telemetryView";
+
+registerHordePlugin({
+   id: "analytics",
+   routes: [{ path: "analytics", element: <TelemetryView /> }],
+   mount: {
+      type: MountType.TopNav,
+      context: "Tools",
+      text: "Analytics",
+      route: `/analytics`
+   }
+})
\ No newline at end of file
diff --git a/Engine/Source/Programs/Horde/HordeDashboard/plugins/analytics/telemetryData.ts b/Engine/Source/Programs/Horde/HordeDashboard/plugins/analytics/telemetryData.ts
new file mode 100644
--- /dev/null
+++ b/Engine/Source/Programs/Horde/HordeDashboard/plugins/analytics/telemetryData.ts
@@ -0,0 +1,198 @@
+import { getMetrics, GetTelemetryMetricsResponse, GetTelemetryViewResponse } from "./api";
+
+export type MetricId = string;
+export type TelemetryViewId = string;
+
+export const graphColors = [
+    "#ff780a", // o2
+    "#3274d9", // b2
+    "#56a64b", // g2
+    "#e02f44", // r2
+    "#f2cc0c", // y2
+    "#a352cc", // p2
+    "#ff9830", // o3
+    "#5794f2", // b3
+    "#73bf69", // g3
+    "#f2495c", // r3
+    "#fade2a", // y3
+    "#b877d9", // p3
+    "#fa6400", // o1
+    "#1f60c4", // b1
+    "#37872d", // g1
+    "#c4162a", // r1
+    "#e0b400", // y1
+    "#e0b400", // p1
+    "#ffb357", // o4
+    "#8ab8ff", // b4
+    "#96d98d", // g4
+    "#733d47", // r4
+    "#ffee52", // y4
+    "#ca95e5", // p4
+    "#ffcb7d", // o5
+    "#c0d8ff", // b5
+    "#c8f2c2", // g5
+    "#ffa6b0", // r5
+    "#fff899", // y5
+    "#deb6f2" // p5
+]
+
+const groupRegex = /,(?=(?:[^"]*"[^"]*")*[^"]*$)/
+
+const loadedMetrics = new Map<string, GetTelemetryMetricsResponse>();
+
+export const clearTelemetryViewMetrics = () => {
+    loadedMetrics.clear();
+}
+
+const getTelemetryViewMetrics = async (view: GetTelemetryViewResponse, categoryName: string, minTime: string, maxTime: string): Promise<GetTelemetryMetricsResponse[] | undefined> => {
+
+    return new Promise<GetTelemetryMetricsResponse[] | undefined>(async (resolve, reject) => {
+
+        const category = view.categories.find(c => c.name === categoryName);
+        if (!category) {
+            reject("Invalid category");
+            return;
+        }
+
+        // get set of metrics
+        const viewMetrics = new Set<string>(category.charts.map(c => c.metrics.map(m => m.metricId).flat()).flat());
+        const needMetrics = new Set(viewMetrics);
+
+        loadedMetrics.forEach((_, id) => needMetrics.delete(id));
+
+        if (needMetrics.size) {
+
+            const need: string[] = Array.from(needMetrics);
+            const allMetrics = await getMetrics(view.telemetryStoreId, { id: need, minTime: minTime, maxTime: maxTime, results: 4096 * 32 });
+
+            for (let i = 0; i < need.length; i++) {
+                const metricId = need[i];
+                const metrics = allMetrics.find(m => m.metricId === metricId);
+                if (!metrics) {
+                    loadedMetrics.set(metricId, { metricId: metricId, groupBy: "", metrics: [] });
+                    continue;
+                }
+
+                const chartMetric = category.charts.map(c => c.metrics).flat().find(m => m.metricId === metricId);
+
+                let groupBy = metrics.groupBy;
+                if (!groupBy) {
+                    reject(`Unable to get groupby for ${metricId}`);
+                    return;
+                }
+                metrics.groupBy = groupBy.replaceAll("$.Payload.", "");
+
+                const groups = metrics.groupBy.split(",");
+
+                metrics.metrics = metrics.metrics.filter(m => {
+                    m.id = metricId;
+                    if (m.group) {
+
+                        // todo: can remove this, filtering out some bad escaping from initial implementation
+                        if (m.group.indexOf("\"\"") !== -1) {
+                            return false;
+                        }
+
+                        const groupValues = m.group.split(groupRegex).map((v => v.replaceAll("\"", "")));
+
+                        m.groupValues = {};
+                        let key: string[] = [];
+                        let skip = false;
+                        groups.forEach((g, idx) => {
+                            const value = groupValues[idx];
+                            if (!value) {
+                                skip = true;
+                                return;
+                            }
+                            m.groupValues![g] = value;
+                            key.push(value);
+                        })
+
+                        if (skip) {
+                            return false;
+                        }
+
+                        if (chartMetric?.alias) {
+                            key.push(chartMetric?.alias)
+                        }
+                        m.key = key.join(":");
+                        m.keyElements = key;
+
+                        if (chartMetric?.threshold) {
+                            m.threshold = chartMetric.threshold;
+                        }
+                    }
+                    return true;
+                })
+
+                loadedMetrics.set(metricId, metrics);
+            }
+        }
+
+        const result: GetTelemetryMetricsResponse[] = [];
+
+        viewMetrics.forEach(id => {
+
+            result.push(loadedMetrics.get(id)!);
+        })
+
+        resolve(result);
+    });
+}
+
+export type TelemetryViewData = {
+    minTime: Date;
+    maxTime: Date;
+    metrics: GetTelemetryMetricsResponse[];
+}
+
+export const getTelemetryViewData = async (view: GetTelemetryViewResponse, categoryName: string, minDate: Date, maxDate: Date): Promise<TelemetryViewData | undefined> => {
+
+    return new Promise<TelemetryViewData | undefined>(async (resolve, reject) => {
+
+        const cat = view.categories.find(c => c.name === categoryName);
+        if (!cat) {
+            reject("Unable to get view category");
+            return;
+        }
+
+        const viewMetrics: GetTelemetryMetricsResponse[] = await getTelemetryViewMetrics(view, categoryName, minDate.toISOString(), maxDate.toISOString()) as GetTelemetryMetricsResponse[];
+
+        if (!viewMetrics?.length) {
+            reject("Unable to get view metrics");
+            return;
+        }
+
+        // get variables
+        view.variables?.forEach(v => {
+            const values = new Set<string>();
+            const varGroup = v.group;
+            viewMetrics.forEach(metric => {
+
+                const groups = metric.groupBy.split(",");
+                for (let i = 0; i < groups.length; i++) {
+                    const group = groups[i];
+                    if (varGroup.toLowerCase() === group.toLowerCase()) {
+                        metric.metrics.forEach(m => {
+                            if (m.groupValues) {
+                                values.add(m.groupValues[group]);
+                            }
+                        })
+                    }
+                }
+            })
+
+            v.values = Array.from(values).sort((a, b) => a.localeCompare(b));
+
+        })
+
+        const result: TelemetryViewData = {
+            minTime: minDate,
+            maxTime: maxDate,
+            metrics: viewMetrics
+        }
+
+        resolve(result);
+    });
+
+}
\ No newline at end of file
diff --git a/Engine/Source/Programs/Horde/HordeDashboard/plugins/analytics/telemetryLineGraph.ts b/Engine/Source/Programs/Horde/HordeDashboard/plugins/analytics/telemetryLineGraph.ts
new file mode 100644
--- /dev/null
+++ b/Engine/Source/Programs/Horde/HordeDashboard/plugins/analytics/telemetryLineGraph.ts
@@ -0,0 +1,291 @@
+import * as d3 from "d3";
+import moment from "moment";
+import { GetTelemetryMetricResponse, GetTelemetryMetricsResponse, GetTelemetryChartResponse } from "./api";
+import dashboard from "horde/backend/Dashboard";
+import { displayTimeZone, msecToElapsed } from "horde/base/utilities/timeUtils";
+import { graphColors } from "./telemetryData";
+
+// Handle bad "@types/d3" types, fix if addressed upstream
+const _d3 = d3 as any;
+
+type SelectionType = d3.Selection<SVGGElement, unknown, null, undefined>;
+type Zoom = d3.ZoomBehavior<Element, unknown>;
+type Scalar = d3.ScaleLinear<number, number>;
+
+type DataPoint = [number, number, string, number, number];
+
+export class TelemetryLineRenderer {
+
+   render(chart: GetTelemetryChartResponse, metrics: GetTelemetryMetricsResponse[], legend: string[], minTime: Date, maxTime: Date, container: HTMLDivElement, onZoom: (name: string, event: any) => void, onTimeSelect: (name: string, minTime: Date, maxTime: Date) => void, onDataHover: (key: string, x: number, y: number, time: Date, value: number, color: string) => void, scale = 1.0): any {
+
+      let minValue = Number.MAX_SAFE_INTEGER
+      let maxValue = Number.MIN_SAFE_INTEGER
+
+      const allMetrics: GetTelemetryMetricResponse[] = [];
+
+      metrics.forEach(metric => {
+         metric.metrics.forEach(m => {
+            allMetrics.push(m);
+            minValue = Math.min(m.value, minValue);
+            maxValue = Math.max(m.value, maxValue);
+         })
+      })
+
+      const ratio = chart.display === "Ratio"
+
+      minValue = chart.min ?? 0;
+
+      if (ratio) {
+         maxValue = 1.0
+      }
+
+      let svg = this.svg;
+      const width = 1000;
+      const height = 300;
+      const margin = { top: 16, right: 32, bottom: 0, left: 64 };
+
+      const x = this.scaleX = d3.scaleLinear()
+         .domain([minTime, maxTime].map(d => d.getTime() / 1000))
+         .range([margin.left, width - margin.right])
+
+      const y = this.scaleY = d3.scaleLinear()
+         .domain([minValue, maxValue / scale])
+         .range([height - margin.bottom, margin.top])
+
+      if (!svg) {
+         svg = d3.select(container)
+            .append("svg") as any as SelectionType
+
+         this.svg = svg;
+      } else {
+         // remove tooltip
+         d3.select(container).selectAll('div').remove();
+         svg.selectAll("*").remove();
+      }
+
+      svg.attr("viewBox", [0, 0, width, height] as any)
+         .attr("width", width)
+         .attr("height", height)
+
+      const clipId = `metrics_${chart.name}_clip}`;
+
+      svg.append("clipPath")
+         .attr("id", clipId)
+         .append("rect")
+         .attr("x", margin.left)
+         .attr("y", 0)
+         .attr("width", width - margin.left - margin.right + 2)
+         .attr("height", height);
+
+      const points = allMetrics.map((m, index) => [x(m.time.getTime() / 1000), y(m.value), m.key, legend.indexOf(m.key) % graphColors.length, index]) as DataPoint[];
+      const groups = _d3.rollup(points, v => Object.assign(v, { z: v[0][2] }), d => d[2]);
+      const gvalues = Array.from(groups.values());
+
+      const line = d3.line().curve(d3.curveMonotoneX);
+      svg.append("g")
+         .attr("clip-path", `url(#${clipId})`)
+         .attr("fill", "none")
+         .attr("stroke-width", 1.5)
+         .attr("stroke-linejoin", "round")
+         .attr("stroke-linecap", "round")
+         .selectAll("path")
+         .data(gvalues)
+         .join("path")
+         .attr("stroke", (d:any) => { return d[0][3] !== undefined ? graphColors[d[0][3] as number] : "#8ab8ff" })
+         .attr("d", line as any);
+
+      svg.append("g")
+         .selectAll("circle")
+         .data(points)
+         .join("circle")
+         .attr("id", i => {
+            return `circle_id_${i[4]}`
+         })
+         .attr("cx", (i) => {
+            return i[0];
+         })
+         .attr("cy", i => i[1])
+         .attr("fill", i => { return i[3] !== undefined ? graphColors[i[3] as number] : "#8ab8ff" })
+         .attr("r", 0);
+
+      let ticks: number[] = [];
+      const inc = (maxTime.getTime() - minTime.getTime()) / 10;
+      for (let i = 0; i < 10; i++) {
+         ticks.push((minTime.getTime() + inc * i) / 1000);
+      }
+
+      const xAxis = (g: SelectionType) => {
+
+         g.attr("transform", `translate(0,18)`)
+            .style("font-family", "Horde Open Sans Regular")
+            .style("font-size", "11px")
+            .call(d3.axisTop(x)
+               .tickValues(ticks)
+               .tickFormat(d => {
+                  const time = moment(new Date((d as number) * 1000)).tz(displayTimeZone());
+                  return time.format("MM/DD HH:mm")
+               })
+               .tickSizeOuter(0))
+            .call(g => g.select(".domain").remove())
+            .call(g => g.selectAll(".tick line").attr("stroke-opacity", dashboard.darktheme ? 0.35 : 0.25)
+               .attr("stroke", dashboard.darktheme ? "#6D6C6B" : "#4D4C4B")
+               .attr("y2", height - margin.bottom))
+      }
+
+      // left axis
+      const yAxis = (g: SelectionType) => {
+
+         g.attr("transform", `translate(${margin.left},0)`)
+            .style("font-family", "Horde Open Sans Regular")
+            .style("font-size", "9px")
+            .call(d3.axisLeft(this.scaleY!)
+               .ticks(10)
+               .tickFormat((d) => {
+
+
+                  if (ratio) {
+                     return (d as number * 100).toString() + "%"
+                  }
+
+                  if (!d) {
+                     return "";
+                  }
+
+                  if (chart.display === "Value") {
+                     return d.toString();
+                  }
+
+                  return msecToElapsed((d as number) * 1000, true, true);
+
+               }))
+      }
+
+      svg.append("g").attr("class", "x-axis").call(xAxis)
+      svg.append("g").attr("class", "y-axis").call(yAxis)
+
+      /*
+      // zoom
+      const zoom = this.zoom = d3.zoom()
+         .scaleExtent([1, 12])
+         .extent([[margin.left, 0], [width - margin.right, height]])
+         .translateExtent([[margin.left, -Infinity], [width - margin.right, Infinity]])
+         .on("zoom", zoomed as any);
+
+      function zoomed(event: any, propogate = true) {
+
+         if (propogate) {
+            onZoom(chart.name, event);
+         }
+
+         x.range([margin.left, width - margin.right].map(d => event.transform.applyX(d)));
+
+         const npoints = allMetrics.map((d) => [x(d.time.getTime() / 1000), y(d.value), d.key]);
+         const ngroups = d3.rollup(npoints, v => Object.assign(v, { z: v[0][2] }), d => d[2]);
+
+         svg!.selectAll("path")
+            .data(ngroups.values())
+            .join("path")
+            .attr("d", line as any);
+
+
+         svg!.selectAll(".x-axis").call(xAxis as any);
+      }
+
+      svg.call(zoom as any);
+      */
+
+      const brush = d3.brushX().extent([[margin.left, margin.top], [width - margin.right, height]]).on('end', (event: any) => {
+         const extent = event?.selection
+         if (!extent) {
+            return;
+         }
+
+         const min = extent[0];// - margin.left;
+         const max = extent[1];// - margin.left;
+         const startDate = x.invert(min);
+         const endDate = x.invert(max);
+
+         onTimeSelect(chart.name, new Date(startDate * 1000), new Date(endDate * 1000))
+
+      });
+
+      svg.append("g")
+         .attr("class", "brush")
+         .call(brush);
+
+      //svg.on("wheel", (event) => { event.preventDefault(); })
+
+      const keyPoints = new Map<string, DataPoint[]>();
+      points.forEach(p => {
+
+         if (!keyPoints.has(p[2])) {
+            keyPoints.set(p[2], [])
+         }
+         keyPoints.get(p[2])!.push(p);
+      });
+
+      const handleMouseMove = (event: any) => {
+
+         const mouseX = _d3.pointer(event)[0];
+         const mouseY = _d3.pointer(event)[1];
+
+         // find closest point on x axis for each data key
+         const closestX = new Map<string, DataPoint>();
+
+         keyPoints.forEach((values, key) => {
+
+            let closest = values.reduce((best, data, i) => {
+
+               let absx = Math.abs(data[0] - mouseX);
+
+               if (absx < best.value) {
+                  return { index: i, value: absx };
+               }
+               else {
+                  return best;
+               }
+
+            }, { index: 0, value: Number.MAX_SAFE_INTEGER });
+
+            closestX.set(key, values[closest.index]);
+         })
+
+         svg!.selectAll(`circle`).attr("r", 0);
+
+         let closestY: DataPoint | undefined;
+
+         closestX.forEach((point) => {
+
+            if (!closestY) {
+               closestY = point;
+            } else if (Math.abs(point[1] - mouseY) < Math.abs(closestY[1] - mouseY)) {
+               closestY = point;
+            }
+            svg!.select(`#circle_id_${point[4]}`).attr("r", 4);
+         })
+
+         if (closestY && Math.abs(closestY[1] - mouseY) < 16 && Math.abs(closestY[0] - mouseX) < 16) {
+            onDataHover(closestY[2], mouseX, mouseY, new Date(x.invert(closestY[0]) * 1000), y.invert(closestY[1]), closestY[3] !== undefined ? graphColors[closestY[3] as number] : "#8ab8ff")
+         } else {
+            onDataHover("__clear__", 0, 0, new Date(), 0, "");
+         }
+      }
+
+      const handleMouseLeave = (event: any) => {
+         svg!.selectAll(`circle`).attr("r", 0);
+         onDataHover("__clear__", 0, 0, new Date(), 0, "");
+      }
+
+      svg.on("mousemove", (event) => handleMouseMove(event))
+      svg.on("mouseleave", (event) => handleMouseLeave(event))
+
+      return undefined;  //zoomed;
+   }
+
+   svg?: SelectionType;
+   zoom?: Zoom;
+   scaleX?: Scalar;;
+   scaleY?: Scalar;;
+   metrics?: GetTelemetryMetricsResponse;
+}
+
diff --git a/Engine/Source/Programs/Horde/HordeDashboard/plugins/analytics/telemetryView.tsx b/Engine/Source/Programs/Horde/HordeDashboard/plugins/analytics/telemetryView.tsx
new file mode 100644
--- /dev/null
+++ b/Engine/Source/Programs/Horde/HordeDashboard/plugins/analytics/telemetryView.tsx
@@ -0,0 +1,1478 @@
+import { ComboBox, DefaultButton, DirectionalHint, FontIcon, IComboBox, IComboBoxOption, IComboBoxStyles, ITooltipHostStyles, ITooltipProps, Icon, Pivot, PivotItem, SelectableOptionMenuItemType, Spinner, SpinnerSize, Stack, Text, TooltipHost, mergeStyleSets, mergeStyles } from "@fluentui/react";
+import { useConst } from '@fluentui/react-hooks';
+import dashboard, { StatusColor } from "horde/backend/Dashboard";
+import { useWindowSize } from "horde/base/utilities/hooks";
+import { displayTimeZone, msecToElapsed } from "horde/base/utilities/timeUtils";
+import { Breadcrumbs } from "horde/components/Breadcrumbs";
+import { TopNav } from "horde/components/TopNav";
+import { getHordeStyling } from "horde/styles/Styles";
+import { action, makeObservable, observable } from "mobx";
+import { observer } from "mobx-react-lite";
+import moment from "moment";
+import React, { useEffect, useId, useState } from "react";
+import { useNavigate, useSearchParams } from "react-router-dom";
+import { GetTelemetryChartResponse, GetTelemetryMetricResponse, GetTelemetryMetricsResponse, GetTelemetryVariableResponse, GetTelemetryViewResponse, getViews } from "./api";
+import { TelemetryViewData, clearTelemetryViewMetrics, getTelemetryViewData, graphColors } from "./telemetryData";
+import { TelemetryLineRenderer } from "./telemetryLineGraph";
+
+const timeSelections: TimeSelection[] = [
+   {
+      text: "Past 1 Hour", key: "time_1_hour", minutes: 60
+   },
+   {
+      text: "Past 2 Hours", key: "time_2_hours", minutes: 60 * 2
+   },
+   {
+      text: "Past 4 Hours", key: "time_4_hours", minutes: 60 * 4
+   },
+   {
+      text: "Past 1 Day", key: "time_1_day", minutes: 60 * 24
+   },
+   {
+      text: "Past 2 Days", key: "time_2_days", minutes: 60 * 24 * 2
+   },
+   {
+      text: "Past 1 Week", key: "time_1_week", minutes: 60 * 24 * 7
+   },
+   {
+      text: "Past 2 Weeks", key: "time_2_weeks", minutes: 60 * 24 * 7 * 2
+   },
+   {
+      text: "Past Month", key: "time_4_weeks", minutes: 60 * 24 * 7 * 4
+   },
+   {
+      text: "Custom Range", key: "time_custom", minutes: 0, hidden: true
+   }
+]
+
+type LegendEntry = {
+   display: string;
+   key: string;
+   min: number;
+   max: number;
+   change: number;
+}
+
+
+type SearchState = {
+   viewId?: string;
+   category?: string;
+   variables?: string[];
+   minutes?: number;
+}
+
+class MetricsHandler {
+   constructor() {
+      makeObservable(this);
+   }
+
+   subscribe() {
+      if (this.updated) { }
+   }
+
+   subscribeToSearch() {
+      if (this.searchUpdated) { }
+   }
+
+   async set(category: string) {
+
+      if (!this.view || category === this._category) {
+         return;
+      }
+
+      this._category = category;
+      this.querying = true;
+      this.setUpdated();
+
+      this.metrics = await getTelemetryViewData(this.view, category, this.minDate, this.maxDate);
+
+      this.querying = false;
+      this.setUpdated();
+   }
+
+   getCharts() {
+      if (!this.searchState.category) {
+         return [];
+      }
+
+      const category = this.view?.categories.find(c => c.name === this.searchState.category);
+      if (!category) {
+         return [];
+      }
+
+      return category.charts;
+   }
+
+   getChartVariables(): GetTelemetryVariableResponse[] {
+
+      const variables = this.view?.variables;
+
+      if (!variables?.length || !this.metrics) {
+         return [];
+      }
+
+      const charts = this.getCharts();
+      const allMetrics = charts.map(c => this.getAllChartMetrics(c.name).map(m => m.metrics).flat()).flat();
+
+      const vars = variables.filter(v => {
+
+         if (!v.values.length) {
+            return false;
+         }
+
+         return !!allMetrics.find(m => m.groupValues && m.groupValues[v.group] !== undefined);
+      });
+
+      return vars;
+   }
+
+   getChart(chartName: string): GetTelemetryChartResponse | undefined {
+      const category = this.view?.categories.find(c => !!c.charts.find(ch => ch.name === chartName));
+      return category?.charts.find(ch => ch.name === chartName);
+   }
+
+   getChartLegend(chartName: string): LegendEntry[] {
+
+      const mins = new Map<string, number>();
+      const maxs = new Map<string, number>();
+      const keyMetrics = new Map<string, GetTelemetryMetricResponse[]>();
+
+      const chart = this.getChart(chartName);
+      if (!chart) {
+         return [];
+      }
+
+      const metrics = this.getFilteredChartMetrics(chartName);
+
+      const replace: string[] = [];
+      this.searchState.variables?.forEach(v => {
+         const elements = v.split(";");
+         if (elements.length === 2) {
+            replace.push(elements[1] + ":")
+         }
+      });
+
+      const legendSet = new Set<string>();
+
+      metrics.forEach(metric => {
+
+         metric.metrics.forEach(m => {
+
+            mins.set(m.key, Math.min(mins.get(m.key) ?? Number.MAX_SAFE_INTEGER, m.value));
+            maxs.set(m.key, Math.max(maxs.get(m.key) ?? Number.MIN_SAFE_INTEGER, m.value));
+
+            if (!keyMetrics.has(m.key)) {
+               keyMetrics.set(m.key, []);
+            }
+
+            keyMetrics.get(m.key)!.push(m);
+
+            legendSet.add(m.key)
+         })
+      })
+
+      const legend = Array.from(legendSet).sort((a, b) => a.localeCompare(b)).map((key) => {
+
+         let display = key;
+         for (let i = 0; i < replace.length; i++) {
+            display = display.replace(replace[i], "")
+         }
+
+         let change = 0;
+
+         const metrics = keyMetrics.get(key)!.sort((a, b) => {
+            return a.time.getTime() - b.time.getTime();
+         })
+
+         if (metrics.length > 1) {
+            const m1 = metrics[0].value;
+            const m2 = metrics[metrics.length - 1].value;
+            if (m1) {
+               change = (m2 - m1) / m1;
+            }
+
+         }
+
+         return {
+            display: display,
+            key: key,
+            min: mins.get(key)!,
+            max: maxs.get(key)!,
+            change: change
+         }
+      });
+
+      function getPrefix(words: string[]) {
+
+         if (!words[0] || words.length === 1) return words[0] || "";
+         let i = 0;
+
+         // eslint-disable-next-line no-loop-func
+         while (words[0][i] && words.every(w => w[i] === words[0][i]))
+            i++;
+
+         return words[0].slice(0, i);
+      }
+
+      // collapse to common value is left most variable only has one value
+      const vars = handler.getChartVariables();
+      if (vars.length > 0 && this.searchState.variables) {
+
+         const v = this.searchState.variables.find(v => {
+            const group = v.split(";")[0];
+            return group === vars[0].group;
+         })
+
+         if (v && legend.length > 0 && v.split(";").length === 2) {
+            const display = legend.map(v => v.display);
+            const prefix = getPrefix(display);
+
+            if (prefix) {
+               legend.forEach(v => {
+                  const ndisplay = v.display.replace(prefix, "");
+                  if (ndisplay) {
+                     v.display = ndisplay;
+                  }
+               });
+            }
+         }
+      }
+
+      return legend;
+
+   }
+
+   private getAllChartMetrics(chartName: string): GetTelemetryMetricsResponse[] {
+
+      if (!this.metrics) {
+         return [];
+      }
+
+      const chart = this.getChart(chartName);
+      if (!chart) {
+         return [];
+      }
+
+      const cmetrics = new Set<string>(chart.metrics.map(cm => cm.metricId));
+
+      let metrics = this.metrics.metrics.filter(m => {
+         return cmetrics.has(m.metricId);
+      })
+
+      return metrics;
+
+   }
+
+   getFilteredChartMetrics(chartName: string, latest = false): GetTelemetryMetricsResponse[] {
+
+      if (!this.metrics) {
+         return [];
+      }
+
+      const chart = this.getChart(chartName);
+      if (!chart) {
+         return [];
+      }
+
+      const cmetrics = new Set<string>(chart.metrics.map(cm => cm.metricId));
+
+      let metrics = this.metrics.metrics.filter(m => {
+
+         return cmetrics.has(m.metricId);
+      }).map(m => { return { ...m } as GetTelemetryMetricsResponse });
+
+      if (chart.min !== undefined || chart.max !== undefined) {
+         metrics.forEach(metric => {
+            metric.metrics = metric.metrics.filter(m => {
+               if (chart.min !== undefined && m.value < chart.min) {
+                  return false;
+               }
+
+               if (chart.max !== undefined && m.value > chart.max) {
+                  return false;
+               }
+
+               return true;
+            })
+         })
+      }
+
+
+      if (latest) {
+         const found = new Set<string>();
+
+         metrics.forEach(metric => {
+            metric.metrics = metric.metrics.sort((a, b) => a.time.getTime() - b.time.getTime()).filter(m => {
+
+               if (found.has(m.key)) {
+                  return false;
+               }
+               found.add(m.key);
+               return true;
+            })
+         })
+
+      }
+
+      this.searchState.variables?.forEach(v => {
+
+         const values = v.split(";");
+         const group = values.shift()!;
+
+         metrics.forEach(metric => {
+            metric.metrics = metric.metrics.filter(m => {
+
+               if (m.groupValues && m.groupValues[group]) {
+                  if (values.indexOf(m.groupValues[group]) === -1) {
+                     return false;
+                  }
+               }
+
+               return true;
+            })
+         })
+
+      })
+
+      return metrics;
+   }
+
+   clear() {
+      this.initialized = false;
+      this._category = undefined;
+      this.anchorMinDate = undefined;
+      this.anchorMaxDate = undefined;
+      this.search = new URLSearchParams();
+      this.searchState = {};
+      this.view = undefined;
+      this.telemetryViews = undefined;
+      this.metrics = undefined;
+      this.filteredKeys.clear();
+      this.zoomHandler.clear();
+      clearTelemetryViewMetrics();
+   }
+
+   valueToString(chart: GetTelemetryChartResponse, valueIn: number): string {
+
+      let value = "";
+
+      if (chart.display === "Ratio") {
+         value = Math.round((valueIn * 100)).toString() + "%"
+      }
+      else if (chart.display === "Value") {
+         value = Math.round(valueIn).toString();
+      }
+      else {
+         value = msecToElapsed((valueIn) * 1000, true, true);
+      }
+
+      return value;
+   }
+
+   async initialize(viewId?: string, categoryIn?: string) {
+
+      if (this.initialized || (viewId && this.view?.id === viewId)) {
+         return;
+      }
+
+      this.initialized = true;
+
+      this.telemetryViews = await getViews();
+
+      this.view = undefined;
+
+      if (!this.allViews.length) {
+         console.log("No telemetry views configured");
+         return;
+      }
+
+      this.search = new URLSearchParams();
+
+      if (!viewId) {
+         const query = new URLSearchParams(window.location.search).get("query");
+         if (query?.length) {
+            try {
+               this.search = new URLSearchParams(atob(query));
+            } catch (reason) {
+               console.error(reason);
+            }
+         }
+
+         this.searchState = this.stateFromSearch();
+         this.view = this.allViews.find(v => v.id === this.searchState.viewId) ?? this.allViews[0];
+
+      } else {
+         this.searchState = {};
+         this.searchState.minutes = 60 * 24;
+         this.searchState.category = categoryIn;
+         this.view = this.allViews.find(v => !viewId || v.id === viewId) ?? this.allViews[0];
+         this.updateTime(this.searchState.minutes);
+      }
+
+      if (!this.view || !this.view.categories.length) {
+         console.error("No view or no categories");
+         return;
+      }
+
+      this.searchState.viewId = this.view.id;
+
+      const category = this.searchState.category ?? this.view.categories[0].name;
+
+      await this.setCategory(category)
+
+      if (!this.searchState?.variables?.length) {
+
+         const vars = this.getChartVariables();
+         vars.forEach(v => {
+            if (v.defaults?.length) {
+               this.setVariables(v.group, v.defaults);
+            }
+         })
+      }
+
+      this.setUpdated();
+   }
+
+   async setCategory(category?: string, fromUser?: boolean) {
+
+      if (fromUser) {
+         this.anchorMinDate = this.anchorMaxDate = undefined;
+      }
+
+      if (category?.length) {
+         if (this.searchState.category !== category) {
+            this.searchState.category = category;
+            this.updateSearch();
+         }
+
+         await this.set(category)
+      }
+   }
+
+   setVariables(group: string, values: string[]) {
+      let newVars = this.searchState.variables?.filter(v => {
+         const [vgroup,] = v.split(";");
+         if (group === vgroup) {
+            return false;
+         }
+         return true;
+      })
+
+      if (!newVars) {
+         newVars = [];
+      }
+
+      if (values.length) {
+         newVars.push(`${group};` + values.join(";"));
+      }
+
+      newVars = newVars.sort((a, b) => a.localeCompare(b));
+      this.searchState.variables = newVars.length ? newVars : undefined;
+
+      this.updateSearch();
+      this.setUpdated();
+
+   }
+
+   setFilterKeys(keys: Set<string>) {
+
+      this.filteredKeys = keys;
+
+      this.setUpdated();
+   }
+
+   @action
+   setUpdated() {
+      this.updated++;
+   }
+
+   @action
+   setSearchUpdated() {
+      this.searchUpdated++;
+   }
+
+   updateSearch(): boolean {
+
+      const state = { ...this.searchState } as SearchState;
+
+      const search = new URLSearchParams();
+      const csearch = this.search.toString();
+
+      if (state.viewId?.length) {
+         search.append("view", state.viewId);
+      }
+
+      if (state.category?.length) {
+         search.append("category", state.category);
+      }
+
+      if (state.minutes) {
+         search.append("minutes", state.minutes.toString());
+      }
+
+      state.variables?.forEach(v => {
+         search.append("v", v);
+      });
+
+
+      if (search.toString() !== csearch) {
+         this.search = search;
+         this.setSearchUpdated();
+         return true;
+      }
+
+      return false;
+   }
+
+   stateFromSearch() {
+
+      const state: SearchState = {};
+
+      state.viewId = this.search.get("view") ?? undefined;
+      state.category = this.search.get("category") ?? undefined;
+      state.variables = this.search.getAll("v") ?? undefined;
+      let minutes = Number.parseInt(this.search.get("minutes") ?? "0")
+      if (!minutes) {
+         minutes = 60 * 24;
+      }
+      state.minutes = minutes;
+
+      this.updateTime(minutes);
+
+      return state;
+   }
+
+   updateTime(minutes: number) {
+
+      this.anchorMinDate = undefined;
+      this.anchorMaxDate = undefined;
+
+      this.selectMinDate = new Date(new Date().getTime() - (minutes * 60000));
+      this.selectMaxDate = new Date();
+   }
+
+   reload() {
+
+      // force an update
+      const category = this._category;
+      this._category = undefined;
+      this.searchState.category = undefined;
+      clearTelemetryViewMetrics();
+      this.setCategory(category);
+
+   }
+
+   setTimeSelection(time: TimeSelection) {
+
+      if (this.searchState.minutes === time.minutes && !this.anchorMinDate) {
+         return;
+      }
+
+      this.searchState.minutes = time.minutes;
+      this.updateTime(time.minutes);
+
+      this.reload();
+   }
+
+   setView(viewId: string | undefined) {
+
+      if (!viewId || this.view?.id === viewId) {
+         return;
+      }
+
+      let category: string | undefined;
+
+      if (this.view) {
+
+         const newView = this.allViews.find(v => v.id === viewId);
+         const ccat = this.view.categories.find(c => c.name === this._category);
+
+         if (ccat && newView) {
+            category = newView.categories.find(c => ccat.name === c.name)?.name;
+         }
+      }
+
+      this._category = undefined;
+      this.anchorMinDate = undefined;
+      this.anchorMaxDate = undefined;
+      this.search = new URLSearchParams();
+      this.searchState = {};
+      this.filteredKeys.clear();
+      this.metrics = undefined;
+
+      clearTelemetryViewMetrics();
+
+      this.initialized = false;
+      this.initialize(viewId, category);
+   }
+
+
+   setZoomHandler(chartName: string, zoomed: any) {
+      this.zoomHandler.set(chartName, zoomed);
+   }
+
+   onZoom(originatingChartName: string, event: any) {
+      this.zoomHandler.forEach((zoomed, chartName) => {
+         if (originatingChartName !== chartName) {
+            zoomed(event, false);
+         }
+      })
+   }
+
+   onTimeSelect(chartName: string, minTime: Date, maxTime: Date) {
+
+      this.anchorMinDate = minTime;
+      this.anchorMaxDate = maxTime;
+
+      this.reload();
+   }
+
+   @observable
+   private updated = 0;
+
+   @observable
+   private searchUpdated: number = 0;
+
+   view?: GetTelemetryViewResponse;
+
+   anchorMinDate?: Date;
+   anchorMaxDate?: Date;
+   restoreMinutes?: number;
+
+   selectMinDate: Date = new Date();
+   selectMaxDate: Date = new Date();
+
+   get minDate(): Date {
+      return this.anchorMinDate ?? this.selectMinDate;
+   }
+
+   get maxDate(): Date {
+      return this.anchorMaxDate ?? this.selectMaxDate;
+   }
+
+   querying = false;
+
+   initialized = false;
+
+   filteredKeys = new Set<string>();
+   zoomHandler = new Map<string, any>();
+
+   searchState: SearchState = {}
+
+   search: URLSearchParams = new URLSearchParams();
+
+   get allViews() { return this.telemetryViews ?? []}
+
+   private telemetryViews?: GetTelemetryViewResponse[];
+
+   private metrics?: TelemetryViewData;
+
+   private _category?: string;
+}
+
+const handler = new MetricsHandler();
+
+const pivotClasses = mergeStyleSets({
+   pivot: {
+      selectors: {
+         ".ms-Pivot-link": {
+            lineHeight: "36px",
+            height: "36px",
+            paddingTop: 0,
+            paddingBottom: 0
+         }
+      }
+   }
+});
+
+const ViewChooser: React.FC = observer(() => {
+
+   handler.subscribe();
+
+   const options: IComboBoxOption[] = handler.allViews.map(v => {
+      return {
+         key: `view_option_${v.id}`,
+         text: v.name,
+         data: v
+      }
+   });
+
+   let key: string | undefined;
+
+   if (handler.view) {
+      key = `view_option_${handler.view.id}`
+   }
+
+   return <Stack>
+      <ComboBox
+         styles={{ root: { width: 120 } }}
+         options={options}
+         selectedKey={key}
+         onChange={(ev, option, index, value) => {
+            handler.setView(option?.data?.id);
+         }}
+      />
+   </Stack>
+})
+
+const TelemetryPivot: React.FC = observer(() => {
+
+   handler.subscribe();
+
+   const view = handler.view;
+
+   if (!view) {
+      return null;
+   }
+
+   const links = view.categories.map(cat => {
+      const key = `item_key_${cat.name}`;
+      return <PivotItem headerText={cat.name} key={key} itemKey={key} />
+   });
+
+   return <Stack horizontal tokens={{ childrenGap: 18 }} verticalAlign="center" verticalFill={false}>
+      <Pivot className={pivotClasses.pivot} defaultSelectedKey={`item_key_${handler.searchState.category}`} linkSize="normal" linkFormat="links" onLinkClick={(item, ev) => {
+         if (item) {
+
+            const catname = item.props.itemKey!.replace("item_key_", "");
+            handler.setCategory(catname, true);
+         }
+      }}>
+         {links}
+      </Pivot>
+   </Stack>;
+})
+
+let multiComboBoxId = 0;
+
+const VariableChooser: React.FC<{ group: string, label?: string, multiSelect: boolean, optionsIn: IComboBoxOption[], initialKeysIn: string[], updateKeys: (group: string, selectedKeys: string[]) => void }> = ({ group, label, multiSelect, optionsIn, initialKeysIn, updateKeys }) => {
+
+   const comboBoxRef = React.useRef<IComboBox>(null);
+
+   let initialKeys = [...initialKeysIn];
+   let options = [...optionsIn];
+
+   if (options.length === initialKeys.length) {
+      initialKeys.push('selectAll');
+   }
+
+   if (multiSelect && options.length > 3) {
+      options.unshift({ key: 'selectAll', text: 'Select All', itemType: SelectableOptionMenuItemType.SelectAll });
+   }
+
+   const comboBoxStyles: Partial<IComboBoxStyles> = { root: { width: 270 } };
+
+   return <ComboBox componentRef={comboBoxRef} key={`multi_option_${group}_${multiComboBoxId++}`} label={label} placeholder="None" defaultSelectedKey={initialKeys} multiSelect={multiSelect} options={options} onResolveOptions={() => options}
+      onChange={multiSelect ? undefined : (event, option, index, value) => {
+         setTimeout(() => { multiComboBoxId++; updateKeys(group, [value ?? ""]) }, 250);
+      }}
+      onMenuDismiss={!multiSelect ? undefined : () => {
+         if (comboBoxRef?.current?.selectedOptions) {
+            const selectedKeys = comboBoxRef.current.selectedOptions.map(o => o.key as string).filter(k => k !== 'selectAll');
+            setTimeout(() => { multiComboBoxId++; updateKeys(group, selectedKeys) }, 250);
+         }
+      }} styles={comboBoxStyles} />
+};
+
+
+const TelemetryChooser: React.FC = observer(() => {
+
+   handler.subscribe();
+
+   const view = handler.view;
+
+   if (!view) {
+      return null;
+   }
+
+   const varStacks: JSX.Element[] = [];
+
+   const vars = handler.getChartVariables();
+
+   vars.forEach(v => {
+
+      const valueSet = new Set<string>();
+      v.defaults.forEach(v => valueSet.add(v));
+      v.values.forEach(v => valueSet.add(v));
+
+      const options: IComboBoxOption[] = Array.from(valueSet).sort((a, b) => a.localeCompare(b)).map(name => {
+         return {
+            key: name,
+            text: name
+         }
+      })
+
+      if (!options.length) {
+         return;
+      }
+
+      let defaultKeys: string[] = options.map(o => o.key as string);
+
+      const state = handler.searchState;
+      if (state.variables?.length) {
+         const sv = state.variables?.find(sv => sv.startsWith(`${v.group};`))
+         if (sv) {
+            defaultKeys = sv.split(";");
+            defaultKeys.shift();
+         }
+      }
+
+      const stack = <Stack key={`key_chooser_${handler.view?.id}_${v.group}`}>
+         <VariableChooser
+            group={v.group}
+            label={v.name}
+            multiSelect={true}
+            optionsIn={options}
+            initialKeysIn={defaultKeys}
+            updateKeys={(group, keys) => {
+               handler.setVariables(group, keys);
+            }}
+         />
+      </Stack>
+
+      varStacks.push(stack);
+
+   })
+
+   if (!varStacks.length) {
+      return null;
+   }
+
+   return <Stack horizontal tokens={{ childrenGap: 24 }} verticalAlign="center" verticalFill={false}>
+      {varStacks}
+   </Stack>;
+})
+
+type TimeSelection = {
+   text: string;
+   key: string;
+   minutes: number;
+   hidden?: boolean;
+}
+
+const TimeChooser: React.FC = observer(() => {
+
+   handler.subscribe();
+
+   let timeComboWidth = 180;
+
+   let key: string | undefined;
+
+   if (handler.anchorMinDate) {
+      key = "time_custom";
+   } else {
+      key = timeSelections.find(t => t.minutes === handler.searchState.minutes)?.key;
+      if (!key) {
+         return null;
+      }
+   }
+
+   return <Stack>
+      <ComboBox
+         styles={{ root: { width: timeComboWidth } }}
+         options={timeSelections}
+         selectedKey={key}
+         onChange={(ev, option, index, value) => {
+            const select = option as TimeSelection;
+            handler.setTimeSelection(select);
+         }}
+      />
+   </Stack>
+
+})
+
+const Legend: React.FC<{ chart: GetTelemetryChartResponse }> = observer(({ chart }) => {
+
+   const { modeColors } = getHordeStyling();
+
+   const tooltipId = useId();
+
+   handler.subscribe();
+
+   const legend = handler.getChartLegend(chart.name);
+
+   const legendStacks: JSX.Element[] = [];
+
+   const calloutProps = { gapSpace: 0 };
+   const hostStyles: Partial<ITooltipHostStyles> = { root: { display: 'inline-block' } };
+
+   legend.forEach((v, index) => {
+
+      const filtered = handler.filteredKeys.has(v.key);
+
+      const tooltipProps: ITooltipProps = {
+
+         onRenderContent: () => {
+
+            const entry = legend.find(n => n.key === v.key);
+            if (!entry) {
+               return null;
+            }
+
+            return <Stack style={{ backgroundColor: modeColors.background, border: "solid", borderWidth: "1px", borderRadius: "3px", borderColor: dashboard.darktheme ? "#413F3D" : "#2D3F5F" }}>
+               <Stack style={{ padding: "16px 16px" }} tokens={{ childrenGap: 8 }}>
+                  <Stack>
+                     <Text variant="small">Min: {handler.valueToString(chart, entry.min)}</Text>
+                  </Stack>
+                  <Stack>
+                     <Text variant="small">Max: {handler.valueToString(chart, entry.max)}</Text>
+                  </Stack>
+                  <Stack>
+                     <Text variant="small">Change: {Math.round((entry.change * 100)).toString() + "%"}</Text>
+                  </Stack>
+               </Stack>
+            </Stack>
+         }
+      };
+
+      const stack = <TooltipHost
+         tooltipProps={tooltipProps}
+         id={tooltipId}
+         delay={0}
+         directionalHint={DirectionalHint.leftCenter}
+         calloutProps={calloutProps}
+         styles={hostStyles}>
+         <Stack key={`key_legend_${v.key}`} style={{ userSelect: "none" }} horizontal verticalAlign="center" tokens={{ childrenGap: 8 }} onClick={(ev) => {
+
+            let keys = legend.map(l => l.key);
+
+            const mod = ev.shiftKey || ev.ctrlKey;
+            if (!mod) {
+
+               if (handler.filteredKeys.size && !handler.filteredKeys.has(v.key)) {
+                  handler.setFilterKeys(new Set());
+               } else {
+                  const nset = new Set<string>();
+                  keys.forEach(k => {
+                     if (k !== v.key) {
+                        nset.add(k);
+                     }
+                  })
+                  handler.setFilterKeys(nset);
+               }
+            } else {
+
+               let nset = new Set(handler.filteredKeys);
+               if (nset.has(v.key)) {
+                  nset.delete(v.key);
+               } else {
+                  if (!nset.size) {
+                     nset = new Set<string>();
+
+                     keys.forEach(k => {
+                        if (k !== v.key) {
+                           nset.add(k);
+                        }
+                     })
+                  } else {
+                     nset.add(v.key);
+                  }
+               }
+
+               if (nset.size === legend.length) {
+                  nset = new Set();
+               }
+
+               handler.setFilterKeys(nset);
+
+            }
+
+         }}>
+            <Stack>
+               <FontIcon style={{ color: filtered ? "#999999" : graphColors[index % graphColors.length], paddingTop: 2 }} iconName="Square" />
+            </Stack>
+            <Stack>
+               <Text style={{ fontSize: "11px", color: filtered ? "#999999" : undefined }}>{v.display}</Text>
+            </Stack>
+         </Stack>
+      </TooltipHost>
+
+      legendStacks.push(stack)
+
+   })
+
+   return <Stack style={{ cursor: "pointer", width: 340, height: 300, overflowY: "auto" }} tokens={{ childrenGap: 4 }}>{legendStacks}</Stack>
+})
+
+
+const indicatorStyles = mergeStyleSets({
+
+   stripes: {
+      backgroundImage: 'repeating-linear-gradient(-45deg, rgba(255, 255, 255, .2) 25%, transparent 25%, transparent 50%, rgba(255, 255, 255, .2) 50%, rgba(255, 255, 255, .2) 75%, transparent 75%, transparent)',
+   }
+
+});
+
+export type IndicatorBarStack = {
+   value: number,
+   title?: string,
+   titleValue?: number,
+   color?: string,
+   onClick?: () => void,
+   stripes?: boolean,
+   brightness?: number
+}
+
+export const IndicatorBar: React.FC<{ stack: IndicatorBarStack[], width: number, height: number, basecolor?: string, style?: any }> = ({ stack, width, height, basecolor, style }) => {
+   stack = stack.filter(s => s.value > 0);
+
+   return (
+      <div className={mergeStyles({ backgroundColor: basecolor, width: width, height: height, verticalAlign: 'middle', display: "flex" }, style)}>
+         {stack.map((item) => {
+
+            let boxShadow = !item.brightness ? `0 0 3px ${item.color}` : undefined;
+            let filter = item.brightness ? `brightness(${item.brightness})` : undefined;
+            if (!dashboard.darktheme) {
+               boxShadow = undefined;
+            }
+
+            const iwidth = width * (item.value / 100);
+            return <span key={`${item.title!}_${metricIdCounter++}`}
+               onClick={item.onClick}
+               className={item.stripes ? indicatorStyles.stripes : undefined}
+               style={{
+                  width: `${iwidth}px`, height: '100%',
+                  backgroundColor: item.color,
+                  margin: "1px",
+                  display: 'block',
+                  borderRadius: "2px",
+                  cursor: item.onClick ? 'pointer' : 'inherit',
+                  backgroundSize: `${height * 2}px ${height * 2}px`,
+                  boxShadow: boxShadow,
+                  filter: filter
+               }} />
+         })}
+      </div>
+   );
+}
+
+let metricIdCounter = 0;
+
+const IndicatorTile: React.FC<{ chart: GetTelemetryChartResponse }> = observer(({ chart }) => {
+
+   const { modeColors } = getHordeStyling();
+
+   handler.subscribe();
+
+   const metrics = handler.getFilteredChartMetrics(chart.name);
+
+   const legend = handler.getChartLegend(chart.name);
+
+   if (!metrics?.length) {
+      return <Text>No Data</Text>;
+   }
+
+   const colors = dashboard.getStatusColors();
+
+   let allMetrics = metrics.map(m => m.metrics).flat().sort((a, b) => a.key.localeCompare(b.key));
+
+   const keyValues = new Map<string, number[]>();
+   allMetrics.forEach(m => {
+
+      if (!keyValues.has(m.key)) {
+         keyValues.set(m.key, []);
+      }
+      keyValues.get(m.key)!.push(m.value);
+   })
+
+   const keyAverages = new Map<string, number>();
+   keyValues.forEach((values, key) => {
+      let avg = 0;
+      values.forEach(v => avg += v);
+      avg /= values.length;
+      keyAverages.set(key, Math.round(avg));
+   })
+
+   const found = new Set<string>();
+
+   allMetrics = allMetrics.sort((a, b) => a.time.getTime() - b.time.getTime());
+   allMetrics = allMetrics.filter(m => {
+
+      if (found.has(m.key)) {
+         return false;
+      }
+
+      found.add(m.key);
+      return true;
+   })
+
+   allMetrics.forEach(m => {
+      m.value = keyAverages.get(m.key) ?? 0;
+   })
+
+   const elements: JSX.Element[] = [];
+
+   allMetrics.forEach(m => {
+
+      const barStack: IndicatorBarStack[] = [];
+
+      const max = chart.max ?? 100;
+      const threshold = m.threshold ?? max;
+
+      const v = m.value / max;
+      const t = threshold / max;
+      for (let i = 0.0; i < 1; i += .1) {
+
+         let color = i <= t ? colors.get(StatusColor.Success)! : colors.get(StatusColor.Failure)!;
+
+         let brightness: number | undefined;
+         if (i > v) {
+            brightness = dashboard.darktheme ? 0.4 : 1;
+            if (!dashboard.darktheme) {
+               color += "4B"
+            }
+         }
+         barStack.push({ value: 10, color: color, brightness: brightness });
+      }
+
+      const name = legend.find(v => v.key === m.key)?.display ?? m.key;
+
+      const element = <Stack horizontal verticalAlign="center" key={`indicator_bar_${metricIdCounter++}`}>
+         <Stack style={{ width: 340 }}>
+            <Text variant="small">{name}</Text>
+         </Stack>
+         <Stack>
+            <IndicatorBar stack={barStack} width={160} height={14} />
+         </Stack>
+         <Stack style={{ width: 90 }} horizontalAlign="end">
+            <Text variant="small">{chart.display === "Value" ? m.value : msecToElapsed(m.value * 1000)}</Text>
+         </Stack>
+      </Stack>
+
+      elements.push(element)
+   })
+
+   return <Stack style={{ paddingTop: 12 }}>
+      <Stack tokens={{ childrenGap: 6 }} style={{ backgroundColor: modeColors.background, padding: 18, height: 300, overflowY: "auto" }}>
+         {elements}
+      </Stack>
+   </Stack>
+})
+
+
+class Tooltip {
+   constructor() {
+      makeObservable(this);
+   }
+
+   subscribe() {
+      if (this.updated) { }
+   }
+
+   @action
+   set(key: string, x: number, y: number, time: Date, value: number, color: string) {
+
+      if (key === "__clear__") {
+         this.show = false;
+         this.updated++;
+         return;
+      }
+      this.show = true;
+      this.color = color;
+      this.key = key;
+      this.x = x;
+      this.y = y;
+      this.time = time;
+      this.value = value;
+      this.updated++;
+   }
+
+   show = false;
+
+   key: string = "";
+   x: number = 0;
+   y: number = 0;
+   time: Date = new Date();
+   value: number = 0;
+   color: string = "";
+
+   @observable
+   private updated = 0;
+}
+
+const GraphTooltip: React.FC<{ chart: GetTelemetryChartResponse, tooltip: Tooltip, legend: LegendEntry[] }> = observer(({ chart, tooltip, legend }) => {
+
+   const { modeColors } = getHordeStyling();
+
+   tooltip.subscribe();
+
+   if (!tooltip.show) {
+      return null;
+   }
+
+   let tipX = tooltip.x;
+   let offsetX = 32;
+   let translateX = "0%";
+
+   if (tipX > 800) {
+      offsetX = -32;
+      translateX = "-100%";
+   }
+
+   const translateY = "-50%";
+
+   const time = moment(tooltip.time).tz(displayTimeZone());
+
+   let value = handler.valueToString(chart, tooltip.value);
+
+   let name = legend.find(k => k.key === tooltip.key)?.display ?? tooltip.key;
+
+   return <div style={{
+      position: "absolute",
+      display: "block",
+      top: `${tooltip.y}px`,
+      left: `${tooltip.x + offsetX}px`,
+      backgroundColor: modeColors.background,
+      zIndex: 1,
+      border: "solid",
+      borderWidth: "1px",
+      borderRadius: "3px",
+      width: "max-content",
+      borderColor: dashboard.darktheme ? "#413F3D" : "#2D3F5F",
+      pointerEvents: "none",
+      transform: `translate(${translateX}, ${translateY})`
+   }}>
+      <Stack style={{ padding: "16px 16px" }} tokens={{ childrenGap: 8 }}>
+         <Stack horizontal tokens={{ childrenGap: 4 }}>
+            <FontIcon style={{ color: tooltip.color, paddingTop: 3 }} iconName="Square" />
+            <Text>{name}</Text>
+         </Stack>
+         <Stack>
+            <Text>{time.format("MM/DD HH:mm")}</Text>
+         </Stack>
+         <Stack>
+            <Text>{value}</Text>
+         </Stack>
+      </Stack>
+   </div>
+
+})
+
+const LineGraphTile: React.FC<{ chart: GetTelemetryChartResponse }> = observer(({ chart }) => {
+
+   const [scale] = useState(1);
+   const [container, setContainer] = useState<HTMLDivElement | null>(null);
+   const renderer = useConst(chart.graph === "Line" ? new TelemetryLineRenderer() : new TelemetryLineRenderer());
+   const tooltip = useConst(new Tooltip());
+
+   const { hordeClasses, modeColors } = getHordeStyling();
+
+   handler.subscribe();
+
+   const graph_container_id = `metric_graph_container_${chart.name}}`;
+
+   const metrics = handler.getFilteredChartMetrics(chart.name);
+
+   let hasMetrics = false;
+   metrics.forEach(metric => {
+      metric.metrics = metric.metrics.filter(m => !handler.filteredKeys.has(m.key));
+      if (metric.metrics.length) {
+         hasMetrics = true;
+      }
+   })
+
+   if (!hasMetrics) {
+      return <Stack style={{ paddingTop: 12 }}><Text>No Matching Data</Text></Stack>;
+   }
+
+   const legend = handler.getChartLegend(chart.name);
+
+   if (container) {
+      try {
+
+         const onZoom = (chartName: string, event: any) => {
+            handler.onZoom(chartName, event);
+         };
+
+         const onTimeSelect = (chartName: string, minTime: Date, maxTime: Date) => {
+            handler.onTimeSelect(chartName, minTime, maxTime);
+         }
+
+         const onDataHover = (key: string, x: number, y: number, time: Date, value: number, color: string) => {
+
+            tooltip.set(key, x, y, time, value, color);
+         }
+
+         const zoomed = renderer.render(chart, metrics, legend.map(v => v.key), handler.minDate!, handler.maxDate!, container, onZoom, onTimeSelect, onDataHover, scale);
+         handler.setZoomHandler(chart.name, zoomed);
+
+      } catch (err) {
+         console.error(err);
+      }
+   }
+   const width = 1024;
+
+   return <Stack className={hordeClasses.horde} key={`metric_graph_stack_${chart.name}`}>
+      <Stack style={{ width: "100%", paddingTop: 16, paddingBottom: 16, paddingLeft: 16, backgroundColor: modeColors.background }} horizontal tokens={{ childrenGap: 12 }}>
+         <Stack style={{ width: width, position: "relative" }}>
+            <GraphTooltip chart={chart} tooltip={tooltip} legend={legend} />
+            <div id={graph_container_id} style={{ shapeRendering: "geometricPrecision", userSelect: "none" }} ref={(ref: HTMLDivElement) => setContainer(ref)} onMouseEnter={() => { }} onMouseLeave={() => { }} />
+         </Stack>
+         <Stack>
+            <Legend chart={chart} />
+         </Stack>
+      </Stack>
+
+   </Stack>;
+})
+
+const TelemetryViewInternal: React.FC = observer(() => {
+
+   const { hordeClasses } = getHordeStyling();
+
+   handler.subscribe();
+
+   if (handler.querying) {
+      return <Stack>
+         <Text>Loading Data</Text>
+         <Spinner size={SpinnerSize.large} />
+      </Stack>
+   }
+
+   const charts = handler.getCharts();
+
+   const indicators = charts.filter(c => c.graph === "Indicator");
+   const lines = charts.filter(c => c.graph === "Line");
+
+   const ipanels = indicators.map(chart => {
+      return <Stack key={`telemetry_view_panel_${chart.name}`} styles={{ root: { paddingTop: 0, paddingRight: 12, width: 720 } }}>
+         <Stack className={hordeClasses.raised} >
+            <Stack tokens={{ childrenGap: 12 }} grow>
+               <Stack >
+                  <Stack >
+                     <Text>{chart.name}</Text>
+                  </Stack>
+                  <Stack>
+                     <IndicatorTile chart={chart} />
+                  </Stack>
+               </Stack>
+            </Stack>
+         </Stack>
+      </Stack>
+   })
+
+   const linepanels = lines.map(chart => {
+      return <Stack key={`telemetry_view_panel_${chart.name}`} styles={{ root: { paddingTop: 0, paddingRight: 12 } }}>
+         <Stack className={hordeClasses.raised} >
+            <Stack tokens={{ childrenGap: 12 }} grow>
+               <Stack >
+                  <Stack >
+                     <Text>{chart.name}</Text>
+                  </Stack>
+                  <Stack>
+                     <LineGraphTile chart={chart} />
+                  </Stack>
+               </Stack>
+            </Stack>
+         </Stack>
+      </Stack>
+   })
+
+   return <Stack tokens={{ childrenGap: 12 }}>
+      <Stack horizontal>
+         {ipanels}
+      </Stack>
+      {linepanels}
+   </Stack>
+})
+
+export const SearchUpdate: React.FC = observer(() => {
+
+   const [, setSearchParams] = useSearchParams();
+   const csearch = "query=" + btoa(handler.search.toString());
+
+   useEffect(() => {
+      if (handler.search.toString().length) {
+         setSearchParams(csearch, { replace: true });
+      }
+
+   }, [csearch, setSearchParams])
+
+   // subscribe
+   handler.subscribeToSearch();
+
+   return null;
+});
+
+export const TelemetryView: React.FC = () => {
+
+   useEffect(() => {
+      handler.initialize();
+      return () => {
+         handler.clear();
+      };
+   }, []);
+
+   const windowSize = useWindowSize();
+   const navigate = useNavigate();
+
+   const vw = Math.max(document.documentElement.clientWidth, window.innerWidth || 0);
+
+   const { hordeClasses, modeColors } = getHordeStyling();
+
+   const rootWidth = 1440;
+   const centerAlign = vw / 2 - 720 /*890*/;
+   const key = `windowsize_metrics_view_${windowSize.width}_${windowSize.height}`;
+
+   const telemetrybDocs = "/docs/Config/Analytics.md";
+
+   return <Stack className={hordeClasses.horde} key="key_metrics_graph_test">
+      <SearchUpdate />
+      <TopNav />
+      <Breadcrumbs items={[{ text: 'Analytics' }]} />
+      <Stack horizontal styles={{ root: { backgroundColor: modeColors.background } }}>
+         <Stack styles={{ root: { width: "100%" } }}>
+            {!handler.allViews.length && <Stack horizontal tokens={{ childrenGap: 6 }} horizontalAlign="center" style={{ paddingTop: 30 }}>
+               <Text variant="mediumPlus">No analytic views found, for more information please see</Text>
+               <a href={telemetrybDocs} style={{ fontSize: "18px", "cursor": "pointer" }} onClick={(ev) => { ev.preventDefault(); ev.stopPropagation(); navigate(telemetrybDocs) }}> Horde analytics documentation.</a>
+            </Stack>}
+            {!!handler.allViews.length && <Stack>
+               <Stack horizontal>
+                  <Stack key={`${key}_1`} style={{ paddingLeft: centerAlign }} />
+                  <Stack style={{ width: rootWidth - 8, maxWidth: windowSize.width - 12, paddingLeft: 0, paddingTop: 24, paddingBottom: 24, paddingRight: 0 }} >
+                     <Stack>
+                        <Stack horizontal verticalAlign="center">
+                           <TelemetryPivot />
+                           <Stack grow />
+                           <Stack horizontal tokens={{ childrenGap: 14 }} verticalAlign="center" verticalFill>
+                              <Stack>
+                                 <ViewChooser />
+                              </Stack>
+                              <Stack >
+                                 <TimeChooser />
+                              </Stack>
+                              <Stack>
+                                 <DefaultButton style={{ minWidth: 52, height: 34 }} onClick={() => handler.reload()}>
+                                    <Icon iconName='Refresh' />
+                                 </DefaultButton>
+                              </Stack>
+                           </Stack>
+                        </Stack>
+                     </Stack>
+                     <Stack horizontal>
+                        <Stack grow />
+                        <Stack style={{ paddingRight: 2, paddingTop: 12, paddingBottom: 12 }}>
+                           <TelemetryChooser />
+                        </Stack>
+                     </Stack>
+                  </Stack>
+               </Stack>
+               <Stack style={{ width: "100%", backgroundColor: modeColors.background }}>
+                  <Stack horizontal>
+                     <Stack /*className={classNames.pointerSuppress} */ style={{ position: "relative", width: "100%", height: 'calc(100vh - 228px)' }}>
+                        <div id="hordeContentArea" style={{ overflowX: "auto", overflowY: "visible" }}>
+                           <Stack horizontal style={{ paddingBottom: 48 }}>
+                              <Stack key={`${key}_2`} style={{ paddingLeft: centerAlign }} />
+                              <Stack style={{ width: rootWidth }} tokens={{ childrenGap: 12 }}>
+                                 <TelemetryViewInternal />
+                              </Stack>
+                           </Stack>
+                        </div>
+                     </Stack>
+                  </Stack>
+               </Stack>
+            </Stack>}
+         </Stack>
+      </Stack>
+   </Stack>
+}
\ No newline at end of file
diff --git a/Engine/Source/Programs/Horde/HordeDashboard/plugins/index.ts b/Engine/Source/Programs/Horde/HordeDashboard/plugins/index.ts
new file mode 100644
--- /dev/null
+++ b/Engine/Source/Programs/Horde/HordeDashboard/plugins/index.ts
@@ -0,0 +1,46 @@
+import { RouteObject } from "react-router-dom";
+
+// Routes => components
+// backend queries
+// mounting points???
+
+export enum MountType {
+   TopNav
+}
+
+export type Mount = {
+   type: MountType,
+   context?: string;
+   text: string;
+   route: string;
+}
+
+export type HordePlugin = {
+   id: string;
+   routes: RouteObject[];
+   mount?: Mount;
+   enabled?: boolean;
+}
+
+const plugins: HordePlugin[] = [];
+
+export function registerHordePlugin(plugin: HordePlugin) {
+   const existing = plugins.find(p => p.id === plugin.id);
+   if (existing) {
+      throw `Duplicate plugin registration for id: ${plugin.id}`;
+   }
+
+   plugins.push(plugin);
+}
+
+export function enableHordePlugins(ids: string[]) {
+   const set = new Set(ids);
+   plugins.forEach(p => {
+      p.enabled = set.has(p.id)
+   });
+}
+
+export function getHordePlugins(): HordePlugin[] {
+   return plugins.filter(p => !!p.enabled);
+}
+
diff --git a/Engine/Source/Programs/Horde/HordeDashboard/plugins/registry.ts b/Engine/Source/Programs/Horde/HordeDashboard/plugins/registry.ts
new file mode 100644
--- /dev/null
+++ b/Engine/Source/Programs/Horde/HordeDashboard/plugins/registry.ts
@@ -0,0 +1,2 @@
+
+import "./analytics"
diff --git a/Engine/Source/Programs/Horde/HordeDashboard/src/App.tsx b/Engine/Source/Programs/Horde/HordeDashboard/src/App.tsx
--- a/Engine/Source/Programs/Horde/HordeDashboard/src/App.tsx
+++ b/Engine/Source/Programs/Horde/HordeDashboard/src/App.tsx
@@ -3,7 +3,8 @@
 import { Image, Spinner, SpinnerSize, Stack, Text, ThemeProvider } from '@fluentui/react';
 import React, { useState } from 'react';
 import { Navigate, Outlet, RouteObject, RouterProvider, createBrowserRouter } from 'react-router-dom';
-import hordePlugins from './legacyPlugins';
+import { enableHordePlugins, getHordePlugins } from '../plugins';
+import "hordePlugins/registry";
 import backend from './backend';
 import { getSiteConfig } from './backend/Config';
 import dashboard from './backend/Dashboard';
@@ -32,17 +33,17 @@ import { AccountsView } from './components/accounts/AccountsView';
 import { HordeLoginView } from './components/accounts/HordeLoginView';
 import { ServiceAccountsView } from './components/accounts/ServiceAccountsView';
 import { AgentRequestsView } from './components/agents/AgentRequestsView';
+import { ArtifactRedirector } from './components/artifacts/ArtifactsRedirector';
 import { DocView } from './components/docs/DocView';
 import { JobDetailViewV2 } from './components/jobDetailsV2/JobDetailViewV2';
 import { PreflightConfigRedirector } from './components/preflights/PreflightConfigCheckRedirector';
 import { ServerStatusView } from './components/server/ServerStatus';
 import { HordeSetupView } from './components/setup/HordeSetupView';
-import { TelemetryView } from './components/telemetry/TelemetryView';
 import { StepIssueReportTest } from './components/test/IssueStepReport';
+import hordePlugins from './legacyPlugins';
 import { preloadFonts } from './styles/Styles';
 import { darkTheme } from './styles/darkTheme';
 import { lightTheme } from './styles/lightTheme';
-import { ArtifactRedirector } from './components/artifacts/ArtifactsRedirector';
 
 let router: any;
 
@@ -52,14 +53,14 @@ const RouteError: React.FC = () => {
 
 const Main: React.FC = () => {
 
-   const [init, setInit] = useState(false);   
+   const [init, setInit] = useState(false);
 
    if (window.location.pathname === "/login") {
-      return <HordeLoginView/>
+      return <HordeLoginView />
    }
 
    if (window.location.pathname === "/setup") {
-      return <HordeSetupView/>
+      return <HordeSetupView />
    }
 
    const config = getSiteConfig();
@@ -107,8 +108,28 @@ const Main: React.FC = () => {
       </ThemeProvider>);
    }
 
+   // legacy plugins
    hordePlugins.loadPlugins(config.plugins);
 
+   const enabledPlugins = backend.enabledPlugins;
+   if (enabledPlugins.length) {
+      console.log("Enabled Server Plugins:")
+      enabledPlugins.forEach(p => {
+         console.log(`    ${p.name} : ${p.version ?? ""}`)
+      })         
+   } else {
+      console.log("No Enabled Plugins")
+   }
+
+   enableHordePlugins(enabledPlugins.map(p => p.name));
+   const plugins = getHordePlugins();
+
+   plugins.forEach(p => {
+      if (!!enabledPlugins.find(plugin => p.id === plugin.name)) {
+         console.log(`Loading Dashboard Plugin: ${p.id}`);
+      }
+   })
+
    if (!router) {
 
       const routes: RouteObject[] = [
@@ -141,7 +162,6 @@ const Main: React.FC = () => {
                { path: "lease/:leaseId", element: <DebugView /> },
                { path: "docs", element: <DocView /> },
                { path: "docs/*", element: <DocView /> },
-               { path: "analytics", element: <TelemetryView /> },
                { path: "accounts", element: <AccountsView /> },
                { path: "accounts/service", element: <ServiceAccountsView /> },
                { path: "test/stepissuereport", element: <StepIssueReportTest /> },
@@ -150,12 +170,15 @@ const Main: React.FC = () => {
          }
       ];
 
+      const pluginRoutes = plugins.map(p => p.routes).flat();
+      routes[0].children!.push(...pluginRoutes);
+
       // mount plugins
-      const pluginRoutes = hordePlugins.routes.map((route) => {
+      const legacyPluginRoutes = hordePlugins.routes.map((route) => {
          return { path: route.path, element: <route.component /> };
       })
 
-      routes[0].children!.push(...pluginRoutes);
+      routes[0].children!.push(...legacyPluginRoutes);
 
       router = createBrowserRouter(routes);
    }
diff --git a/Engine/Source/Programs/Horde/HordeDashboard/src/backend/Api.ts b/Engine/Source/Programs/Horde/HordeDashboard/src/backend/Api.ts
--- a/Engine/Source/Programs/Horde/HordeDashboard/src/backend/Api.ts
+++ b/Engine/Source/Programs/Horde/HordeDashboard/src/backend/Api.ts
@@ -457,7 +457,7 @@ export type GetDashboardConfigResponse = {
 	artifactTypes: string[];
 
 	/** Telemetry views */
-	telemetryViews: GetTelemetryViewResponse[];
+	telemetryViews: any[];
 }
 
 /**Parameters to register a new agent */
@@ -3266,136 +3266,6 @@ export type GetUtilizationTelemetryResponse = {
 	numAgents: number;
 }
 
-export type MetricsQuery = {
-	id: string[];
-	minTime?: string;
-	maxTime?: string;
-	group?: string;
-	results?: number;
-}
-
-/// Metrics matching a particular query
-export type GetTelemetryMetricsResponse = {
-
-	metricId: string;
-
-	groupBy: string;
-
-	/// Metrics matching the search terms	
-	metrics: GetTelemetryMetricResponse[];
-}
-
-/// Information about a particular metric
-export type GetTelemetryMetricResponse = {
-
-	/// Start time for the sample	
-	time: Date;
-
-	/// Name of the group	
-	group?: string;
-
-	/// Value for the metric	
-	value: number;
-
-	// Added locally in the dashboard
-	// GetTelemetryMetricsResponse id
-	id: string;
-
-	// added locally by dashboard
-	key: string;
-
-	// added locally by dashboard
-	keyElements: string[];
-
-	threshold?: number;
-
-	// calculated on dashboard, group name => value
-	groupValues?: Record<string, string>;
-}
-
-export type TelemetryDisplayType = "Time" | "Ratio" | "Value";
-export type TelemetryGraphType = "Line" | "Indicator";
-
-/// Metric attached to a telemetry chart	
-export type GetTelemetryChartMetricResponse = {
-
-	/// Associated metric id	
-	metricId: string;
-
-	/// The threshold for KPI values	
-	threshold?: number;
-
-	/// The metric alias for display purposes	
-	alias?: string;
-
-}
-
-/// Telemetry chart configuraton
-export type GetTelemetryChartResponse = {
-
-	/// The name of the chart, will be displayed on the dashboard	
-	name: string;
-
-	/// The unit to display	
-	display: TelemetryDisplayType;
-
-	/// The graph type 	
-	graph: TelemetryGraphType;
-
-	/// List of configured metrics	
-	metrics: GetTelemetryChartMetricResponse[];
-
-	/// The min unit value for clamping chart	
-	min?: number;
-
-	/// The max unit value for clamping chart	
-	max?: number;
-}
-
-/// A chart categody, will be displayed on the dashbord under an associated pivot
-export type GetTelemetryCategoryResponse = {
-
-	/// The name of the category
-	name: string;
-
-	/// The charts contained within the category
-	charts: GetTelemetryChartResponse[];
-}
-
-/// A telemetry view variable used for filtering the charting data
-export type GetTelemetryVariableResponse = {
-	/// The name of the variable for display purposes
-	name: string;
-
-	/// The associated data group attached to the variable 
-	group: string;
-
-	/// default values to select
-	defaults: string[];
-
-	/// Populated on dashboard
-	values: string[];
-}
-
-/// A telemetry view of related metrics, divided into categofies
-export type GetTelemetryViewResponse = {
-
-	/// Identifier for the view
-	id: string;
-
-	/// The name of the view
-	name: string;
-
-	/// The telemetry store id the view uses
-	telemetryStoreId: string;
-
-	///  The variables used to filter the view data
-	variables: GetTelemetryVariableResponse[];
-
-	/// The categories contained within the view
-	categories: GetTelemetryCategoryResponse[];
-}
-
 export type UserClaim = {
 
 	type: string;
diff --git a/Engine/Source/Programs/Horde/HordeDashboard/src/backend/Backend.ts b/Engine/Source/Programs/Horde/HordeDashboard/src/backend/Backend.ts
--- a/Engine/Source/Programs/Horde/HordeDashboard/src/backend/Backend.ts
+++ b/Engine/Source/Programs/Horde/HordeDashboard/src/backend/Backend.ts
@@ -1,7 +1,7 @@
 // Copyright Epic Games, Inc. All Rights Reserved.
 
 import templateCache from '../backend/TemplateCache';
-import { AccountClaimMessage, AgentData, AgentQuery, ApproveAgentsRequest, AuditLogEntry, AuditLogQuery, BatchUpdatePoolRequest, ChangeSummaryData, CreateAccountRequest, CreateBisectTaskRequest, CreateBisectTaskResponse, CreateDeviceRequest, CreateDeviceResponse, CreateExternalIssueRequest, CreateExternalIssueResponse, CreateJobRequest, CreateJobResponse, CreateNoticeRequest, CreatePoolRequest, CreateServiceAccountRequest, CreateServiceAccountResponse, CreateSoftwareResponse, CreateSubscriptionRequest, CreateSubscriptionResponse, CreateZipRequest, DashboardPreference, DevicePoolTelemetryQuery, DeviceTelemetryQuery, EventData, FindArtifactsResponse, FindIssueResponse, FindJobTimingsResponse, GetAccountResponse, GetAgentLeaseResponse, GetAgentSessionResponse, GetAgentTelemetrySampleResponse, GetArtifactDirectoryResponse, GetArtifactResponse, GetArtifactZipRequest, GetBisectTaskResponse, GetDashboardConfigResponse, GetDevicePlatformResponse, GetDevicePoolResponse, GetDevicePoolTelemetryResponse, GetDeviceReservationResponse, GetDeviceResponse, GetDeviceTelemetryResponse, GetExternalIssueProjectResponse, GetExternalIssueResponse, GetIssueStreamResponse, GetJobsTabResponse, GetJobStepRefResponse, GetJobStepTraceResponse, GetJobTimingResponse, GetLogEventResponse, GetLogFileResponse, GetNoticeResponse, GetNotificationResponse, GetPendingAgentsResponse, GetPerforceServerStatusResponse, GetPoolResponse, GetPoolSummaryResponse, GetServerInfoResponse, GetServerSettingsResponse, GetServiceAccountResponse, GetSoftwareResponse, GetSubscriptionResponse, GetTelemetryMetricsResponse, GetTestDataDetailsResponse, GetTestDataRefResponse, GetTestMetaResponse, GetTestResponse, GetTestsRequest, GetTestStreamResponse, GetToolResponse, GetToolSummaryResponse, GetUserResponse, GetUtilizationTelemetryResponse, GlobalConfig, IssueData, IssueQuery, IssueQueryV2, JobData, JobQuery, JobsTabColumnType, JobStepOutcome, JobStreamQuery, JobTimingsQuery, LeaseData, LogData, LogEventQuery, LogLineData, MetricsQuery, PoolData, PoolQuery, PreflightConfigResponse, ProjectData, ScheduleData, ScheduleQuery, SearchLogFileResponse, ServerStatusResponse, ServerUpdateResponse, SessionData, StreamData, TabType, TestData, UpdateAccountRequest, UpdateAgentRequest, UpdateBisectTaskRequest, UpdateCurrentAccountRequest, UpdateDeploymentRequest, UpdateDeviceRequest, UpdateGlobalConfigRequest, UpdateIssueRequest, UpdateJobRequest, UpdateLeaseRequest, UpdateNoticeRequest, UpdateNotificationsRequest, UpdatePoolRequest, UpdateServerSettingsRequest, UpdateServiceAccountRequest, UpdateServiceAccountResponse, UpdateStepRequest, UpdateStepResponse, UpdateTemplateRefRequest, UpdateUserRequest, UsersQuery } from './Api';
+import { AccountClaimMessage, AgentData, AgentQuery, ApproveAgentsRequest, AuditLogEntry, AuditLogQuery, BatchUpdatePoolRequest, ChangeSummaryData, CreateAccountRequest, CreateBisectTaskRequest, CreateBisectTaskResponse, CreateDeviceRequest, CreateDeviceResponse, CreateExternalIssueRequest, CreateExternalIssueResponse, CreateJobRequest, CreateJobResponse, CreateNoticeRequest, CreatePoolRequest, CreateServiceAccountRequest, CreateServiceAccountResponse, CreateSoftwareResponse, CreateSubscriptionRequest, CreateSubscriptionResponse, CreateZipRequest, DashboardPreference, DevicePoolTelemetryQuery, DeviceTelemetryQuery, EventData, FindArtifactsResponse, FindIssueResponse, FindJobTimingsResponse, GetAccountResponse, GetAgentLeaseResponse, GetAgentTelemetrySampleResponse, GetArtifactDirectoryResponse, GetArtifactResponse, GetBisectTaskResponse, GetDashboardConfigResponse, GetDevicePlatformResponse, GetDevicePoolResponse, GetDevicePoolTelemetryResponse, GetDeviceReservationResponse, GetDeviceResponse, GetDeviceTelemetryResponse, GetExternalIssueProjectResponse, GetExternalIssueResponse, GetIssueStreamResponse, GetJobsTabResponse, GetJobStepRefResponse, GetJobStepTraceResponse, GetJobTimingResponse, GetLogEventResponse, GetLogFileResponse, GetNoticeResponse, GetNotificationResponse, GetPendingAgentsResponse, GetPerforceServerStatusResponse, GetPoolResponse, GetPoolSummaryResponse, GetServerInfoResponse, GetServerSettingsResponse, GetServiceAccountResponse, GetSoftwareResponse, GetSubscriptionResponse, GetTestDataDetailsResponse, GetTestDataRefResponse, GetTestMetaResponse, GetTestResponse, GetTestsRequest, GetTestStreamResponse, GetToolResponse, GetToolSummaryResponse, GetUserResponse, GetUtilizationTelemetryResponse, GlobalConfig, IssueData, IssueQuery, IssueQueryV2, JobData, JobQuery, JobsTabColumnType, JobStepOutcome, JobStreamQuery, JobTimingsQuery, LeaseData, LogData, LogEventQuery, LogLineData, PoolData, PoolQuery, PreflightConfigResponse, ProjectData, ScheduleData, ScheduleQuery, SearchLogFileResponse, ServerPluginInfoResponse, ServerStatusResponse, ServerUpdateResponse, SessionData, StreamData, TabType, TestData, UpdateAccountRequest, UpdateAgentRequest, UpdateBisectTaskRequest, UpdateCurrentAccountRequest, UpdateDeploymentRequest, UpdateDeviceRequest, UpdateGlobalConfigRequest, UpdateIssueRequest, UpdateJobRequest, UpdateLeaseRequest, UpdateNoticeRequest, UpdateNotificationsRequest, UpdatePoolRequest, UpdateServerSettingsRequest, UpdateServiceAccountRequest, UpdateServiceAccountResponse, UpdateStepRequest, UpdateStepResponse, UpdateTemplateRefRequest, UpdateUserRequest, UsersQuery } from './Api';
 import dashboard, { Dashboard } from './Dashboard';
 import { ChallengeStatus, Fetch } from './Fetch';
 import { projectStore } from './ProjectStore';
@@ -15,7 +15,7 @@ export class Backend {
 
     getDashboardConfig(): Promise<GetDashboardConfigResponse> {
         return new Promise<GetDashboardConfigResponse>((resolve, reject) => {
-            this.backend.get("/api/v1/dashboard/config").then((response) => {
+            this._fetch.get("/api/v1/dashboard/config").then((response) => {
                 resolve(response.data as GetDashboardConfigResponse);
             }).catch(reason => {
                 reject(reason);
@@ -27,11 +27,11 @@ export class Backend {
 
         return new Promise<ProjectData[]>((resolve, reject) => {
 
-            this.backend.get("/api/v1/projects", { params: { categories: true } }).then(async (response) => {
+            this._fetch.get("/api/v1/projects", { params: { categories: true } }).then(async (response) => {
 
                 const projects = response.data as ProjectData[];
 
-                await this.backend.get("/api/v1/streams").then(response => {
+                await this._fetch.get("/api/v1/streams").then(response => {
 
                     const streams = (response.data as StreamData[]);
 
@@ -72,7 +72,7 @@ export class Backend {
 
         return new Promise<ProjectData>((resolve, reject) => {
 
-            this.backend.get(`/api/v1/projects/${project.id}`).then((response) => {
+            this._fetch.get(`/api/v1/projects/${project.id}`).then((response) => {
 
                 const project = response.data as ProjectData;
                 resolve(project);
@@ -86,7 +86,7 @@ export class Backend {
     updateTemplateRef(streamId: string, templateRefId: string, request: UpdateTemplateRefRequest): Promise<boolean> {
 
         return new Promise<boolean>((resolve, reject) => {
-            this.backend.put(`/api/v1/streams/${streamId}/templates/${templateRefId}`, request).then((value) => {
+            this._fetch.put(`/api/v1/streams/${streamId}/templates/${templateRefId}`, request).then((value) => {
                 resolve(true);
             }).catch(reason => {
                 reject(reason);
@@ -96,7 +96,7 @@ export class Backend {
 
     getPool(poolId: string): Promise<GetPoolResponse> {
         return new Promise<GetPoolResponse>((resolve, reject) => {
-            this.backend.get(`/api/v1/pools/${poolId}`).then((response) => {
+            this._fetch.get(`/api/v1/pools/${poolId}`).then((response) => {
                 resolve(response.data as GetPoolResponse);
             }).catch(reason => { reject(reason); });
         });
@@ -109,7 +109,7 @@ export class Backend {
         };
 
         return new Promise<PoolData[]>((resolve, reject) => {
-            this.backend.get("/api/v1/pools", { params: params }).then((response) => {
+            this._fetch.get("/api/v1/pools", { params: params }).then((response) => {
                 const pools = response.data as PoolData[];
                 resolve(pools);
             }).catch(reason => { reject(reason); });
@@ -119,7 +119,7 @@ export class Backend {
     getPoolsV2(query: PoolQuery): Promise<GetPoolSummaryResponse[]> {
 
         return new Promise<GetPoolSummaryResponse[]>((resolve, reject) => {
-            this.backend.get("/api/v2/pools", { params: query }).then((response) => {
+            this._fetch.get("/api/v2/pools", { params: query }).then((response) => {
                 const pools = response.data as GetPoolSummaryResponse[];
                 resolve(pools);
             }).catch(reason => { reject(reason); });
@@ -128,7 +128,7 @@ export class Backend {
 
     getAgentRegistrationRequests(): Promise<GetPendingAgentsResponse> {
         return new Promise<GetPendingAgentsResponse>((resolve, reject) => {
-            this.backend.get(`/api/v1/enrollment`).then((response) => {
+            this._fetch.get(`/api/v1/enrollment`).then((response) => {
                 const agent = response.data as GetPendingAgentsResponse;
                 resolve(agent);
             }).catch(reason => { reject(reason); });
@@ -138,7 +138,7 @@ export class Backend {
     // create a new account
     registerAgents(request: ApproveAgentsRequest): Promise<void> {
         return new Promise<void>((resolve, reject) => {
-            this.backend.post(`/api/v1/enrollment`, request).then(() => {
+            this._fetch.post(`/api/v1/enrollment`, request).then(() => {
                 resolve();
             }).catch(reason => {
                 reject(reason);
@@ -149,7 +149,7 @@ export class Backend {
 
     getAgent(id: string): Promise<AgentData> {
         return new Promise<AgentData>((resolve, reject) => {
-            this.backend.get(`/api/v1/agents/${id}`).then((response) => {
+            this._fetch.get(`/api/v1/agents/${id}`).then((response) => {
                 const agent = response.data as AgentData;
                 resolve(agent);
             }).catch(reason => { reject(reason); });
@@ -158,7 +158,7 @@ export class Backend {
 
     getAgentHistory(id: string, query: AuditLogQuery): Promise<AuditLogEntry[]> {
         return new Promise<AuditLogEntry[]>((resolve, reject) => {
-            this.backend.get(`/api/v1/agents/${encodeURIComponent(id)}/history`, { params: query }).then((response) => {
+            this._fetch.get(`/api/v1/agents/${encodeURIComponent(id)}/history`, { params: query }).then((response) => {
                 const history = (response.data?.entries ?? []) as AuditLogEntry[];
                 resolve(history);
             }).catch(reason => { reject(reason); });
@@ -168,7 +168,7 @@ export class Backend {
     getAgents(query: AgentQuery): Promise<AgentData[]> {
         return new Promise<AgentData[]>((resolve, reject) => {
             // for whenever we bring back modifiedDate, swap to this line
-            this.backend.get(`/api/v1/agents`, { params: query }).then((response) => {
+            this._fetch.get(`/api/v1/agents`, { params: query }).then((response) => {
                 //this.backend.get(`/api/v1/agents`).then((response) => {
                 const agents = response.data as AgentData[];
                 resolve(agents);
@@ -178,7 +178,7 @@ export class Backend {
 
     async getLeaseLog(leaseId: string): Promise<GetLogFileResponse> {
         return new Promise<GetLogFileResponse>((resolve, reject) => {
-            this.backend.get(`/api/v1/leases/${leaseId}/log`).then((response) => {
+            this._fetch.get(`/api/v1/leases/${leaseId}/log`).then((response) => {
                 resolve(response.data as GetLogFileResponse);
             }).catch(reason => { reject(reason); });
         });
@@ -188,14 +188,14 @@ export class Backend {
 
         if (agentId) {
             return new Promise<LeaseData>((resolve, reject) => {
-                this.backend.get(`/api/v1/agents/${agentId}/leases/${leaseId}`).then((response) => {
+                this._fetch.get(`/api/v1/agents/${agentId}/leases/${leaseId}`).then((response) => {
                     resolve(response.data as LeaseData);
                 }).catch(reason => { reject(reason); });
             });
         }
 
         return new Promise<LeaseData>((resolve, reject) => {
-            this.backend.get(`/api/v1/leases/${leaseId}`).then((response) => {
+            this._fetch.get(`/api/v1/leases/${leaseId}`).then((response) => {
                 resolve(response.data as LeaseData);
             }).catch(reason => { reject(reason); });
         });
@@ -208,7 +208,7 @@ export class Backend {
         params.push("Count=" + count);
         const paramString = params.join("&");
         return new Promise<LeaseData[]>((resolve, reject) => {
-            this.backend.get(`/api/v1/agents/${agentId}/leases?${paramString}`).then(async (response) => {
+            this._fetch.get(`/api/v1/agents/${agentId}/leases?${paramString}`).then(async (response) => {
                 const leases = response.data as LeaseData[];
                 leases.forEach(function (lease) {
                     lease.startTime = new Date(Date.parse(lease.startTime as string));
@@ -250,7 +250,7 @@ export class Backend {
         params.push("Count=" + count);
         const paramString = params.join("&");
         return new Promise<SessionData[]>((resolve, reject) => {
-            this.backend.get(`/api/v1/agents/${agentId}/sessions?${paramString}`).then((response) => {
+            this._fetch.get(`/api/v1/agents/${agentId}/sessions?${paramString}`).then((response) => {
                 const sessions = response.data as SessionData[];
                 sessions.forEach(function (session) {
                     session.startTime = new Date(Date.parse(session.startTime as string));
@@ -269,7 +269,7 @@ export class Backend {
         params.push("finishime=" + maxTime.toISOString());
         const paramString = params.join("&");
         return new Promise<GetAgentLeaseResponse[]>((resolve, reject) => {
-            this.backend.get(`/api/v1/agents/${agentId}/leases?${paramString}`).then((response) => {
+            this._fetch.get(`/api/v1/agents/${agentId}/leases?${paramString}`).then((response) => {
                 (response.data as GetAgentLeaseResponse[]).forEach(s => {
                     s.startTime = s.startTime ? new Date(s.startTime) : new Date();
                     s.finishTime = s.finishTime ? new Date(s.finishTime) : undefined;
@@ -286,7 +286,7 @@ export class Backend {
         params.push("maxTime=" + maxTime.toISOString());
         const paramString = params.join("&");
         return new Promise<GetAgentTelemetrySampleResponse[]>((resolve, reject) => {
-            this.backend.get(`/api/v1/agents/${agentId}/telemetry?${paramString}`).then((response) => {
+            this._fetch.get(`/api/v1/agents/${agentId}/telemetry?${paramString}`).then((response) => {
                 const data = response.data.samples as GetAgentTelemetrySampleResponse[];
                 data.forEach((d) => {
                     d.time = new Date(d.time);
@@ -299,7 +299,7 @@ export class Backend {
     retryJobStep(jobId: string, batchId: string, stepId: string): Promise<any> {
 
         return new Promise<any>((resolve, reject) => {
-            this.backend.put(`/api/v1/jobs/${jobId}/batches/${batchId}/steps/${stepId}`, {
+            this._fetch.put(`/api/v1/jobs/${jobId}/batches/${batchId}/steps/${stepId}`, {
                 Retry: true
             }).then((value) => {
                 resolve(value.data as JobData[]);
@@ -314,7 +314,7 @@ export class Backend {
 
         return new Promise<JobData>((resolve, reject) => {
 
-            this.backend.get(`/api/v1/jobs/${id}`, {
+            this._fetch.get(`/api/v1/jobs/${id}`, {
                 params: query,
                 show404Error: show404Error
             }).then((value) => {
@@ -341,7 +341,7 @@ export class Backend {
 
         return new Promise<JobData[]>((resolve, reject) => {
 
-            this.backend.get(`/api/v1/jobs`, {
+            this._fetch.get(`/api/v1/jobs`, {
                 params: query
             }).then((value) => {
 
@@ -368,7 +368,7 @@ export class Backend {
 
         return new Promise<JobData[]>((resolve, reject) => {
 
-            this.backend.get(`/api/v1/jobs/streams/${streamId}`, {
+            this._fetch.get(`/api/v1/jobs/streams/${streamId}`, {
                 params: query
             }).then((value) => {
                 const jobs = value.data as JobData[];
@@ -394,7 +394,7 @@ export class Backend {
 
         return new Promise<JobData[]>((resolve, reject) => {
 
-            this.backend.get("/api/v1/jobs", {
+            this._fetch.get("/api/v1/jobs", {
                 params: query
             }).then((value) => {
                 const jobs = value.data as JobData[];
@@ -412,7 +412,7 @@ export class Backend {
     getJobTiming(jobId: string): Promise<GetJobTimingResponse> {
 
         return new Promise<GetJobTimingResponse>((resolve, reject) => {
-            this.backend.get(`/api/v1/jobs/${jobId}/timing`).then((value) => {
+            this._fetch.get(`/api/v1/jobs/${jobId}/timing`).then((value) => {
                 resolve(value.data as GetJobTimingResponse);
             }).catch(reason => {
                 reject(reason);
@@ -424,7 +424,7 @@ export class Backend {
     getBatchJobTiming(query: JobTimingsQuery): Promise<FindJobTimingsResponse> {
 
         return new Promise<FindJobTimingsResponse>((resolve, reject) => {
-            this.backend.get(`/api/v1/jobs/timing`, {
+            this._fetch.get(`/api/v1/jobs/timing`, {
                 params: query
             }).then((response) => {
                 let timingsWrapper = response.data as FindJobTimingsResponse;
@@ -440,7 +440,7 @@ export class Backend {
 
         return new Promise<EventData[]>((resolve, reject) => {
 
-            this.backend.get(`/api/v1/logs/${logId}/events`, { params: query }).then((value) => {
+            this._fetch.get(`/api/v1/logs/${logId}/events`, { params: query }).then((value) => {
                 resolve(value.data as EventData[]);
             }).catch(reason => {
                 reject(reason);
@@ -458,7 +458,7 @@ export class Backend {
                 count: count
             };
 
-            this.backend.get(`/api/v1/logs/${logId}/search`, {
+            this._fetch.get(`/api/v1/logs/${logId}/search`, {
                 params: params
             }).then((value) => {
                 resolve(value.data as SearchLogFileResponse);
@@ -478,7 +478,7 @@ export class Backend {
                 templateId: templateId
             };
 
-            this.backend.get(`/api/v1/streams/${streamId}/history`, {
+            this._fetch.get(`/api/v1/streams/${streamId}/history`, {
                 params: params
             }).then((value) => {
                 let results = (value.data ?? []) as GetJobStepRefResponse[];
@@ -492,7 +492,7 @@ export class Backend {
 
     getJobStepTrace(jobId: string, batchId: string, stepId: string) {
         return new Promise<GetJobStepTraceResponse>((resolve, reject) => {
-            this.backend.get(`/api/v1/jobs/${jobId}/batches/${batchId}/steps/${stepId}/trace`, {
+            this._fetch.get(`/api/v1/jobs/${jobId}/batches/${batchId}/steps/${stepId}/trace`, {
                 suppress404: true
             }).then((value) => {
                 resolve(value.data as GetJobStepTraceResponse);
@@ -505,8 +505,7 @@ export class Backend {
     getArtifacts(streamId?: string, minChange?: number, maxChange?: number, name?: string, type?: string): Promise<FindArtifactsResponse> {
 
         return new Promise<FindArtifactsResponse>((resolve, reject) => {
-
-            this.backend.get(`/api/v2/artifacts`, { params: { streamId: streamId, minChange: minChange, maxChange: maxChange, name: name, type: type, maxResults: 512 } })
+            this._fetch.get(`/api/v2/artifacts`, { params: { streamId: streamId, minChange: minChange, maxChange: maxChange, name: name, type: type, maxResults: 512 } })
                 .then(response => { resolve(response.data); })
                 .catch(reason => reject(reason));
         })
@@ -524,7 +523,7 @@ export class Backend {
 
         return new Promise<FindArtifactsResponse>((resolve, reject) => {
 
-            this.backend.get(`/api/v2/artifacts`, { params: { id: uniqueIds, key: uniqueKeys } })
+            this._fetch.get(`/api/v2/artifacts`, { params: { id: uniqueIds, key: uniqueKeys } })
                 .then(response => { resolve(response.data); })
                 .catch(reason => reject(reason));
         })
@@ -535,7 +534,7 @@ export class Backend {
 
         return new Promise<GetArtifactDirectoryResponse>((resolve, reject) => {
 
-            this.backend.get(`/api/v2/artifacts/${id}/browse`, { params: { path: path } })
+            this._fetch.get(`/api/v2/artifacts/${id}/browse`, { params: { path: path } })
                 .then(response => { resolve(response.data); })
                 .catch(reason => reject(reason));
         })
@@ -544,7 +543,7 @@ export class Backend {
     getArtifactData(artifactId: string): Promise<GetArtifactResponse> {
         const url = `/api/v2/artifacts/${artifactId}`;
         return new Promise<GetArtifactResponse>((resolve, reject) => {
-            this.backend.get(url).then((value) => {
+            this._fetch.get(url).then((value) => {
                 resolve(value.data as GetArtifactResponse);
             }).catch(reason => {
                 reject(reason);
@@ -555,7 +554,7 @@ export class Backend {
     getArtifactV2(artifactId: string, path: string): Promise<object> {
         const url = `/api/v2/artifacts/${artifactId}/file?path=${encodeURIComponent(path)}`;
         return new Promise<object>((resolve, reject) => {
-            this.backend.get(url).then((value) => {
+            this._fetch.get(url).then((value) => {
                 resolve(value.data as object);
             }).catch(reason => {
                 resolve([]);
@@ -589,7 +588,7 @@ export class Backend {
 
         return new Promise<LogData>((resolve, reject) => {
 
-            this.backend.get(`/api/v1/logs/${logId}`).then((value) => {
+            this._fetch.get(`/api/v1/logs/${logId}`).then((value) => {
                 resolve(value.data as LogData);
             }).catch(reason => {
                 reject(reason);
@@ -612,7 +611,7 @@ export class Backend {
 
         return new Promise<LogLineData>((resolve, reject) => {
 
-            this.backend.get(`/api/v1/logs/${logId}/lines`, {
+            this._fetch.get(`/api/v1/logs/${logId}/lines`, {
                 params: params
             }).then((value) => {
                 resolve(value.data as LogLineData);
@@ -627,7 +626,7 @@ export class Backend {
 
         return new Promise<void>((resolve, reject) => {
 
-            this.backend.get("/api/v1/debug/exception").then(() => {
+            this._fetch.get("/api/v1/debug/exception").then(() => {
                 resolve();
             }).catch(reason => {
                 reject(reason);
@@ -640,7 +639,7 @@ export class Backend {
 
         return new Promise<CreateJobResponse>((resolve, reject) => {
 
-            this.backend.post("/api/v1/jobs", request).then((value) => {
+            this._fetch.post("/api/v1/jobs", request).then((value) => {
                 resolve(value.data as CreateJobResponse);
             }).catch(reason => {
                 reject(reason);
@@ -651,7 +650,7 @@ export class Backend {
     updateJob(jobId: string, request: UpdateJobRequest): Promise<boolean> {
 
         return new Promise<boolean>((resolve, reject) => {
-            this.backend.put(`/api/v1/jobs/${jobId}`, request).then((value) => {
+            this._fetch.put(`/api/v1/jobs/${jobId}`, request).then((value) => {
                 resolve(true);
             }).catch(reason => {
                 reject(reason);
@@ -661,7 +660,7 @@ export class Backend {
 
     updateAgent(agentId: string, request: UpdateAgentRequest): Promise<boolean> {
         return new Promise<boolean>((resolve, reject) => {
-            this.backend.put(`api/v1/agents/${agentId}`, request).then((value) => {
+            this._fetch.put(`api/v1/agents/${agentId}`, request).then((value) => {
                 resolve(value.data as boolean);
             }).catch(reason => {
                 reject(reason);
@@ -671,7 +670,7 @@ export class Backend {
 
     deleteAgent(agentId: string): Promise<any> {
         return new Promise<boolean>((resolve, reject) => {
-            this.backend.delete(`api/v1/agents/${agentId}`).then((value) => {
+            this._fetch.delete(`api/v1/agents/${agentId}`).then((value) => {
                 resolve(value.data);
             }).catch(reason => {
                 reject(reason);
@@ -681,7 +680,7 @@ export class Backend {
 
     createPool(request: CreatePoolRequest): Promise<string> {
         return new Promise<string>((resolve, reject) => {
-            this.backend.post(`api/v1/pools`, request).then((value) => {
+            this._fetch.post(`api/v1/pools`, request).then((value) => {
                 resolve(value.data as string);
             }).catch(reason => {
                 reject(reason);
@@ -691,7 +690,7 @@ export class Backend {
 
     updatePool(poolId: string, request: UpdatePoolRequest): Promise<boolean> {
         return new Promise<boolean>((resolve, reject) => {
-            this.backend.put(`api/v1/pools/${poolId}`, request).then((value) => {
+            this._fetch.put(`api/v1/pools/${poolId}`, request).then((value) => {
                 resolve(value.data as boolean);
             }).catch(reason => {
                 reject(reason);
@@ -701,7 +700,7 @@ export class Backend {
 
     deletePool(poolId: string): Promise<boolean> {
         return new Promise<boolean>((resolve, reject) => {
-            this.backend.delete(`api/v1/pools/${poolId}`).then((value) => {
+            this._fetch.delete(`api/v1/pools/${poolId}`).then((value) => {
                 resolve(value.data as boolean);
             }).catch(reason => {
                 reject(reason);
@@ -711,7 +710,7 @@ export class Backend {
 
     batchUpdatePools(request: BatchUpdatePoolRequest[]): Promise<boolean> {
         return new Promise<boolean>((resolve, reject) => {
-            this.backend.put(`api/v1/pools`, request).then((value) => {
+            this._fetch.put(`api/v1/pools`, request).then((value) => {
                 resolve(value.data as boolean);
             }).catch(reason => {
                 reject(reason);
@@ -723,7 +722,7 @@ export class Backend {
 
         return new Promise<CreateSoftwareResponse>((resolve, reject) => {
 
-            this.backend.post("/api/v1/software", formData, { formData: true }).then((value) => {
+            this._fetch.post("/api/v1/software", formData, { formData: true }).then((value) => {
                 resolve(value.data as CreateSoftwareResponse);
             }).catch(reason => {
                 reject(reason);
@@ -735,7 +734,7 @@ export class Backend {
     updateSoftware(id: string, isDefault: boolean): Promise<any> {
 
         return new Promise<any>((resolve, reject) => {
-            this.backend.put(`/api/v1/software/${id}`, {
+            this._fetch.put(`/api/v1/software/${id}`, {
                 Default: isDefault
             }).then((value) => {
                 resolve(value.data as any);
@@ -751,7 +750,7 @@ export class Backend {
 
         return new Promise<ScheduleData>((resolve, reject) => {
 
-            this.backend.get(`/api/v1/schedules/${id}`).then((value) => {
+            this._fetch.get(`/api/v1/schedules/${id}`).then((value) => {
                 const data = value.data as ScheduleData;
                 data.nextTriggerTimesUTC.forEach((d, index) => {
                     data.nextTriggerTimesUTC[index] = new Date(d as any as string);
@@ -767,7 +766,7 @@ export class Backend {
     getSchedules(query: ScheduleQuery): Promise<ScheduleData[]> {
 
         return new Promise<ScheduleData[]>((resolve, reject) => {
-            this.backend.get(`/api/v1/schedules`, {
+            this._fetch.get(`/api/v1/schedules`, {
                 params: query
             }).then((value) => {
                 resolve(value.data as ScheduleData[]);
@@ -783,7 +782,7 @@ export class Backend {
 
         return new Promise<GetSoftwareResponse[]>((resolve, reject) => {
 
-            this.backend.get(`/api/v1/software`).then((value) => {
+            this._fetch.get(`/api/v1/software`).then((value) => {
                 resolve(value.data as GetSoftwareResponse[]);
             }).catch(reason => {
                 reject(reason);
@@ -796,7 +795,7 @@ export class Backend {
 
         return new Promise<any>((resolve, reject) => {
 
-            this.backend.delete(`/api/v1/software/${id}`).then((value) => {
+            this._fetch.delete(`/api/v1/software/${id}`).then((value) => {
                 resolve(value.data);
             }).catch(reason => {
                 reject(reason);
@@ -814,7 +813,7 @@ export class Backend {
                 return;
             }
 
-            this.backend.get(`/api/v1/streams/${streamId}/changes/${changelist}`).then((value) => {
+            this._fetch.get(`/api/v1/streams/${streamId}/changes/${changelist}`).then((value) => {
                 resolve(value.data as ChangeSummaryData);
             }).catch(reason => {
                 console.error(reason);
@@ -839,7 +838,7 @@ export class Backend {
                 params.max = maxChange;
             }
 
-            this.backend.get(`/api/v1/streams/${streamId}/changes`, {
+            this._fetch.get(`/api/v1/streams/${streamId}/changes`, {
                 params: params
             }).then((value) => {
                 resolve(value.data as ChangeSummaryData[]);
@@ -855,7 +854,7 @@ export class Backend {
 
         return new Promise<string>((resolve, reject) => {
 
-            this.backend.get("/api/v1/admin/token").then((response) => {
+            this._fetch.get("/api/v1/admin/token").then((response) => {
                 resolve(response.data as string);
             }).catch(reason => reject(reason));
 
@@ -866,7 +865,7 @@ export class Backend {
     updateUser(request: UpdateUserRequest): Promise<void> {
 
         return new Promise<void>((resolve, reject) => {
-            this.backend.put("/api/v1/user", request).then((response) => {
+            this._fetch.put("/api/v1/user", request).then((response) => {
                 resolve();
             }).catch(reason => reject(reason));
 
@@ -879,7 +878,7 @@ export class Backend {
         query = query ?? {};
 
         return new Promise<GetUserResponse[]>((resolve, reject) => {
-            this.backend.get("/api/v1/users", {
+            this._fetch.get("/api/v1/users", {
                 params: query as any
             }).then((response) => {
                 resolve(response.data as GetUserResponse[]);
@@ -892,7 +891,7 @@ export class Backend {
 
         return new Promise<GetUserResponse>((resolve, reject) => {
 
-            this.backend.get("/api/v1/user", { suppress404: true }).then((response) => {
+            this._fetch.get("/api/v1/user", { suppress404: true }).then((response) => {
 
                 let data = response.data as GetUserResponse;
 
@@ -941,7 +940,7 @@ export class Backend {
     getIssue(issueId: number): Promise<IssueData> {
 
         return new Promise<IssueData>((resolve, reject) => {
-            this.backend.get(`/api/v1/issues/${issueId}`).then((value) => {
+            this._fetch.get(`/api/v1/issues/${issueId}`).then((value) => {
                 resolve(value.data as IssueData);
             }).catch((reason) => {
                 reject(reason);
@@ -952,7 +951,7 @@ export class Backend {
 
     getIssueHistory(id: string, query: AuditLogQuery): Promise<AuditLogEntry[]> {
         return new Promise<AuditLogEntry[]>((resolve, reject) => {
-            this.backend.get(`/api/v1/issues/${encodeURIComponent(id)}/history`, { params: query }).then((response) => {
+            this._fetch.get(`/api/v1/issues/${encodeURIComponent(id)}/history`, { params: query }).then((response) => {
                 const history = (response.data?.entries ?? []) as AuditLogEntry[];
                 resolve(history);
             }).catch(reason => { reject(reason); });
@@ -966,7 +965,7 @@ export class Backend {
 
         return new Promise<FindIssueResponse[]>((resolve, reject) => {
 
-            this.backend.get("/api/v2/issues", {
+            this._fetch.get("/api/v2/issues", {
                 params: query
             }).then((value) => {
                 resolve(value.data as FindIssueResponse[]);
@@ -983,7 +982,7 @@ export class Backend {
 
         return new Promise<IssueData[]>((resolve, reject) => {
 
-            this.backend.get("/api/v1/issues", {
+            this._fetch.get("/api/v1/issues", {
                 params: query
             }).then((value) => {
                 resolve(value.data as IssueData[]);
@@ -1000,7 +999,7 @@ export class Backend {
 
         return new Promise<IssueData[]>((resolve, reject) => {
 
-            this.backend.get(`/api/v1/issues`, { params: { id: unique.map(id => id.toString()) } })
+            this._fetch.get(`/api/v1/issues`, { params: { id: unique.map(id => id.toString()) } })
                 .then(response => { resolve(response.data); })
                 .catch(reason => reject(reason));
         })
@@ -1008,7 +1007,7 @@ export class Backend {
 
     getIssueStreams(issueId: number): Promise<GetIssueStreamResponse[]> {
         return new Promise<GetIssueStreamResponse[]>((resolve, reject) => {
-            this.backend.get(`/api/v1/issues/${issueId}/streams`).then((value) => {
+            this._fetch.get(`/api/v1/issues/${issueId}/streams`).then((value) => {
                 resolve(value.data as GetIssueStreamResponse[]);
             }).catch((reason) => {
                 reject(reason);
@@ -1030,7 +1029,7 @@ export class Backend {
         };
 
         return new Promise<GetLogEventResponse[]>((resolve, reject) => {
-            this.backend.get(`/api/v1/issues/${issueId}/events`, { params: params }).then((value) => {
+            this._fetch.get(`/api/v1/issues/${issueId}/events`, { params: params }).then((value) => {
                 resolve(value.data as GetLogEventResponse[]);
             }).catch((reason) => {
                 reject(reason);
@@ -1041,7 +1040,7 @@ export class Backend {
 
     updateIssue(id: number, request: UpdateIssueRequest): Promise<boolean> {
         return new Promise<boolean>((resolve, reject) => {
-            this.backend.put(`/api/v1/issues/${id}`, request).then((value) => {
+            this._fetch.put(`/api/v1/issues/${id}`, request).then((value) => {
                 resolve(true);
             }).catch((reason) => {
                 reject(reason);
@@ -1057,7 +1056,7 @@ export class Backend {
         };
 
         return new Promise<GetExternalIssueResponse[]>((resolve, reject) => {
-            this.backend.get(`/api/v1/issues/external`, { params: params }).then((value) => {
+            this._fetch.get(`/api/v1/issues/external`, { params: params }).then((value) => {
                 resolve(value.data as GetExternalIssueResponse[]);
             }).catch((reason) => {
                 reject(reason);
@@ -1072,7 +1071,7 @@ export class Backend {
         };
 
         return new Promise<GetExternalIssueProjectResponse[]>((resolve, reject) => {
-            this.backend.get(`/api/v1/issues/external/projects`, { params: params }).then((value) => {
+            this._fetch.get(`/api/v1/issues/external/projects`, { params: params }).then((value) => {
                 resolve(value.data as GetExternalIssueProjectResponse[]);
             }).catch((reason) => {
                 reject(reason);
@@ -1083,7 +1082,7 @@ export class Backend {
 
     createExternalIssue(request: CreateExternalIssueRequest): Promise<CreateExternalIssueResponse> {
         return new Promise<CreateExternalIssueResponse>((resolve, reject) => {
-            this.backend.post(`/api/v1/issues/external`, request).then((value) => {
+            this._fetch.post(`/api/v1/issues/external`, request).then((value) => {
                 resolve(value.data as CreateExternalIssueResponse);
             }).catch(reason => {
                 reject(reason);
@@ -1097,7 +1096,7 @@ export class Backend {
                 filter: 'id,key,change,jobId,stepId,streamId,' + (filter ? filter : 'data')
             };
 
-            this.backend.get(`/api/v1/testdata/${id}`, {
+            this._fetch.get(`/api/v1/testdata/${id}`, {
                 params: params
             }).then((value) => {
                 resolve(value.data as TestData);
@@ -1109,7 +1108,7 @@ export class Backend {
 
     getTestMetadata(automationProjects?: string[], platforms?: string[], targets?: string[], configurations?: string[]): Promise<GetTestMetaResponse[]> {
         return new Promise<GetTestMetaResponse[]>((resolve, reject) => {
-            this.backend.get(`/api/v2/testdata/metadata`, {
+            this._fetch.get(`/api/v2/testdata/metadata`, {
                 params: { project: automationProjects, platform: platforms, target: targets, configuration: configurations }
             }).then((value) => {
                 resolve(value.data as GetTestMetaResponse[]);
@@ -1121,7 +1120,7 @@ export class Backend {
 
     getTestRefs(streamIds: string[], metaIds: string[], minCreateTime?: string, maxCreateTime?: string, minChange?: number, maxChange?: number, testIds?: string[], suiteIds?: string[]): Promise<GetTestDataRefResponse[]> {
         return new Promise<GetTestDataRefResponse[]>((resolve, reject) => {
-            this.backend.get(`/api/v2/testdata/refs`, {
+            this._fetch.get(`/api/v2/testdata/refs`, {
                 params: { id: streamIds, mid: metaIds, tid: testIds, sid: suiteIds, minCreateTime: minCreateTime, maxCreateTime: maxCreateTime, minChange: minChange, maxChange: maxChange }
             }).then((value) => {
                 resolve(value.data as GetTestDataRefResponse[]);
@@ -1133,7 +1132,7 @@ export class Backend {
 
     getTestDetails(refIds: string[]): Promise<GetTestDataDetailsResponse[]> {
         return new Promise<GetTestDataDetailsResponse[]>((resolve, reject) => {
-            this.backend.get(`/api/v2/testdata/details`, {
+            this._fetch.get(`/api/v2/testdata/details`, {
                 params: { id: refIds }
             }).then((value) => {
                 resolve(value.data as GetTestDataDetailsResponse[]);
@@ -1146,7 +1145,7 @@ export class Backend {
     getTests(testIds: string[]): Promise<GetTestResponse[]> {
         const request: GetTestsRequest = { testIds: testIds };
         return new Promise<GetTestResponse[]>((resolve, reject) => {
-            this.backend.post(`/api/v2/testdata/tests`, request).then((value) => {
+            this._fetch.post(`/api/v2/testdata/tests`, request).then((value) => {
                 resolve(value.data as GetTestResponse[]);
             }).catch((reason) => {
                 reject(reason);
@@ -1156,7 +1155,7 @@ export class Backend {
 
     getTestStreams(streamIds: string[]): Promise<GetTestStreamResponse[]> {
         return new Promise<GetTestStreamResponse[]>((resolve, reject) => {
-            this.backend.get(`/api/v2/testdata/streams`, {
+            this._fetch.get(`/api/v2/testdata/streams`, {
                 params: { id: streamIds }
             }).then((value) => {
                 resolve(value.data as GetTestStreamResponse[]);
@@ -1178,7 +1177,7 @@ export class Backend {
                 params.jobStepId = stepId;
             }
 
-            this.backend.get(`/api/v1/testdata`, {
+            this._fetch.get(`/api/v1/testdata`, {
                 params: params
             }).then((value) => {
                 resolve(value.data as TestData[]);
@@ -1204,7 +1203,7 @@ export class Backend {
                 params.index = index;
             }
 
-            this.backend.get(`/api/v1/testdata`, {
+            this._fetch.get(`/api/v1/testdata`, {
                 params: params
             }).then((value) => {
                 resolve(value.data as TestData[]);
@@ -1230,7 +1229,7 @@ export class Backend {
 
     updateJobStep(jobId: string, batchId: string, stepId: string, request: UpdateStepRequest): Promise<UpdateStepResponse> {
         return new Promise<UpdateStepResponse>((resolve, reject) => {
-            this.backend.put(`api/v1/jobs/${jobId}/batches/${batchId}/steps/${stepId}`, request).then((value) => {
+            this._fetch.put(`api/v1/jobs/${jobId}/batches/${batchId}/steps/${stepId}`, request).then((value) => {
                 resolve(value.data as UpdateStepResponse);
             }).catch(reason => {
                 reject(reason);
@@ -1248,7 +1247,7 @@ export class Backend {
             url = `${url}/batches/${batchId}/steps/${stepId}`;
         }
         return new Promise<GetNotificationResponse>((resolve, reject) => {
-            this.backend.get(`${url}/notifications`).then((value) => {
+            this._fetch.get(`${url}/notifications`).then((value) => {
                 resolve(value.data as GetNotificationResponse);
             }).catch(reason => {
                 reject(reason);
@@ -1267,7 +1266,7 @@ export class Backend {
             url = `${url}/batches/${batchId}/steps/${stepId}`;
         }
         return new Promise<boolean>((resolve, reject) => {
-            this.backend.put(`${url}/notifications`, request).then((value) => {
+            this._fetch.put(`${url}/notifications`, request).then((value) => {
                 resolve(true);
             }).catch(reason => {
                 reject(reason);
@@ -1281,7 +1280,7 @@ export class Backend {
         });
         let url = `api/v1/subscriptions`;
         return new Promise<CreateSubscriptionResponse[]>((resolve, reject) => {
-            this.backend.post(`${url}`, requests).then((value) => {
+            this._fetch.post(`${url}`, requests).then((value) => {
                 let created = value.data as CreateSubscriptionResponse[];
                 resolve(created);
             }).catch(reason => {
@@ -1293,7 +1292,7 @@ export class Backend {
     getSubscriptions(): Promise<GetSubscriptionResponse[]> {
         let url = `api/v1/subscriptions?userId=${dashboard.userId}`;
         return new Promise<GetSubscriptionResponse[]>((resolve, reject) => {
-            this.backend.get(`${url}`, { suppress404: true }).then((value) => {
+            this._fetch.get(`${url}`, { suppress404: true }).then((value) => {
                 resolve(value.data as GetSubscriptionResponse[]);
             }).catch(reason => {
                 reject(reason);
@@ -1304,7 +1303,7 @@ export class Backend {
     deleteSubscription(id: string): Promise<boolean> {
         let url = `api/v1/subscriptions/${id}`;
         return new Promise<boolean>((resolve, reject) => {
-            this.backend.delete(`${url}`).then((value) => {
+            this._fetch.delete(`${url}`).then((value) => {
                 resolve(true);
             }).catch(reason => {
                 reject(reason);
@@ -1315,7 +1314,7 @@ export class Backend {
     getUtilizationData(endDate: string, range: number, tzOffset: number): Promise<GetUtilizationTelemetryResponse[]> {
         let url = `api/v1/reports/utilization/${endDate}?range=${range}&tzOffset=${tzOffset}`;
         return new Promise<GetUtilizationTelemetryResponse[]>((resolve, reject) => {
-            this.backend.get(`${url}`).then((value) => {
+            this._fetch.get(`${url}`).then((value) => {
                 resolve(value.data as GetUtilizationTelemetryResponse[]);
             }).catch(reason => {
                 reject(reason);
@@ -1323,30 +1322,10 @@ export class Backend {
         });
     }
 
-    getMetrics(telemetryStoreId: string, query: MetricsQuery): Promise<GetTelemetryMetricsResponse[]> {
-        query.id = query.id.map(id => encodeURIComponent(id));
-        return new Promise<GetTelemetryMetricsResponse[]>((resolve, reject) => {
-            this.backend.get(`api/v1/telemetry/${telemetryStoreId}/metrics`, { params: query }).then((response) => {
-                const result = response.data as GetTelemetryMetricsResponse[];
-                result?.forEach(r => {
-                    r?.metrics.forEach(m => {
-                        if (m.time) {
-                            m.time = new Date(m.time);
-                        } else {
-                            console.warn("Metrics missing time property");
-                            m.time = new Date();
-                        }
-                    })
-                })
-                resolve(result);
-            }).catch(reason => { reject(reason); });
-        });
-    }
-
     getPerforceServerStatus(): Promise<GetPerforceServerStatusResponse[]> {
 
         return new Promise<GetPerforceServerStatusResponse[]>((resolve, reject) => {
-            this.backend.get(`/api/v1/perforce/status`).then((value) => {
+            this._fetch.get(`/api/v1/perforce/status`).then((value) => {
                 resolve(value.data as GetPerforceServerStatusResponse[]);
             }).catch(reason => {
                 reject(reason);
@@ -1357,7 +1336,7 @@ export class Backend {
     getServerStatus(): Promise<ServerStatusResponse> {
 
         return new Promise<ServerStatusResponse>((resolve, reject) => {
-            this.backend.get(`/api/v1/server/status`).then((value) => {
+            this._fetch.get(`/api/v1/server/status`).then((value) => {
                 const result = value.data as ServerStatusResponse;
                 // convert from string date to Date
                 result.statuses.forEach(s => s.updates.forEach(u => u.updatedAt = new Date(u.updatedAt)))
@@ -1371,7 +1350,7 @@ export class Backend {
     getDevices(): Promise<GetDeviceResponse[]> {
 
         return new Promise<GetDeviceResponse[]>((resolve, reject) => {
-            this.backend.get(`/api/v2/devices`).then((value) => {
+            this._fetch.get(`/api/v2/devices`).then((value) => {
                 resolve(value.data as GetDeviceResponse[]);
             }).catch(reason => {
                 reject(reason);
@@ -1382,7 +1361,7 @@ export class Backend {
     getDevicePoolTelemetry(query?: DevicePoolTelemetryQuery): Promise<GetDevicePoolTelemetryResponse[]> {
 
         return new Promise<GetDevicePoolTelemetryResponse[]>((resolve, reject) => {
-            this.backend.get(`/api/v2/devices/pools/telemetry`, { params: query }).then((value) => {
+            this._fetch.get(`/api/v2/devices/pools/telemetry`, { params: query }).then((value) => {
                 resolve(value.data as GetDevicePoolTelemetryResponse[]);
             }).catch(reason => {
                 reject(reason);
@@ -1393,7 +1372,7 @@ export class Backend {
     getDeviceTelemetry(query?: DeviceTelemetryQuery): Promise<GetDeviceTelemetryResponse[]> {
 
         return new Promise<GetDeviceTelemetryResponse[]>((resolve, reject) => {
-            this.backend.get(`/api/v2/devices/telemetry`, { params: query }).then((value) => {
+            this._fetch.get(`/api/v2/devices/telemetry`, { params: query }).then((value) => {
                 resolve(value.data as GetDeviceTelemetryResponse[]);
             }).catch(reason => {
                 reject(reason);
@@ -1403,7 +1382,7 @@ export class Backend {
 
     addDevice(request: CreateDeviceRequest): Promise<CreateDeviceResponse> {
         return new Promise<CreateDeviceResponse>((resolve, reject) => {
-            this.backend.post(`/api/v2/devices`, request).then((value) => {
+            this._fetch.post(`/api/v2/devices`, request).then((value) => {
                 resolve(value.data as CreateDeviceResponse);
             }).catch(reason => {
                 reject(reason);
@@ -1413,7 +1392,7 @@ export class Backend {
 
     checkoutDevice(deviceId: string, checkout: boolean): Promise<void> {
         return new Promise<void>((resolve, reject) => {
-            this.backend.put(`/api/v2/devices/${deviceId}/checkout`, { checkout: checkout }).then(() => {
+            this._fetch.put(`/api/v2/devices/${deviceId}/checkout`, { checkout: checkout }).then(() => {
                 resolve();
             }).catch(reason => {
                 reject(reason);
@@ -1424,7 +1403,7 @@ export class Backend {
 
     modifyDevice(deviceId: string, request: UpdateDeviceRequest): Promise<void> {
         return new Promise<void>((resolve, reject) => {
-            this.backend.put(`/api/v2/devices/${deviceId}`, request).then((value) => {
+            this._fetch.put(`/api/v2/devices/${deviceId}`, request).then((value) => {
                 resolve();
             }).catch(reason => {
                 reject(reason);
@@ -1434,7 +1413,7 @@ export class Backend {
 
     updateLease(leaseId: string, update: UpdateLeaseRequest): Promise<void> {
         return new Promise<void>((resolve, reject) => {
-            this.backend.put(`/api/v1/leases/${leaseId}`, update).then(() => {
+            this._fetch.put(`/api/v1/leases/${leaseId}`, update).then(() => {
                 resolve();
             }).catch(reason => {
                 reject(reason);
@@ -1445,7 +1424,7 @@ export class Backend {
 
     deleteDevice(deviceId: string): Promise<void> {
         return new Promise<void>((resolve, reject) => {
-            this.backend.delete(`/api/v2/devices/${deviceId}`).then((response) => {
+            this._fetch.delete(`/api/v2/devices/${deviceId}`).then((response) => {
                 resolve();
             }).catch(reason => {
                 reject(reason);
@@ -1456,7 +1435,7 @@ export class Backend {
     getDevicePlatforms(): Promise<GetDevicePlatformResponse[]> {
 
         return new Promise<GetDevicePlatformResponse[]>((resolve, reject) => {
-            this.backend.get(`/api/v2/devices/platforms`).then((value) => {
+            this._fetch.get(`/api/v2/devices/platforms`).then((value) => {
                 resolve(value.data as GetDevicePlatformResponse[]);
             }).catch(reason => {
                 reject(reason);
@@ -1468,7 +1447,7 @@ export class Backend {
     getDevicePools(): Promise<GetDevicePoolResponse[]> {
 
         return new Promise<GetDevicePoolResponse[]>((resolve, reject) => {
-            this.backend.get(`/api/v2/devices/pools`).then((value) => {
+            this._fetch.get(`/api/v2/devices/pools`).then((value) => {
                 resolve(value.data as GetDevicePoolResponse[]);
             }).catch(reason => {
                 reject(reason);
@@ -1480,7 +1459,7 @@ export class Backend {
     getDeviceReservations(): Promise<GetDeviceReservationResponse[]> {
 
         return new Promise<GetDeviceReservationResponse[]>((resolve, reject) => {
-            this.backend.get(`/api/v2/devices/reservations`).then((value) => {
+            this._fetch.get(`/api/v2/devices/reservations`).then((value) => {
                 resolve(value.data as GetDeviceReservationResponse[]);
             }).catch(reason => {
                 reject(reason);
@@ -1491,7 +1470,7 @@ export class Backend {
 
     getServerSettings(): Promise<GetServerSettingsResponse> {
         return new Promise<GetServerSettingsResponse>((resolve, reject) => {
-            this.backend.get(`/api/v1/config/serversettings`).then((value) => {
+            this._fetch.get(`/api/v1/config/serversettings`).then((value) => {
                 resolve(value.data as GetServerSettingsResponse);
             }).catch(reason => {
                 reject(reason);
@@ -1501,7 +1480,7 @@ export class Backend {
 
     updateServerSettings(request: UpdateServerSettingsRequest): Promise<ServerUpdateResponse> {
         return new Promise<ServerUpdateResponse>((resolve, reject) => {
-            this.backend.put(`/api/v1/config/serversettings`, request).then((value) => {
+            this._fetch.put(`/api/v1/config/serversettings`, request).then((value) => {
                 resolve(value.data as ServerUpdateResponse);
             }).catch(reason => {
                 reject(reason);
@@ -1511,7 +1490,7 @@ export class Backend {
 
     getGlobalConfig(): Promise<GlobalConfig> {
         return new Promise<GlobalConfig>((resolve, reject) => {
-            this.backend.get(`/api/v1/config/global`).then((value) => {
+            this._fetch.get(`/api/v1/config/global`).then((value) => {
                 resolve(value.data as GlobalConfig);
             }).catch(reason => {
                 reject(reason);
@@ -1522,7 +1501,7 @@ export class Backend {
     // updates global configuation
     updateGlobalConfig(request: UpdateGlobalConfigRequest): Promise<void> {
         return new Promise<void>((resolve, reject) => {
-            this.backend.put(`/api/v1/config/global`, request).then(() => {
+            this._fetch.put(`/api/v1/config/global`, request).then(() => {
                 resolve();
             }).catch(reason => {
                 reject(reason);
@@ -1532,7 +1511,7 @@ export class Backend {
 
     getServerInfo(): Promise<GetServerInfoResponse> {
         return new Promise<GetServerInfoResponse>((resolve, reject) => {
-            this.backend.get(`/api/v1/server/info`).then((value) => {
+            this._fetch.get(`/api/v1/server/info`).then((value) => {
                 resolve(value.data as GetServerInfoResponse);
             }).catch(reason => {
                 reject(reason);
@@ -1543,7 +1522,7 @@ export class Backend {
     // create a new notice
     createNotice(request: CreateNoticeRequest): Promise<void> {
         return new Promise<void>((resolve, reject) => {
-            this.backend.post(`/api/v1/notices`, request).then(() => {
+            this._fetch.post(`/api/v1/notices`, request).then(() => {
                 resolve();
             }).catch(reason => {
                 reject(reason);
@@ -1554,7 +1533,7 @@ export class Backend {
     // get all notices
     getNotices(): Promise<GetNoticeResponse[]> {
         return new Promise<GetNoticeResponse[]>((resolve, reject) => {
-            this.backend.get(`/api/v1/notices`).then((value) => {
+            this._fetch.get(`/api/v1/notices`).then((value) => {
                 resolve((value.data as GetNoticeResponse[]).map(notice => {
                     notice.startTime = notice.startTime ? new Date(notice.startTime as string) : undefined;
                     notice.finishTime = notice.finishTime ? new Date(notice.finishTime as string) : undefined;
@@ -1569,7 +1548,7 @@ export class Backend {
     // get tools
     getTools(): Promise<GetToolSummaryResponse[]> {
         return new Promise<GetToolSummaryResponse[]>((resolve, reject) => {
-            this.backend.get(`/api/v1/tools`).then((value) => {
+            this._fetch.get(`/api/v1/tools`).then((value) => {
                 resolve(value.data?.tools as GetToolSummaryResponse[]);
             }).catch(reason => {
                 reject(reason);
@@ -1579,7 +1558,7 @@ export class Backend {
 
     getTool(toolId: string): Promise<GetToolResponse> {
         return new Promise<GetToolResponse>((resolve, reject) => {
-            this.backend.get(`/api/v1/tools/${toolId}`).then((value) => {
+            this._fetch.get(`/api/v1/tools/${toolId}`).then((value) => {
                 resolve(value.data as GetToolResponse);
             }).catch(reason => {
                 reject(reason);
@@ -1590,7 +1569,7 @@ export class Backend {
     /// Updates the state of an active deployment.
     updateDeployment(toolId: string, deploymentId: string, update: UpdateDeploymentRequest): Promise<void> {
         return new Promise<void>((resolve, reject) => {
-            this.backend.patch(`/api/v1/tools/${toolId}/deployments/${deploymentId}`, update).then(() => {
+            this._fetch.patch(`/api/v1/tools/${toolId}/deployments/${deploymentId}`, update).then(() => {
                 resolve();
             }).catch(reason => {
                 reject(reason);
@@ -1600,7 +1579,7 @@ export class Backend {
     // update a notice
     updateNotice(request: UpdateNoticeRequest): Promise<void> {
         return new Promise<void>((resolve, reject) => {
-            this.backend.put(`/api/v1/notices`, request).then(() => {
+            this._fetch.put(`/api/v1/notices`, request).then(() => {
                 resolve();
             }).catch(reason => {
                 reject(reason);
@@ -1611,7 +1590,7 @@ export class Backend {
     // delete a notice
     deleteNotice(id: string): Promise<void> {
         return new Promise<void>((resolve, reject) => {
-            this.backend.delete(`/api/v1/notices/${id}`).then(() => {
+            this._fetch.delete(`/api/v1/notices/${id}`).then(() => {
                 resolve();
             }).catch(reason => {
                 reject(reason);
@@ -1622,7 +1601,7 @@ export class Backend {
     // create job bisection
     createBisectTask(create: CreateBisectTaskRequest): Promise<CreateBisectTaskResponse> {
         return new Promise<CreateBisectTaskResponse>((resolve, reject) => {
-            this.backend.post(`/api/v1/bisect`, create).then((response) => {
+            this._fetch.post(`/api/v1/bisect`, create).then((response) => {
                 resolve(response.data);
             }).catch(reason => {
                 reject(reason);
@@ -1633,7 +1612,7 @@ export class Backend {
     // get a bisection task
     getBisectTask(id: string): Promise<GetBisectTaskResponse> {
         return new Promise<GetBisectTaskResponse>((resolve, reject) => {
-            this.backend.get(`/api/v1/bisect/${id}`).then((value) => {
+            this._fetch.get(`/api/v1/bisect/${id}`).then((value) => {
                 resolve(value.data as GetBisectTaskResponse);
             }).catch(reason => {
                 reject(reason);
@@ -1644,7 +1623,7 @@ export class Backend {
     getBisections(query: { id?: string[], ownerId?: string, jobId?: string, minCreateTime?: string, maxCreateTime?: string, index?: number, count?: number }) {
 
         return new Promise<GetBisectTaskResponse[]>((resolve, reject) => {
-            this.backend.get(`/api/v1/bisect`, { params: query }).then((value) => {
+            this._fetch.get(`/api/v1/bisect`, { params: query }).then((value) => {
                 resolve(value.data as GetBisectTaskResponse[]);
             }).catch(reason => {
                 reject(reason);
@@ -1655,7 +1634,7 @@ export class Backend {
     // update a visection task
     updateBisectTask(id: string, request: UpdateBisectTaskRequest): Promise<void> {
         return new Promise<void>((resolve, reject) => {
-            this.backend.patch(`/api/v1/bisect/${id}`, request).then(() => {
+            this._fetch.patch(`/api/v1/bisect/${id}`, request).then(() => {
                 resolve();
             }).catch(reason => {
                 reject(reason);
@@ -1666,7 +1645,7 @@ export class Backend {
     // get a job's bisection tasks
     getJobBisectTasks(jobId: string): Promise<GetBisectTaskResponse[]> {
         return new Promise<GetBisectTaskResponse[]>((resolve, reject) => {
-            this.backend.get(`/api/v1/bisect/job/${jobId}`).then((value) => {
+            this._fetch.get(`/api/v1/bisect/job/${jobId}`).then((value) => {
                 resolve(value.data as GetBisectTaskResponse[]);
             }).catch(reason => {
                 reject(reason);
@@ -1677,7 +1656,7 @@ export class Backend {
     checkPreflightConfig(shelvedChange: number): Promise<PreflightConfigResponse> {
 
         return new Promise<PreflightConfigResponse>((resolve, reject) => {
-            this.backend.post(`/api/v1/server/preflightconfig`, { shelvedChange: shelvedChange }).then((value) => {
+            this._fetch.post(`/api/v1/server/preflightconfig`, { shelvedChange: shelvedChange }).then((value) => {
                 resolve(value.data as PreflightConfigResponse);
             }).catch(reason => {
                 reject(reason);
@@ -1690,7 +1669,7 @@ export class Backend {
 
     getAccountEntitlements(): Promise<any> {
         return new Promise<any>((resolve, reject) => {
-            this.backend.get(`/account/entitlements`).then((value) => {
+            this._fetch.get(`/account/entitlements`).then((value) => {
                 resolve(value.data as any);
             }).catch(reason => {
                 reject(reason);
@@ -1701,7 +1680,7 @@ export class Backend {
     // update current account 
     updateCurrentAccount(request: UpdateCurrentAccountRequest): Promise<boolean> {
         return new Promise<boolean>((resolve, reject) => {
-            this.backend.put(`/api/v1/accounts/current`, request).then(() => {
+            this._fetch.put(`/api/v1/accounts/current`, request).then(() => {
                 resolve(true);
             }).catch((reason) => {
                 reject(reason);
@@ -1711,7 +1690,7 @@ export class Backend {
 
     getAccounts(): Promise<GetAccountResponse[]> {
         return new Promise<GetAccountResponse[]>((resolve, reject) => {
-            this.backend.get(`/api/v1/accounts`).then((value) => {
+            this._fetch.get(`/api/v1/accounts`).then((value) => {
                 resolve(value.data as GetAccountResponse[]);
             }).catch(reason => {
                 reject(reason);
@@ -1721,7 +1700,7 @@ export class Backend {
 
     getAccountGroups(): Promise<AccountClaimMessage[]> {
         return new Promise<AccountClaimMessage[]>((resolve, reject) => {
-            this.backend.get(`/api/v1/dashboard/accountgroups`).then((value) => {
+            this._fetch.get(`/api/v1/dashboard/accountgroups`).then((value) => {
                 resolve(value.data as AccountClaimMessage[]);
             }).catch(reason => {
                 reject(reason);
@@ -1732,7 +1711,7 @@ export class Backend {
     // create a new account
     createAccount(request: CreateAccountRequest): Promise<void> {
         return new Promise<void>((resolve, reject) => {
-            this.backend.post(`/api/v1/accounts`, request).then(() => {
+            this._fetch.post(`/api/v1/accounts`, request).then(() => {
                 resolve();
             }).catch(reason => {
                 reject(reason);
@@ -1743,7 +1722,7 @@ export class Backend {
     // delete an account
     deleteAccount(id: string): Promise<void> {
         return new Promise<void>((resolve, reject) => {
-            this.backend.delete(`/api/v1/accounts/${id}`).then(() => {
+            this._fetch.delete(`/api/v1/accounts/${id}`).then(() => {
                 resolve();
             }).catch(reason => {
                 reject(reason);
@@ -1754,7 +1733,7 @@ export class Backend {
     // update an account
     updateAccount(id: string, request: UpdateAccountRequest): Promise<void> {
         return new Promise<void>((resolve, reject) => {
-            this.backend.put(`/api/v1/accounts/${id}`, request).then(() => {
+            this._fetch.put(`/api/v1/accounts/${id}`, request).then(() => {
                 resolve();
             }).catch(reason => {
                 reject(reason);
@@ -1766,7 +1745,7 @@ export class Backend {
 
     getServiceAccounts(): Promise<GetServiceAccountResponse[]> {
         return new Promise<GetServiceAccountResponse[]>((resolve, reject) => {
-            this.backend.get(`/api/v1/serviceaccounts`).then((value) => {
+            this._fetch.get(`/api/v1/serviceaccounts`).then((value) => {
                 resolve(value.data as GetServiceAccountResponse[]);
             }).catch(reason => {
                 reject(reason);
@@ -1777,7 +1756,7 @@ export class Backend {
     // create a new account
     createServiceAccount(request: CreateServiceAccountRequest): Promise<CreateServiceAccountResponse> {
         return new Promise<CreateServiceAccountResponse>((resolve, reject) => {
-            this.backend.post(`/api/v1/serviceaccounts`, request).then((value) => {
+            this._fetch.post(`/api/v1/serviceaccounts`, request).then((value) => {
                 resolve(value.data as CreateServiceAccountResponse);
             }).catch(reason => {
                 reject(reason);
@@ -1788,7 +1767,7 @@ export class Backend {
     // delete an account
     deleteServiceAccount(id: string): Promise<void> {
         return new Promise<void>((resolve, reject) => {
-            this.backend.delete(`/api/v1/serviceaccounts/${id}`).then(() => {
+            this._fetch.delete(`/api/v1/serviceaccounts/${id}`).then(() => {
                 resolve();
             }).catch(reason => {
                 reject(reason);
@@ -1799,7 +1778,7 @@ export class Backend {
     // update an account
     updateServiceAccount(id: string, request: UpdateServiceAccountRequest): Promise<UpdateServiceAccountResponse> {
         return new Promise<UpdateServiceAccountResponse>((resolve, reject) => {
-            this.backend.put(`/api/v1/serviceaccounts/${id}`, request).then((value) => {
+            this._fetch.put(`/api/v1/serviceaccounts/${id}`, request).then((value) => {
                 resolve(value.data as UpdateServiceAccountResponse);
             }).catch(reason => {
                 reject(reason);
@@ -1827,8 +1806,8 @@ export class Backend {
 
     async serverLogout(redirect: string) {
         try {
-            this.backend.logout = true;
-            await this.backend.get("/api/v1/dashboard/logout", { params: { dashboard: true } });
+            this._fetch.logout = true;
+            await this._fetch.get("/api/v1/dashboard/logout", { params: { dashboard: true } });
             window.location.assign(redirect);
 
         } catch (err) {
@@ -1860,22 +1839,27 @@ export class Backend {
 
     }
 
+    get enabledPlugins(): ServerPluginInfoResponse[] {
+        return this.serverInfo?.plugins?.filter(p => p.loaded) ?? [];
+    }
 
     init() {
 
 
         return new Promise<boolean>(async (resolve, reject) => {
 
-            this.backend.setBaseUrl(this.serverUrl);
-            this.backend.setDebugToken(this.debugToken);
+            this._fetch.setBaseUrl(this.serverUrl);
+            this._fetch.setDebugToken(this.debugToken);
 
-            const challenge = await this.backend.challenge();
+            const challenge = await this._fetch.challenge();
 
             if (challenge === ChallengeStatus.Unauthorized) {
-                this.backend.login(window.location.toString());
+                this._fetch.login(window.location.toString());
                 return;
             }
 
+            this.serverInfo = await this.getServerInfo();
+
             await dashboard.update();
 
             if (dashboard.localCache) {
@@ -1891,14 +1875,20 @@ export class Backend {
         });
     }
 
+    get fetch() {
+        return this._fetch;
+    }
+
     constructor() {
 
-        this.backend = new Fetch();
+        this._fetch = new Fetch();
     }
 
+    serverInfo?: GetServerInfoResponse;
+
     updateID?: any;
     logout: boolean = false;
 
-    private backend: Fetch;
+    private _fetch: Fetch;
 
 }
diff --git a/Engine/Source/Programs/Horde/HordeDashboard/src/backend/Dashboard.ts b/Engine/Source/Programs/Horde/HordeDashboard/src/backend/Dashboard.ts
--- a/Engine/Source/Programs/Horde/HordeDashboard/src/backend/Dashboard.ts
+++ b/Engine/Source/Programs/Horde/HordeDashboard/src/backend/Dashboard.ts
@@ -2,7 +2,7 @@
 
 import { action, makeObservable, observable } from 'mobx';
 import backend from '.';
-import { AuthMethod, DashboardPreference, GetDashboardAgentCategoryResponse, GetDashboardConfigResponse, GetDashboardPoolCategoryResponse, GetJobTemplateSettingsResponse, GetTelemetryViewResponse, GetUserResponse, UserClaim } from './Api';
+import { AuthMethod, DashboardPreference, GetDashboardAgentCategoryResponse, GetDashboardConfigResponse, GetDashboardPoolCategoryResponse, GetJobTemplateSettingsResponse, GetUserResponse, UserClaim } from './Api';
 import { getSiteConfig } from './Config';
 
 export enum StatusColor {
@@ -656,13 +656,6 @@ export class Dashboard {
         return this.data;
     }
 
-    get telemetryViews(): GetTelemetryViewResponse[] {
-        if (!this.config?.telemetryViews) {
-            return [];
-        }
-        return this.config?.telemetryViews;
-    }
-
     private async postPreferences(reload?: boolean): Promise<boolean> {
 
         // cancel any pending        
diff --git a/Engine/Source/Programs/Horde/HordeDashboard/src/components/TopNav.tsx b/Engine/Source/Programs/Horde/HordeDashboard/src/components/TopNav.tsx
--- a/Engine/Source/Programs/Horde/HordeDashboard/src/components/TopNav.tsx
+++ b/Engine/Source/Programs/Horde/HordeDashboard/src/components/TopNav.tsx
@@ -13,7 +13,7 @@ import { VersionModal } from './VersionModal';
 import { getHordeTheme } from '../styles/theme';
 import { getHordeStyling } from '../styles/Styles';
 import { ArtifactQueryState, FindArtifactsModal } from './artifacts/ArtifactsSearch';
-
+import { getHordePlugins, MountType } from 'hordePlugins';
 
 const getStyles = () => {
 
@@ -438,13 +438,18 @@ export const TopNav: React.FC<{ suppressServer?: boolean }> = observer(({ suppre
 
       const features = dashboard.user?.dashboardFeatures;
 
-      const serviceItems: IContextualMenuItem[] = [];
+      let serviceItems: IContextualMenuItem[] = [];
 
-      serviceItems.push({
-         key: "admin_analytics",
-         text: "Analytics",
-         link: `/analytics`
-      });
+      // get the plugins which mount into the tools bar
+      const plugins = getHordePlugins().filter(p => p.mount?.type === MountType.TopNav);
+
+      plugins.forEach(p => {
+         serviceItems.push({
+            key: `admin_plugin_${p.mount!.text}`,
+            text: p.mount!.text,
+            link: p.mount!.route
+         });   
+      })
 
       if (features?.showTests !== false) {
          serviceItems.push({
@@ -468,6 +473,11 @@ export const TopNav: React.FC<{ suppressServer?: boolean }> = observer(({ suppre
          },
       });
 
+      serviceItems = serviceItems.sort((a, b) => {
+         return (a.text ?? "").localeCompare(b.text ?? "")
+      })
+
+      // Always put downloads at the button for easy navigation
       serviceItems.push({
          key: "software_tools",
          text: "Downloads",
diff --git a/Engine/Source/Programs/Horde/HordeDashboard/src/components/telemetry/TelemetryData.ts b/Engine/Source/Programs/Horde/HordeDashboard/src/components/telemetry/TelemetryData.ts
deleted file mode 100644
diff --git a/Engine/Source/Programs/Horde/HordeDashboard/src/components/telemetry/TelemetryLineGraph.ts b/Engine/Source/Programs/Horde/HordeDashboard/src/components/telemetry/TelemetryLineGraph.ts
deleted file mode 100644
diff --git a/Engine/Source/Programs/Horde/HordeDashboard/src/components/telemetry/TelemetryView.tsx b/Engine/Source/Programs/Horde/HordeDashboard/src/components/telemetry/TelemetryView.tsx
deleted file mode 100644
diff --git a/Engine/Source/Programs/Horde/HordeDashboard/src/legacyPlugins/index.tsx b/Engine/Source/Programs/Horde/HordeDashboard/src/legacyPlugins/index.tsx
--- a/Engine/Source/Programs/Horde/HordeDashboard/src/legacyPlugins/index.tsx
+++ b/Engine/Source/Programs/Horde/HordeDashboard/src/legacyPlugins/index.tsx
@@ -19,7 +19,7 @@ export class Plugins {
 
    loadPlugins(pluginList: string[] | undefined): void {
 
-      console.log("loading plugins...");
+      console.log("Loading Legacy Plugins");
 
       this.plugins.push(...[AutomatedTestSessionPlugin, SimpleTestReportPlugin, UnrealAutomatedTestsPlugin]);
 
diff --git a/Engine/Source/Programs/Horde/HordeDashboard/tsconfig.json b/Engine/Source/Programs/Horde/HordeDashboard/tsconfig.json
--- a/Engine/Source/Programs/Horde/HordeDashboard/tsconfig.json
+++ b/Engine/Source/Programs/Horde/HordeDashboard/tsconfig.json
@@ -17,7 +17,13 @@
     "resolveJsonModule": true,
     "isolatedModules": true,
     "noEmit": true,
-    "jsx": "react-jsx",    
+    "jsx": "react-jsx",
+    
+    "paths": {
+      "horde/*" : ["./src/*"],
+      "hordePlugins": ["./plugins"],
+      "hordePlugins/*": ["./plugins/*"]
+    },
 
     /* Linting */
     "strict": false,
@@ -28,6 +34,6 @@
     "noUnusedParameters": true,    
     */
   },
-  "include": ["src"],
+  "include": ["src", "plugins"],
   "references": [{ "path": "./tsconfig.node.json" }]
 }
diff --git a/Engine/Source/Programs/Horde/HordeDashboard/vite.config.ts b/Engine/Source/Programs/Horde/HordeDashboard/vite.config.ts
--- a/Engine/Source/Programs/Horde/HordeDashboard/vite.config.ts
+++ b/Engine/Source/Programs/Horde/HordeDashboard/vite.config.ts
@@ -1,4 +1,5 @@
 import { defineConfig } from 'vite'
+import tsconfigPaths from 'vite-tsconfig-paths'
 import react from '@vitejs/plugin-react'
 
 const proxyTarget = "http://127.0.0.1:13340"
@@ -14,14 +15,15 @@ export default defineConfig({
         plugins: ['decorators-legacy'],
       },
     },
-  })],
+  }),
+  tsconfigPaths()],
   build: {
     chunkSizeWarningLimit: 8192,
     rollupOptions: {
       output: {
-          entryFileNames: `[name].${cacheBreak}.js`,
-          chunkFileNames: `[name].${cacheBreak}.js`,
-          assetFileNames: `[name].${cacheBreak}.[ext]`
+        entryFileNames: `[name].${cacheBreak}.js`,
+        chunkFileNames: `[name].${cacheBreak}.js`,
+        assetFileNames: `[name].${cacheBreak}.[ext]`
       }
     }
   },
