Horde: Fix analytic plugin initial load

diff --git a/Engine/Source/Programs/Horde/HordeDashboard/plugins/analytics/telemetryView.tsx b/Engine/Source/Programs/Horde/HordeDashboard/plugins/analytics/telemetryView.tsx
--- a/Engine/Source/Programs/Horde/HordeDashboard/plugins/analytics/telemetryView.tsx
+++ b/Engine/Source/Programs/Horde/HordeDashboard/plugins/analytics/telemetryView.tsx
@@ -369,9 +369,12 @@ class MetricsHandler {
       return value;
    }
 
-   async initialize(viewId?: string, categoryIn?: string) {
+   async initialize(viewId?: string, categoryIn?: string, callback?: () => void) {
 
       if (this.initialized || (viewId && this.view?.id === viewId)) {
+         if (callback) {
+            callback();
+         }
          return;
       }
 
@@ -383,6 +386,9 @@ class MetricsHandler {
 
       if (!this.allViews.length) {
          console.log("No telemetry views configured");
+         if (callback) {
+            callback();
+         }
          return;
       }
 
@@ -431,6 +437,11 @@ class MetricsHandler {
       }
 
       this.setUpdated();
+
+      if (callback) {
+         callback();
+      }
+
    }
 
    async setCategory(category?: string, fromUser?: boolean) {
@@ -661,7 +672,7 @@ class MetricsHandler {
 
    search: URLSearchParams = new URLSearchParams();
 
-   get allViews() { return this.telemetryViews ?? []}
+   get allViews() { return this.telemetryViews ?? [] }
 
    private telemetryViews?: GetTelemetryViewResponse[];
 
@@ -1396,8 +1407,10 @@ export const SearchUpdate: React.FC = observer(() => {
 
 export const TelemetryView: React.FC = () => {
 
+   const [initialized, setInitialized] = useState(false);
+
    useEffect(() => {
-      handler.initialize();
+      handler.initialize(undefined, undefined, () => setInitialized(true));
       return () => {
          handler.clear();
       };
@@ -1422,11 +1435,12 @@ export const TelemetryView: React.FC = () => {
       <Breadcrumbs items={[{ text: 'Analytics' }]} />
       <Stack horizontal styles={{ root: { backgroundColor: modeColors.background } }}>
          <Stack styles={{ root: { width: "100%" } }}>
-            {!handler.allViews.length && <Stack horizontal tokens={{ childrenGap: 6 }} horizontalAlign="center" style={{ paddingTop: 30 }}>
+            {!initialized && <Stack style={{ paddingTop: 24 }}><Spinner size={SpinnerSize.large} /></Stack>}
+            {initialized && !handler.allViews.length && <Stack horizontal tokens={{ childrenGap: 6 }} horizontalAlign="center" style={{ paddingTop: 30 }}>
                <Text variant="mediumPlus">No analytic views found, for more information please see</Text>
                <a href={telemetrybDocs} style={{ fontSize: "18px", "cursor": "pointer" }} onClick={(ev) => { ev.preventDefault(); ev.stopPropagation(); navigate(telemetrybDocs) }}> Horde analytics documentation.</a>
             </Stack>}
-            {!!handler.allViews.length && <Stack>
+            {initialized && !!handler.allViews.length && <Stack>
                <Stack horizontal>
                   <Stack key={`${key}_1`} style={{ paddingLeft: centerAlign }} />
                   <Stack style={{ width: rootWidth - 8, maxWidth: windowSize.width - 12, paddingLeft: 0, paddingTop: 24, paddingBottom: 24, paddingRight: 0 }} >
