Horde: Make text field entries read only instead of disabled so they can be readily copies
#skipci

diff --git a/Engine/Source/Programs/Horde/HordeDashboard/src/components/build/NewBuildV2.tsx b/Engine/Source/Programs/Horde/HordeDashboard/src/components/build/NewBuildV2.tsx
--- a/Engine/Source/Programs/Horde/HordeDashboard/src/components/build/NewBuildV2.tsx
+++ b/Engine/Source/Programs/Horde/HordeDashboard/src/components/build/NewBuildV2.tsx
@@ -1054,7 +1054,7 @@ const TextParameter: React.FC<{ param: TextParameterData, disabled?: boolean }>
          spellCheck={false}
          autoComplete="off"
          value={value}
-         disabled={options.readOnly || disabled}
+         readOnly={options.readOnly || disabled}
          onChange={(ev, value) => {
             options.onTextChanged(param.id, value ?? "");
          }}
