Horde: Move session update and notification logic into AgentCollection.

#rnx

diff --git a/Engine/Source/Programs/Horde/HordeServer/Startup.cs b/Engine/Source/Programs/Horde/HordeServer/Startup.cs
--- a/Engine/Source/Programs/Horde/HordeServer/Startup.cs
+++ b/Engine/Source/Programs/Horde/HordeServer/Startup.cs
@@ -358,8 +358,6 @@ namespace HordeServer
 			services.AddHostedService(provider => provider.GetRequiredService<ConfigService>());
 
 			// Auditing
-			services.AddSingleton<IAuditLog<AgentId>>(sp => sp.GetRequiredService<IAuditLogFactory<AgentId>>().Create("Agents.Log", "AgentId"));
-
 			services.AddSingleton(typeof(IAuditLogFactory<>), typeof(AuditLogFactory<>));
 			services.AddSingleton(typeof(ISingletonDocument<>), typeof(SingletonDocument<>));
 
diff --git a/Engine/Source/Programs/Horde/Plugins/Build/HordeServer.Build.Tests/Agents/AgentCollectionTests.cs b/Engine/Source/Programs/Horde/Plugins/Build/HordeServer.Build.Tests/Agents/AgentCollectionTests.cs
deleted file mode 100644
diff --git a/Engine/Source/Programs/Horde/Plugins/Build/HordeServer.Build.Tests/BuildTestSetup.cs b/Engine/Source/Programs/Horde/Plugins/Build/HordeServer.Build.Tests/BuildTestSetup.cs
--- a/Engine/Source/Programs/Horde/Plugins/Build/HordeServer.Build.Tests/BuildTestSetup.cs
+++ b/Engine/Source/Programs/Horde/Plugins/Build/HordeServer.Build.Tests/BuildTestSetup.cs
@@ -123,8 +123,6 @@ namespace HordeServer.Tests
 
 			services.AddHttpClient<RpcService>();
 
-			services.AddSingleton<IAuditLog<AgentId>>(sp => sp.GetRequiredService<IAuditLogFactory<AgentId>>().Create("Agents.Log", "AgentId"));
-
 			services.AddSingleton<IAgentCollection, AgentCollection>();
 			services.AddSingleton<IAgentTelemetryCollection, AgentTelemetryCollection>();
 			services.AddSingleton<ArtifactCollection>();
diff --git a/Engine/Source/Programs/Horde/Plugins/Build/HordeServer.Build/Issues/IssueService.cs b/Engine/Source/Programs/Horde/Plugins/Build/HordeServer.Build/Issues/IssueService.cs
--- a/Engine/Source/Programs/Horde/Plugins/Build/HordeServer.Build/Issues/IssueService.cs
+++ b/Engine/Source/Programs/Horde/Plugins/Build/HordeServer.Build/Issues/IssueService.cs
@@ -17,7 +17,6 @@ using HordeServer.Commits;
 using HordeServer.Jobs;
 using HordeServer.Jobs.Graphs;
 using HordeServer.Logs;
-using HordeServer.Server;
 using HordeServer.Streams;
 using HordeServer.Users;
 using Microsoft.Extensions.DependencyInjection;
diff --git a/Engine/Source/Programs/Horde/Plugins/Compute/HordeServer.Compute.Tests/Agents/AgentCollectionTests.cs b/Engine/Source/Programs/Horde/Plugins/Compute/HordeServer.Compute.Tests/Agents/AgentCollectionTests.cs
new file mode 100644
--- /dev/null
+++ b/Engine/Source/Programs/Horde/Plugins/Compute/HordeServer.Compute.Tests/Agents/AgentCollectionTests.cs
@@ -0,0 +1,175 @@
+// Copyright Epic Games, Inc. All Rights Reserved.
+
+using EpicGames.Horde.Agents;
+using EpicGames.Horde.Agents.Leases;
+using EpicGames.Horde.Agents.Pools;
+using Google.Protobuf.WellKnownTypes;
+using HordeCommon.Rpc.Messages;
+using HordeServer.Agents;
+
+namespace HordeServer.Tests.Agents;
+
+[TestClass]
+public class AgentCollectionTests : ComputeTestSetup
+{
+	private IAgent _agent;
+	private readonly CreateLeaseOptions _lease1;
+	private readonly CreateLeaseOptions _lease2;
+	private readonly CreateLeaseOptions _leaseWithParent3;
+	private readonly CreateLeaseOptions _leaseWithParent4;
+
+	public AgentCollectionTests()
+	{
+		_agent = AgentCollection.AddAsync(new AgentId("test"), ephemeral: true, enrollmentKey: "").Result;
+		_lease1 = new CreateLeaseOptions(LeaseId.Parse("aaaaaaaaaaaaaaaaaaaaaaaa"), null, "lease", null, null, null, null, false, new Empty());
+		_lease2 = new CreateLeaseOptions(LeaseId.Parse("bbbbbbbbbbbbbbbbbbbbbbbb"), null, "lease", null, null, null, null, false, new Empty());
+		_leaseWithParent3 = new CreateLeaseOptions(LeaseId.Parse("cccccccccccccccccccccccc"), _lease1.Id, "leaseWithParent3", null, null, null, null, false, new Empty());
+		_leaseWithParent4 = new CreateLeaseOptions(LeaseId.Parse("dddddddddddddddddddddddd"), _lease1.Id, "leaseWithParent4", null, null, null, null, false, new Empty());
+	}
+
+	[TestMethod]
+	public async Task AddLeaseAsync()
+	{
+		await _agent.TryCreateLeaseAsync(_lease1);
+		await UpdateAgentAsync();
+
+		await _agent.TryCreateLeaseAsync(_leaseWithParent3);
+		await UpdateAgentAsync();
+
+		await _agent.TryCreateLeaseAsync(_leaseWithParent4);
+		await UpdateAgentAsync();
+
+		List<LeaseId> leases = await AgentCollection.FindActiveLeaseIdsAsync();
+
+		Assert.AreEqual(3, leases.Count);
+		Assert.IsTrue(leases.Contains(_lease1.Id));
+		Assert.IsTrue(leases.Contains(_leaseWithParent3.Id));
+		Assert.IsTrue(leases.Contains(_leaseWithParent4.Id));
+	}
+
+	[TestMethod]
+	public async Task StartSessionAsync()
+	{
+		await _agent.TryCreateLeaseAsync(_lease1);
+		await UpdateAgentAsync();
+
+		await _agent.TryCreateLeaseAsync(_lease2);
+		await UpdateAgentAsync();
+
+		await _agent.TryCreateSessionAsync(new CreateSessionOptions(AgentStatus.Ok,
+			new RpcAgentCapabilities(), new List<PoolId>(), DateTime.UtcNow, null));
+
+		List<LeaseId> leases = await AgentCollection.FindActiveLeaseIdsAsync();
+
+		Assert.AreEqual(0, leases.Count);
+	}
+
+	[TestMethod]
+	public async Task UpdateSession_WithEmptyLeasesAsync()
+	{
+		await _agent.TryCreateLeaseAsync(_lease1);
+		await UpdateAgentAsync();
+
+		await _agent.TryCreateLeaseAsync(_lease2);
+		await UpdateAgentAsync();
+
+		await _agent.TryUpdateSessionAsync(new UpdateSessionOptions { Leases = new List<RpcLease>() });
+
+		List<LeaseId> leases = await AgentCollection.FindActiveLeaseIdsAsync();
+
+		Assert.AreEqual(0, leases.Count);
+	}
+
+	[TestMethod]
+	public async Task UpdateSession_WithNewLeasesAsync()
+	{
+		await _agent.TryCreateLeaseAsync(_lease1);
+		await UpdateAgentAsync();
+
+		IAgent? newAgent = await _agent.TryCreateLeaseAsync(_lease1);
+		Assert.IsNotNull(newAgent);
+
+		newAgent = await _agent.TryCreateLeaseAsync(_lease2);
+		Assert.IsNotNull(newAgent);
+
+		List<LeaseId> leases = await AgentCollection.FindActiveLeaseIdsAsync();
+
+		Assert.AreEqual(2, leases.Count);
+		Assert.IsTrue(leases.Contains(_lease1.Id));
+		Assert.IsTrue(leases.Contains(_lease2.Id));
+	}
+
+	[TestMethod]
+	public async Task UpdateSession_WithOneLeaseRemovedAsync()
+	{
+		await _agent.TryCreateLeaseAsync(_lease1);
+		await UpdateAgentAsync();
+
+		await _agent.TryCreateLeaseAsync(_lease2);
+		await UpdateAgentAsync();
+
+		await _agent.TryUpdateSessionAsync(new UpdateSessionOptions { Leases = new List<RpcLease> { new RpcLease { Id = _lease1.Id } } });
+
+		List<LeaseId> leases = await AgentCollection.FindActiveLeaseIdsAsync();
+
+		Assert.AreEqual(1, leases.Count);
+		Assert.IsTrue(leases.Contains(_lease1.Id));
+	}
+
+	[TestMethod]
+	public async Task CancelLeaseAsync()
+	{
+		await _agent.TryCreateLeaseAsync(_lease1);
+		await UpdateAgentAsync();
+
+		await _agent.TryCreateLeaseAsync(_lease2);
+		await UpdateAgentAsync();
+
+		await _agent.TryCancelLeaseAsync(0);
+
+		List<LeaseId> leases = await AgentCollection.FindActiveLeaseIdsAsync();
+
+		Assert.AreEqual(1, leases.Count);
+		Assert.IsTrue(leases.Contains(_lease2.Id));
+	}
+
+	[TestMethod]
+	public async Task TerminateSessionAsync()
+	{
+		await _agent.TryCreateLeaseAsync(_lease1);
+		await UpdateAgentAsync();
+
+		await _agent.TryCreateLeaseAsync(_lease2);
+		await UpdateAgentAsync();
+
+		await _agent.TryTerminateSessionAsync();
+
+		List<LeaseId> leases = await AgentCollection.FindActiveLeaseIdsAsync();
+
+		Assert.AreEqual(0, leases.Count);
+	}
+
+	[TestMethod]
+	public async Task GetChildLeaseIdsAsync()
+	{
+		await _agent.TryCreateLeaseAsync(_lease1);
+		await UpdateAgentAsync();
+
+		await _agent.TryCreateLeaseAsync(_leaseWithParent3);
+		await UpdateAgentAsync();
+
+		await _agent.TryCreateLeaseAsync(_leaseWithParent4);
+		await UpdateAgentAsync();
+
+		List<LeaseId> leases = await AgentCollection.GetChildLeaseIdsAsync(_leaseWithParent3.ParentId!.Value);
+
+		Assert.AreEqual(2, leases.Count);
+		Assert.IsTrue(leases.Contains(_leaseWithParent3.Id));
+		Assert.IsTrue(leases.Contains(_leaseWithParent4.Id));
+	}
+
+	private async Task UpdateAgentAsync()
+	{
+		_agent = (await AgentCollection.GetAsync(_agent.Id))!;
+	}
+}
\ No newline at end of file
diff --git a/Engine/Source/Programs/Horde/Plugins/Compute/HordeServer.Compute.Tests/ComputeTestSetup.cs b/Engine/Source/Programs/Horde/Plugins/Compute/HordeServer.Compute.Tests/ComputeTestSetup.cs
--- a/Engine/Source/Programs/Horde/Plugins/Compute/HordeServer.Compute.Tests/ComputeTestSetup.cs
+++ b/Engine/Source/Programs/Horde/Plugins/Compute/HordeServer.Compute.Tests/ComputeTestSetup.cs
@@ -1,6 +1,6 @@
 // Copyright Epic Games, Inc. All Rights Reserved.
 
-using HordeServer.Acls;
+using HordeServer.Agents;
 using HordeServer.Compute;
 using HordeServer.Logs;
 using HordeServer.Storage;
@@ -11,6 +11,7 @@ namespace HordeServer.Tests
 	public class ComputeTestSetup : ServerTestSetup
 	{
 		public ComputeService ComputeService => ServiceProvider.GetRequiredService<ComputeService>();
+		public IAgentCollection AgentCollection => ServiceProvider.GetRequiredService<IAgentCollection>();
 		public StorageService StorageService => ServiceProvider.GetRequiredService<StorageService>();
 		public ILogCollection LogCollection => ServiceProvider.GetRequiredService<ILogCollection>();
 
@@ -19,12 +20,5 @@ namespace HordeServer.Tests
 			AddPlugin<StoragePlugin>();
 			AddPlugin<ComputePlugin>();
 		}
-
-		protected override void ConfigureServices(IServiceCollection services)
-		{
-			base.ConfigureServices(services);
-
-			services.AddSingleton<IDefaultAclModifier, ComputeAclModifier>();
-		}
 	}
 }
diff --git a/Engine/Source/Programs/Horde/Plugins/Compute/HordeServer.Compute/Agents/AgentCollection.cs b/Engine/Source/Programs/Horde/Plugins/Compute/HordeServer.Compute/Agents/AgentCollection.cs
--- a/Engine/Source/Programs/Horde/Plugins/Compute/HordeServer.Compute/Agents/AgentCollection.cs
+++ b/Engine/Source/Programs/Horde/Plugins/Compute/HordeServer.Compute/Agents/AgentCollection.cs
@@ -1,6 +1,7 @@
 // Copyright Epic Games, Inc. All Rights Reserved.
 
 using System.Collections.ObjectModel;
+using System.Collections.Concurrent;
 using System.Diagnostics.CodeAnalysis;
 using System.Linq.Expressions;
 using EpicGames.Core;
@@ -9,8 +10,12 @@ using EpicGames.Horde.Agents.Leases;
 using EpicGames.Horde.Agents.Pools;
 using EpicGames.Horde.Agents.Sessions;
 using EpicGames.Redis;
+using Google.Protobuf;
 using Google.Protobuf.WellKnownTypes;
+using HordeCommon.Rpc;
+using HordeCommon.Rpc.Messages;
 using HordeCommon.Rpc.Tasks;
+using HordeServer.Agents.Leases;
 using HordeServer.Agents.Sessions;
 using HordeServer.Auditing;
 using HordeServer.Server;
@@ -19,13 +24,14 @@ using MongoDB.Bson;
 using MongoDB.Bson.Serialization.Attributes;
 using MongoDB.Driver;
 using StackExchange.Redis;
+using Microsoft.Extensions.Hosting;
 
 namespace HordeServer.Agents
 {
 	/// <summary>
 	/// Collection of agent documents
 	/// </summary>
-	public class AgentCollection : IAgentCollection
+	public sealed class AgentCollection : IAgentCollection, IHostedService, IAsyncDisposable
 	{
 		class Agent : IAgent
 		{
@@ -137,6 +143,12 @@ namespace HordeServer.Agents
 				AgentDocument? newDocument = await _collection.TryUpdateWorkspacesAsync(_document, workspaces, requestConform, cancellationToken);
 				return _collection.CreateAgentObject(newDocument);
 			}
+
+			public async Task<IAgent?> WaitForUpdateAsync(CancellationToken cancellationToken = default)
+			{
+				AgentDocument? newDocument = await _collection.WaitForUpdateAsync(_document, cancellationToken);
+				return _collection.CreateAgentObject(newDocument);
+			}
 		}
 
 		/// <summary>
@@ -260,15 +272,19 @@ namespace HordeServer.Agents
 
 		readonly IMongoCollection<AgentDocument> _agentCollection;
 		readonly IMongoCollection<SessionDocument> _sessionCollection;
+		readonly ILeaseCollection _leaseCollection;
 		readonly IAuditLog<AgentId> _auditLog;
 		readonly IRedisService _redisService;
 		readonly IClock _clock;
 		readonly RedisChannel<AgentId> _updateEventChannel;
+		readonly ConcurrentDictionary<AgentId, TaskCompletionSource> _agentIdToTcs = new ConcurrentDictionary<AgentId, TaskCompletionSource>();
+		IAsyncDisposable? _updateEventSubscription;
+		readonly ILogger _logger;
 
 		/// <summary>
 		/// Constructor
 		/// </summary>
-		public AgentCollection(IMongoService mongoService, IRedisService redisService, IClock clock, IAuditLog<AgentId> auditLog)
+		public AgentCollection(IMongoService mongoService, IRedisService redisService, ILeaseCollection leaseCollection, IClock clock, IAuditLog<AgentId> auditLog, ILogger<AgentCollection> logger)
 		{
 			List<MongoIndex<AgentDocument>> agentIndexes = new List<MongoIndex<AgentDocument>>();
 			agentIndexes.Add(keys => keys.Ascending(x => x.Deleted).Ascending(x => x.Id).Ascending(x => x.Pools));
@@ -281,9 +297,37 @@ namespace HordeServer.Agents
 			_sessionCollection = mongoService.GetCollection<SessionDocument>("Sessions", sessionIndexes);
 
 			_redisService = redisService;
+			_leaseCollection = leaseCollection;
 			_clock = clock;
 			_updateEventChannel = new RedisChannel<AgentId>(RedisChannel.Literal("agents/notify"));
 			_auditLog = auditLog;
+			_logger = logger;
+		}
+
+		/// <inheritdoc/>
+		public async ValueTask DisposeAsync()
+		{
+			if (_updateEventSubscription != null)
+			{
+				await _updateEventSubscription.DisposeAsync();
+				_updateEventSubscription = null;
+			}
+		}
+
+		/// <inheritdoc/>
+		public async Task StartAsync(CancellationToken cancellationToken)
+		{
+			_updateEventSubscription = await _redisService.GetDatabase().Multiplexer.SubscribeAsync(_updateEventChannel, OnAgentUpdate);
+		}
+
+		/// <inheritdoc/>
+		public async Task StopAsync(CancellationToken cancellationToken)
+		{
+			if (_updateEventSubscription != null)
+			{
+				await _updateEventSubscription.DisposeAsync();
+				_updateEventSubscription = null;
+			}
 		}
 
 		async ValueTask<AgentDocument?> PostLoadAsync(AgentDocument? document, CancellationToken cancellationToken)
@@ -356,12 +400,7 @@ namespace HordeServer.Agents
 				.Unset(x => x.Deleted)
 				.Unset(x => x.SessionId);
 
-			AgentDocument? newAgent = await TryUpdateAsync(agentDocument, update, cancellationToken);
-			if (newAgent != null)
-			{
-				await PublishUpdateEventAsync(agent.Id);
-			}
-			return newAgent;
+			return await TryUpdateAsync(agentDocument, update, cancellationToken);
 		}
 
 		/// <inheritdoc/>
@@ -384,10 +423,16 @@ namespace HordeServer.Agents
 
 		/// <inheritdoc/>
 		public async Task<IAgent?> GetAsync(AgentId agentId, CancellationToken cancellationToken)
+		{
+			AgentDocument? document = await GetDocumentAsync(agentId, cancellationToken);
+			return CreateAgentObject(document);
+		}
+
+		async Task<AgentDocument?> GetDocumentAsync(AgentId agentId, CancellationToken cancellationToken)
 		{
 			AgentDocument? document = await _agentCollection.Find<AgentDocument>(x => x.Id == agentId).FirstOrDefaultAsync(cancellationToken);
 			document = await PostLoadAsync(document, cancellationToken);
-			return CreateAgentObject(document);
+			return document;
 		}
 
 		/// <inheritdoc/>
@@ -524,7 +569,14 @@ namespace HordeServer.Agents
 			Expression<Func<AgentDocument, bool>> filter = x => x.Id == current.Id && x.UpdateIndex == prevUpdateIndex;
 			UpdateDefinition<AgentDocument> updateWithIndex = update.Set(x => x.UpdateIndex, current.UpdateIndex).Set(x => x.UpdateTime, current.UpdateTime);
 
-			return await _agentCollection.FindOneAndUpdateAsync<AgentDocument>(filter, updateWithIndex, new FindOneAndUpdateOptions<AgentDocument, AgentDocument> { ReturnDocument = ReturnDocument.After }, cancellationToken);
+			AgentDocument? newDocument = await _agentCollection.FindOneAndUpdateAsync<AgentDocument>(filter, updateWithIndex, new FindOneAndUpdateOptions<AgentDocument, AgentDocument> { ReturnDocument = ReturnDocument.After }, cancellationToken);
+			if (newDocument == null)
+			{
+				return null;
+			}
+
+			PublishUpdateEvent(newDocument.Id);
+			return newDocument;
 		}
 
 		/// <inheritdoc/>
@@ -599,15 +651,7 @@ namespace HordeServer.Agents
 			}
 
 			// Apply the update
-			AgentDocument? newAgent = await TryUpdateAsync(agent, updateBuilder.Combine(updates), cancellationToken);
-			if (newAgent != null)
-			{
-				if (newAgent.RequestRestart != agent.RequestRestart || newAgent.RequestConform != agent.RequestConform || newAgent.RequestShutdown != agent.RequestShutdown || newAgent.RequestForceRestart != agent.RequestForceRestart)
-				{
-					await PublishUpdateEventAsync(agent.Id);
-				}
-			}
-			return newAgent;
+			return await TryUpdateAsync(agent, updateBuilder.Combine(updates), cancellationToken);
 		}
 
 		/// <inheritdoc/>
@@ -617,6 +661,7 @@ namespace HordeServer.Agents
 			UpdateDefinitionBuilder<AgentDocument> updateBuilder = Builders<AgentDocument>.Update;
 			List<UpdateDefinition<AgentDocument>> updates = new List<UpdateDefinition<AgentDocument>>();
 
+			AgentStatus status = options.Status ?? agent.Status;
 			if (options.Status != null && agent.Status != options.Status.Value)
 			{
 				updates.Add(updateBuilder.Set(x => x.Status, options.Status.Value));
@@ -647,30 +692,57 @@ namespace HordeServer.Agents
 			}
 			if (options.Leases != null)
 			{
-				foreach (AgentLease lease in options.Leases)
+				// Flag for whether the leases array should be updated
+				bool updateLeases = false;
+				Dictionary<LeaseId, RpcLease> leaseIdToNewState = options.Leases.ToDictionary(x => x.Id, x => x);
+
+				// Remove any completed leases from the agent
+				List<AgentLease> newLeases = agent.Leases?.ConvertAll(x => new AgentLease(x)) ?? new List<AgentLease>();
+				for (int idx = 0; idx < newLeases.Count; idx++)
 				{
-					if (lease.Payload != null && (agent.Leases == null || !agent.Leases.Any(x => x.Id == lease.Id)))
+					AgentLease lease = newLeases[idx];
+					if (!leaseIdToNewState.TryGetValue(lease.Id, out RpcLease? newLease) || lease.State == LeaseState.Cancelled)
 					{
-						GetNewLeaseUpdates(agent, lease, updates);
+						if (newLease == null || newLease.State == RpcLeaseState.Cancelled || newLease.State == RpcLeaseState.Completed)
+						{
+							newLeases.RemoveAt(idx--);
+							updateLeases = true;
+						}
+					}
+					else
+					{
+						if ((LeaseState)newLease.State != lease.State)
+						{
+							if (newLease.State == RpcLeaseState.Cancelled || newLease.State == RpcLeaseState.Completed)
+							{
+								newLeases.RemoveAt(idx--);
+							}
+							else if (newLease.State == RpcLeaseState.Active && lease.State == LeaseState.Pending)
+							{
+								lease.State = LeaseState.Active;
+							}
+							updateLeases = true;
+						}
 					}
 				}
 
-				List<AgentLease> currentLeases = agent.Leases ?? new List<AgentLease>();
-				List<AgentLease> newLeases = options.Leases;
-				List<AgentLease> leasesToAdd = newLeases.Where(nl => currentLeases.All(cl => cl.Id != nl.Id)).ToList();
-				List<AgentLease> leasesToRemove = currentLeases.Where(cl => newLeases.All(nl => nl.Id != cl.Id)).ToList();
-
-				foreach (AgentLease lease in leasesToAdd)
+				// If the agent is stopping, cancel all the leases. Clear out the current session once it's complete.
+				if (status == AgentStatus.Stopping || status == AgentStatus.Busy)
 				{
-					await AddActiveLeaseAsync(lease);
+					foreach (AgentLease newLease in newLeases)
+					{
+						if (newLease.State != LeaseState.Cancelled)
+						{
+							newLease.State = LeaseState.Cancelled;
+							updateLeases = true;
+						}
+					}
 				}
 
-				foreach (AgentLease lease in leasesToRemove)
+				if (updateLeases)
 				{
-					await RemoveActiveLeaseAsync(lease);
+					updates.Add(updateBuilder.Set(x => x.Leases, newLeases));
 				}
-
-				updates.Add(updateBuilder.Set(x => x.Leases, options.Leases));
 			}
 
 			// Update the pools
@@ -694,7 +766,33 @@ namespace HordeServer.Agents
 			}
 
 			// Update the agent, and try to create new lease documents if we succeed
-			return await TryUpdateAsync(agent, updateBuilder.Combine(updates), cancellationToken);
+			AgentDocument? newAgent = await TryUpdateAsync(agent, updateBuilder.Combine(updates), cancellationToken);
+			if (newAgent == null)
+			{
+				return null;
+			}
+
+			// Remove any complete leases
+			if (agent.Leases != null)
+			{
+				HashSet<LeaseId> newLeaseIds = new HashSet<LeaseId>(newAgent.Leases?.Select(x => x.Id) ?? Enumerable.Empty<LeaseId>());
+				foreach (AgentLease oldLease in agent.Leases.Where(x => !newLeaseIds.Contains(x.Id)))
+				{
+					await RemoveActiveLeaseAsync(oldLease);
+
+					RpcLease? rpcLease = options.Leases?.FirstOrDefault(x => x.Id == oldLease.Id);
+					if (rpcLease != null && rpcLease.State == RpcLeaseState.Completed)
+					{
+						await _leaseCollection.TrySetOutcomeAsync(oldLease.Id, _clock.UtcNow, LeaseOutcome.Success, rpcLease.Output.ToArray(), CancellationToken.None);
+					}
+					else
+					{
+						await _leaseCollection.TrySetOutcomeAsync(oldLease.Id, _clock.UtcNow, LeaseOutcome.Cancelled, null, CancellationToken.None);
+					}
+				}
+			}
+
+			return newAgent;
 		}
 
 		static void SetPoolProperties(List<string> properties, List<PoolId> pools)
@@ -845,7 +943,12 @@ namespace HordeServer.Agents
 		/// <inheritdoc/>
 		async Task<AgentDocument?> TryTerminateSessionAsync(AgentDocument agent, CancellationToken cancellationToken)
 		{
-			DateTime utcNow = _clock.UtcNow;
+			// Get the time that the session finishes at
+			DateTime finishTime = _clock.UtcNow;
+			if (agent.SessionExpiresAt.HasValue && agent.SessionExpiresAt.Value < finishTime)
+			{
+				finishTime = agent.SessionExpiresAt.Value;
+			}
 
 			UpdateDefinition<AgentDocument> update = new BsonDocument();
 
@@ -853,18 +956,13 @@ namespace HordeServer.Agents
 			update = update.Unset(x => x.SessionExpiresAt);
 			update = update.Unset(x => x.Leases);
 			update = update.Set(x => x.Status, AgentStatus.Stopped);
-			update = update.Set(x => x.LastStatusChange, utcNow);
+			update = update.Set(x => x.LastStatusChange, finishTime);
 
 			if (agent.Ephemeral)
 			{
 				update = update.Set(x => x.Deleted, true);
 			}
 
-			foreach (AgentLease agentLease in agent.Leases ?? Enumerable.Empty<AgentLease>())
-			{
-				await RemoveActiveLeaseAsync(agentLease);
-			}
-
 			AgentDocument? newAgent = await TryUpdateAsync(agent, update, cancellationToken);
 			if (newAgent == null)
 			{
@@ -872,8 +970,24 @@ namespace HordeServer.Agents
 			}
 
 			// Update the session document
-			GetLogger(newAgent.Id).LogInformation("Terminated session {SessionId}", agent.SessionId);
-			await _sessionCollection.UpdateOneAsync(x => x.Id == agent.SessionId, Builders<SessionDocument>.Update.Set(x => x.FinishTime, utcNow), cancellationToken: CancellationToken.None);
+			ILogger agentLogger = GetLogger(agent.Id);
+			agentLogger.LogInformation("Terminated session {SessionId}", agent.SessionId);
+			await _sessionCollection.UpdateOneAsync(x => x.Id == agent.SessionId, Builders<SessionDocument>.Update.Set(x => x.FinishTime, finishTime), cancellationToken: CancellationToken.None);
+
+			// Remove any outstanding leases
+			if (agent.Leases != null)
+			{
+				foreach (AgentLease agentLease in agent.Leases)
+				{
+					agentLogger.LogInformation("Removing lease {LeaseId} during session terminate...", agentLease.Id);
+
+					// Remove it from the active list used for autoscaling
+					await RemoveActiveLeaseAsync(agentLease);
+
+					// Update the lease
+					await _leaseCollection.TrySetOutcomeAsync(agentLease.Id, finishTime, LeaseOutcome.Failed, null, CancellationToken.None);
+				}
+			}
 
 			return newAgent;
 		}
@@ -895,14 +1009,26 @@ namespace HordeServer.Agents
 
 			List<UpdateDefinition<AgentDocument>> updates = new List<UpdateDefinition<AgentDocument>>();
 			updates.Add(Builders<AgentDocument>.Update.Set(x => x.Leases, leases));
-			GetNewLeaseUpdates(agent, newLease, updates);
+			GetNewLeaseUpdates(agent, newLease.Payload, updates);
 
 			UpdateDefinition<AgentDocument> update = Builders<AgentDocument>.Update.Combine(updates);
+
 			AgentDocument? updatedDoc = await TryUpdateAsync(agent, update, cancellationToken);
+			if (updatedDoc == null)
+			{
+				return null;
+			}
+
+			await AddActiveLeaseAsync(newLease);
 
-			if (updatedDoc != null)
+			try
+			{
+				DateTime startTime = _clock.UtcNow;
+				await _leaseCollection.AddAsync(options.Id, options.ParentId, options.Name, agent.Id, agent.SessionId!.Value, options.StreamId, options.PoolId, options.LogId, startTime, Any.Pack(options.Payload).ToByteArray(), cancellationToken);
+			}
+			catch (Exception ex)
 			{
-				await AddActiveLeaseAsync(newLease);
+				_logger.LogError(ex, "Unable to create lease {LeaseId} for agent {AgentId}; lease already exists?", options.Id, agent.Id);
 			}
 
 			return updatedDoc;
@@ -933,32 +1059,29 @@ namespace HordeServer.Agents
 			}
 		}
 
-		static void GetNewLeaseUpdates(AgentDocument agent, AgentLease lease, List<UpdateDefinition<AgentDocument>> updates)
+		static void GetNewLeaseUpdates(AgentDocument agent, ReadOnlySpan<byte> payloadData, List<UpdateDefinition<AgentDocument>> updates)
 		{
-			if (lease.Payload != null)
+			Any payload = Any.Parser.ParseFrom(payloadData);
+			if (payload.TryUnpack(out ConformTask conformTask))
+			{
+				int newConformAttemptCount = (agent.ConformAttemptCount ?? 0) + 1;
+				updates.Add(Builders<AgentDocument>.Update.Set(x => x.ConformAttemptCount, newConformAttemptCount));
+				updates.Add(Builders<AgentDocument>.Update.Set(x => x.LastConformTime, DateTime.UtcNow));
+			}
+			else if (payload.TryUnpack(out UpgradeTask upgradeTask))
 			{
-				Any payload = Any.Parser.ParseFrom(lease.Payload.ToArray());
-				if (payload.TryUnpack(out ConformTask conformTask))
+				string newVersion = upgradeTask.SoftwareId;
+
+				int versionIdx = newVersion.IndexOf(':', StringComparison.Ordinal);
+				if (versionIdx != -1)
 				{
-					int newConformAttemptCount = (agent.ConformAttemptCount ?? 0) + 1;
-					updates.Add(Builders<AgentDocument>.Update.Set(x => x.ConformAttemptCount, newConformAttemptCount));
-					updates.Add(Builders<AgentDocument>.Update.Set(x => x.LastConformTime, DateTime.UtcNow));
+					newVersion = newVersion.Substring(versionIdx + 1);
 				}
-				else if (payload.TryUnpack(out UpgradeTask upgradeTask))
-				{
-					string newVersion = upgradeTask.SoftwareId;
-
-					int versionIdx = newVersion.IndexOf(':', StringComparison.Ordinal);
-					if (versionIdx != -1)
-					{
-						newVersion = newVersion.Substring(versionIdx + 1);
-					}
 
-					int newUpgradeAttemptCount = (agent.UpgradeAttemptCount ?? 0) + 1;
-					updates.Add(Builders<AgentDocument>.Update.Set(x => x.LastUpgradeVersion, newVersion));
-					updates.Add(Builders<AgentDocument>.Update.Set(x => x.UpgradeAttemptCount, newUpgradeAttemptCount));
-					updates.Add(Builders<AgentDocument>.Update.Set(x => x.LastUpgradeTime, DateTime.UtcNow));
-				}
+				int newUpgradeAttemptCount = (agent.UpgradeAttemptCount ?? 0) + 1;
+				updates.Add(Builders<AgentDocument>.Update.Set(x => x.LastUpgradeVersion, newVersion));
+				updates.Add(Builders<AgentDocument>.Update.Set(x => x.UpgradeAttemptCount, newUpgradeAttemptCount));
+				updates.Add(Builders<AgentDocument>.Update.Set(x => x.LastUpgradeTime, DateTime.UtcNow));
 			}
 		}
 
@@ -969,7 +1092,7 @@ namespace HordeServer.Agents
 			AgentDocument? newAgent = await TryUpdateAsync(agent, update, cancellationToken);
 			if (newAgent != null)
 			{
-				await PublishUpdateEventAsync(agent.Id);
+				PublishUpdateEvent(agent.Id);
 				if (agent.Leases != null && leaseIdx < agent.Leases.Count)
 				{
 					await RemoveActiveLeaseAsync(agent.Leases[leaseIdx]);
@@ -979,21 +1102,46 @@ namespace HordeServer.Agents
 		}
 
 		/// <inheritdoc/>
-		public IAuditLogChannel<AgentId> GetLogger(AgentId agentId)
+		void PublishUpdateEvent(AgentId agentId)
 		{
-			return _auditLog[agentId];
+			_ = _redisService.GetDatabase().PublishAsync(_updateEventChannel, agentId, CommandFlags.FireAndForget);
 		}
 
-		/// <inheritdoc/>
-		public Task PublishUpdateEventAsync(AgentId agentId)
+		void OnAgentUpdate(AgentId agentId)
 		{
-			return _redisService.GetDatabase().PublishAsync(_updateEventChannel, agentId);
+			if (_agentIdToTcs.TryGetValue(agentId, out TaskCompletionSource? tcs))
+			{
+				tcs.TrySetResult();
+			}
+		}
+
+		async Task<AgentDocument?> WaitForUpdateAsync(AgentDocument agent, CancellationToken cancellationToken)
+		{
+			TaskCompletionSource newTcs = new TaskCompletionSource();
+			try
+			{
+				TaskCompletionSource tcs = _agentIdToTcs.GetOrAdd(agent.Id, newTcs);
+
+				AgentDocument? newAgent = await GetDocumentAsync(agent.Id, cancellationToken);
+				if (newAgent != null && newAgent.UpdateIndex == agent.UpdateIndex)
+				{
+					await tcs.Task.WaitAsync(cancellationToken);
+					newAgent = await GetDocumentAsync(agent.Id, cancellationToken);
+				}
+
+				return newAgent;
+			}
+			finally
+			{
+				newTcs.TrySetResult();
+				_agentIdToTcs.TryRemove(new KeyValuePair<AgentId, TaskCompletionSource>(agent.Id, newTcs));
+			}
 		}
 
 		/// <inheritdoc/>
-		public async Task<IAsyncDisposable> SubscribeToUpdateEventsAsync(Action<AgentId> onUpdate)
+		public IAuditLogChannel<AgentId> GetLogger(AgentId agentId)
 		{
-			return await _redisService.GetDatabase().Multiplexer.SubscribeAsync(_updateEventChannel, onUpdate);
+			return _auditLog[agentId];
 		}
 	}
 }
diff --git a/Engine/Source/Programs/Horde/Plugins/Compute/HordeServer.Compute/Agents/AgentService.cs b/Engine/Source/Programs/Horde/Plugins/Compute/HordeServer.Compute/Agents/AgentService.cs
--- a/Engine/Source/Programs/Horde/Plugins/Compute/HordeServer.Compute/Agents/AgentService.cs
+++ b/Engine/Source/Programs/Horde/Plugins/Compute/HordeServer.Compute/Agents/AgentService.cs
@@ -1,6 +1,5 @@
 // Copyright Epic Games, Inc. All Rights Reserved.
 
-using System.Diagnostics;
 using System.Diagnostics.Metrics;
 using System.Security.Claims;
 using EpicGames.Core;
@@ -10,9 +9,7 @@ using EpicGames.Horde.Agents.Pools;
 using EpicGames.Horde.Agents.Sessions;
 using EpicGames.Redis;
 using EpicGames.Serialization;
-using Google.Protobuf;
 using Google.Protobuf.WellKnownTypes;
-using HordeCommon.Rpc;
 using HordeCommon.Rpc.Messages;
 using HordeServer.Acls;
 using HordeServer.Agents.Leases;
@@ -97,9 +94,6 @@ namespace HordeServer.Agents
 		/// <summary>OpenTelemetry measurements (gauges)</summary>
 		IEnumerable<Measurement<int>> _measurements = new List<Measurement<int>>();
 
-		/// <summary>Subscription for update events</summary>
-		IAsyncDisposable? _subscription;
-
 		/// <summary>
 		/// Constructor
 		/// </summary>
@@ -140,17 +134,11 @@ namespace HordeServer.Agents
 		{
 			await _ticker.StartAsync();
 			await _sharedTicker.StartAsync();
-			_subscription = await Agents.SubscribeToUpdateEventsAsync(OnAgentUpdate);
 		}
 
 		/// <inheritdoc/>
 		public async Task StopAsync(CancellationToken cancellationToken)
 		{
-			if (_subscription != null)
-			{
-				await _subscription.DisposeAsync();
-				_subscription = null;
-			}
 			await _ticker.StopAsync();
 			await _sharedTicker.StopAsync();
 		}
@@ -356,21 +344,6 @@ namespace HordeServer.Agents
 			return newDynamicPools;
 		}
 
-		/// <summary>
-		/// Callback for an agents 
-		/// </summary>
-		/// <param name="agentId"></param>
-		void OnAgentUpdate(AgentId agentId)
-		{
-			lock (_waitingAgents)
-			{
-				if (_waitingAgents.TryGetValue(agentId, out CancellationTokenSource? cancellationSource))
-				{
-					cancellationSource.Cancel();
-				}
-			}
-		}
-
 		/// <summary>
 		/// Creates a new agent session
 		/// </summary>
@@ -401,13 +374,6 @@ namespace HordeServer.Agents
 				{
 					DateTime utcNow = _clock.UtcNow;
 
-					// Remove any outstanding leases
-					foreach (AgentLease lease in agent.Leases)
-					{
-						agentLogger.LogInformation("Removing outstanding lease {LeaseId}", lease.Id);
-						await RemoveLeaseAsync(agent, lease, utcNow, LeaseOutcome.Failed, null, cancellationToken);
-					}
-
 					// Get the new pools for the agent
 					List<PoolId> dynamicPools = await GetDynamicPoolsAsync(agent, cancellationToken);
 
@@ -500,62 +466,48 @@ namespace HordeServer.Agents
 				using CancellationTokenSource cancellationSource = CancellationTokenSource.CreateLinkedTokenSource(_applicationLifetime.ApplicationStopping, cancellationToken);
 				cancellationSource.CancelAfter(maxWaitTime);
 
-				// Assign a new lease
-				(ITaskSource, CreateLeaseOptions)? result = null;
-				try
-				{
-					// Add the cancellation source to the set of waiting agents
-					lock (_waitingAgents)
-					{
-						_waitingAgents[agent.Id] = cancellationSource;
-					}
+				// If an agent changes, cancel the update
+				_ = agent.WaitForUpdateAsync(cancellationSource.Token).ContinueWith(x => cancellationSource.Cancel());
 
-					// Create all the tasks to wait for
-					List<Task<(ITaskSource, CreateLeaseOptions)?>> tasks = new List<Task<(ITaskSource, CreateLeaseOptions)?>>();
-					foreach (ITaskSource taskSource in _taskSources)
+				// Create all the tasks to wait for
+				List<Task<(ITaskSource, CreateLeaseOptions)?>> tasks = new List<Task<(ITaskSource, CreateLeaseOptions)?>>();
+				foreach (ITaskSource taskSource in _taskSources)
+				{
+					if (CanUseTaskSource(agent, taskSource) && !cancellationSource.IsCancellationRequested)
 					{
-						if (CanUseTaskSource(agent, taskSource) && !cancellationSource.IsCancellationRequested)
-						{
-							Task<(ITaskSource, CreateLeaseOptions)?> task = await GuardedAssignLeaseAsync(taskSource, agent, cancellationSource);
-							tasks.Add(task);
-						}
+						Task<(ITaskSource, CreateLeaseOptions)?> task = await GuardedAssignLeaseAsync(taskSource, agent, cancellationSource);
+						tasks.Add(task);
 					}
+				}
 
-					// If no task source is valid, just add a delay
-					if (tasks.Count == 0)
-					{
-						_logger.LogInformation("No task source valid for agent {AgentId}. Waiting {WaitTimeMs} ms", agent.Id, maxWaitTime.TotalMilliseconds);
-						await AsyncUtils.DelayNoThrow(maxWaitTime, cancellationToken);
-						break;
-					}
+				// If no task source is valid, just add a delay
+				if (tasks.Count == 0)
+				{
+					_logger.LogInformation("No task source valid for agent {AgentId}. Waiting {WaitTimeMs} ms", agent.Id, maxWaitTime.TotalMilliseconds);
+					await AsyncUtils.DelayNoThrow(maxWaitTime, cancellationToken);
+					break;
+				}
 
-					// Wait for all the tasks to complete. Once the first task completes it will set the cancellation source, triggering the 
-					// others to terminate.
-					await Task.WhenAll(tasks);
+				// Wait for all the tasks to complete. Once the first task completes it will set the cancellation source, triggering the 
+				// others to terminate.
+				await Task.WhenAll(tasks);
 
-					// Find the first result
-					foreach (Task<(ITaskSource, CreateLeaseOptions)?> task in tasks)
+				// Find the first result
+				(ITaskSource, CreateLeaseOptions)? result = null;
+				foreach (Task<(ITaskSource, CreateLeaseOptions)?> task in tasks)
+				{
+					(ITaskSource, CreateLeaseOptions)? taskResult;
+					if (task.TryGetResult(out taskResult) && taskResult != null)
 					{
-						(ITaskSource, CreateLeaseOptions)? taskResult;
-						if (task.TryGetResult(out taskResult) && taskResult != null)
+						(ITaskSource taskSource, CreateLeaseOptions taskLease) = taskResult.Value;
+						if (result == null)
 						{
-							(ITaskSource taskSource, CreateLeaseOptions taskLease) = taskResult.Value;
-							if (result == null)
-							{
-								result = (taskSource, taskLease);
-							}
-							else
-							{
-								await taskSource.CancelLeaseAsync(agent, taskLease.Id, Any.Pack(taskLease.Payload), CancellationToken.None);
-							}
+							result = (taskSource, taskLease);
+						}
+						else
+						{
+							await taskSource.CancelLeaseAsync(agent, taskLease.Id, Any.Pack(taskLease.Payload), CancellationToken.None);
 						}
-					}
-				}
-				finally
-				{
-					lock (_waitingAgents)
-					{
-						_waitingAgents.Remove(agent.Id);
 					}
 				}
 
@@ -577,7 +529,6 @@ namespace HordeServer.Agents
 				if (newAgent != null)
 				{
 					await source.OnLeaseStartedAsync(newAgent, lease.Id, Any.Pack(lease.Payload), Agents.GetLogger(agent.Id), CancellationToken.None);
-					await CreateLeaseAsync(agent, lease, CancellationToken.None);
 					return newAgent;
 				}
 				else
@@ -705,8 +656,6 @@ namespace HordeServer.Agents
 		{
 			DateTime utcNow = _clock.UtcNow;
 
-			Stopwatch timer = Stopwatch.StartNew();
-
 			IAgent? agent = inAgent;
 			while (agent != null)
 			{
@@ -738,62 +687,11 @@ namespace HordeServer.Agents
 					sessionExpiresAt = utcNow + SessionExpiryTime;
 				}
 
-				// Flag for whether the leases array should be updated
-				bool updateLeases = false;
-				List<AgentLease> leases = agent.Leases.ConvertAll(x => new AgentLease(x));
-
-				// Remove any completed leases from the agent
-				Dictionary<LeaseId, RpcLease> leaseIdToNewState = newLeases.ToDictionary(x => x.Id, x => x);
-				for (int idx = 0; idx < leases.Count; idx++)
-				{
-					AgentLease lease = leases[idx];
-					if (lease.State == LeaseState.Cancelled)
-					{
-						RpcLease? newLease;
-						if (!leaseIdToNewState.TryGetValue(lease.Id, out newLease) || newLease.State == RpcLeaseState.Cancelled || newLease.State == RpcLeaseState.Completed)
-						{
-							await RemoveLeaseAsync(agent, lease, utcNow, LeaseOutcome.Cancelled, null, cancellationToken);
-							leases.RemoveAt(idx--);
-							updateLeases = true;
-						}
-					}
-					else
-					{
-						RpcLease? newLease;
-						if (leaseIdToNewState.TryGetValue(lease.Id, out newLease) && (LeaseState)newLease.State != lease.State)
-						{
-							if (newLease.State == RpcLeaseState.Cancelled || newLease.State == RpcLeaseState.Completed)
-							{
-								await RemoveLeaseAsync(agent, lease, utcNow, (LeaseOutcome)newLease.Outcome, newLease.Output.ToByteArray(), cancellationToken);
-								leases.RemoveAt(idx--);
-							}
-							else if (newLease.State == RpcLeaseState.Active && lease.State == LeaseState.Pending)
-							{
-								lease.State = LeaseState.Active;
-							}
-							updateLeases = true;
-						}
-					}
-				}
-
-				// If the agent is stopping, cancel all the leases. Clear out the current session once it's complete.
-				if (status == AgentStatus.Stopping || status == AgentStatus.Busy)
-				{
-					foreach (AgentLease lease in leases)
-					{
-						if (lease.State != LeaseState.Cancelled)
-						{
-							lease.State = LeaseState.Cancelled;
-							updateLeases = true;
-						}
-					}
-				}
-
 				// Get the new dynamic pools for the agent
 				List<PoolId> dynamicPools = await GetDynamicPoolsAsync(agent, cancellationToken);
 
 				// Update the agent, and try to create new lease documents if we succeed
-				IAgent? newAgent = await agent.TryUpdateSessionAsync(new UpdateSessionOptions(status, sessionExpiresAt, capabilities, dynamicPools, updateLeases ? leases : null), cancellationToken);
+				IAgent? newAgent = await agent.TryUpdateSessionAsync(new UpdateSessionOptions(status, sessionExpiresAt, capabilities, dynamicPools, newLeases), cancellationToken);
 				if (newAgent != null)
 				{
 					agent = newAgent;
@@ -826,7 +724,7 @@ namespace HordeServer.Agents
 		/// <param name="agent">The agent whose current session should be terminated</param>
 		/// <param name="cancellationToken">Cancellation token for the operation</param>
 		/// <returns>An up-to-date IAgent if the session was terminated</returns>
-		private async Task<IAgent?> TryTerminateSessionAsync(IAgent agent, CancellationToken cancellationToken = default)
+		private static async Task<IAgent?> TryTerminateSessionAsync(IAgent agent, CancellationToken cancellationToken = default)
 		{
 			// Make sure the agent has a valid session id
 			if (agent.SessionId == null)
@@ -834,53 +732,8 @@ namespace HordeServer.Agents
 				return agent;
 			}
 
-			// Get the time that the session finishes at
-			DateTime finishTime = _clock.UtcNow;
-			if (agent.SessionExpiresAt.HasValue && agent.SessionExpiresAt.Value < finishTime)
-			{
-				finishTime = agent.SessionExpiresAt.Value;
-			}
-
-			// Save off the session id and current leases
-			IReadOnlyList<IAgentLease> leases = agent.Leases;
-
 			// Clear the current session
-			IAgent? newAgent = await agent.TryTerminateSessionAsync(cancellationToken);
-			if (newAgent != null)
-			{
-				agent = newAgent;
-
-				// Remove any outstanding leases
-				foreach (IAgentLease lease in leases)
-				{
-					Agents.GetLogger(agent.Id).LogInformation("Removing lease {LeaseId} during session terminate...", lease.Id);
-					await RemoveLeaseAsync(agent, lease, finishTime, LeaseOutcome.Failed, null, cancellationToken);
-				}
-
-				return agent;
-			}
-			return null;
-		}
-
-		/// <summary>
-		/// Creates a new lease document
-		/// </summary>
-		/// <param name="agent">Agent that will be executing the lease</param>
-		/// <param name="options">The new agent lease</param>
-		/// <param name="cancellationToken">Cancellation token for the operation</param>
-		/// <returns>New lease document</returns>
-		internal Task<ILease> CreateLeaseAsync(IAgent agent, CreateLeaseOptions options, CancellationToken cancellationToken = default)
-		{
-			try
-			{
-				DateTime startTime = _clock.UtcNow;
-				return _leases.AddAsync(options.Id, options.ParentId, options.Name, agent.Id, agent.SessionId!.Value, options.StreamId, options.PoolId, options.LogId, startTime, Any.Pack(options.Payload).ToByteArray(), cancellationToken);
-			}
-			catch (Exception ex)
-			{
-				_logger.LogError(ex, "Unable to create lease {LeaseId} for agent {AgentId}; lease already exists?", options.Id, agent.Id);
-				throw;
-			}
+			return await agent.TryTerminateSessionAsync(cancellationToken);
 		}
 
 		/// <summary>
@@ -924,51 +777,6 @@ namespace HordeServer.Agents
 			return _leases.GetAsync(leaseId, cancellationToken);
 		}
 
-		/// <summary>
-		/// Removes a lease with the given id. Updates the lease state in the database, and removes the item from the agent's leases array.
-		/// </summary>
-		/// <param name="agent">The agent to remove a lease from</param>
-		/// <param name="lease">The lease to cancel</param>
-		/// <param name="utcNow">The current time</param>
-		/// <param name="outcome">Final status of the lease</param>
-		/// <param name="output">Output from executing the task</param>
-		/// <param name="cancellationToken">Cancellation token for the operation</param>
-		/// <returns>Async task</returns>
-		private async Task RemoveLeaseAsync(IAgent agent, IAgentLease lease, DateTime utcNow, LeaseOutcome outcome, byte[]? output, CancellationToken cancellationToken = default)
-		{
-			// Figure out what time the lease finished
-			DateTime finishTime = utcNow;
-			if (agent.SessionExpiresAt.HasValue && agent.SessionExpiresAt.Value < finishTime)
-			{
-				finishTime = agent.SessionExpiresAt.Value;
-			}
-
-			// Update the lease
-			await _leases.TrySetOutcomeAsync(lease.Id, finishTime, outcome, output, cancellationToken);
-
-			// Temporarily disabling due to gRPC timeouts.
-#if false
-			// Terminate any child leases
-			List<ILease> childLeases = await _leases.FindLeasesAsync(parentId: lease.Id);
-			foreach (ILease childLease in childLeases)
-			{
-				if (childLease.Outcome == LeaseOutcome.Unspecified)
-				{
-					IAgent? otherAgent = await GetAgentAsync(childLease.AgentId);
-					if (otherAgent != null)
-					{
-						AgentLease? otherLease = otherAgent.Leases.FirstOrDefault(x => x.Id == childLease.Id);
-						if (otherLease != null)
-						{
-							_logger.LogInformation("Terminating child lease {LeaseId} with parent {ParentLeaseId}", otherLease.Id, lease.Id);
-							await RemoveLeaseAsync(otherAgent, otherLease, utcNow, LeaseOutcome.Failed, null);
-						}
-					}
-				}
-			}
-#endif
-		}
-
 		/// <summary>
 		/// Update the rate table for different agent types
 		/// </summary>
diff --git a/Engine/Source/Programs/Horde/Plugins/Compute/HordeServer.Compute/Agents/IAgent.cs b/Engine/Source/Programs/Horde/Plugins/Compute/HordeServer.Compute/Agents/IAgent.cs
--- a/Engine/Source/Programs/Horde/Plugins/Compute/HordeServer.Compute/Agents/IAgent.cs
+++ b/Engine/Source/Programs/Horde/Plugins/Compute/HordeServer.Compute/Agents/IAgent.cs
@@ -273,6 +273,12 @@ namespace HordeServer.Agents
 		/// <param name="cancellationToken">Cancellation token for the operation</param>
 		/// <returns>New agent state if it succeeded, otherwise null</returns>
 		Task<IAgent?> TryCancelLeaseAsync(int leaseIdx, CancellationToken cancellationToken = default);
+
+		/// <summary>
+		/// Wait for this agent to be updated
+		/// </summary>
+		/// <param name="cancellationToken">Cancellation token for the operation</param>
+		Task<IAgent?> WaitForUpdateAsync(CancellationToken cancellationToken = default);
 	}
 
 	/// <summary>
@@ -312,7 +318,7 @@ namespace HordeServer.Agents
 	/// <param name="Capabilities">Capbilities for the session</param>
 	/// <param name="DynamicPools">New list of dynamic pools for the agent</param>
 	/// <param name="Leases">New set of leases</param>
-	public record class UpdateSessionOptions(AgentStatus? Status = null, DateTime? SessionExpiresAt = null, RpcAgentCapabilities? Capabilities = null, IReadOnlyList<PoolId>? DynamicPools = null, List<AgentLease>? Leases = null);
+	public record class UpdateSessionOptions(AgentStatus? Status = null, DateTime? SessionExpiresAt = null, RpcAgentCapabilities? Capabilities = null, IReadOnlyList<PoolId>? DynamicPools = null, IEnumerable<RpcLease>? Leases = null);
 
 	/// <summary>
 	/// Extension methods for IAgent
diff --git a/Engine/Source/Programs/Horde/Plugins/Compute/HordeServer.Compute/Agents/IAgentCollection.cs b/Engine/Source/Programs/Horde/Plugins/Compute/HordeServer.Compute/Agents/IAgentCollection.cs
--- a/Engine/Source/Programs/Horde/Plugins/Compute/HordeServer.Compute/Agents/IAgentCollection.cs
+++ b/Engine/Source/Programs/Horde/Plugins/Compute/HordeServer.Compute/Agents/IAgentCollection.cs
@@ -100,18 +100,5 @@ namespace HordeServer.Agents
 		/// <param name="agentId"></param>
 		/// <returns></returns>
 		IAuditLogChannel<AgentId> GetLogger(AgentId agentId);
-
-		/// <summary>
-		/// Sends a notification that an update event has ocurred
-		/// </summary>
-		/// <param name="agentId">Agent that has been updated</param>
-		Task PublishUpdateEventAsync(AgentId agentId);
-
-		/// <summary>
-		/// Subscribe to notifications on agent states being updated
-		/// </summary>
-		/// <param name="onUpdate">Callback for updates</param>
-		/// <returns>Disposable subscription object</returns>
-		Task<IAsyncDisposable> SubscribeToUpdateEventsAsync(Action<AgentId> onUpdate);
 	}
 }
diff --git a/Engine/Source/Programs/Horde/Plugins/Compute/HordeServer.Compute/Compute/ComputeService.cs b/Engine/Source/Programs/Horde/Plugins/Compute/HordeServer.Compute/Compute/ComputeService.cs
--- a/Engine/Source/Programs/Horde/Plugins/Compute/HordeServer.Compute/Compute/ComputeService.cs
+++ b/Engine/Source/Programs/Horde/Plugins/Compute/HordeServer.Compute/Compute/ComputeService.cs
@@ -434,8 +434,6 @@ namespace HordeServer.Compute
 							IAgent? newAgent = await agent.TryCreateLeaseAsync(createLeaseOptions, cancellationToken);
 							if (newAgent != null)
 							{
-								await _agentCollection.PublishUpdateEventAsync(agent.Id);
-								await _agentService.CreateLeaseAsync(newAgent, createLeaseOptions, cancellationToken);
 								span.SetAttribute("allocatedLeaseId", leaseId.ToString());
 								span.SetAttribute("allocatedAgentId", newAgent.Id.ToString());
 								span.SetAttribute("allocatedConnectionAddress", resource.ConnectionAddress);
diff --git a/Engine/Source/Programs/Horde/Plugins/Compute/HordeServer.Compute/ComputePlugin.cs b/Engine/Source/Programs/Horde/Plugins/Compute/HordeServer.Compute/ComputePlugin.cs
--- a/Engine/Source/Programs/Horde/Plugins/Compute/HordeServer.Compute/ComputePlugin.cs
+++ b/Engine/Source/Programs/Horde/Plugins/Compute/HordeServer.Compute/ComputePlugin.cs
@@ -6,6 +6,7 @@ using Amazon.CloudWatch;
 using Amazon.EC2;
 using Amazon.Extensions.NETCore.Setup;
 using Amazon.SQS;
+using EpicGames.Horde.Agents;
 using HordeServer.Acls;
 using HordeServer.Agents;
 using HordeServer.Agents.Enrollment;
@@ -14,6 +15,7 @@ using HordeServer.Agents.Leases;
 using HordeServer.Agents.Pools;
 using HordeServer.Agents.Relay;
 using HordeServer.Agents.Telemetry;
+using HordeServer.Auditing;
 using HordeServer.Aws;
 using HordeServer.Compute;
 using HordeServer.Logs;
@@ -59,6 +61,8 @@ namespace HordeServer
 		/// <inheritdoc/>
 		public void ConfigureServices(IServiceCollection services)
 		{
+			services.AddSingleton<IAuditLog<AgentId>>(sp => sp.GetRequiredService<IAuditLogFactory<AgentId>>().Create("Agents.Log", "AgentId"));
+
 			services.AddSingleton<IDefaultAclModifier, ComputeAclModifier>();
 			services.AddSingleton<IPluginResponseFilter, ComputeResponseFilter>();
 
