Horde: Updating schema documentation

#rnx

diff --git a/Engine/Source/Programs/Horde/Docs/Deployment/AgentSettings.md b/Engine/Source/Programs/Horde/Docs/Deployment/AgentSettings.md
--- a/Engine/Source/Programs/Horde/Docs/Deployment/AgentSettings.md
+++ b/Engine/Source/Programs/Horde/Docs/Deployment/AgentSettings.md
@@ -20,7 +20,8 @@ Name | Description
 `writeStepOutputToLogger` | `boolean`<br>Whether to write step output to the logging device
 `enableAwsEc2Support` | `boolean`<br>Queries information about the current agent through the AWS EC2 interface
 `useLocalStorageClient` | `boolean`<br>Option to use a local storage client rather than connecting through the server. Primarily for convenience when debugging / iterating locally.
-`computePort` | `integer`<br>Incoming port for listening for compute work. Needs to be tied with a lease.
+`computeIp` | `string`<br>Incoming IP for listening for compute work. If not set, it will be automatically resolved.
+`computePort` | `integer`<br>Incoming port for listening for compute work. Needs to be tied with a lease. Set port to 0 to disable incoming compute requests.
 `enableTelemetry` | `boolean`<br>Whether to send telemetry back to Horde server
 `telemetryReportInterval` | `integer`<br>How often to report telemetry events to server in milliseconds
 `bundleCacheSize` | `integer`<br>Maximum size of the bundle cache, in megabytes.
