Horde: Run issue service on all server instances, not just workers.

#rnx

diff --git a/Engine/Source/Programs/Horde/Plugins/Build/HordeServer.Build/BuildPlugin.cs b/Engine/Source/Programs/Horde/Plugins/Build/HordeServer.Build/BuildPlugin.cs
--- a/Engine/Source/Programs/Horde/Plugins/Build/HordeServer.Build/BuildPlugin.cs
+++ b/Engine/Source/Programs/Horde/Plugins/Build/HordeServer.Build/BuildPlugin.cs
@@ -112,10 +112,12 @@ namespace HordeServer
 			// Notifications can be triggered from any instance, so always make sure we're ticking the background task.
 			services.AddHostedService(provider => (NotificationService)provider.GetRequiredService<INotificationService>());
 
+			// Issues need to be assigned asynchronously on any pod
+			services.AddHostedService(provider => provider.GetRequiredService<IssueService>());
+
 			if (_serverInfo.IsRunModeActive(RunMode.Worker) && !_serverInfo.ReadOnlyMode)
 			{
 				services.AddHostedService<AgentReportService>();
-				services.AddHostedService(provider => provider.GetRequiredService<IssueService>());
 				services.AddHostedService<IssueReportService>();
 				services.AddHostedService<IssueTagService>();
 				services.AddHostedService<JobExpirationService>();
