Horde: Add support for explicitly deleting artifacts.

#rnx

diff --git a/Engine/Source/Programs/Horde/Plugins/Build/HordeServer.Build/Artifacts/ArtifactAclAction.cs b/Engine/Source/Programs/Horde/Plugins/Build/HordeServer.Build/Artifacts/ArtifactAclAction.cs
--- a/Engine/Source/Programs/Horde/Plugins/Build/HordeServer.Build/Artifacts/ArtifactAclAction.cs
+++ b/Engine/Source/Programs/Horde/Plugins/Build/HordeServer.Build/Artifacts/ArtifactAclAction.cs
@@ -19,6 +19,11 @@ namespace HordeServer.Artifacts
 		/// </summary>
 		public static AclAction WriteArtifact { get; } = new AclAction("WriteArtifact");
 
+		/// <summary>
+		/// Permission to delete to an artifact
+		/// </summary>
+		public static AclAction DeleteArtifact { get; } = new AclAction("DeleteArtifact");
+
 		/// <summary>
 		/// Ability to create an artifact. Typically just for debugging; agents have this access for a particular session.
 		/// </summary>
diff --git a/Engine/Source/Programs/Horde/Plugins/Build/HordeServer.Build/Artifacts/ArtifactCollection.cs b/Engine/Source/Programs/Horde/Plugins/Build/HordeServer.Build/Artifacts/ArtifactCollection.cs
--- a/Engine/Source/Programs/Horde/Plugins/Build/HordeServer.Build/Artifacts/ArtifactCollection.cs
+++ b/Engine/Source/Programs/Horde/Plugins/Build/HordeServer.Build/Artifacts/ArtifactCollection.cs
@@ -30,6 +30,7 @@ namespace HordeServer.Artifacts
 	{
 		class Artifact : IArtifact
 		{
+			readonly ArtifactCollection _collection;
 			readonly ArtifactDocument _document;
 
 			public ArtifactDocument Document => _document;
@@ -46,8 +47,14 @@ namespace HordeServer.Artifacts
 			RefName IArtifact.RefName => _document.RefName;
 			DateTime IArtifact.CreatedAtUtc => (_document.CreatedAtUtc == default) ? _document.Id.CreationTime : _document.CreatedAtUtc;
 
-			public Artifact(ArtifactDocument document)
-				=> _document = document;
+			public Artifact(ArtifactCollection collection, ArtifactDocument document)
+			{
+				_collection = collection;
+				_document = document;
+			}
+
+			public Task DeleteAsync(CancellationToken cancellationToken)
+				=> _collection.DeleteArtifactAsync(_document, cancellationToken);
 		}
 
 		class ArtifactDocument
@@ -92,9 +99,6 @@ namespace HordeServer.Artifacts
 			[BsonElement("ref")]
 			public RefName RefName { get; set; }
 
-			[BsonElement("scp")]
-			public AclScopeName AclScope { get; set; }
-
 			[BsonElement("cre")]
 			public DateTime CreatedAtUtc { get; set; }
 
@@ -106,7 +110,7 @@ namespace HordeServer.Artifacts
 			{
 			}
 
-			public ArtifactDocument(ObjectId id, ArtifactName name, ArtifactType type, string? description, StreamId streamId, CommitIdWithOrder commitId, IEnumerable<string> keys, IEnumerable<string> metadata, NamespaceId namespaceId, RefName refName, DateTime createdAtUtc, AclScopeName scopeName)
+			public ArtifactDocument(ObjectId id, ArtifactName name, ArtifactType type, string? description, StreamId streamId, CommitIdWithOrder commitId, IEnumerable<string> keys, IEnumerable<string> metadata, NamespaceId namespaceId, RefName refName, DateTime createdAtUtc)
 			{
 				Id = id;
 				Name = name;
@@ -119,7 +123,6 @@ namespace HordeServer.Artifacts
 				NamespaceId = namespaceId;
 				RefName = refName;
 				CreatedAtUtc = createdAtUtc;
-				AclScope = scopeName;
 			}
 		}
 
@@ -193,7 +196,7 @@ namespace HordeServer.Artifacts
 		public static string GetArtifactPath(StreamId streamId, ArtifactType type) => $"{type}/{streamId}";
 
 		/// <inheritdoc/>
-		public async Task<IArtifact> AddAsync(ArtifactName name, ArtifactType type, string? description, StreamId streamId, CommitId commitId, IEnumerable<string> keys, IEnumerable<string> metadata, AclScopeName scopeName, CancellationToken cancellationToken = default)
+		public async Task<IArtifact> AddAsync(ArtifactName name, ArtifactType type, string? description, StreamId streamId, CommitId commitId, IEnumerable<string> keys, IEnumerable<string> metadata, CancellationToken cancellationToken = default)
 		{
 			if (name.Id.IsEmpty)
 			{
@@ -222,10 +225,10 @@ namespace HordeServer.Artifacts
 
 			CommitIdWithOrder commitIdWithOrder = await _commitService.GetOrderedAsync(streamId, commitId, cancellationToken);
 
-			ArtifactDocument artifactDocument = new ArtifactDocument(id, name, type, description, streamId, commitIdWithOrder, keys.Select(x => NormalizeKey(x)), metadata, namespaceId, refName, _clock.UtcNow, scopeName);
+			ArtifactDocument artifactDocument = new ArtifactDocument(id, name, type, description, streamId, commitIdWithOrder, keys.Select(x => NormalizeKey(x)), metadata, namespaceId, refName, _clock.UtcNow);
 			await _artifactCollection.InsertOneAsync(artifactDocument, null, cancellationToken);
 
-			return new Artifact(artifactDocument);
+			return new Artifact(this, artifactDocument);
 		}
 
 		async Task AddExpiryRecordAsync(StreamId streamId, ArtifactType type, DateTime utcNow, CancellationToken cancellationToken = default)
@@ -271,7 +274,7 @@ namespace HordeServer.Artifacts
 				{
 					foreach (ArtifactDocument artifactDocument in cursor.Current)
 					{
-						yield return new Artifact(artifactDocument);
+						yield return new Artifact(this, artifactDocument);
 					}
 				}
 			}
@@ -361,6 +364,17 @@ namespace HordeServer.Artifacts
 			}
 		}
 
+		async Task DeleteArtifactAsync(ArtifactDocument artifact, CancellationToken cancellationToken)
+		{
+			_logger.LogInformation("Deleting {StreamId} artifact {ArtifactId}, ref {RefName} (created {CreateTime})", artifact.StreamId, artifact.Id, artifact.RefName, artifact.CreatedAtUtc);
+
+			IStorageClient storageClient = _storageService.CreateClient(artifact.NamespaceId);
+			await storageClient.DeleteRefAsync(artifact.RefName, cancellationToken);
+
+			FilterDefinition<ArtifactDocument> filter = Builders<ArtifactDocument>.Filter.Eq(x => x.Id, artifact.Id);
+			await _artifactCollection.DeleteOneAsync(filter, cancellationToken);
+		}
+
 		async Task DeleteArtifactsAsync(IEnumerable<ArtifactDocument> deleteArtifacts, CancellationToken cancellationToken)
 		{
 			foreach (IGrouping<NamespaceId, ArtifactDocument> artifactGroup in deleteArtifacts.GroupBy(x => x.NamespaceId))
@@ -389,7 +403,7 @@ namespace HordeServer.Artifacts
 				return null;
 			}
 
-			return new Artifact(document);
+			return new Artifact(this, document);
 		}
 	}
 }
diff --git a/Engine/Source/Programs/Horde/Plugins/Build/HordeServer.Build/Artifacts/ArtifactsController.cs b/Engine/Source/Programs/Horde/Plugins/Build/HordeServer.Build/Artifacts/ArtifactsController.cs
--- a/Engine/Source/Programs/Horde/Plugins/Build/HordeServer.Build/Artifacts/ArtifactsController.cs
+++ b/Engine/Source/Programs/Horde/Plugins/Build/HordeServer.Build/Artifacts/ArtifactsController.cs
@@ -85,7 +85,7 @@ namespace HordeServer.Artifacts
 			{
 				if (streamConfig.Authorize(ArtifactAclAction.WriteArtifact, User))
 				{
-					return await CreateArtifactInternalAsync(request.Name, request.Type, request.Description, request.StreamId.Value, request.CommitId, request.Keys, request.Metadata, AclScopeName.Root, cancellationToken);
+					return await CreateArtifactInternalAsync(request.Name, request.Type, request.Description, request.StreamId.Value, request.CommitId, request.Keys, request.Metadata, cancellationToken);
 				}
 			}
 
@@ -137,15 +137,15 @@ namespace HordeServer.Artifacts
 					scopeName = templateRefConfig.Acl.ScopeName;
 				}
 
-				return await CreateArtifactInternalAsync(request.Name, request.Type, request.Description, streamId, request.CommitId ?? job.CommitId, keys, request.Metadata, scopeName, cancellationToken);
+				return await CreateArtifactInternalAsync(request.Name, request.Type, request.Description, streamId, request.CommitId ?? job.CommitId, keys, request.Metadata, cancellationToken);
 			}
 
 			return Forbid(ArtifactAclAction.WriteArtifact);
 		}
 
-		async Task<ActionResult<CreateArtifactResponse>> CreateArtifactInternalAsync(ArtifactName name, ArtifactType type, string? description, StreamId streamId, CommitId commitId, List<string> keys, List<string> metadata, AclScopeName scopeName, CancellationToken cancellationToken)
+		async Task<ActionResult<CreateArtifactResponse>> CreateArtifactInternalAsync(ArtifactName name, ArtifactType type, string? description, StreamId streamId, CommitId commitId, List<string> keys, List<string> metadata, CancellationToken cancellationToken)
 		{
-			IArtifact artifact = await _artifactCollection.AddAsync(name, type, description, streamId, commitId, keys, metadata, scopeName, cancellationToken);
+			IArtifact artifact = await _artifactCollection.AddAsync(name, type, description, streamId, commitId, keys, metadata, cancellationToken);
 			RefName? prevRefName = await GetPrevRefNameForArtifactAsync(artifact, cancellationToken);
 
 			List<AclClaimConfig> claims = new List<AclClaimConfig>();
@@ -233,6 +233,28 @@ namespace HordeServer.Artifacts
 			return PropertyFilter.Apply(new GetArtifactResponse(artifact), filter);
 		}
 
+		/// <summary>
+		/// Deletes an artifact object
+		/// </summary>
+		/// <param name="id">Identifier of the artifact to retrieve</param>
+		[HttpDelete]
+		[Route("/api/v2/artifacts/{id}")]
+		public async Task<ActionResult> DeleteArtifactAsync(ArtifactId id, CancellationToken cancellationToken)
+		{
+			IArtifact? artifact = await _artifactCollection.GetAsync(id, HttpContext.RequestAborted);
+			if (artifact == null)
+			{
+				return NotFound(id);
+			}
+			if (!_buildConfig.AuthorizeArtifact(artifact.Type, artifact.StreamId, ArtifactAclAction.DeleteArtifact, User))
+			{
+				return Forbid(ArtifactAclAction.DeleteArtifact, artifact.StreamId);
+			}
+
+			await artifact.DeleteAsync(cancellationToken);
+			return Ok();
+		}
+
 		/// <summary>
 		/// Retrieves bundles for a particular artifact
 		/// </summary>
diff --git a/Engine/Source/Programs/Shared/EpicGames.Horde/Artifacts/ArtifactHttpCollection.cs b/Engine/Source/Programs/Shared/EpicGames.Horde/Artifacts/ArtifactHttpCollection.cs
--- a/Engine/Source/Programs/Shared/EpicGames.Horde/Artifacts/ArtifactHttpCollection.cs
+++ b/Engine/Source/Programs/Shared/EpicGames.Horde/Artifacts/ArtifactHttpCollection.cs
@@ -16,10 +16,12 @@ namespace EpicGames.Horde.Artifacts
 	{
 		class Artifact : IArtifact
 		{
+			readonly ArtifactHttpCollection _collection;
 			readonly GetArtifactResponse _response;
 
-			public Artifact(GetArtifactResponse response)
+			public Artifact(ArtifactHttpCollection collection, GetArtifactResponse response)
 			{
+				_collection = collection;
 				_response = response;
 			}
 
@@ -34,6 +36,9 @@ namespace EpicGames.Horde.Artifacts
 			public NamespaceId NamespaceId => _response.NamespaceId;
 			public RefName RefName => _response.RefName;
 			public DateTime CreatedAtUtc => _response.CreatedAtUtc;
+
+			public Task DeleteAsync(CancellationToken cancellationToken)
+				=> _collection.DeleteAsync(_response.Id, cancellationToken);
 		}
 
 		readonly IHordeClient _hordeClient;
@@ -42,7 +47,7 @@ namespace EpicGames.Horde.Artifacts
 			=> _hordeClient = hordeClient;
 
 		/// <inheritdoc/>
-		public Task<IArtifact> AddAsync(ArtifactName name, ArtifactType type, string? description, StreamId streamId, CommitId commitId, IEnumerable<string> keys, IEnumerable<string> metadata, AclScopeName scopeName, CancellationToken cancellationToken = default)
+		public Task<IArtifact> AddAsync(ArtifactName name, ArtifactType type, string? description, StreamId streamId, CommitId commitId, IEnumerable<string> keys, IEnumerable<string> metadata, CancellationToken cancellationToken = default)
 		{
 			throw new NotImplementedException();
 		}
@@ -55,16 +60,22 @@ namespace EpicGames.Horde.Artifacts
 			List<GetArtifactResponse> responses = await hordeHttpClient.FindArtifactsAsync(streamId, minCommitId, maxCommitId, name, type, keys, maxResults, cancellationToken);
 			foreach (GetArtifactResponse response in responses)
 			{
-				yield return new Artifact(response);
+				yield return new Artifact(this, response);
 			}
 		}
 
+		async Task DeleteAsync(ArtifactId artifactId, CancellationToken cancellationToken = default)
+		{
+			HordeHttpClient hordeHttpClient = _hordeClient.CreateHttpClient();
+			await hordeHttpClient.DeleteArtifactAsync(artifactId, cancellationToken);
+		}
+
 		/// <inheritdoc/>
 		public async Task<IArtifact?> GetAsync(ArtifactId artifactId, CancellationToken cancellationToken = default)
 		{
 			HordeHttpClient hordeHttpClient = _hordeClient.CreateHttpClient();
 			GetArtifactResponse? response = await hordeHttpClient.GetArtifactAsync(artifactId, cancellationToken);
-			return new Artifact(response);
+			return new Artifact(this, response);
 		}
 	}
 }
diff --git a/Engine/Source/Programs/Shared/EpicGames.Horde/Artifacts/IArtifact.cs b/Engine/Source/Programs/Shared/EpicGames.Horde/Artifacts/IArtifact.cs
--- a/Engine/Source/Programs/Shared/EpicGames.Horde/Artifacts/IArtifact.cs
+++ b/Engine/Source/Programs/Shared/EpicGames.Horde/Artifacts/IArtifact.cs
@@ -2,6 +2,8 @@
 
 using System;
 using System.Collections.Generic;
+using System.Threading;
+using System.Threading.Tasks;
 using EpicGames.Horde.Commits;
 using EpicGames.Horde.Storage;
 using EpicGames.Horde.Streams;
@@ -16,56 +18,62 @@ namespace EpicGames.Horde.Artifacts
 		/// <summary>
 		/// Identifier for the Artifact. Randomly generated.
 		/// </summary>
-		public ArtifactId Id { get; }
+		ArtifactId Id { get; }
 
 		/// <summary>
 		/// Name of the artifact
 		/// </summary>
-		public ArtifactName Name { get; }
+		ArtifactName Name { get; }
 
 		/// <summary>
 		/// Type of artifact
 		/// </summary>
-		public ArtifactType Type { get; }
+		ArtifactType Type { get; }
 
 		/// <summary>
 		/// Description for the artifact
 		/// </summary>
-		public string? Description { get; }
+		string? Description { get; }
 
 		/// <summary>
 		/// Identifier for the stream that produced the artifact
 		/// </summary>
-		public StreamId StreamId { get; }
+		StreamId StreamId { get; }
 
 		/// <summary>
 		/// Change that the artifact corresponds to
 		/// </summary>
-		public CommitIdWithOrder CommitId { get; }
+		CommitIdWithOrder CommitId { get; }
 
 		/// <summary>
 		/// Keys used to collate artifacts
 		/// </summary>
-		public IReadOnlyList<string> Keys { get; }
+		IReadOnlyList<string> Keys { get; }
 
 		/// <summary>
 		/// Metadata for the artifact
 		/// </summary>
-		public IReadOnlyList<string> Metadata { get; }
+		IReadOnlyList<string> Metadata { get; }
 
 		/// <summary>
 		/// Storage namespace containing the data
 		/// </summary>
-		public NamespaceId NamespaceId { get; }
+		NamespaceId NamespaceId { get; }
 
 		/// <summary>
 		/// Name of the ref containing the root data object
 		/// </summary>
-		public RefName RefName { get; }
+		RefName RefName { get; }
 
 		/// <summary>
 		/// Time at which the artifact was created
 		/// </summary>
-		public DateTime CreatedAtUtc { get; }
+		DateTime CreatedAtUtc { get; }
+
+		/// <summary>
+		/// Deletes this artifact
+		/// </summary>
+		/// <param name="cancellationToken">Cancellation token for the operation</param>
+		Task DeleteAsync(CancellationToken cancellationToken);
 	}
 }
diff --git a/Engine/Source/Programs/Shared/EpicGames.Horde/Artifacts/IArtifactCollection.cs b/Engine/Source/Programs/Shared/EpicGames.Horde/Artifacts/IArtifactCollection.cs
--- a/Engine/Source/Programs/Shared/EpicGames.Horde/Artifacts/IArtifactCollection.cs
+++ b/Engine/Source/Programs/Shared/EpicGames.Horde/Artifacts/IArtifactCollection.cs
@@ -5,6 +5,7 @@ using System.Threading;
 using System.Threading.Tasks;
 using EpicGames.Horde.Acls;
 using EpicGames.Horde.Commits;
+using EpicGames.Horde.Storage;
 using EpicGames.Horde.Streams;
 
 namespace EpicGames.Horde.Artifacts
@@ -24,10 +25,9 @@ namespace EpicGames.Horde.Artifacts
 		/// <param name="commitId">Commit that the artifact was built from</param>
 		/// <param name="keys">Keys for the artifact</param>
 		/// <param name="metadata">Metadata for the artifact</param>
-		/// <param name="scopeName">Inherited scope used for permissions</param>
 		/// <param name="cancellationToken">Cancellation token for the operation</param>
 		/// <returns>The new log file document</returns>
-		Task<IArtifact> AddAsync(ArtifactName name, ArtifactType type, string? description, StreamId streamId, CommitId commitId, IEnumerable<string> keys, IEnumerable<string> metadata, AclScopeName scopeName, CancellationToken cancellationToken = default);
+		Task<IArtifact> AddAsync(ArtifactName name, ArtifactType type, string? description, StreamId streamId, CommitId commitId, IEnumerable<string> keys, IEnumerable<string> metadata, CancellationToken cancellationToken = default);
 
 		/// <summary>
 		/// Finds artifacts with the given keys.
diff --git a/Engine/Source/Programs/Shared/EpicGames.Horde/HordeHttpClient.cs b/Engine/Source/Programs/Shared/EpicGames.Horde/HordeHttpClient.cs
--- a/Engine/Source/Programs/Shared/EpicGames.Horde/HordeHttpClient.cs
+++ b/Engine/Source/Programs/Shared/EpicGames.Horde/HordeHttpClient.cs
@@ -3,6 +3,7 @@
 using System;
 using System.Collections.Generic;
 using System.IO;
+using System.Linq;
 using System.Net;
 using System.Net.Http;
 using System.Net.Http.Headers;
@@ -139,9 +140,19 @@ namespace EpicGames.Horde
 		/// <param name="keys">Keys used to identify the artifact</param>
 		/// <param name="metadata">Metadata for the artifact</param>
 		/// <param name="cancellationToken">Cancellation token for the operation</param>
-		public Task<CreateArtifactResponse> CreateArtifactAsync(ArtifactName name, ArtifactType type, string? description, StreamId streamId, CommitId commitId, List<string>? keys = null, List<string>? metadata = null, CancellationToken cancellationToken = default)
+		public Task<CreateArtifactResponse> CreateArtifactAsync(ArtifactName name, ArtifactType type, string? description, StreamId streamId, CommitId commitId, IEnumerable<string>? keys = null, IEnumerable<string>? metadata = null, CancellationToken cancellationToken = default)
 		{
-			return PostAsync<CreateArtifactResponse, CreateArtifactRequest>(_httpClient, $"api/v2/artifacts", new CreateArtifactRequest(name, type, description, streamId, keys ?? new List<string>(), metadata ?? new List<string>()) { CommitId = commitId }, cancellationToken);
+			return PostAsync<CreateArtifactResponse, CreateArtifactRequest>(_httpClient, $"api/v2/artifacts", new CreateArtifactRequest(name, type, description, streamId, keys?.ToList() ?? new List<string>(), metadata?.ToList() ?? new List<string>()) { CommitId = commitId }, cancellationToken);
+		}
+
+		/// <summary>
+		/// Deletes an artifact
+		/// </summary>
+		/// <param name="id">Identifier for the artifact</param>
+		/// <param name="cancellationToken">Cancellation token for the operation</param>
+		public async Task DeleteArtifactAsync(ArtifactId id, CancellationToken cancellationToken = default)
+		{
+			await DeleteAsync(_httpClient, $"api/v2/artifacts/{id}", cancellationToken);
 		}
 
 		/// <summary>
@@ -527,6 +538,18 @@ namespace EpicGames.Horde
 
 		#region Utility Methods
 
+		/// <summary>
+		/// Deletes a resource from an HTTP endpoint
+		/// </summary>
+		/// <param name="httpClient">Http client instance</param>
+		/// <param name="relativePath">The url to retrieve</param>
+		/// <param name="cancellationToken">Cancels the request</param>
+		internal static async Task DeleteAsync(HttpClient httpClient, string relativePath, CancellationToken cancellationToken = default)
+		{
+			using HttpResponseMessage response = await httpClient.DeleteAsync(relativePath, cancellationToken);
+			response.EnsureSuccessStatusCode();
+		}
+
 		/// <summary>
 		/// Gets a resource from an HTTP endpoint and parses it as a JSON object
 		/// </summary>
