Horde: Missing files.

#rnx

diff --git a/Engine/Source/Programs/Horde/Plugins/Build/HordeServer.Build.Tests/Artifacts/ArtifactTests.cs b/Engine/Source/Programs/Horde/Plugins/Build/HordeServer.Build.Tests/Artifacts/ArtifactTests.cs
--- a/Engine/Source/Programs/Horde/Plugins/Build/HordeServer.Build.Tests/Artifacts/ArtifactTests.cs
+++ b/Engine/Source/Programs/Horde/Plugins/Build/HordeServer.Build.Tests/Artifacts/ArtifactTests.cs
@@ -37,7 +37,7 @@ namespace HordeServer.Tests.Artifacts
 			StreamId streamId = new StreamId("foo");
 
 			IArtifactCollection artifactCollection = ServiceProvider.GetRequiredService<IArtifactCollection>();
-			IArtifact artifact = await artifactCollection.AddAsync(new ArtifactName("default"), ArtifactType.StepOutput, null, streamId, CommitIdWithOrder.FromPerforceChange(1), new string[] { "test1", "test2" }, Array.Empty<string>(), AclScopeName.Root);
+			IArtifact artifact = await artifactCollection.AddAsync(new ArtifactName("default"), ArtifactType.StepOutput, null, streamId, CommitIdWithOrder.FromPerforceChange(1), new string[] { "test1", "test2" }, Array.Empty<string>());
 
 			{
 				List<IArtifact> artifacts = await artifactCollection.FindAsync(streamId, keys: new[] { "test1" }).ToListAsync();
@@ -71,7 +71,7 @@ namespace HordeServer.Tests.Artifacts
 			StreamId streamId = new StreamId("foo");
 
 			IArtifactCollection artifactCollection = ServiceProvider.GetRequiredService<IArtifactCollection>();
-			IArtifact artifact = await artifactCollection.AddAsync(new ArtifactName("default"), type, null, streamId, CommitIdWithOrder.FromPerforceChange(1), new string[] { "test1", "test2" }, Array.Empty<string>(), AclScopeName.Root);
+			IArtifact artifact = await artifactCollection.AddAsync(new ArtifactName("default"), type, null, streamId, CommitIdWithOrder.FromPerforceChange(1), new string[] { "test1", "test2" }, Array.Empty<string>());
 
 			{
 				List<IArtifact> artifacts = await artifactCollection.FindAsync(streamId, keys: new[] { "test1" }).ToListAsync();
@@ -102,7 +102,7 @@ namespace HordeServer.Tests.Artifacts
 			StreamId streamId = new StreamId("foo");
 
 			IArtifactCollection artifactCollection = ServiceProvider.GetRequiredService<IArtifactCollection>();
-			IArtifact artifact = await artifactCollection.AddAsync(new ArtifactName("default"), type, null, streamId, CommitIdWithOrder.FromPerforceChange(1), new string[] { "test1", "test2" }, Array.Empty<string>(), AclScopeName.Root);
+			IArtifact artifact = await artifactCollection.AddAsync(new ArtifactName("default"), type, null, streamId, CommitIdWithOrder.FromPerforceChange(1), new string[] { "test1", "test2" }, Array.Empty<string>());
 
 			{
 				List<IArtifact> artifacts = await artifactCollection.FindAsync(streamId, keys: new[] { "test1" }).ToListAsync();
@@ -144,11 +144,11 @@ namespace HordeServer.Tests.Artifacts
 			IArtifactCollection artifactCollection = ServiceProvider.GetRequiredService<IArtifactCollection>();
 			for (int idx = 0; idx < 10; idx++)
 			{
-				await artifactCollection.AddAsync(new ArtifactName($"default-{idx}"), type, null, fooStreamId, CommitIdWithOrder.FromPerforceChange(1), new string[] { "test1", "test2" }, Array.Empty<string>(), AclScopeName.Root);
+				await artifactCollection.AddAsync(new ArtifactName($"default-{idx}"), type, null, fooStreamId, CommitIdWithOrder.FromPerforceChange(1), new string[] { "test1", "test2" }, Array.Empty<string>());
 			}
 			for (int idx = 0; idx < 10; idx++)
 			{
-				await artifactCollection.AddAsync(new ArtifactName($"default-{idx}"), type, null, barStreamId, CommitIdWithOrder.FromPerforceChange(1), new string[] { "test1", "test2" }, Array.Empty<string>(), AclScopeName.Root);
+				await artifactCollection.AddAsync(new ArtifactName($"default-{idx}"), type, null, barStreamId, CommitIdWithOrder.FromPerforceChange(1), new string[] { "test1", "test2" }, Array.Empty<string>());
 			}
 
 			await Clock.AdvanceAsync(TimeSpan.FromDays(2.0));
@@ -199,7 +199,7 @@ namespace HordeServer.Tests.Artifacts
 			// Orphan the artifact type
 			{
 				UpdateConfig(config => config.Plugins.GetBuildConfig().ArtifactTypes.Add(new ArtifactTypeConfig { Type = type, KeepCount = 4 }));
-				IArtifact artifact = await artifactCollection.AddAsync(new ArtifactName($"default"), type, null, fooStreamId, CommitIdWithOrder.FromPerforceChange(1), new string[] { "test1", "test2" }, Array.Empty<string>(), AclScopeName.Root);
+				IArtifact artifact = await artifactCollection.AddAsync(new ArtifactName($"default"), type, null, fooStreamId, CommitIdWithOrder.FromPerforceChange(1), new string[] { "test1", "test2" }, Array.Empty<string>());
 
 				await Clock.AdvanceAsync(TimeSpan.FromDays(30.0));
 				IArtifact? artifact2 = await artifactCollection.GetAsync(artifact.Id);
@@ -217,7 +217,7 @@ namespace HordeServer.Tests.Artifacts
 			// Orphan the stream
 			{
 				UpdateConfig(config => config.Plugins.GetBuildConfig().ArtifactTypes.Add(new ArtifactTypeConfig { Type = type, KeepCount = 4 }));
-				IArtifact artifact = await artifactCollection.AddAsync(new ArtifactName($"default"), type, null, fooStreamId, CommitIdWithOrder.FromPerforceChange(1), new string[] { "test1", "test2" }, Array.Empty<string>(), AclScopeName.Root);
+				IArtifact artifact = await artifactCollection.AddAsync(new ArtifactName($"default"), type, null, fooStreamId, CommitIdWithOrder.FromPerforceChange(1), new string[] { "test1", "test2" }, Array.Empty<string>());
 
 				await Clock.AdvanceAsync(TimeSpan.FromDays(30.0));
 				IArtifact? artifact2 = await artifactCollection.GetAsync(artifact.Id);
diff --git a/Engine/Source/Programs/Horde/Plugins/Build/HordeServer.Build/Jobs/JobRpcService.cs b/Engine/Source/Programs/Horde/Plugins/Build/HordeServer.Build/Jobs/JobRpcService.cs
--- a/Engine/Source/Programs/Horde/Plugins/Build/HordeServer.Build/Jobs/JobRpcService.cs
+++ b/Engine/Source/Programs/Horde/Plugins/Build/HordeServer.Build/Jobs/JobRpcService.cs
@@ -818,7 +818,7 @@ namespace HordeServer.Jobs
 				description = request.Name;
 			}
 
-			IArtifact artifact = await _artifactCollection.AddAsync(name, type, description, job.StreamId, job.CommitId, keys, request.Metadata, templateConfig.Acl.ScopeName, context.CancellationToken);
+			IArtifact artifact = await _artifactCollection.AddAsync(name, type, description, job.StreamId, job.CommitId, keys, request.Metadata, context.CancellationToken);
 
 			List<AclClaimConfig> claims = new List<AclClaimConfig>();
 			claims.Add(new AclClaimConfig(HordeClaimTypes.WriteNamespace, $"{artifact.NamespaceId}:{artifact.RefName}"));
