Horde: Register the agent collection as a hosted service.

#rnx

diff --git a/Engine/Source/Programs/Horde/Plugins/Compute/HordeServer.Compute/ComputePlugin.cs b/Engine/Source/Programs/Horde/Plugins/Compute/HordeServer.Compute/ComputePlugin.cs
--- a/Engine/Source/Programs/Horde/Plugins/Compute/HordeServer.Compute/ComputePlugin.cs
+++ b/Engine/Source/Programs/Horde/Plugins/Compute/HordeServer.Compute/ComputePlugin.cs
@@ -66,7 +66,9 @@ namespace HordeServer
 			services.AddSingleton<IDefaultAclModifier, ComputeAclModifier>();
 			services.AddSingleton<IPluginResponseFilter, ComputeResponseFilter>();
 
-			services.AddSingleton<IAgentCollection, AgentCollection>();
+			services.AddSingleton<AgentCollection>();
+			services.AddSingleton<IAgentCollection>(sp => sp.GetRequiredService<AgentCollection>());
+			services.AddHostedService(sp => sp.GetRequiredService<AgentCollection>());
 
 			services.AddSingleton<LeaseCollection>();
 			services.AddSingleton<ILeaseCollection>(sp => sp.GetRequiredService<LeaseCollection>());
