Horde: Separating out static analysis fixes and renames from problematic compute changes.

#rnx

diff --git a/Engine/Source/Programs/Horde/HordeServer/Startup.cs b/Engine/Source/Programs/Horde/HordeServer/Startup.cs
--- a/Engine/Source/Programs/Horde/HordeServer/Startup.cs
+++ b/Engine/Source/Programs/Horde/HordeServer/Startup.cs
@@ -358,8 +358,6 @@ namespace HordeServer
 			services.AddHostedService(provider => provider.GetRequiredService<ConfigService>());
 
 			// Auditing
-			services.AddSingleton<IAuditLog<AgentId>>(sp => sp.GetRequiredService<IAuditLogFactory<AgentId>>().Create("Agents.Log", "AgentId"));
-
 			services.AddSingleton(typeof(IAuditLogFactory<>), typeof(AuditLogFactory<>));
 			services.AddSingleton(typeof(ISingletonDocument<>), typeof(SingletonDocument<>));
 
diff --git a/Engine/Source/Programs/Horde/Plugins/Build/HordeServer.Build.Tests/Agents/AgentCollectionTests.cs b/Engine/Source/Programs/Horde/Plugins/Build/HordeServer.Build.Tests/Agents/AgentCollectionTests.cs
deleted file mode 100644
diff --git a/Engine/Source/Programs/Horde/Plugins/Build/HordeServer.Build.Tests/Artifacts/ArtifactTests.cs b/Engine/Source/Programs/Horde/Plugins/Build/HordeServer.Build.Tests/Artifacts/ArtifactTests.cs
--- a/Engine/Source/Programs/Horde/Plugins/Build/HordeServer.Build.Tests/Artifacts/ArtifactTests.cs
+++ b/Engine/Source/Programs/Horde/Plugins/Build/HordeServer.Build.Tests/Artifacts/ArtifactTests.cs
@@ -1,6 +1,5 @@
 // Copyright Epic Games, Inc. All Rights Reserved.
 
-using EpicGames.Horde.Acls;
 using EpicGames.Horde.Artifacts;
 using EpicGames.Horde.Commits;
 using EpicGames.Horde.Streams;
diff --git a/Engine/Source/Programs/Horde/Plugins/Build/HordeServer.Build.Tests/BuildTestSetup.cs b/Engine/Source/Programs/Horde/Plugins/Build/HordeServer.Build.Tests/BuildTestSetup.cs
--- a/Engine/Source/Programs/Horde/Plugins/Build/HordeServer.Build.Tests/BuildTestSetup.cs
+++ b/Engine/Source/Programs/Horde/Plugins/Build/HordeServer.Build.Tests/BuildTestSetup.cs
@@ -18,7 +18,6 @@ using HordeServer.Agents.Relay;
 using HordeServer.Agents.Telemetry;
 using HordeServer.Agents.Utilization;
 using HordeServer.Artifacts;
-using HordeServer.Auditing;
 using HordeServer.Commits;
 using HordeServer.Compute;
 using HordeServer.Dashboard;
@@ -123,8 +122,6 @@ namespace HordeServer.Tests
 
 			services.AddHttpClient<RpcService>();
 
-			services.AddSingleton<IAuditLog<AgentId>>(sp => sp.GetRequiredService<IAuditLogFactory<AgentId>>().Create("Agents.Log", "AgentId"));
-
 			services.AddSingleton<IAgentCollection, AgentCollection>();
 			services.AddSingleton<IAgentTelemetryCollection, AgentTelemetryCollection>();
 			services.AddSingleton<ArtifactCollection>();
diff --git a/Engine/Source/Programs/Horde/Plugins/Build/HordeServer.Build/Artifacts/ArtifactCollection.cs b/Engine/Source/Programs/Horde/Plugins/Build/HordeServer.Build/Artifacts/ArtifactCollection.cs
--- a/Engine/Source/Programs/Horde/Plugins/Build/HordeServer.Build/Artifacts/ArtifactCollection.cs
+++ b/Engine/Source/Programs/Horde/Plugins/Build/HordeServer.Build/Artifacts/ArtifactCollection.cs
@@ -3,7 +3,6 @@
 using System.Diagnostics;
 using System.Runtime.CompilerServices;
 using EpicGames.Core;
-using EpicGames.Horde.Acls;
 using EpicGames.Horde.Artifacts;
 using EpicGames.Horde.Commits;
 using EpicGames.Horde.Storage;
diff --git a/Engine/Source/Programs/Horde/Plugins/Build/HordeServer.Build/Artifacts/ArtifactsController.cs b/Engine/Source/Programs/Horde/Plugins/Build/HordeServer.Build/Artifacts/ArtifactsController.cs
--- a/Engine/Source/Programs/Horde/Plugins/Build/HordeServer.Build/Artifacts/ArtifactsController.cs
+++ b/Engine/Source/Programs/Horde/Plugins/Build/HordeServer.Build/Artifacts/ArtifactsController.cs
@@ -237,6 +237,7 @@ namespace HordeServer.Artifacts
 		/// Deletes an artifact object
 		/// </summary>
 		/// <param name="id">Identifier of the artifact to retrieve</param>
+		/// <param name="cancellationToken">Cancellation token for the operation</param>
 		[HttpDelete]
 		[Route("/api/v2/artifacts/{id}")]
 		public async Task<ActionResult> DeleteArtifactAsync(ArtifactId id, CancellationToken cancellationToken)
diff --git a/Engine/Source/Programs/Horde/Plugins/Build/HordeServer.Build/Issues/IssueService.cs b/Engine/Source/Programs/Horde/Plugins/Build/HordeServer.Build/Issues/IssueService.cs
--- a/Engine/Source/Programs/Horde/Plugins/Build/HordeServer.Build/Issues/IssueService.cs
+++ b/Engine/Source/Programs/Horde/Plugins/Build/HordeServer.Build/Issues/IssueService.cs
@@ -17,7 +17,6 @@ using HordeServer.Commits;
 using HordeServer.Jobs;
 using HordeServer.Jobs.Graphs;
 using HordeServer.Logs;
-using HordeServer.Server;
 using HordeServer.Streams;
 using HordeServer.Users;
 using Microsoft.Extensions.DependencyInjection;
diff --git a/Engine/Source/Programs/Horde/Plugins/Build/HordeServer.Build/Jobs/JobRpcService.cs b/Engine/Source/Programs/Horde/Plugins/Build/HordeServer.Build/Jobs/JobRpcService.cs
--- a/Engine/Source/Programs/Horde/Plugins/Build/HordeServer.Build/Jobs/JobRpcService.cs
+++ b/Engine/Source/Programs/Horde/Plugins/Build/HordeServer.Build/Jobs/JobRpcService.cs
@@ -807,11 +807,6 @@ namespace HordeServer.Jobs
 			keys.Add(job.GetArtifactKey(step));
 			keys.AddRange(request.Keys);
 
-			if (!_buildConfig.Value.TryGetTemplate(job.StreamId, job.TemplateId, out TemplateRefConfig? templateConfig))
-			{
-				throw new StructuredRpcException(StatusCode.NotFound, "Couldn't find template {TemplateId} in stream {StreamId}", job.TemplateId, job.StreamId);
-			}
-
 			string? description = request.Description;
 			if (String.IsNullOrEmpty(description))
 			{
diff --git a/Engine/Source/Programs/Horde/Plugins/Compute/HordeServer.Compute.Tests/Agents/AgentCollectionTests.cs b/Engine/Source/Programs/Horde/Plugins/Compute/HordeServer.Compute.Tests/Agents/AgentCollectionTests.cs
new file mode 100644
--- /dev/null
+++ b/Engine/Source/Programs/Horde/Plugins/Compute/HordeServer.Compute.Tests/Agents/AgentCollectionTests.cs
@@ -0,0 +1,175 @@
+// Copyright Epic Games, Inc. All Rights Reserved.
+
+using EpicGames.Horde.Agents;
+using EpicGames.Horde.Agents.Leases;
+using EpicGames.Horde.Agents.Pools;
+using Google.Protobuf.WellKnownTypes;
+using HordeCommon.Rpc.Messages;
+using HordeServer.Agents;
+
+namespace HordeServer.Tests.Agents;
+
+[TestClass]
+public class AgentCollectionTests : ComputeTestSetup
+{
+	private IAgent _agent;
+	private readonly CreateLeaseOptions _lease1;
+	private readonly CreateLeaseOptions _lease2;
+	private readonly CreateLeaseOptions _leaseWithParent3;
+	private readonly CreateLeaseOptions _leaseWithParent4;
+
+	public AgentCollectionTests()
+	{
+		_agent = AgentCollection.AddAsync(new AgentId("test"), ephemeral: true, enrollmentKey: "").Result;
+		_lease1 = new CreateLeaseOptions(LeaseId.Parse("aaaaaaaaaaaaaaaaaaaaaaaa"), null, "lease", null, null, null, null, false, new Empty());
+		_lease2 = new CreateLeaseOptions(LeaseId.Parse("bbbbbbbbbbbbbbbbbbbbbbbb"), null, "lease", null, null, null, null, false, new Empty());
+		_leaseWithParent3 = new CreateLeaseOptions(LeaseId.Parse("cccccccccccccccccccccccc"), _lease1.Id, "leaseWithParent3", null, null, null, null, false, new Empty());
+		_leaseWithParent4 = new CreateLeaseOptions(LeaseId.Parse("dddddddddddddddddddddddd"), _lease1.Id, "leaseWithParent4", null, null, null, null, false, new Empty());
+	}
+
+	[TestMethod]
+	public async Task AddLeaseAsync()
+	{
+		await _agent.TryCreateLeaseAsync(_lease1);
+		await UpdateAgentAsync();
+
+		await _agent.TryCreateLeaseAsync(_leaseWithParent3);
+		await UpdateAgentAsync();
+
+		await _agent.TryCreateLeaseAsync(_leaseWithParent4);
+		await UpdateAgentAsync();
+
+		List<LeaseId> leases = await AgentCollection.FindActiveLeaseIdsAsync();
+
+		Assert.AreEqual(3, leases.Count);
+		Assert.IsTrue(leases.Contains(_lease1.Id));
+		Assert.IsTrue(leases.Contains(_leaseWithParent3.Id));
+		Assert.IsTrue(leases.Contains(_leaseWithParent4.Id));
+	}
+
+	[TestMethod]
+	public async Task StartSessionAsync()
+	{
+		await _agent.TryCreateLeaseAsync(_lease1);
+		await UpdateAgentAsync();
+
+		await _agent.TryCreateLeaseAsync(_lease2);
+		await UpdateAgentAsync();
+
+		await _agent.TryCreateSessionAsync(new CreateSessionOptions(AgentStatus.Ok,
+			new RpcAgentCapabilities(), new List<PoolId>(), DateTime.UtcNow, null));
+
+		List<LeaseId> leases = await AgentCollection.FindActiveLeaseIdsAsync();
+
+		Assert.AreEqual(0, leases.Count);
+	}
+
+	[TestMethod]
+	public async Task UpdateSession_WithEmptyLeasesAsync()
+	{
+		await _agent.TryCreateLeaseAsync(_lease1);
+		await UpdateAgentAsync();
+
+		await _agent.TryCreateLeaseAsync(_lease2);
+		await UpdateAgentAsync();
+
+		await _agent.TryUpdateSessionAsync(new UpdateSessionOptions { Leases = new List<AgentLease>() });
+
+		List<LeaseId> leases = await AgentCollection.FindActiveLeaseIdsAsync();
+
+		Assert.AreEqual(0, leases.Count);
+	}
+
+	[TestMethod]
+	public async Task UpdateSession_WithNewLeasesAsync()
+	{
+		await _agent.TryCreateLeaseAsync(_lease1);
+		await UpdateAgentAsync();
+
+		IAgent? newAgent = await _agent.TryCreateLeaseAsync(_lease1);
+		Assert.IsNotNull(newAgent);
+
+		newAgent = await _agent.TryCreateLeaseAsync(_lease2);
+		Assert.IsNotNull(newAgent);
+
+		List<LeaseId> leases = await AgentCollection.FindActiveLeaseIdsAsync();
+
+		Assert.AreEqual(2, leases.Count);
+		Assert.IsTrue(leases.Contains(_lease1.Id));
+		Assert.IsTrue(leases.Contains(_lease2.Id));
+	}
+
+	[TestMethod]
+	public async Task UpdateSession_WithOneLeaseRemovedAsync()
+	{
+		await _agent.TryCreateLeaseAsync(_lease1);
+		await UpdateAgentAsync();
+
+		await _agent.TryCreateLeaseAsync(_lease2);
+		await UpdateAgentAsync();
+
+		await _agent.TryUpdateSessionAsync(new UpdateSessionOptions { Leases = new List<AgentLease> { new AgentLease(_lease1.Id, null, LeaseState.Active, null, false, null) } });
+
+		List<LeaseId> leases = await AgentCollection.FindActiveLeaseIdsAsync();
+
+		Assert.AreEqual(1, leases.Count);
+		Assert.IsTrue(leases.Contains(_lease1.Id));
+	}
+
+	[TestMethod]
+	public async Task CancelLeaseAsync()
+	{
+		await _agent.TryCreateLeaseAsync(_lease1);
+		await UpdateAgentAsync();
+
+		await _agent.TryCreateLeaseAsync(_lease2);
+		await UpdateAgentAsync();
+
+		await _agent.TryCancelLeaseAsync(0);
+
+		List<LeaseId> leases = await AgentCollection.FindActiveLeaseIdsAsync();
+
+		Assert.AreEqual(1, leases.Count);
+		Assert.IsTrue(leases.Contains(_lease2.Id));
+	}
+
+	[TestMethod]
+	public async Task TerminateSessionAsync()
+	{
+		await _agent.TryCreateLeaseAsync(_lease1);
+		await UpdateAgentAsync();
+
+		await _agent.TryCreateLeaseAsync(_lease2);
+		await UpdateAgentAsync();
+
+		await _agent.TryTerminateSessionAsync();
+
+		List<LeaseId> leases = await AgentCollection.FindActiveLeaseIdsAsync();
+
+		Assert.AreEqual(0, leases.Count);
+	}
+
+	[TestMethod]
+	public async Task GetChildLeaseIdsAsync()
+	{
+		await _agent.TryCreateLeaseAsync(_lease1);
+		await UpdateAgentAsync();
+
+		await _agent.TryCreateLeaseAsync(_leaseWithParent3);
+		await UpdateAgentAsync();
+
+		await _agent.TryCreateLeaseAsync(_leaseWithParent4);
+		await UpdateAgentAsync();
+
+		List<LeaseId> leases = await AgentCollection.GetChildLeaseIdsAsync(_leaseWithParent3.ParentId!.Value);
+
+		Assert.AreEqual(2, leases.Count);
+		Assert.IsTrue(leases.Contains(_leaseWithParent3.Id));
+		Assert.IsTrue(leases.Contains(_leaseWithParent4.Id));
+	}
+
+	private async Task UpdateAgentAsync()
+	{
+		_agent = (await AgentCollection.GetAsync(_agent.Id))!;
+	}
+}
\ No newline at end of file
diff --git a/Engine/Source/Programs/Horde/Plugins/Compute/HordeServer.Compute.Tests/ComputeTestSetup.cs b/Engine/Source/Programs/Horde/Plugins/Compute/HordeServer.Compute.Tests/ComputeTestSetup.cs
--- a/Engine/Source/Programs/Horde/Plugins/Compute/HordeServer.Compute.Tests/ComputeTestSetup.cs
+++ b/Engine/Source/Programs/Horde/Plugins/Compute/HordeServer.Compute.Tests/ComputeTestSetup.cs
@@ -1,6 +1,6 @@
 // Copyright Epic Games, Inc. All Rights Reserved.
 
-using HordeServer.Acls;
+using HordeServer.Agents;
 using HordeServer.Compute;
 using HordeServer.Logs;
 using HordeServer.Storage;
@@ -11,6 +11,7 @@ namespace HordeServer.Tests
 	public class ComputeTestSetup : ServerTestSetup
 	{
 		public ComputeService ComputeService => ServiceProvider.GetRequiredService<ComputeService>();
+		public IAgentCollection AgentCollection => ServiceProvider.GetRequiredService<IAgentCollection>();
 		public StorageService StorageService => ServiceProvider.GetRequiredService<StorageService>();
 		public ILogCollection LogCollection => ServiceProvider.GetRequiredService<ILogCollection>();
 
@@ -19,12 +20,5 @@ namespace HordeServer.Tests
 			AddPlugin<StoragePlugin>();
 			AddPlugin<ComputePlugin>();
 		}
-
-		protected override void ConfigureServices(IServiceCollection services)
-		{
-			base.ConfigureServices(services);
-
-			services.AddSingleton<IDefaultAclModifier, ComputeAclModifier>();
-		}
 	}
 }
diff --git a/Engine/Source/Programs/Horde/Plugins/Compute/HordeServer.Compute/ComputePlugin.cs b/Engine/Source/Programs/Horde/Plugins/Compute/HordeServer.Compute/ComputePlugin.cs
--- a/Engine/Source/Programs/Horde/Plugins/Compute/HordeServer.Compute/ComputePlugin.cs
+++ b/Engine/Source/Programs/Horde/Plugins/Compute/HordeServer.Compute/ComputePlugin.cs
@@ -6,6 +6,7 @@ using Amazon.CloudWatch;
 using Amazon.EC2;
 using Amazon.Extensions.NETCore.Setup;
 using Amazon.SQS;
+using EpicGames.Horde.Agents;
 using HordeServer.Acls;
 using HordeServer.Agents;
 using HordeServer.Agents.Enrollment;
@@ -14,6 +15,7 @@ using HordeServer.Agents.Leases;
 using HordeServer.Agents.Pools;
 using HordeServer.Agents.Relay;
 using HordeServer.Agents.Telemetry;
+using HordeServer.Auditing;
 using HordeServer.Aws;
 using HordeServer.Compute;
 using HordeServer.Logs;
@@ -59,10 +61,13 @@ namespace HordeServer
 		/// <inheritdoc/>
 		public void ConfigureServices(IServiceCollection services)
 		{
+			services.AddSingleton<IAuditLog<AgentId>>(sp => sp.GetRequiredService<IAuditLogFactory<AgentId>>().Create("Agents.Log", "AgentId"));
+
 			services.AddSingleton<IDefaultAclModifier, ComputeAclModifier>();
 			services.AddSingleton<IPluginResponseFilter, ComputeResponseFilter>();
 
-			services.AddSingleton<IAgentCollection, AgentCollection>();
+			services.AddSingleton<AgentCollection>();
+			services.AddSingleton<IAgentCollection>(sp => sp.GetRequiredService<AgentCollection>());
 
 			services.AddSingleton<LeaseCollection>();
 			services.AddSingleton<ILeaseCollection>(sp => sp.GetRequiredService<LeaseCollection>());
diff --git a/Engine/Source/Programs/Shared/EpicGames.Horde/Artifacts/ArtifactHttpCollection.cs b/Engine/Source/Programs/Shared/EpicGames.Horde/Artifacts/ArtifactHttpCollection.cs
--- a/Engine/Source/Programs/Shared/EpicGames.Horde/Artifacts/ArtifactHttpCollection.cs
+++ b/Engine/Source/Programs/Shared/EpicGames.Horde/Artifacts/ArtifactHttpCollection.cs
@@ -5,7 +5,6 @@ using System.Collections.Generic;
 using System.Runtime.CompilerServices;
 using System.Threading;
 using System.Threading.Tasks;
-using EpicGames.Horde.Acls;
 using EpicGames.Horde.Commits;
 using EpicGames.Horde.Storage;
 using EpicGames.Horde.Streams;
diff --git a/Engine/Source/Programs/Shared/EpicGames.Horde/Artifacts/IArtifactCollection.cs b/Engine/Source/Programs/Shared/EpicGames.Horde/Artifacts/IArtifactCollection.cs
--- a/Engine/Source/Programs/Shared/EpicGames.Horde/Artifacts/IArtifactCollection.cs
+++ b/Engine/Source/Programs/Shared/EpicGames.Horde/Artifacts/IArtifactCollection.cs
@@ -3,9 +3,7 @@
 using System.Collections.Generic;
 using System.Threading;
 using System.Threading.Tasks;
-using EpicGames.Horde.Acls;
 using EpicGames.Horde.Commits;
-using EpicGames.Horde.Storage;
 using EpicGames.Horde.Streams;
 
 namespace EpicGames.Horde.Artifacts
