csproj update coverlet.collector -> 6.0.2

#rnx

diff --git a/Engine/Source/Programs/AutomationTool.Tests/AutomationTool.Tests.csproj b/Engine/Source/Programs/AutomationTool.Tests/AutomationTool.Tests.csproj
--- a/Engine/Source/Programs/AutomationTool.Tests/AutomationTool.Tests.csproj
+++ b/Engine/Source/Programs/AutomationTool.Tests/AutomationTool.Tests.csproj
@@ -13,7 +13,10 @@
     <PackageReference Include="Microsoft.NET.Test.Sdk" Version="17.5.0" />
     <PackageReference Include="MSTest.TestAdapter" Version="2.2.10" />
     <PackageReference Include="MSTest.TestFramework" Version="2.2.10" />
-    <PackageReference Include="coverlet.collector" Version="3.2.0" />
+    <PackageReference Include="coverlet.collector" Version="6.0.2">
+      <PrivateAssets>all</PrivateAssets>
+      <IncludeAssets>runtime; build; native; contentfiles; analyzers; buildtransitive</IncludeAssets>
+    </PackageReference>
   </ItemGroup>
 
   <ItemGroup>
diff --git a/Engine/Source/Programs/Horde/Drivers/JobDriver.Tests/JobDriver.Tests.csproj b/Engine/Source/Programs/Horde/Drivers/JobDriver.Tests/JobDriver.Tests.csproj
--- a/Engine/Source/Programs/Horde/Drivers/JobDriver.Tests/JobDriver.Tests.csproj
+++ b/Engine/Source/Programs/Horde/Drivers/JobDriver.Tests/JobDriver.Tests.csproj
@@ -24,7 +24,7 @@
     <PackageReference Include="Moq" Version="4.20.70" />
     <PackageReference Include="MSTest.TestAdapter" Version="3.1.1" />
     <PackageReference Include="MSTest.TestFramework" Version="3.1.1" />
-    <PackageReference Include="coverlet.collector" Version="6.0.0">
+    <PackageReference Include="coverlet.collector" Version="6.0.2">
       <PrivateAssets>all</PrivateAssets>
       <IncludeAssets>runtime; build; native; contentfiles; analyzers; buildtransitive</IncludeAssets>
     </PackageReference>
diff --git a/Engine/Source/Programs/Horde/HordeAgent.Tests/HordeAgent.Tests.csproj b/Engine/Source/Programs/Horde/HordeAgent.Tests/HordeAgent.Tests.csproj
--- a/Engine/Source/Programs/Horde/HordeAgent.Tests/HordeAgent.Tests.csproj
+++ b/Engine/Source/Programs/Horde/HordeAgent.Tests/HordeAgent.Tests.csproj
@@ -24,7 +24,7 @@
     <PackageReference Include="Moq" Version="4.20.70" />
     <PackageReference Include="MSTest.TestAdapter" Version="3.1.1" />
     <PackageReference Include="MSTest.TestFramework" Version="3.1.1" />
-    <PackageReference Include="coverlet.collector" Version="6.0.0">
+    <PackageReference Include="coverlet.collector" Version="6.0.2">
       <PrivateAssets>all</PrivateAssets>
       <IncludeAssets>runtime; build; native; contentfiles; analyzers; buildtransitive</IncludeAssets>
     </PackageReference>
diff --git a/Engine/Source/Programs/Horde/HordeServer.Tests.Shared/HordeServer.Tests.Shared.csproj b/Engine/Source/Programs/Horde/HordeServer.Tests.Shared/HordeServer.Tests.Shared.csproj
--- a/Engine/Source/Programs/Horde/HordeServer.Tests.Shared/HordeServer.Tests.Shared.csproj
+++ b/Engine/Source/Programs/Horde/HordeServer.Tests.Shared/HordeServer.Tests.Shared.csproj
@@ -28,7 +28,7 @@
     <PackageReference Include="MSTest.TestAdapter" Version="3.1.1" />
     <PackageReference Include="MSTest.TestFramework" Version="3.1.1" />
     <PackageReference Include="Microsoft.Extensions.Caching.Memory" Version="8.0.0" />
-    <PackageReference Include="coverlet.collector" Version="6.0.0">
+    <PackageReference Include="coverlet.collector" Version="6.0.2">
       <PrivateAssets>all</PrivateAssets>
       <IncludeAssets>runtime; build; native; contentfiles; analyzers; buildtransitive</IncludeAssets>
     </PackageReference>
diff --git a/Engine/Source/Programs/Horde/HordeServer.Tests/HordeServer.Tests.csproj b/Engine/Source/Programs/Horde/HordeServer.Tests/HordeServer.Tests.csproj
--- a/Engine/Source/Programs/Horde/HordeServer.Tests/HordeServer.Tests.csproj
+++ b/Engine/Source/Programs/Horde/HordeServer.Tests/HordeServer.Tests.csproj
@@ -24,7 +24,7 @@
     <PackageReference Include="MSTest.TestAdapter" Version="3.1.1" />
     <PackageReference Include="MSTest.TestFramework" Version="3.1.1" />
     <PackageReference Include="Microsoft.Extensions.Caching.Memory" Version="8.0.0" />
-    <PackageReference Include="coverlet.collector" Version="6.0.0">
+    <PackageReference Include="coverlet.collector" Version="6.0.2">
       <PrivateAssets>all</PrivateAssets>
       <IncludeAssets>runtime; build; native; contentfiles; analyzers; buildtransitive</IncludeAssets>
     </PackageReference>
@@ -32,7 +32,7 @@
 
   <ItemGroup>
     <ProjectReference Include="..\..\Shared\EpicGames.Serialization\EpicGames.Serialization.csproj" />
-    <ProjectReference Include="..\HordeServer.Tests.Shared\HordeServer.Tests.Shared.csproj"/>
+    <ProjectReference Include="..\HordeServer.Tests.Shared\HordeServer.Tests.Shared.csproj" />
     <ProjectReference Include="..\HordeAgent\HordeAgent.csproj">
       <Aliases>HordeAgent</Aliases>
     </ProjectReference>
diff --git a/Engine/Source/Programs/Horde/Plugins/Analytics/HordeServer.Analytics.Tests/HordeServer.Analytics.Tests.csproj b/Engine/Source/Programs/Horde/Plugins/Analytics/HordeServer.Analytics.Tests/HordeServer.Analytics.Tests.csproj
--- a/Engine/Source/Programs/Horde/Plugins/Analytics/HordeServer.Analytics.Tests/HordeServer.Analytics.Tests.csproj
+++ b/Engine/Source/Programs/Horde/Plugins/Analytics/HordeServer.Analytics.Tests/HordeServer.Analytics.Tests.csproj
@@ -19,7 +19,10 @@
   </PropertyGroup>
 
   <ItemGroup>
-    <PackageReference Include="coverlet.collector" Version="6.0.0" />
+    <PackageReference Include="coverlet.collector" Version="6.0.2">
+      <PrivateAssets>all</PrivateAssets>
+      <IncludeAssets>runtime; build; native; contentfiles; analyzers; buildtransitive</IncludeAssets>
+    </PackageReference>
     <PackageReference Include="Microsoft.NET.Test.Sdk" Version="17.8.0" />
     <PackageReference Include="MSTest.TestAdapter" Version="3.1.1" />
     <PackageReference Include="MSTest.TestFramework" Version="3.1.1" />
diff --git a/Engine/Source/Programs/Horde/Plugins/Build/HordeServer.Build.Tests/HordeServer.Build.Tests.csproj b/Engine/Source/Programs/Horde/Plugins/Build/HordeServer.Build.Tests/HordeServer.Build.Tests.csproj
--- a/Engine/Source/Programs/Horde/Plugins/Build/HordeServer.Build.Tests/HordeServer.Build.Tests.csproj
+++ b/Engine/Source/Programs/Horde/Plugins/Build/HordeServer.Build.Tests/HordeServer.Build.Tests.csproj
@@ -20,7 +20,10 @@
   </PropertyGroup>
 
   <ItemGroup>
-    <PackageReference Include="coverlet.collector" Version="6.0.0" />
+    <PackageReference Include="coverlet.collector" Version="6.0.2">
+      <PrivateAssets>all</PrivateAssets>
+      <IncludeAssets>runtime; build; native; contentfiles; analyzers; buildtransitive</IncludeAssets>
+    </PackageReference>
     <PackageReference Include="Microsoft.NET.Test.Sdk" Version="17.8.0" />
     <PackageReference Include="MSTest.TestAdapter" Version="3.1.1" />
     <PackageReference Include="MSTest.TestFramework" Version="3.1.1" />
diff --git a/Engine/Source/Programs/Horde/Plugins/Compute/HordeServer.Compute.Tests/HordeServer.Compute.Tests.csproj b/Engine/Source/Programs/Horde/Plugins/Compute/HordeServer.Compute.Tests/HordeServer.Compute.Tests.csproj
--- a/Engine/Source/Programs/Horde/Plugins/Compute/HordeServer.Compute.Tests/HordeServer.Compute.Tests.csproj
+++ b/Engine/Source/Programs/Horde/Plugins/Compute/HordeServer.Compute.Tests/HordeServer.Compute.Tests.csproj
@@ -20,7 +20,10 @@
   </PropertyGroup>
 
   <ItemGroup>
-    <PackageReference Include="coverlet.collector" Version="6.0.0" />
+    <PackageReference Include="coverlet.collector" Version="6.0.2">
+      <PrivateAssets>all</PrivateAssets>
+      <IncludeAssets>runtime; build; native; contentfiles; analyzers; buildtransitive</IncludeAssets>
+    </PackageReference>
     <PackageReference Include="Microsoft.NET.Test.Sdk" Version="17.8.0" />
     <PackageReference Include="MSTest.TestAdapter" Version="3.1.1" />
     <PackageReference Include="MSTest.TestFramework" Version="3.1.1" />
diff --git a/Engine/Source/Programs/Horde/Plugins/Ddc/HordeServer.Ddc.Tests/HordeServer.Ddc.Tests.csproj b/Engine/Source/Programs/Horde/Plugins/Ddc/HordeServer.Ddc.Tests/HordeServer.Ddc.Tests.csproj
--- a/Engine/Source/Programs/Horde/Plugins/Ddc/HordeServer.Ddc.Tests/HordeServer.Ddc.Tests.csproj
+++ b/Engine/Source/Programs/Horde/Plugins/Ddc/HordeServer.Ddc.Tests/HordeServer.Ddc.Tests.csproj
@@ -20,7 +20,10 @@
   </PropertyGroup>
 
   <ItemGroup>
-    <PackageReference Include="coverlet.collector" Version="6.0.0" />
+    <PackageReference Include="coverlet.collector" Version="6.0.2">
+      <PrivateAssets>all</PrivateAssets>
+      <IncludeAssets>runtime; build; native; contentfiles; analyzers; buildtransitive</IncludeAssets>
+    </PackageReference>
     <PackageReference Include="Microsoft.NET.Test.Sdk" Version="17.8.0" />
     <PackageReference Include="MSTest.TestAdapter" Version="3.1.1" />
     <PackageReference Include="MSTest.TestFramework" Version="3.1.1" />
diff --git a/Engine/Source/Programs/Horde/Plugins/Storage/HordeServer.Storage.Tests/HordeServer.Storage.Tests.csproj b/Engine/Source/Programs/Horde/Plugins/Storage/HordeServer.Storage.Tests/HordeServer.Storage.Tests.csproj
--- a/Engine/Source/Programs/Horde/Plugins/Storage/HordeServer.Storage.Tests/HordeServer.Storage.Tests.csproj
+++ b/Engine/Source/Programs/Horde/Plugins/Storage/HordeServer.Storage.Tests/HordeServer.Storage.Tests.csproj
@@ -20,7 +20,10 @@
   </PropertyGroup>
 
   <ItemGroup>
-    <PackageReference Include="coverlet.collector" Version="6.0.0" />
+    <PackageReference Include="coverlet.collector" Version="6.0.2">
+      <PrivateAssets>all</PrivateAssets>
+      <IncludeAssets>runtime; build; native; contentfiles; analyzers; buildtransitive</IncludeAssets>
+    </PackageReference>
     <PackageReference Include="Microsoft.NET.Test.Sdk" Version="17.8.0" />
     <PackageReference Include="MSTest.TestAdapter" Version="3.1.1" />
     <PackageReference Include="MSTest.TestFramework" Version="3.1.1" />
diff --git a/Engine/Source/Programs/Horde/Plugins/Tools/HordeServer.Tools.Tests/HordeServer.Tools.Tests.csproj b/Engine/Source/Programs/Horde/Plugins/Tools/HordeServer.Tools.Tests/HordeServer.Tools.Tests.csproj
--- a/Engine/Source/Programs/Horde/Plugins/Tools/HordeServer.Tools.Tests/HordeServer.Tools.Tests.csproj
+++ b/Engine/Source/Programs/Horde/Plugins/Tools/HordeServer.Tools.Tests/HordeServer.Tools.Tests.csproj
@@ -20,7 +20,10 @@
   </PropertyGroup>
 
   <ItemGroup>
-    <PackageReference Include="coverlet.collector" Version="6.0.0" />
+    <PackageReference Include="coverlet.collector" Version="6.0.2">
+      <PrivateAssets>all</PrivateAssets>
+      <IncludeAssets>runtime; build; native; contentfiles; analyzers; buildtransitive</IncludeAssets>
+    </PackageReference>
     <PackageReference Include="Microsoft.NET.Test.Sdk" Version="17.8.0" />
     <PackageReference Include="MSTest.TestAdapter" Version="3.1.1" />
     <PackageReference Include="MSTest.TestFramework" Version="3.1.1" />
diff --git a/Engine/Source/Programs/Shared/EpicGames.BuildGraph.Tests/EpicGames.BuildGraph.Tests.csproj b/Engine/Source/Programs/Shared/EpicGames.BuildGraph.Tests/EpicGames.BuildGraph.Tests.csproj
--- a/Engine/Source/Programs/Shared/EpicGames.BuildGraph.Tests/EpicGames.BuildGraph.Tests.csproj
+++ b/Engine/Source/Programs/Shared/EpicGames.BuildGraph.Tests/EpicGames.BuildGraph.Tests.csproj
@@ -48,7 +48,7 @@
     <PackageReference Include="Microsoft.NET.Test.Sdk" Version="17.3.0" />
     <PackageReference Include="MSTest.TestAdapter" Version="2.2.10" />
     <PackageReference Include="MSTest.TestFramework" Version="2.2.10" />
-    <PackageReference Include="coverlet.collector" Version="3.1.2">
+    <PackageReference Include="coverlet.collector" Version="6.0.2">
       <PrivateAssets>all</PrivateAssets>
       <IncludeAssets>runtime; build; native; contentfiles; analyzers; buildtransitive</IncludeAssets>
     </PackageReference>
diff --git a/Engine/Source/Programs/Shared/EpicGames.Core.Tests/EpicGames.Core.Tests.csproj b/Engine/Source/Programs/Shared/EpicGames.Core.Tests/EpicGames.Core.Tests.csproj
--- a/Engine/Source/Programs/Shared/EpicGames.Core.Tests/EpicGames.Core.Tests.csproj
+++ b/Engine/Source/Programs/Shared/EpicGames.Core.Tests/EpicGames.Core.Tests.csproj
@@ -51,7 +51,7 @@
     <PackageReference Include="Microsoft.NET.Test.Sdk" Version="17.3.0" />
     <PackageReference Include="MSTest.TestAdapter" Version="2.2.10" />
     <PackageReference Include="MSTest.TestFramework" Version="2.2.10" />
-    <PackageReference Include="coverlet.collector" Version="3.1.2">
+    <PackageReference Include="coverlet.collector" Version="6.0.2">
       <PrivateAssets>all</PrivateAssets>
       <IncludeAssets>runtime; build; native; contentfiles; analyzers; buildtransitive</IncludeAssets>
     </PackageReference>
diff --git a/Engine/Source/Programs/Shared/EpicGames.Horde.Tests/EpicGames.Horde.Tests.csproj b/Engine/Source/Programs/Shared/EpicGames.Horde.Tests/EpicGames.Horde.Tests.csproj
--- a/Engine/Source/Programs/Shared/EpicGames.Horde.Tests/EpicGames.Horde.Tests.csproj
+++ b/Engine/Source/Programs/Shared/EpicGames.Horde.Tests/EpicGames.Horde.Tests.csproj
@@ -34,7 +34,7 @@
     <PackageReference Include="Microsoft.Extensions.Caching.Memory" Version="6.0.1" />
     <PackageReference Include="MSTest.TestAdapter" Version="2.2.10" />
     <PackageReference Include="MSTest.TestFramework" Version="2.2.10" />
-    <PackageReference Include="coverlet.collector" Version="3.1.2">
+    <PackageReference Include="coverlet.collector" Version="6.0.2">
       <PrivateAssets>all</PrivateAssets>
       <IncludeAssets>runtime; build; native; contentfiles; analyzers; buildtransitive</IncludeAssets>
     </PackageReference>
diff --git a/Engine/Source/Programs/Shared/EpicGames.IoHash.Tests/EpicGames.IoHash.Tests.csproj b/Engine/Source/Programs/Shared/EpicGames.IoHash.Tests/EpicGames.IoHash.Tests.csproj
--- a/Engine/Source/Programs/Shared/EpicGames.IoHash.Tests/EpicGames.IoHash.Tests.csproj
+++ b/Engine/Source/Programs/Shared/EpicGames.IoHash.Tests/EpicGames.IoHash.Tests.csproj
@@ -20,7 +20,7 @@
     <PackageReference Include="Microsoft.NET.Test.Sdk" Version="17.3.0" />
     <PackageReference Include="MSTest.TestAdapter" Version="2.2.10" />
     <PackageReference Include="MSTest.TestFramework" Version="2.2.10" />
-    <PackageReference Include="coverlet.collector" Version="3.1.2">
+    <PackageReference Include="coverlet.collector" Version="6.0.2">
       <PrivateAssets>all</PrivateAssets>
       <IncludeAssets>runtime; build; native; contentfiles; analyzers; buildtransitive</IncludeAssets>
     </PackageReference>
diff --git a/Engine/Source/Programs/Shared/EpicGames.Perforce.Fixture/EpicGames.Perforce.Fixture.csproj b/Engine/Source/Programs/Shared/EpicGames.Perforce.Fixture/EpicGames.Perforce.Fixture.csproj
--- a/Engine/Source/Programs/Shared/EpicGames.Perforce.Fixture/EpicGames.Perforce.Fixture.csproj
+++ b/Engine/Source/Programs/Shared/EpicGames.Perforce.Fixture/EpicGames.Perforce.Fixture.csproj
@@ -36,7 +36,7 @@
     <PackageReference Include="Microsoft.NET.Test.Sdk" Version="17.3.0" />
     <PackageReference Include="MSTest.TestAdapter" Version="2.2.10" />
     <PackageReference Include="MSTest.TestFramework" Version="2.2.10" />
-    <PackageReference Include="coverlet.collector" Version="3.1.2">
+    <PackageReference Include="coverlet.collector" Version="6.0.2">
       <PrivateAssets>all</PrivateAssets>
       <IncludeAssets>runtime; build; native; contentfiles; analyzers; buildtransitive</IncludeAssets>
     </PackageReference>
diff --git a/Engine/Source/Programs/Shared/EpicGames.Perforce.Managed.Tests/EpicGames.Perforce.Managed.Tests.csproj b/Engine/Source/Programs/Shared/EpicGames.Perforce.Managed.Tests/EpicGames.Perforce.Managed.Tests.csproj
--- a/Engine/Source/Programs/Shared/EpicGames.Perforce.Managed.Tests/EpicGames.Perforce.Managed.Tests.csproj
+++ b/Engine/Source/Programs/Shared/EpicGames.Perforce.Managed.Tests/EpicGames.Perforce.Managed.Tests.csproj
@@ -34,7 +34,7 @@
     <PackageReference Include="Microsoft.NET.Test.Sdk" Version="17.3.0" />
     <PackageReference Include="MSTest.TestAdapter" Version="2.2.10" />
     <PackageReference Include="MSTest.TestFramework" Version="2.2.10" />
-    <PackageReference Include="coverlet.collector" Version="3.1.2">
+    <PackageReference Include="coverlet.collector" Version="6.0.2">
       <PrivateAssets>all</PrivateAssets>
       <IncludeAssets>runtime; build; native; contentfiles; analyzers; buildtransitive</IncludeAssets>
     </PackageReference>
diff --git a/Engine/Source/Programs/Shared/EpicGames.Perforce.Tests/EpicGames.Perforce.Tests.csproj b/Engine/Source/Programs/Shared/EpicGames.Perforce.Tests/EpicGames.Perforce.Tests.csproj
--- a/Engine/Source/Programs/Shared/EpicGames.Perforce.Tests/EpicGames.Perforce.Tests.csproj
+++ b/Engine/Source/Programs/Shared/EpicGames.Perforce.Tests/EpicGames.Perforce.Tests.csproj
@@ -33,7 +33,7 @@
     <PackageReference Include="Microsoft.NET.Test.Sdk" Version="17.3.0" />
     <PackageReference Include="MSTest.TestAdapter" Version="2.2.10" />
     <PackageReference Include="MSTest.TestFramework" Version="2.2.10" />
-    <PackageReference Include="coverlet.collector" Version="3.1.2">
+    <PackageReference Include="coverlet.collector" Version="6.0.2">
       <PrivateAssets>all</PrivateAssets>
       <IncludeAssets>runtime; build; native; contentfiles; analyzers; buildtransitive</IncludeAssets>
     </PackageReference>
diff --git a/Engine/Source/Programs/Shared/EpicGames.Redis.Tests/EpicGames.Redis.Tests.csproj b/Engine/Source/Programs/Shared/EpicGames.Redis.Tests/EpicGames.Redis.Tests.csproj
--- a/Engine/Source/Programs/Shared/EpicGames.Redis.Tests/EpicGames.Redis.Tests.csproj
+++ b/Engine/Source/Programs/Shared/EpicGames.Redis.Tests/EpicGames.Redis.Tests.csproj
@@ -34,7 +34,7 @@
     <PackageReference Include="Microsoft.NET.Test.Sdk" Version="17.8.0" />
     <PackageReference Include="MSTest.TestAdapter" Version="3.1.1" />
     <PackageReference Include="MSTest.TestFramework" Version="3.1.1" />
-    <PackageReference Include="coverlet.collector" Version="6.0.0">
+    <PackageReference Include="coverlet.collector" Version="6.0.2">
       <PrivateAssets>all</PrivateAssets>
       <IncludeAssets>runtime; build; native; contentfiles; analyzers; buildtransitive</IncludeAssets>
     </PackageReference>
diff --git a/Engine/Source/Programs/Shared/EpicGames.Serialization.Tests/EpicGames.Serialization.Tests.csproj b/Engine/Source/Programs/Shared/EpicGames.Serialization.Tests/EpicGames.Serialization.Tests.csproj
--- a/Engine/Source/Programs/Shared/EpicGames.Serialization.Tests/EpicGames.Serialization.Tests.csproj
+++ b/Engine/Source/Programs/Shared/EpicGames.Serialization.Tests/EpicGames.Serialization.Tests.csproj
@@ -36,7 +36,7 @@
     <PackageReference Include="Microsoft.NET.Test.Sdk" Version="17.3.0" />
     <PackageReference Include="MSTest.TestAdapter" Version="2.2.10" />
     <PackageReference Include="MSTest.TestFramework" Version="2.2.10" />
-    <PackageReference Include="coverlet.collector" Version="3.1.2">
+    <PackageReference Include="coverlet.collector" Version="6.0.2">
       <PrivateAssets>all</PrivateAssets>
       <IncludeAssets>runtime; build; native; contentfiles; analyzers; buildtransitive</IncludeAssets>
     </PackageReference>
diff --git a/Engine/Source/Programs/Shared/EpicGames.Tracing.Tests/EpicGames.Tracing.Tests.csproj b/Engine/Source/Programs/Shared/EpicGames.Tracing.Tests/EpicGames.Tracing.Tests.csproj
--- a/Engine/Source/Programs/Shared/EpicGames.Tracing.Tests/EpicGames.Tracing.Tests.csproj
+++ b/Engine/Source/Programs/Shared/EpicGames.Tracing.Tests/EpicGames.Tracing.Tests.csproj
@@ -32,7 +32,10 @@
     <PackageReference Include="Microsoft.NET.Test.Sdk" Version="16.5.0" />
     <PackageReference Include="MSTest.TestAdapter" Version="2.1.0" />
     <PackageReference Include="MSTest.TestFramework" Version="2.1.0" />
-    <PackageReference Include="coverlet.collector" Version="1.2.0" />
+    <PackageReference Include="coverlet.collector" Version="6.0.2">
+      <PrivateAssets>all</PrivateAssets>
+      <IncludeAssets>runtime; build; native; contentfiles; analyzers; buildtransitive</IncludeAssets>
+    </PackageReference>
   </ItemGroup>
 
   <ItemGroup>
diff --git a/Engine/Source/Programs/UnrealBuildTool.Tests/UnrealBuildTool.Tests.csproj b/Engine/Source/Programs/UnrealBuildTool.Tests/UnrealBuildTool.Tests.csproj
--- a/Engine/Source/Programs/UnrealBuildTool.Tests/UnrealBuildTool.Tests.csproj
+++ b/Engine/Source/Programs/UnrealBuildTool.Tests/UnrealBuildTool.Tests.csproj
@@ -20,7 +20,7 @@
     <PackageReference Include="Microsoft.NET.Test.Sdk" Version="17.3.0" />
     <PackageReference Include="MSTest.TestAdapter" Version="2.2.10" />
     <PackageReference Include="MSTest.TestFramework" Version="2.2.10" />
-    <PackageReference Include="coverlet.collector" Version="3.1.2">
+    <PackageReference Include="coverlet.collector" Version="6.0.2">
       <PrivateAssets>all</PrivateAssets>
       <IncludeAssets>runtime; build; native; contentfiles; analyzers; buildtransitive</IncludeAssets>
     </PackageReference>
