Horde: Fix agent collection tests.

#rnx

diff --git a/Engine/Source/Programs/Horde/Plugins/Compute/HordeServer.Compute.Tests/Agents/AgentCollectionTests.cs b/Engine/Source/Programs/Horde/Plugins/Compute/HordeServer.Compute.Tests/Agents/AgentCollectionTests.cs
--- a/Engine/Source/Programs/Horde/Plugins/Compute/HordeServer.Compute.Tests/Agents/AgentCollectionTests.cs
+++ b/Engine/Source/Programs/Horde/Plugins/Compute/HordeServer.Compute.Tests/Agents/AgentCollectionTests.cs
@@ -4,6 +4,7 @@ using EpicGames.Horde.Agents;
 using EpicGames.Horde.Agents.Leases;
 using EpicGames.Horde.Agents.Pools;
 using Google.Protobuf.WellKnownTypes;
+using HordeCommon.Rpc;
 using HordeCommon.Rpc.Messages;
 using HordeServer.Agents;
 
@@ -73,7 +74,20 @@ public class AgentCollectionTests : ComputeTestSetup
 		await _agent.TryCreateLeaseAsync(_lease2);
 		await UpdateAgentAsync();
 
-		await _agent.TryUpdateSessionAsync(new UpdateSessionOptions { Leases = new List<RpcLease>() });
+		await _agent.TryUpdateSessionAsync(new UpdateSessionOptions
+		{
+			Leases = new List<RpcLease>
+			{
+				new RpcLease{ Id = _lease1.Id, State = RpcLeaseState.Active },
+				new RpcLease{ Id = _lease2.Id, State = RpcLeaseState.Active }
+			}
+		});
+		await UpdateAgentAsync();
+
+		await _agent.TryUpdateSessionAsync(new UpdateSessionOptions
+		{
+			Leases = new List<RpcLease>()
+		});
 
 		List<LeaseId> leases = await AgentCollection.FindActiveLeaseIdsAsync();
 
@@ -108,6 +122,16 @@ public class AgentCollectionTests : ComputeTestSetup
 		await _agent.TryCreateLeaseAsync(_lease2);
 		await UpdateAgentAsync();
 
+		await _agent.TryUpdateSessionAsync(new UpdateSessionOptions
+		{
+			Leases = new List<RpcLease>
+			{
+				new RpcLease{ Id = _lease1.Id, State = RpcLeaseState.Active },
+				new RpcLease{ Id = _lease2.Id, State = RpcLeaseState.Active }
+			}
+		});
+		await UpdateAgentAsync();
+
 		await _agent.TryUpdateSessionAsync(new UpdateSessionOptions { Leases = new List<RpcLease> { new RpcLease { Id = _lease1.Id } } });
 
 		List<LeaseId> leases = await AgentCollection.FindActiveLeaseIdsAsync();
