Horde: Updating schema documentation

#rnx

diff --git a/Engine/Source/Programs/Horde/Docs/Config/Schema/AclActions.md b/Engine/Source/Programs/Horde/Docs/Config/Schema/AclActions.md
--- a/Engine/Source/Programs/Horde/Docs/Config/Schema/AclActions.md
+++ b/Engine/Source/Programs/Horde/Docs/Config/Schema/AclActions.md
@@ -69,6 +69,7 @@
 | ---- | ----------- |
 | `ReadArtifact` | Permission to read from an artifact |
 | `WriteArtifact` | Permission to write to an artifact |
+| `DeleteArtifact` | Permission to delete to an artifact |
 | `UploadArtifact` | Ability to create an artifact. Typically just for debugging; agents have this access for a particular session. |
 | `DownloadArtifact` | Ability to download an artifact |
 
@@ -173,6 +174,7 @@
 | ---- | ----------- |
 | `ReadArtifact` |  |
 | `WriteArtifact` |  |
+| `DeleteArtifact` |  |
 | `UploadArtifact` |  |
 | `DownloadArtifact` |  |
 
@@ -327,6 +329,7 @@
 | ---- | ----------- |
 | `ReadArtifact` |  |
 | `WriteArtifact` |  |
+| `DeleteArtifact` |  |
 | `UploadArtifact` |  |
 | `DownloadArtifact` |  |
 
@@ -494,6 +497,7 @@
 | ---- | ----------- |
 | `ReadArtifact` |  |
 | `WriteArtifact` |  |
+| `DeleteArtifact` |  |
 | `UploadArtifact` |  |
 | `DownloadArtifact` |  |
 
@@ -667,6 +671,7 @@
 | ---- | ----------- |
 | `ReadArtifact` |  |
 | `WriteArtifact` |  |
+| `DeleteArtifact` |  |
 | `UploadArtifact` |  |
 | `DownloadArtifact` |  |
 
@@ -850,6 +855,7 @@
 | ---- | ----------- |
 | `ReadArtifact` |  |
 | `WriteArtifact` |  |
+| `DeleteArtifact` |  |
 | `UploadArtifact` |  |
 | `DownloadArtifact` |  |
 
@@ -1040,6 +1046,7 @@
 | ---- | ----------- |
 | `ReadArtifact` |  |
 | `WriteArtifact` |  |
+| `DeleteArtifact` |  |
 | `UploadArtifact` |  |
 | `DownloadArtifact` |  |
 
