Horde: Add lock around invoking nftables CLI tool

diff --git a/Engine/Source/Programs/Horde/Plugins/Compute/HordeServer.Compute/Agents/Relay/Nftables.cs b/Engine/Source/Programs/Horde/Plugins/Compute/HordeServer.Compute/Agents/Relay/Nftables.cs
--- a/Engine/Source/Programs/Horde/Plugins/Compute/HordeServer.Compute/Agents/Relay/Nftables.cs
+++ b/Engine/Source/Programs/Horde/Plugins/Compute/HordeServer.Compute/Agents/Relay/Nftables.cs
@@ -31,6 +31,7 @@ public class NftablesException : Exception
 public class Nftables
 {
 	private const string NftablesExecutable = "nft";
+	private static readonly SemaphoreSlim s_nftSemaphore = new (1, 1);
 
 	/// <summary>
 	/// Prefix the invocation of 'nft' executable with 'sudo'
@@ -104,13 +105,21 @@ public class Nftables
 		{
 			return (_configurableExitCode ?? 0, _configurableOutputText ?? "");
 		}
-
-		byte[]? stdinData = stdin != null ? Encoding.UTF8.GetBytes(stdin) : null;
-		using ManagedProcessGroup processGroup = new();
-		using ManagedProcess process = new(processGroup, executable, CommandLineArguments.Join(argumentsCopy), null, null, stdinData, ProcessPriorityClass.Normal);
-		await process.WaitForExitAsync(cancellationToken);
-		string outputText = await process.StdOutText.ReadToEndAsync(cancellationToken);
-		return (process.ExitCode, outputText);
+		
+		await s_nftSemaphore.WaitAsync(cancellationToken);
+		try
+		{
+			byte[]? stdinData = stdin != null ? Encoding.UTF8.GetBytes(stdin) : null;
+			using ManagedProcessGroup processGroup = new();
+			using ManagedProcess process = new(processGroup, executable, CommandLineArguments.Join(argumentsCopy), null, null, stdinData, ProcessPriorityClass.Normal);
+			await process.WaitForExitAsync(cancellationToken);
+			string outputText = await process.StdOutText.ReadToEndAsync(cancellationToken);
+			return (process.ExitCode, outputText);
+		}
+		finally
+		{
+			s_nftSemaphore.Release();
+		}
 	}
 
 	private void LogPortMappings(List<PortMapping> leaseMappings)
