Horde: Move functionality for expiring agents and sessions into AgentCollection.

#rnx

diff --git a/Engine/Source/Programs/Horde/Plugins/Build/HordeServer.Build.Tests/Agents/AgentServiceTest.cs b/Engine/Source/Programs/Horde/Plugins/Build/HordeServer.Build.Tests/Agents/AgentServiceTest.cs
--- a/Engine/Source/Programs/Horde/Plugins/Build/HordeServer.Build.Tests/Agents/AgentServiceTest.cs
+++ b/Engine/Source/Programs/Horde/Plugins/Build/HordeServer.Build.Tests/Agents/AgentServiceTest.cs
@@ -178,13 +178,13 @@ public class AgentServiceTest : BuildTestSetup
 
 		// Let background task run for purging outdated sessions, which will terminate session for our agent
 		await Clock.AdvanceAsync(TimeSpan.FromHours(1));
-		await AgentService.TickSharedAsync(CancellationToken.None);
+		await ((AgentCollection)AgentCollection).TickSharedAsync(CancellationToken.None);
 
 		// Ephemeral agent is marked as deleted once its session is terminated
 		Assert.IsTrue((await AgentService.GetAgentAsync(agent.Id))!.Deleted);
 
 		await Clock.AdvanceAsync(TimeSpan.FromHours(25));
-		await AgentService.TickSharedAsync(CancellationToken.None);
+		await ((AgentCollection)AgentCollection).TickSharedAsync(CancellationToken.None);
 
 		// Once more time has passed, the ephemeral agent marked as deleted is removed from database
 		Assert.IsNull(await AgentService.GetAgentAsync(agent.Id));
diff --git a/Engine/Source/Programs/Horde/Plugins/Compute/HordeServer.Compute/Agents/AgentCollection.cs b/Engine/Source/Programs/Horde/Plugins/Compute/HordeServer.Compute/Agents/AgentCollection.cs
--- a/Engine/Source/Programs/Horde/Plugins/Compute/HordeServer.Compute/Agents/AgentCollection.cs
+++ b/Engine/Source/Programs/Horde/Plugins/Compute/HordeServer.Compute/Agents/AgentCollection.cs
@@ -25,6 +25,7 @@ using MongoDB.Bson.Serialization.Attributes;
 using MongoDB.Driver;
 using StackExchange.Redis;
 using Microsoft.Extensions.Hosting;
+using OpenTelemetry.Trace;
 
 namespace HordeServer.Agents
 {
@@ -277,14 +278,16 @@ namespace HordeServer.Agents
 		readonly IRedisService _redisService;
 		readonly IClock _clock;
 		readonly RedisChannel<AgentId> _updateEventChannel;
+		readonly ITicker _sharedTicker;
 		readonly ConcurrentDictionary<AgentId, TaskCompletionSource> _agentIdToTcs = new ConcurrentDictionary<AgentId, TaskCompletionSource>();
 		IAsyncDisposable? _updateEventSubscription;
+		readonly Tracer _tracer;
 		readonly ILogger _logger;
 
 		/// <summary>
 		/// Constructor
 		/// </summary>
-		public AgentCollection(IMongoService mongoService, IRedisService redisService, ILeaseCollection leaseCollection, IClock clock, IAuditLog<AgentId> auditLog, ILogger<AgentCollection> logger)
+		public AgentCollection(IMongoService mongoService, IRedisService redisService, ILeaseCollection leaseCollection, IClock clock, IAuditLog<AgentId> auditLog, Tracer tracer, ILogger<AgentCollection> logger)
 		{
 			List<MongoIndex<AgentDocument>> agentIndexes = new List<MongoIndex<AgentDocument>>();
 			agentIndexes.Add(keys => keys.Ascending(x => x.Deleted).Ascending(x => x.Id).Ascending(x => x.Pools));
@@ -300,13 +303,16 @@ namespace HordeServer.Agents
 			_leaseCollection = leaseCollection;
 			_clock = clock;
 			_updateEventChannel = new RedisChannel<AgentId>(RedisChannel.Literal("agents/notify"));
+			_sharedTicker = clock.AddSharedTicker($"{nameof(AgentCollection)}.{nameof(TickSharedAsync)}", TimeSpan.FromSeconds(30.0), TickSharedAsync, logger);
 			_auditLog = auditLog;
+			_tracer = tracer;
 			_logger = logger;
 		}
 
 		/// <inheritdoc/>
 		public async ValueTask DisposeAsync()
 		{
+			await _sharedTicker.DisposeAsync();
 			if (_updateEventSubscription != null)
 			{
 				await _updateEventSubscription.DisposeAsync();
@@ -318,11 +324,13 @@ namespace HordeServer.Agents
 		public async Task StartAsync(CancellationToken cancellationToken)
 		{
 			_updateEventSubscription = await _redisService.GetDatabase().Multiplexer.SubscribeAsync(_updateEventChannel, OnAgentUpdate);
+			await _sharedTicker.StartAsync();
 		}
 
 		/// <inheritdoc/>
 		public async Task StopAsync(CancellationToken cancellationToken)
 		{
+			await _sharedTicker.StopAsync();
 			if (_updateEventSubscription != null)
 			{
 				await _updateEventSubscription.DisposeAsync();
@@ -330,6 +338,71 @@ namespace HordeServer.Agents
 			}
 		}
 
+		internal async ValueTask TickSharedAsync(CancellationToken stoppingToken)
+		{
+			using TelemetrySpan span = _tracer.StartActiveSpan($"{nameof(AgentCollection)}.{nameof(TickSharedAsync)}");
+
+			await TerminateExpiredSessionsAsync(stoppingToken);
+			await DeleteExpiredEphemeralAgentsAsync(stoppingToken);
+		}
+
+		private async Task TerminateExpiredSessionsAsync(CancellationToken cancellationToken)
+		{
+			using TelemetrySpan span = _tracer.StartActiveSpan($"{nameof(AgentCollection)}.{nameof(TerminateExpiredSessionsAsync)}");
+
+			int c = 0;
+			while (!cancellationToken.IsCancellationRequested)
+			{
+				// Find all the agents which are ready to be expired
+				const int MaxAgents = 100;
+				DateTime utcNow = _clock.UtcNow;
+
+				FilterDefinition<AgentDocument> filter = Builders<AgentDocument>.Filter.Exists(x => x.SessionExpiresAt) & Builders<AgentDocument>.Filter.Lt(x => x.SessionExpiresAt, utcNow);
+
+				List<AgentDocument> expiredAgents = await _agentCollection.Find(filter).ToListAsync(cancellationToken);
+				expiredAgents = await PostLoadAsync(expiredAgents, cancellationToken);
+
+				// Transition each agent to being offline
+				foreach (AgentDocument expiredAgent in expiredAgents)
+				{
+					cancellationToken.ThrowIfCancellationRequested();
+					_logger.LogDebug("Terminating session {SessionId} for agent {Agent}", expiredAgent.SessionId, expiredAgent.Id);
+					await TryTerminateSessionAsync(expiredAgent, cancellationToken);
+				}
+				c += expiredAgents.Count;
+
+				// Try again if we didn't fetch everything
+				if (expiredAgents.Count < MaxAgents)
+				{
+					break;
+				}
+			}
+			span.SetAttribute("NumAgentsTerminated", c);
+		}
+
+		private async Task DeleteExpiredEphemeralAgentsAsync(CancellationToken cancellationToken)
+		{
+			using TelemetrySpan span = _tracer.StartActiveSpan($"{nameof(AgentCollection)}.{nameof(DeleteExpiredEphemeralAgentsAsync)}");
+			int c = 0;
+
+			List<AgentDocument> deletedDocuments = await _agentCollection.Find(x => x.Deleted).ToListAsync(cancellationToken);
+			deletedDocuments = await PostLoadAsync(deletedDocuments, cancellationToken);
+
+			foreach (AgentDocument agent in deletedDocuments)
+			{
+				cancellationToken.ThrowIfCancellationRequested();
+				bool noStatusChangeDuringPeriod = _clock.UtcNow > agent.LastStatusChange + TimeSpan.FromHours(1);
+				if (agent is { Status: AgentStatus.Stopped, Ephemeral: true } && noStatusChangeDuringPeriod)
+				{
+					_logger.LogDebug("Deleting ephemeral agent {Agent}", agent.Id);
+					await ForceDeleteAsync(agent.Id, cancellationToken);
+					c++;
+				}
+			}
+
+			span.SetAttribute("NumAgentsDeleted", c);
+		}
+
 		async ValueTask<AgentDocument?> PostLoadAsync(AgentDocument? document, CancellationToken cancellationToken)
 		{
 			while (document != null)
@@ -495,24 +568,6 @@ namespace HordeServer.Agents
 			return documents.ConvertAll(x => CreateAgentObject(x));
 		}
 
-		/// <inheritdoc/>
-		public async Task<IReadOnlyList<IAgent>> FindExpiredAsync(DateTime utcNow, int maxAgents, CancellationToken cancellationToken)
-		{
-			FilterDefinition<AgentDocument> filter = Builders<AgentDocument>.Filter.Exists(x => x.SessionExpiresAt) & Builders<AgentDocument>.Filter.Lt(x => x.SessionExpiresAt, utcNow);
-
-			List<AgentDocument> documents = await _agentCollection.Find(filter).Limit(maxAgents).ToListAsync(cancellationToken);
-			documents = await PostLoadAsync(documents, cancellationToken);
-			return documents.ConvertAll(x => CreateAgentObject(x));
-		}
-
-		/// <inheritdoc/>
-		public async Task<IReadOnlyList<IAgent>> FindDeletedAsync(CancellationToken cancellationToken)
-		{
-			List<AgentDocument> documents = await _agentCollection.Find(x => x.Deleted).ToListAsync(cancellationToken);
-			documents = await PostLoadAsync(documents, cancellationToken);
-			return documents.ConvertAll(x => CreateAgentObject(x));
-		}
-
 		/// <inheritdoc/>
 		public async Task<List<LeaseId>> FindActiveLeaseIdsAsync(CancellationToken cancellationToken)
 		{
diff --git a/Engine/Source/Programs/Horde/Plugins/Compute/HordeServer.Compute/Agents/AgentService.cs b/Engine/Source/Programs/Horde/Plugins/Compute/HordeServer.Compute/Agents/AgentService.cs
--- a/Engine/Source/Programs/Horde/Plugins/Compute/HordeServer.Compute/Agents/AgentService.cs
+++ b/Engine/Source/Programs/Horde/Plugins/Compute/HordeServer.Compute/Agents/AgentService.cs
@@ -75,7 +75,6 @@ namespace HordeServer.Agents
 		readonly Tracer _tracer;
 		readonly ILogger _logger;
 		readonly ITicker _ticker;
-		readonly ITicker _sharedTicker;
 
 		readonly RedisStringKey<AgentRateTable> _agentRateTableData = new("agent-rates");
 
@@ -119,7 +118,6 @@ namespace HordeServer.Agents
 			_redisService = redisService;
 			_clock = clock;
 			_ticker = clock.AddTicker($"{nameof(AgentService)}.{nameof(TickAsync)}", TimeSpan.FromSeconds(30.0), TickAsync, logger);
-			_sharedTicker = clock.AddSharedTicker($"{nameof(AgentService)}.{nameof(TickSharedAsync)}", TimeSpan.FromSeconds(30.0), TickSharedAsync, logger);
 			_tracer = tracer;
 			_logger = logger;
 
@@ -130,14 +128,12 @@ namespace HordeServer.Agents
 		public async Task StartAsync(CancellationToken cancellationToken)
 		{
 			await _ticker.StartAsync();
-			await _sharedTicker.StartAsync();
 		}
 
 		/// <inheritdoc/>
 		public async Task StopAsync(CancellationToken cancellationToken)
 		{
 			await _ticker.StopAsync();
-			await _sharedTicker.StopAsync();
 		}
 
 		/// <inheritdoc/>
@@ -146,7 +142,6 @@ namespace HordeServer.Agents
 			await _cachedRates.DisposeAsync();
 			await _cachedPools.DisposeAsync();
 			await _ticker.DisposeAsync();
-			await _sharedTicker.DisposeAsync();
 		}
 
 		/// <summary>
@@ -843,63 +838,6 @@ namespace HordeServer.Agents
 			await CollectMetricsAsync(stoppingToken);
 		}
 
-		internal async ValueTask TickSharedAsync(CancellationToken stoppingToken)
-		{
-			using TelemetrySpan span = _tracer.StartActiveSpan($"{nameof(AgentService)}.{nameof(TickSharedAsync)}");
-
-			await TerminateExpiredSessionsAsync(stoppingToken);
-			await DeleteExpiredEphemeralAgentsAsync(stoppingToken);
-		}
-
-		private async Task TerminateExpiredSessionsAsync(CancellationToken cancellationToken)
-		{
-			using TelemetrySpan span = _tracer.StartActiveSpan($"{nameof(AgentService)}.{nameof(TerminateExpiredSessionsAsync)}");
-
-			int c = 0;
-			while (!cancellationToken.IsCancellationRequested)
-			{
-				// Find all the agents which are ready to be expired
-				const int MaxAgents = 100;
-				DateTime utcNow = _clock.UtcNow;
-				IReadOnlyList<IAgent> expiredAgents = await Agents.FindExpiredAsync(utcNow, MaxAgents, cancellationToken);
-
-				// Transition each agent to being offline
-				foreach (IAgent expiredAgent in expiredAgents)
-				{
-					cancellationToken.ThrowIfCancellationRequested();
-					_logger.LogDebug("Terminating session {SessionId} for agent {Agent}", expiredAgent.SessionId, expiredAgent.Id);
-					await TryTerminateSessionAsync(expiredAgent, cancellationToken);
-				}
-				c += expiredAgents.Count;
-
-				// Try again if we didn't fetch everything
-				if (expiredAgents.Count < MaxAgents)
-				{
-					break;
-				}
-			}
-			span.SetAttribute("NumAgentsTerminated", c);
-		}
-
-		private async Task DeleteExpiredEphemeralAgentsAsync(CancellationToken cancellationToken)
-		{
-			using TelemetrySpan span = _tracer.StartActiveSpan($"{nameof(AgentService)}.{nameof(DeleteExpiredEphemeralAgentsAsync)}");
-			int c = 0;
-			foreach (IAgent agent in await Agents.FindDeletedAsync(cancellationToken))
-			{
-				cancellationToken.ThrowIfCancellationRequested();
-				bool noStatusChangeDuringPeriod = _clock.UtcNow > agent.LastStatusChange + TimeSpan.FromHours(1);
-				if (agent is { Status: AgentStatus.Stopped, Ephemeral: true } && noStatusChangeDuringPeriod)
-				{
-					_logger.LogDebug("Deleting ephemeral agent {Agent}", agent.Id);
-					await DeleteAgentAsync(agent, true, cancellationToken);
-					c++;
-				}
-			}
-
-			span.SetAttribute("NumAgentsDeleted", c);
-		}
-
 		/// <summary>
 		/// Updates the in-memory cache of current agents
 		/// Called by ticker to avoid blocking reads. Updates in the background, in favor of slightly more stale agents
diff --git a/Engine/Source/Programs/Horde/Plugins/Compute/HordeServer.Compute/Agents/IAgentCollection.cs b/Engine/Source/Programs/Horde/Plugins/Compute/HordeServer.Compute/Agents/IAgentCollection.cs
--- a/Engine/Source/Programs/Horde/Plugins/Compute/HordeServer.Compute/Agents/IAgentCollection.cs
+++ b/Engine/Source/Programs/Horde/Plugins/Compute/HordeServer.Compute/Agents/IAgentCollection.cs
@@ -63,22 +63,6 @@ namespace HordeServer.Agents
 		/// <returns>List of agents matching the given criteria</returns>
 		Task<IReadOnlyList<IAgent>> FindAsync(PoolId? poolId = null, DateTime? modifiedAfter = null, string? property = null, AgentStatus? status = null, bool? enabled = null, bool includeDeleted = false, int? index = null, int? count = null, bool consistentRead = true, CancellationToken cancellationToken = default);
 
-		/// <summary>
-		/// Finds all agents with an expired session
-		/// </summary>
-		/// <param name="utcNow">The current time</param>
-		/// <param name="maxAgents">Maximum number of agents to return</param>
-		/// <param name="cancellationToken">Cancellation token for the operation</param>
-		/// <returns>List of agents</returns>
-		Task<IReadOnlyList<IAgent>> FindExpiredAsync(DateTime utcNow, int maxAgents, CancellationToken cancellationToken = default);
-
-		/// <summary>
-		/// Finds all agents marked as deleted
-		/// </summary>
-		/// <param name="cancellationToken">Cancellation token for the operation</param>
-		/// <returns>List of agents</returns>
-		Task<IReadOnlyList<IAgent>> FindDeletedAsync(CancellationToken cancellationToken = default);
-
 		/// <summary>
 		/// Finds all active agent lease IDs
 		/// </summary>
