Horde: Remove "force delete" functionality from public IAgentCollection API. This is only used internally for ephemeral agents.

#rnx

diff --git a/Engine/Source/Programs/Horde/Plugins/Compute/HordeServer.Compute/Agents/AgentCollection.cs b/Engine/Source/Programs/Horde/Plugins/Compute/HordeServer.Compute/Agents/AgentCollection.cs
--- a/Engine/Source/Programs/Horde/Plugins/Compute/HordeServer.Compute/Agents/AgentCollection.cs
+++ b/Engine/Source/Programs/Horde/Plugins/Compute/HordeServer.Compute/Agents/AgentCollection.cs
@@ -340,7 +340,7 @@ namespace HordeServer.Agents
 
 		internal async ValueTask TickSharedAsync(CancellationToken stoppingToken)
 		{
-			using TelemetrySpan span = _tracer.StartActiveSpan($"{nameof(AgentCollection)}.{nameof(TickSharedAsync)}");
+			using TelemetrySpan span = _tracer.StartActiveSpan($"{nameof(AgentService)}.{nameof(TickSharedAsync)}");
 
 			await TerminateExpiredSessionsAsync(stoppingToken);
 			await DeleteExpiredEphemeralAgentsAsync(stoppingToken);
@@ -348,7 +348,7 @@ namespace HordeServer.Agents
 
 		private async Task TerminateExpiredSessionsAsync(CancellationToken cancellationToken)
 		{
-			using TelemetrySpan span = _tracer.StartActiveSpan($"{nameof(AgentCollection)}.{nameof(TerminateExpiredSessionsAsync)}");
+			using TelemetrySpan span = _tracer.StartActiveSpan($"{nameof(AgentService)}.{nameof(TerminateExpiredSessionsAsync)}");
 
 			int c = 0;
 			while (!cancellationToken.IsCancellationRequested)
@@ -382,7 +382,7 @@ namespace HordeServer.Agents
 
 		private async Task DeleteExpiredEphemeralAgentsAsync(CancellationToken cancellationToken)
 		{
-			using TelemetrySpan span = _tracer.StartActiveSpan($"{nameof(AgentCollection)}.{nameof(DeleteExpiredEphemeralAgentsAsync)}");
+			using TelemetrySpan span = _tracer.StartActiveSpan($"{nameof(AgentService)}.{nameof(DeleteExpiredEphemeralAgentsAsync)}");
 			int c = 0;
 
 			List<AgentDocument> deletedDocuments = await _agentCollection.Find(x => x.Deleted).ToListAsync(cancellationToken);
@@ -488,7 +488,7 @@ namespace HordeServer.Agents
 		}
 
 		/// <inheritdoc/>
-		public async Task ForceDeleteAsync(AgentId agentId, CancellationToken cancellationToken)
+		async Task ForceDeleteAsync(AgentId agentId, CancellationToken cancellationToken)
 		{
 			await _sessionCollection.DeleteManyAsync(x => x.AgentId == agentId, cancellationToken);
 			await _agentCollection.DeleteOneAsync(x => x.Id == agentId, cancellationToken);
diff --git a/Engine/Source/Programs/Horde/Plugins/Compute/HordeServer.Compute/Agents/AgentService.cs b/Engine/Source/Programs/Horde/Plugins/Compute/HordeServer.Compute/Agents/AgentService.cs
--- a/Engine/Source/Programs/Horde/Plugins/Compute/HordeServer.Compute/Agents/AgentService.cs
+++ b/Engine/Source/Programs/Horde/Plugins/Compute/HordeServer.Compute/Agents/AgentService.cs
@@ -293,22 +293,15 @@ namespace HordeServer.Agents
 		/// Marks the agent as deleted
 		/// </summary>
 		/// <param name="agent">The agent to delete</param>
-		/// <param name="forceDelete">Whether to fully delete the agent as opposed to just marking it as deleted</param>
 		/// <param name="cancellationToken">Cancellation token for the operation</param>
 		/// <returns>Async task</returns>
-		public async Task DeleteAgentAsync(IAgent? agent, bool forceDelete = false, CancellationToken cancellationToken = default)
+		public async Task DeleteAgentAsync(IAgent? agent, CancellationToken cancellationToken = default)
 		{
 			if (agent == null)
 			{
 				return;
 			}
 
-			if (forceDelete)
-			{
-				await Agents.ForceDeleteAsync(agent.Id, cancellationToken);
-				return;
-			}
-
 			while (agent is { Deleted: false })
 			{
 				IAgent? newAgent = await agent.TryDeleteAsync(cancellationToken);
diff --git a/Engine/Source/Programs/Horde/Plugins/Compute/HordeServer.Compute/Agents/IAgentCollection.cs b/Engine/Source/Programs/Horde/Plugins/Compute/HordeServer.Compute/Agents/IAgentCollection.cs
--- a/Engine/Source/Programs/Horde/Plugins/Compute/HordeServer.Compute/Agents/IAgentCollection.cs
+++ b/Engine/Source/Programs/Horde/Plugins/Compute/HordeServer.Compute/Agents/IAgentCollection.cs
@@ -23,14 +23,6 @@ namespace HordeServer.Agents
 		/// <param name="cancellationToken">Cancellation token for the operation</param>
 		Task<IAgent> AddAsync(AgentId id, bool ephemeral, string enrollmentKey, CancellationToken cancellationToken = default);
 
-		/// <summary>
-		/// Deletes an agent
-		/// </summary>
-		/// <param name="agentId">Unique id of the agent</param>
-		/// <param name="cancellationToken">Cancellation token for the operation</param>
-		/// <returns>Async task</returns>
-		Task ForceDeleteAsync(AgentId agentId, CancellationToken cancellationToken = default);
-
 		/// <summary>
 		/// Gets an agent by ID
 		/// </summary>
