Horde: Change agent last status change time to "last online time".

#rnx

diff --git a/Engine/Source/Programs/Horde/Plugins/Build/HordeServer.Build.Tests/Agents/AgentServiceTest.cs b/Engine/Source/Programs/Horde/Plugins/Build/HordeServer.Build.Tests/Agents/AgentServiceTest.cs
--- a/Engine/Source/Programs/Horde/Plugins/Build/HordeServer.Build.Tests/Agents/AgentServiceTest.cs
+++ b/Engine/Source/Programs/Horde/Plugins/Build/HordeServer.Build.Tests/Agents/AgentServiceTest.cs
@@ -59,25 +59,19 @@ public class AgentServiceTest : BuildTestSetup
 	}
 
 	[TestMethod]
-	public async Task LastStatusChangeDuringSessionCreateAsync()
+	public async Task LastOnlineChangeDuringSessionCreateAsync()
 	{
 		// No session created yet, status change timestamp is empty
 		IAgent agent = await AgentService.CreateAgentAsync("agent1", false, "");
-		Assert.AreEqual(AgentStatus.Unspecified, agent.Status);
-		Assert.IsFalse(agent.LastStatusChange.HasValue);
+		Assert.AreEqual(AgentStatus.Stopped, agent.Status);
+		Assert.IsTrue(agent.LastOnlineTime.HasValue);
 
 		// A session has been created, status change timestamp is current time
 		agent = await AgentService.CreateSessionAsync(agent, AgentStatus.Ok, new RpcAgentCapabilities(), "v1");
 		Assert.AreEqual(AgentStatus.Ok, agent.Status);
-		Assert.AreEqual(ToUnixTime(Clock.UtcNow), ToUnixTime(agent.LastStatusChange!.Value));
-		DateTime lastStatusChange = agent.LastStatusChange!.Value;
+		Assert.IsNull(agent.LastOnlineTime);
 
 		await Clock.AdvanceAsync(TimeSpan.FromMinutes(1));
-
-		// The session is re-created, status change timestamp is same as when it first got created
-		agent = await AgentService.CreateSessionAsync(agent, AgentStatus.Ok, new RpcAgentCapabilities(), "v1");
-		Assert.AreEqual(AgentStatus.Ok, agent.Status);
-		Assert.AreEqual(ToUnixTime(lastStatusChange), ToUnixTime(agent.LastStatusChange!.Value));
 	}
 
 	private static int s_agentId = 1;
@@ -101,25 +95,24 @@ public class AgentServiceTest : BuildTestSetup
 	}
 
 	[TestMethod]
-	[DataRow(AgentStatus.Ok, false)]
-	[DataRow(AgentStatus.Unhealthy, true)]
-	[DataRow(AgentStatus.Stopping, true)]
-	[DataRow(AgentStatus.Stopped, true)]
-	[DataRow(AgentStatus.Unspecified, true)]
-	public async Task LastStatusChangeAsync(AgentStatus status, bool expectTimestampUpdate)
+	[DataRow(AgentStatus.Ok)]
+	[DataRow(AgentStatus.Unhealthy)]
+	[DataRow(AgentStatus.Stopping)]
+	[DataRow(AgentStatus.Stopped)]
+	[DataRow(AgentStatus.Unspecified)]
+	public async Task LastOnlineChangeAsync(AgentStatus status)
 	{
 		IAgent agent = await CreateAgentSessionAsync();
-		DateTime lastStatusChange = agent.LastStatusChange!.Value;
 		await Clock.AdvanceAsync(TimeSpan.FromMinutes(1));
 
 		agent = (await AgentService.UpdateSessionAsync(agent, agent.SessionId!.Value, status, null, new List<RpcLease>()))!;
-		if (expectTimestampUpdate)
+		if (agent.Status == AgentStatus.Stopped)
 		{
-			AssertNotEqual(lastStatusChange, agent.LastStatusChange);
+			AssertEqual(Clock.UtcNow, agent.LastOnlineTime);
 		}
 		else
 		{
-			AssertEqual(lastStatusChange, agent.LastStatusChange);
+			Assert.IsNull(agent.LastOnlineTime);
 		}
 	}
 
diff --git a/Engine/Source/Programs/Horde/Plugins/Build/HordeServer.Build/Agents/Pools/PoolUpdateService.cs b/Engine/Source/Programs/Horde/Plugins/Build/HordeServer.Build/Agents/Pools/PoolUpdateService.cs
--- a/Engine/Source/Programs/Horde/Plugins/Build/HordeServer.Build/Agents/Pools/PoolUpdateService.cs
+++ b/Engine/Source/Programs/Horde/Plugins/Build/HordeServer.Build/Agents/Pools/PoolUpdateService.cs
@@ -95,7 +95,7 @@ namespace HordeServer.Agents.Pools
 
 		private bool HasGracePeriodExpired(IAgent agent, IReadOnlyList<IPoolConfig> pools, TimeSpan? globalGracePeriod)
 		{
-			if (agent.LastStatusChange == null)
+			if (agent.LastOnlineTime == null)
 			{
 				return false;
 			}
@@ -106,7 +106,7 @@ namespace HordeServer.Agents.Pools
 				return false;
 			}
 
-			DateTime expirationTime = agent.LastStatusChange.Value + gracePeriod.Value;
+			DateTime expirationTime = agent.LastOnlineTime.Value + gracePeriod.Value;
 			return _clock.UtcNow > expirationTime;
 		}
 
diff --git a/Engine/Source/Programs/Horde/Plugins/Compute/HordeServer.Compute.Tests/Agents/AgentCollectionTests.cs b/Engine/Source/Programs/Horde/Plugins/Compute/HordeServer.Compute.Tests/Agents/AgentCollectionTests.cs
--- a/Engine/Source/Programs/Horde/Plugins/Compute/HordeServer.Compute.Tests/Agents/AgentCollectionTests.cs
+++ b/Engine/Source/Programs/Horde/Plugins/Compute/HordeServer.Compute.Tests/Agents/AgentCollectionTests.cs
@@ -58,7 +58,7 @@ public class AgentCollectionTests : ComputeTestSetup
 		await UpdateAgentAsync();
 
 		await _agent.TryCreateSessionAsync(new CreateSessionOptions(AgentStatus.Ok,
-			new RpcAgentCapabilities(), new List<PoolId>(), DateTime.UtcNow, null));
+			new RpcAgentCapabilities(), new List<PoolId>(), null));
 
 		List<LeaseId> leases = await AgentCollection.FindActiveLeaseIdsAsync();
 
diff --git a/Engine/Source/Programs/Horde/Plugins/Compute/HordeServer.Compute/Agents/AgentCollection.cs b/Engine/Source/Programs/Horde/Plugins/Compute/HordeServer.Compute/Agents/AgentCollection.cs
--- a/Engine/Source/Programs/Horde/Plugins/Compute/HordeServer.Compute/Agents/AgentCollection.cs
+++ b/Engine/Source/Programs/Horde/Plugins/Compute/HordeServer.Compute/Agents/AgentCollection.cs
@@ -46,7 +46,7 @@ namespace HordeServer.Agents
 			SessionId? IAgent.SessionId => _document.SessionId;
 			DateTime? IAgent.SessionExpiresAt => _document.SessionExpiresAt;
 			AgentStatus IAgent.Status => _document.Status;
-			DateTime? IAgent.LastStatusChange => _document.LastStatusChange;
+			DateTime? IAgent.LastOnlineTime => _document.LastOnlineTime;
 			bool IAgent.Enabled => _document.Enabled;
 			bool IAgent.Ephemeral => _document.Ephemeral;
 			bool IAgent.Deleted => _document.Deleted;
@@ -164,7 +164,7 @@ namespace HordeServer.Agents
 			public DateTime? SessionExpiresAt { get; set; }
 
 			public AgentStatus Status { get; set; }
-			public DateTime? LastStatusChange { get; set; }
+			public DateTime? LastOnlineTime { get; set; }
 
 			[BsonRequired]
 			public bool Enabled { get; set; } = true;
@@ -236,6 +236,7 @@ namespace HordeServer.Agents
 				Id = id;
 				Ephemeral = ephemeral;
 				EnrollmentKey = enrollmentKey;
+				Status = AgentStatus.Stopped;
 			}
 		}
 
@@ -383,7 +384,7 @@ namespace HordeServer.Agents
 		private async Task DeleteExpiredEphemeralAgentsAsync(CancellationToken cancellationToken)
 		{
 			using TelemetrySpan span = _tracer.StartActiveSpan($"{nameof(AgentService)}.{nameof(DeleteExpiredEphemeralAgentsAsync)}");
-			int c = 0;
+			int numDeleted = 0;
 
 			List<AgentDocument> deletedDocuments = await _agentCollection.Find(x => x.Deleted).ToListAsync(cancellationToken);
 			deletedDocuments = await PostLoadAsync(deletedDocuments, cancellationToken);
@@ -391,16 +392,21 @@ namespace HordeServer.Agents
 			foreach (AgentDocument agent in deletedDocuments)
 			{
 				cancellationToken.ThrowIfCancellationRequested();
-				bool noStatusChangeDuringPeriod = _clock.UtcNow > agent.LastStatusChange + TimeSpan.FromHours(1);
-				if (agent is { Status: AgentStatus.Stopped, Ephemeral: true } && noStatusChangeDuringPeriod)
+				if (agent.LastOnlineTime.HasValue && _clock.UtcNow > agent.LastOnlineTime.Value + TimeSpan.FromHours(1))
 				{
 					_logger.LogDebug("Deleting ephemeral agent {Agent}", agent.Id);
-					await ForceDeleteAsync(agent.Id, cancellationToken);
-					c++;
+
+					DeleteResult result = await _agentCollection.DeleteOneAsync(x => x.Id == agent.Id && x.UpdateIndex == agent.UpdateIndex, cancellationToken);
+					if (result.DeletedCount > 0)
+					{
+						await _sessionCollection.DeleteManyAsync(x => x.AgentId == agent.Id, CancellationToken.None);
+					}
+
+					numDeleted++;
 				}
 			}
 
-			span.SetAttribute("NumAgentsDeleted", c);
+			span.SetAttribute("NumAgentsDeleted", numDeleted);
 		}
 
 		async ValueTask<AgentDocument?> PostLoadAsync(AgentDocument? document, CancellationToken cancellationToken)
@@ -458,6 +464,8 @@ namespace HordeServer.Agents
 		public async Task<IAgent> AddAsync(AgentId id, bool ephemeral, string enrollmentKey, CancellationToken cancellationToken)
 		{
 			AgentDocument agent = new AgentDocument(id, ephemeral, enrollmentKey);
+			agent.LastOnlineTime = _clock.UtcNow;
+
 			await _agentCollection.InsertOneAsync(agent, null, cancellationToken);
 			return CreateAgentObject(agent);
 		}
@@ -487,13 +495,6 @@ namespace HordeServer.Agents
 			return await TryUpdateAsync(agent, update, cancellationToken);
 		}
 
-		/// <inheritdoc/>
-		async Task ForceDeleteAsync(AgentId agentId, CancellationToken cancellationToken)
-		{
-			await _sessionCollection.DeleteManyAsync(x => x.AgentId == agentId, cancellationToken);
-			await _agentCollection.DeleteOneAsync(x => x.Id == agentId, cancellationToken);
-		}
-
 		/// <inheritdoc/>
 		public async Task<IAgent?> GetAsync(AgentId agentId, CancellationToken cancellationToken)
 		{
@@ -719,7 +720,15 @@ namespace HordeServer.Agents
 			if (options.Status != null && agent.Status != options.Status.Value)
 			{
 				updates.Add(updateBuilder.Set(x => x.Status, options.Status.Value));
-				updates.Add(updateBuilder.Set(x => x.LastStatusChange, _clock.UtcNow));
+
+				if (options.Status == AgentStatus.Stopped)
+				{
+					updates.Add(updateBuilder.Set(x => x.LastOnlineTime, _clock.UtcNow));
+				}
+				else
+				{
+					updates.Add(updateBuilder.Unset(x => x.LastOnlineTime));
+				}
 			}
 			if (options.SessionExpiresAt != null)
 			{
@@ -949,6 +958,7 @@ namespace HordeServer.Agents
 				updates.Add(updateBuilder.Set(x => x.SessionId, newSession.Id));
 				updates.Add(updateBuilder.Set(x => x.SessionExpiresAt, newSession.StartTime + AgentService.SessionExpiryTime));
 				updates.Add(updateBuilder.Set(x => x.Status, options.Status));
+				updates.Add(updateBuilder.Unset(x => x.LastOnlineTime));
 				updates.Add(updateBuilder.Unset(x => x.Leases));
 				updates.Add(updateBuilder.Unset(x => x.Deleted));
 				updates.Add(updateBuilder.Set(x => x.Properties, newProperties));
@@ -966,11 +976,6 @@ namespace HordeServer.Agents
 					updates.Add(updateBuilder.Unset(x => x.UpgradeAttemptCount));
 				}
 
-				if (agent.Status != options.Status)
-				{
-					updates.Add(updateBuilder.Set(x => x.LastStatusChange, options.LastStatusChange));
-				}
-
 				foreach (AgentLease agentLease in agent.Leases ?? Enumerable.Empty<AgentLease>())
 				{
 					await RemoveActiveLeaseAsync(agentLease);
@@ -1009,7 +1014,7 @@ namespace HordeServer.Agents
 			update = update.Unset(x => x.SessionExpiresAt);
 			update = update.Unset(x => x.Leases);
 			update = update.Set(x => x.Status, AgentStatus.Stopped);
-			update = update.Set(x => x.LastStatusChange, finishTime);
+			update = update.Set(x => x.LastOnlineTime, finishTime);
 
 			if (agent.Ephemeral)
 			{
diff --git a/Engine/Source/Programs/Horde/Plugins/Compute/HordeServer.Compute/Agents/AgentService.cs b/Engine/Source/Programs/Horde/Plugins/Compute/HordeServer.Compute/Agents/AgentService.cs
--- a/Engine/Source/Programs/Horde/Plugins/Compute/HordeServer.Compute/Agents/AgentService.cs
+++ b/Engine/Source/Programs/Horde/Plugins/Compute/HordeServer.Compute/Agents/AgentService.cs
@@ -340,7 +340,6 @@ namespace HordeServer.Agents
 		/// <returns>New agent state</returns>
 		public async Task<IAgent> CreateSessionAsync(IAgent agent, AgentStatus status, RpcAgentCapabilities capabilities, string? version, CancellationToken cancellationToken = default)
 		{
-			DateTime? lastStatusChange = null;
 			for (; ; )
 			{
 				IAuditLogChannel<AgentId> agentLogger = Agents.GetLogger(agent.Id);
@@ -349,21 +348,16 @@ namespace HordeServer.Agents
 				IAgent? newAgent;
 				if (agent.SessionId != null)
 				{
-					// Save last status change timestamp to avoid registering the change to "stopped" when it's re-created immediately after to "ok".
-					lastStatusChange = agent.LastStatusChange;
-
 					// Try to terminate the current session
 					await TryTerminateSessionAsync(agent, cancellationToken);
 				}
 				else
 				{
-					DateTime utcNow = _clock.UtcNow;
-
 					// Get the new pools for the agent
 					List<PoolId> dynamicPools = await GetDynamicPoolsAsync(agent, cancellationToken);
 
 					// Reset the agent to use the new session
-					newAgent = await agent.TryCreateSessionAsync(new CreateSessionOptions(status, capabilities, dynamicPools, lastStatusChange ?? utcNow, version), cancellationToken);
+					newAgent = await agent.TryCreateSessionAsync(new CreateSessionOptions(status, capabilities, dynamicPools, version), cancellationToken);
 					if (newAgent != null)
 					{
 						LogPropertyChanges(agentLogger, agent.Properties, newAgent.Properties);
diff --git a/Engine/Source/Programs/Horde/Plugins/Compute/HordeServer.Compute/Agents/AgentsController.cs b/Engine/Source/Programs/Horde/Plugins/Compute/HordeServer.Compute/Agents/AgentsController.cs
--- a/Engine/Source/Programs/Horde/Plugins/Compute/HordeServer.Compute/Agents/AgentsController.cs
+++ b/Engine/Source/Programs/Horde/Plugins/Compute/HordeServer.Compute/Agents/AgentsController.cs
@@ -190,7 +190,7 @@ namespace HordeServer.Agents
 				new List<string>(agent.Properties),
 				new Dictionary<string, int>(agent.Resources),
 				agent.UpdateTime,
-				agent.LastStatusChange,
+				agent.LastOnlineTime,
 				agent.Pools.Select(x => x.ToString()).ToList(),
 				new { Devices = new[] { new { agent.Properties, agent.Resources } } },
 				leases,
diff --git a/Engine/Source/Programs/Horde/Plugins/Compute/HordeServer.Compute/Agents/IAgent.cs b/Engine/Source/Programs/Horde/Plugins/Compute/HordeServer.Compute/Agents/IAgent.cs
--- a/Engine/Source/Programs/Horde/Plugins/Compute/HordeServer.Compute/Agents/IAgent.cs
+++ b/Engine/Source/Programs/Horde/Plugins/Compute/HordeServer.Compute/Agents/IAgent.cs
@@ -50,9 +50,9 @@ namespace HordeServer.Agents
 		public AgentStatus Status { get; }
 
 		/// <summary>
-		/// Time at which last status change took place.
+		/// Last time that the agent was online. Null if the agent is currently online.
 		/// </summary>
-		public DateTime? LastStatusChange { get; }
+		public DateTime? LastOnlineTime { get; }
 
 		/// <summary>
 		/// Whether the agent is enabled
@@ -306,9 +306,8 @@ namespace HordeServer.Agents
 	/// <param name="Status">Status of the agent</param>
 	/// <param name="Capabilities">Capabilities for the agent</param>
 	/// <param name="DynamicPools">New list of dynamic pools for the agent</param>
-	/// <param name="LastStatusChange">Time to force status change timestamp to</param>
 	/// <param name="Version">Current version of the agent software</param>
-	public record class CreateSessionOptions(AgentStatus Status, RpcAgentCapabilities Capabilities, IReadOnlyList<PoolId> DynamicPools, DateTime LastStatusChange, string? Version);
+	public record class CreateSessionOptions(AgentStatus Status, RpcAgentCapabilities Capabilities, IReadOnlyList<PoolId> DynamicPools, string? Version);
 
 	/// <summary>
 	/// Options for updating a new agent session
diff --git a/Engine/Source/Programs/Shared/EpicGames.Horde/Agents/AgentMessages.cs b/Engine/Source/Programs/Shared/EpicGames.Horde/Agents/AgentMessages.cs
--- a/Engine/Source/Programs/Shared/EpicGames.Horde/Agents/AgentMessages.cs
+++ b/Engine/Source/Programs/Shared/EpicGames.Horde/Agents/AgentMessages.cs
@@ -90,34 +90,34 @@ namespace EpicGames.Horde.Agents
 	/// <summary>
 	/// Information about an agent
 	/// </summary>
-	/// <param name="Id"> The agent's unique ID </param>
-	/// <param name="Name"> Friendly name of the agent </param>
-	/// <param name="Enabled"> Whether the agent is currently enabled </param>
+	/// <param name="Id">The agent's unique ID</param>
+	/// <param name="Name">Friendly name of the agent</param>
+	/// <param name="Enabled">Whether the agent is currently enabled</param>
 	/// <param name="Status">Status of the agent</param>
-	/// <param name="Rate"> Cost estimate per-hour for this agent </param>
-	/// <param name="SessionId"> The current session id </param>
-	/// <param name="Ephemeral"> Whether the agent is ephemeral </param>
-	/// <param name="Online"> Whether the agent is currently online </param>
-	/// <param name="Deleted"> Whether this agent has expired </param>
-	/// <param name="PendingConform"> Whether a conform job is pending </param>
-	/// <param name="PendingFullConform"> Whether a full conform job is pending </param>
-	/// <param name="PendingRestart"> Whether a restart is pending </param>
-	/// <param name="PendingShutdown"> Whether a shutdown is pending </param>
-	/// <param name="LastShutdownReason"> The reason for the last shutdown </param>
-	/// <param name="LastConformTime"> Last time a conform was attempted </param>
-	/// <param name="ConformAttemptCount"> Number of times a conform has been attempted </param>
-	/// <param name="NextConformTime"> Last time a conform was attempted </param>
-	/// <param name="Version"> The current client version </param>
-	/// <param name="Properties"> Properties for the agent </param>
-	/// <param name="Resources"> Resources for the agent </param>
-	/// <param name="UpdateTime"> Last update time of this agent </param>
-	/// <param name="LastStatusChange"> Last time agent's status was changed </param>
-	/// <param name="Pools"> Pools for this agent </param>
-	/// <param name="Capabilities"> Capabilities of this agent </param>
-	/// <param name="Leases"> Array of active leases. </param>
+	/// <param name="Rate">Cost estimate per-hour for this agent</param>
+	/// <param name="SessionId">The current session id</param>
+	/// <param name="Ephemeral">Whether the agent is ephemeral</param>
+	/// <param name="Online">Whether the agent is currently online</param>
+	/// <param name="Deleted">Whether this agent has expired</param>
+	/// <param name="PendingConform">Whether a conform job is pending</param>
+	/// <param name="PendingFullConform">Whether a full conform job is pending</param>
+	/// <param name="PendingRestart">Whether a restart is pending</param>
+	/// <param name="PendingShutdown">Whether a shutdown is pending</param>
+	/// <param name="LastShutdownReason">The reason for the last shutdown</param>
+	/// <param name="LastConformTime">Last time a conform was attempted</param>
+	/// <param name="ConformAttemptCount">Number of times a conform has been attempted</param>
+	/// <param name="NextConformTime">Last time a conform was attempted</param>
+	/// <param name="Version">The current client version</param>
+	/// <param name="Properties">Properties for the agent</param>
+	/// <param name="Resources">Resources for the agent</param>
+	/// <param name="UpdateTime">Last update time of this agent</param>
+	/// <param name="LastOnlineTime">Last time that the agent was online</param>
+	/// <param name="Pools">Pools for this agent</param>
+	/// <param name="Capabilities">Capabilities of this agent</param>
+	/// <param name="Leases">Array of active leases.</param>
 	/// <param name="Workspaces">Current workspaces synced on the agent</param>
-	/// <param name="Comment"> Comment for this agent </param>
-	public record GetAgentResponse(AgentId Id, string Name, bool Enabled, AgentStatus Status, double? Rate, SessionId? SessionId, bool Ephemeral, bool Online, bool Deleted, bool PendingConform, bool PendingFullConform, bool PendingRestart, bool PendingShutdown, string LastShutdownReason, DateTime LastConformTime, int? ConformAttemptCount, DateTime? NextConformTime, string? Version, List<string> Properties, Dictionary<string, int> Resources, DateTime? UpdateTime, DateTime? LastStatusChange, List<string>? Pools, object? Capabilities, List<GetAgentLeaseResponse> Leases, List<GetAgentWorkspaceResponse> Workspaces, string? Comment);
+	/// <param name="Comment">Comment for this agent</param>
+	public record GetAgentResponse(AgentId Id, string Name, bool Enabled, AgentStatus Status, double? Rate, SessionId? SessionId, bool Ephemeral, bool Online, bool Deleted, bool PendingConform, bool PendingFullConform, bool PendingRestart, bool PendingShutdown, string LastShutdownReason, DateTime LastConformTime, int? ConformAttemptCount, DateTime? NextConformTime, string? Version, List<string> Properties, Dictionary<string, int> Resources, DateTime? UpdateTime, DateTime? LastOnlineTime, List<string>? Pools, object? Capabilities, List<GetAgentLeaseResponse> Leases, List<GetAgentWorkspaceResponse> Workspaces, string? Comment);
 
 	/// <summary>
 	/// Telemetry data for an agent
