Horde: Fix test failure.

#rnx

diff --git a/Engine/Source/Programs/Horde/HordeAgent/Services/SessionFactory.cs b/Engine/Source/Programs/Horde/HordeAgent/Services/SessionFactory.cs
--- a/Engine/Source/Programs/Horde/HordeAgent/Services/SessionFactory.cs
+++ b/Engine/Source/Programs/Horde/HordeAgent/Services/SessionFactory.cs
@@ -179,7 +179,6 @@ namespace HordeAgent.Services
 				// Create the session information
 				RpcCreateSessionRequest sessionRequest = new RpcCreateSessionRequest();
 				sessionRequest.Id = registrationInfo.Id;
-				sessionRequest.Status = RpcAgentStatus.Ok;
 				sessionRequest.Capabilities = capabilities;
 				sessionRequest.Version = AgentApp.Version;
 
diff --git a/Engine/Source/Programs/Horde/Plugins/Build/HordeServer.Build.Tests/Agents/AgentServiceTest.cs b/Engine/Source/Programs/Horde/Plugins/Build/HordeServer.Build.Tests/Agents/AgentServiceTest.cs
--- a/Engine/Source/Programs/Horde/Plugins/Build/HordeServer.Build.Tests/Agents/AgentServiceTest.cs
+++ b/Engine/Source/Programs/Horde/Plugins/Build/HordeServer.Build.Tests/Agents/AgentServiceTest.cs
@@ -44,7 +44,7 @@ public class AgentServiceTest : BuildTestSetup
 	{
 		Fixture fixture = await CreateFixtureAsync();
 
-		IAgent agent = await AgentService.CreateSessionAsync(fixture.Agent1, AgentStatus.Ok, new RpcAgentCapabilities(), 
+		IAgent agent = await AgentService.CreateSessionAsync(fixture.Agent1, new RpcAgentCapabilities(), 
 			"test");
 
 		Assert.IsTrue(AgentService.AuthorizeSession(agent, GetUser(agent), out string _));
@@ -67,7 +67,7 @@ public class AgentServiceTest : BuildTestSetup
 		Assert.IsTrue(agent.LastOnlineTime.HasValue);
 
 		// A session has been created, status change timestamp is current time
-		agent = await AgentService.CreateSessionAsync(agent, AgentStatus.Ok, new RpcAgentCapabilities(), "v1");
+		agent = await AgentService.CreateSessionAsync(agent, new RpcAgentCapabilities(), "v1");
 		Assert.AreEqual(AgentStatus.Ok, agent.Status);
 		Assert.IsNull(agent.LastOnlineTime);
 
@@ -78,7 +78,7 @@ public class AgentServiceTest : BuildTestSetup
 	private async Task<IAgent> CreateAgentSessionAsync()
 	{
 		IAgent agent = await AgentService.CreateAgentAsync("agentServiceTest-" + s_agentId++, false, "");
-		agent = await AgentService.CreateSessionAsync(agent, AgentStatus.Ok, new RpcAgentCapabilities(), "v1");
+		agent = await AgentService.CreateSessionAsync(agent, new RpcAgentCapabilities(), "v1");
 		return agent;
 	}
 
@@ -129,11 +129,11 @@ public class AgentServiceTest : BuildTestSetup
 		}
 
 		List<string> props = new() { $"{KnownPropertyNames.AwsInstanceType}=m5.large" };
-		IAgent agent = await AgentService.CreateSessionAsync(fixture.Agent1, AgentStatus.Ok, new RpcAgentCapabilities(props), "test");
+		IAgent agent = await AgentService.CreateSessionAsync(fixture.Agent1, new RpcAgentCapabilities(props), "test");
 		Assert.IsFalse(await AuditLogContains("AWS EC2 instance type changed"));
 
 		props = new() { $"{KnownPropertyNames.AwsInstanceType}=c6.xlarge" };
-		agent = await AgentService.CreateSessionAsync(agent, AgentStatus.Ok, new RpcAgentCapabilities(props), "test");
+		agent = await AgentService.CreateSessionAsync(agent, new RpcAgentCapabilities(props), "test");
 		Assert.IsTrue(await AuditLogContains("AWS EC2 instance type changed"));
 	}
 
@@ -142,8 +142,8 @@ public class AgentServiceTest : BuildTestSetup
 	{
 		IAgent agent1 = await AgentService.CreateAgentAsync("agent1", false, "");
 		IAgent agent2 = await AgentService.CreateAgentAsync("agent2", false, "");
-		await AgentService.CreateSessionAsync(agent1, AgentStatus.Ok, new RpcAgentCapabilities(new List<string>() { "aws-instance-type=c5.24xlarge", "osfamily=windows" }), "test");
-		await AgentService.CreateSessionAsync(agent2, AgentStatus.Ok, new RpcAgentCapabilities(new List<string>() { "aws-instance-type=c4.4xLARge", "osfamily=WinDowS" }), "test");
+		await AgentService.CreateSessionAsync(agent1, new RpcAgentCapabilities(new List<string>() { "aws-instance-type=c5.24xlarge", "osfamily=windows" }), "test");
+		await AgentService.CreateSessionAsync(agent2, new RpcAgentCapabilities(new List<string>() { "aws-instance-type=c4.4xLARge", "osfamily=WinDowS" }), "test");
 
 		List<AgentRateConfig> agentRateConfigs = new()
 		{
diff --git a/Engine/Source/Programs/Horde/Plugins/Build/HordeServer.Build.Tests/Agents/Pools/PoolUpdateServiceTest.cs b/Engine/Source/Programs/Horde/Plugins/Build/HordeServer.Build.Tests/Agents/Pools/PoolUpdateServiceTest.cs
--- a/Engine/Source/Programs/Horde/Plugins/Build/HordeServer.Build.Tests/Agents/Pools/PoolUpdateServiceTest.cs
+++ b/Engine/Source/Programs/Horde/Plugins/Build/HordeServer.Build.Tests/Agents/Pools/PoolUpdateServiceTest.cs
@@ -29,8 +29,8 @@ public class PoolUpdateServiceTest : BuildTestSetup
 	{
 		_pool = await CreatePoolAsync(new PoolConfig { Name = "testPool", EnableAutoscaling = true });
 		_enabledAgent = await CreateAgentAsync(_pool, true);
-		_disabledAgent = await CreateAgentAsync(_pool, false);
-		_disabledAgentBeyondGracePeriod = await CreateAgentAsync(_pool, enabled: false, adjustClockBy: -TimeSpan.FromHours(9));
+		_disabledAgent = await CreateAgentAsync(_pool, false, status: AgentStatus.Stopped);
+		_disabledAgentBeyondGracePeriod = await CreateAgentAsync(_pool, enabled: false, adjustClockBy: -TimeSpan.FromHours(9), status: AgentStatus.Stopped);
 	}
 
 	private async Task RefreshAgentsAsync()
diff --git a/Engine/Source/Programs/Horde/Plugins/Build/HordeServer.Build.Tests/BuildTestSetup.cs b/Engine/Source/Programs/Horde/Plugins/Build/HordeServer.Build.Tests/BuildTestSetup.cs
--- a/Engine/Source/Programs/Horde/Plugins/Build/HordeServer.Build.Tests/BuildTestSetup.cs
+++ b/Engine/Source/Programs/Horde/Plugins/Build/HordeServer.Build.Tests/BuildTestSetup.cs
@@ -259,9 +259,9 @@ namespace HordeServer.Tests
 		}
 
 		private static int s_agentIdCounter = 1;
-		public Task<IAgent> CreateAgentAsync(IPool pool, bool enabled = true, bool requestShutdown = false, List<string>? properties = null, List<AgentWorkspaceInfo>? workspaces = null, TimeSpan? adjustClockBy = null)
+		public Task<IAgent> CreateAgentAsync(IPool pool, bool enabled = true, bool requestShutdown = false, List<string>? properties = null, List<AgentWorkspaceInfo>? workspaces = null, TimeSpan? adjustClockBy = null, AgentStatus status = AgentStatus.Ok)
 		{
-			return CreateAgentAsync(pool.Id, enabled, requestShutdown, properties, workspaces, adjustClockBy);
+			return CreateAgentAsync(pool.Id, enabled, requestShutdown, properties, workspaces, adjustClockBy, status: status);
 		}
 
 		/// <summary>
@@ -286,7 +286,8 @@ namespace HordeServer.Tests
 			TimeSpan? adjustClockBy = null,
 			string? awsInstanceId = null,
 			CreateLeaseOptions? lease = null,
-			bool ephemeral = false)
+			bool ephemeral = false,
+			AgentStatus status = AgentStatus.Ok)
 		{
 			DateTime now = Clock.UtcNow;
 			if (adjustClockBy != null)
@@ -306,9 +307,15 @@ namespace HordeServer.Tests
 			agent = await agent.TryUpdateAsync(new UpdateAgentOptions { Enabled = enabled, ExplicitPools = poolId != null ? [poolId.Value] : [] });
 			Assert.IsNotNull(agent);
 
-			agent = await AgentService.CreateSessionAsync(agent, AgentStatus.Ok, new RpcAgentCapabilities(tempProps), null);
+			agent = await AgentService.CreateSessionAsync(agent, new RpcAgentCapabilities(tempProps), null);
 			Assert.IsNotNull(agent);
 
+			if (status != AgentStatus.Ok)
+			{
+				agent = await agent.TryUpdateSessionAsync(new UpdateSessionOptions { Status = status });
+				Assert.IsNotNull(agent);
+			}
+
 			if (workspaces is { Count: > 0 })
 			{
 				await agent.TryUpdateWorkspacesAsync(workspaces, false);
diff --git a/Engine/Source/Programs/Horde/Plugins/Build/HordeServer.Build.Tests/Fleet/AwsAutoScalingLifecycleTest.cs b/Engine/Source/Programs/Horde/Plugins/Build/HordeServer.Build.Tests/Fleet/AwsAutoScalingLifecycleTest.cs
--- a/Engine/Source/Programs/Horde/Plugins/Build/HordeServer.Build.Tests/Fleet/AwsAutoScalingLifecycleTest.cs
+++ b/Engine/Source/Programs/Horde/Plugins/Build/HordeServer.Build.Tests/Fleet/AwsAutoScalingLifecycleTest.cs
@@ -103,7 +103,7 @@ public class AwsAutoScalingLifecycleServiceTest : BuildTestSetup
 		LifecycleActionEvent lae = new() { Ec2InstanceId = "i-1234", LifecycleActionToken = "action-token-test", Origin = "AutoScalingGroup" };
 		List<string> props = new() { KnownPropertyNames.AwsInstanceId + "=" + lae.Ec2InstanceId };
 		IAgent agent = await CreateAgentAsync(new PoolId("pool1"), properties: props);
-		agent = await AgentService.CreateSessionAsync(agent, AgentStatus.Ok, new RpcAgentCapabilities(props), "v1");
+		agent = await AgentService.CreateSessionAsync(agent, new RpcAgentCapabilities(props), "v1");
 		await _asgLifecycleService.InitiateTerminationAsync(lae, CancellationToken.None);
 		TimeSpan extraMargin = TimeSpan.FromSeconds(10);
 
@@ -116,7 +116,8 @@ public class AwsAutoScalingLifecycleServiceTest : BuildTestSetup
 		AssertLifecycleUpdate(lae, AwsAutoScalingLifecycleService.ActionContinue, _lifecycleUpdates[1]);
 
 		// Agent responded to shutdown request, now safe to notify AWS ASG the termination can take place
-		agent = await AgentService.CreateSessionAsync(agent, AgentStatus.Stopped, new RpcAgentCapabilities(props), "v1");
+		agent = await AgentService.CreateSessionAsync(agent, new RpcAgentCapabilities(props), "v1");
+		agent = (await agent.TryUpdateSessionAsync(new UpdateSessionOptions { Status = AgentStatus.Stopped, Capabilities = new RpcAgentCapabilities(props) }))!;
 		await Clock.AdvanceAsync(_asgLifecycleService.LifecycleUpdaterInterval + extraMargin);
 		AssertLifecycleUpdate(lae, AwsAutoScalingLifecycleService.ActionAbandon, _lifecycleUpdates[2]);
 
diff --git a/Engine/Source/Programs/Horde/Plugins/Build/HordeServer.Build.Tests/Jobs/JobServiceTests.cs b/Engine/Source/Programs/Horde/Plugins/Build/HordeServer.Build.Tests/Jobs/JobServiceTests.cs
--- a/Engine/Source/Programs/Horde/Plugins/Build/HordeServer.Build.Tests/Jobs/JobServiceTests.cs
+++ b/Engine/Source/Programs/Horde/Plugins/Build/HordeServer.Build.Tests/Jobs/JobServiceTests.cs
@@ -253,7 +253,7 @@ namespace HordeServer.Tests.Jobs
 			agent = await agent.TryUpdateAsync(new UpdateAgentOptions { Enabled = true, ExplicitPools = new List<PoolId> { new PoolId("win") } });
 			Assert.IsNotNull(agent);
 
-			await AgentService.CreateSessionAsync(agent, AgentStatus.Ok, new RpcAgentCapabilities(), null);
+			await AgentService.CreateSessionAsync(agent, new RpcAgentCapabilities(), null);
 
 			ITemplate template = await TemplateCollection.GetOrAddAsync(new TemplateConfig { Name = "Test template" });
 			IGraph graph = await GraphCollection.AddAsync(template, null);
diff --git a/Engine/Source/Programs/Horde/Plugins/Build/HordeServer.Build.Tests/Jobs/JobTaskSourceTests.cs b/Engine/Source/Programs/Horde/Plugins/Build/HordeServer.Build.Tests/Jobs/JobTaskSourceTests.cs
--- a/Engine/Source/Programs/Horde/Plugins/Build/HordeServer.Build.Tests/Jobs/JobTaskSourceTests.cs
+++ b/Engine/Source/Programs/Horde/Plugins/Build/HordeServer.Build.Tests/Jobs/JobTaskSourceTests.cs
@@ -165,7 +165,7 @@ namespace HordeServer.Tests.Jobs
 				agent = await agent.TryUpdateAsync(new UpdateAgentOptions { Enabled = isAgentEnabled, ExplicitPools = new List<PoolId> { pool.Id } });
 				Assert.IsNotNull(agent);
 
-				await AgentService.CreateSessionAsync(agent, AgentStatus.Ok, new RpcAgentCapabilities(), null);
+				await AgentService.CreateSessionAsync(agent, new RpcAgentCapabilities(), null);
 			}
 
 			JobTaskSource.OnJobScheduled += (pool, poolHasAgentsOnline, job, graph, batchId) =>
diff --git a/Engine/Source/Programs/Horde/Plugins/Compute/HordeServer.Compute.Tests/Agents/AgentCollectionTests.cs b/Engine/Source/Programs/Horde/Plugins/Compute/HordeServer.Compute.Tests/Agents/AgentCollectionTests.cs
--- a/Engine/Source/Programs/Horde/Plugins/Compute/HordeServer.Compute.Tests/Agents/AgentCollectionTests.cs
+++ b/Engine/Source/Programs/Horde/Plugins/Compute/HordeServer.Compute.Tests/Agents/AgentCollectionTests.cs
@@ -57,7 +57,7 @@ public class AgentCollectionTests : ComputeTestSetup
 		await _agent.TryCreateLeaseAsync(_lease2);
 		await UpdateAgentAsync();
 
-		await _agent.TryCreateSessionAsync(new CreateSessionOptions(AgentStatus.Ok,
+		await _agent.TryCreateSessionAsync(new CreateSessionOptions(
 			new RpcAgentCapabilities(), new List<PoolId>(), null));
 
 		List<LeaseId> leases = await AgentCollection.FindActiveLeaseIdsAsync();
diff --git a/Engine/Source/Programs/Horde/Plugins/Compute/HordeServer.Compute/Agents/AgentCollection.cs b/Engine/Source/Programs/Horde/Plugins/Compute/HordeServer.Compute/Agents/AgentCollection.cs
--- a/Engine/Source/Programs/Horde/Plugins/Compute/HordeServer.Compute/Agents/AgentCollection.cs
+++ b/Engine/Source/Programs/Horde/Plugins/Compute/HordeServer.Compute/Agents/AgentCollection.cs
@@ -957,7 +957,7 @@ namespace HordeServer.Agents
 				List<UpdateDefinition<AgentDocument>> updates = new List<UpdateDefinition<AgentDocument>>();
 				updates.Add(updateBuilder.Set(x => x.SessionId, newSession.Id));
 				updates.Add(updateBuilder.Set(x => x.SessionExpiresAt, newSession.StartTime + AgentService.SessionExpiryTime));
-				updates.Add(updateBuilder.Set(x => x.Status, options.Status));
+				updates.Add(updateBuilder.Set(x => x.Status, AgentStatus.Ok));
 				updates.Add(updateBuilder.Unset(x => x.LastOnlineTime));
 				updates.Add(updateBuilder.Unset(x => x.Leases));
 				updates.Add(updateBuilder.Unset(x => x.Deleted));
diff --git a/Engine/Source/Programs/Horde/Plugins/Compute/HordeServer.Compute/Agents/AgentService.cs b/Engine/Source/Programs/Horde/Plugins/Compute/HordeServer.Compute/Agents/AgentService.cs
--- a/Engine/Source/Programs/Horde/Plugins/Compute/HordeServer.Compute/Agents/AgentService.cs
+++ b/Engine/Source/Programs/Horde/Plugins/Compute/HordeServer.Compute/Agents/AgentService.cs
@@ -333,12 +333,11 @@ namespace HordeServer.Agents
 		/// Creates a new agent session
 		/// </summary>
 		/// <param name="agent">The agent to create a session for</param>
-		/// <param name="status">Current status of the agent</param>
 		/// <param name="capabilities">Capabilities for the agent</param>
 		/// <param name="version">Version of the software that's running</param>
 		/// <param name="cancellationToken">Cancellation token for the operation</param>
 		/// <returns>New agent state</returns>
-		public async Task<IAgent> CreateSessionAsync(IAgent agent, AgentStatus status, RpcAgentCapabilities capabilities, string? version, CancellationToken cancellationToken = default)
+		public async Task<IAgent> CreateSessionAsync(IAgent agent, RpcAgentCapabilities capabilities, string? version, CancellationToken cancellationToken = default)
 		{
 			for (; ; )
 			{
@@ -357,7 +356,7 @@ namespace HordeServer.Agents
 					List<PoolId> dynamicPools = await GetDynamicPoolsAsync(agent, cancellationToken);
 
 					// Reset the agent to use the new session
-					newAgent = await agent.TryCreateSessionAsync(new CreateSessionOptions(status, capabilities, dynamicPools, version), cancellationToken);
+					newAgent = await agent.TryCreateSessionAsync(new CreateSessionOptions(capabilities, dynamicPools, version), cancellationToken);
 					if (newAgent != null)
 					{
 						LogPropertyChanges(agentLogger, agent.Properties, newAgent.Properties);
diff --git a/Engine/Source/Programs/Horde/Plugins/Compute/HordeServer.Compute/Agents/IAgent.cs b/Engine/Source/Programs/Horde/Plugins/Compute/HordeServer.Compute/Agents/IAgent.cs
--- a/Engine/Source/Programs/Horde/Plugins/Compute/HordeServer.Compute/Agents/IAgent.cs
+++ b/Engine/Source/Programs/Horde/Plugins/Compute/HordeServer.Compute/Agents/IAgent.cs
@@ -303,11 +303,10 @@ namespace HordeServer.Agents
 	/// <summary>
 	/// Options for starting a new agent session
 	/// </summary>
-	/// <param name="Status">Status of the agent</param>
 	/// <param name="Capabilities">Capabilities for the agent</param>
 	/// <param name="DynamicPools">New list of dynamic pools for the agent</param>
 	/// <param name="Version">Current version of the agent software</param>
-	public record class CreateSessionOptions(AgentStatus Status, RpcAgentCapabilities Capabilities, IReadOnlyList<PoolId> DynamicPools, string? Version);
+	public record class CreateSessionOptions(RpcAgentCapabilities Capabilities, IReadOnlyList<PoolId> DynamicPools, string? Version);
 
 	/// <summary>
 	/// Options for updating a new agent session
diff --git a/Engine/Source/Programs/Horde/Plugins/Compute/HordeServer.Compute/Server/RpcService.cs b/Engine/Source/Programs/Horde/Plugins/Compute/HordeServer.Compute/Server/RpcService.cs
--- a/Engine/Source/Programs/Horde/Plugins/Compute/HordeServer.Compute/Server/RpcService.cs
+++ b/Engine/Source/Programs/Horde/Plugins/Compute/HordeServer.Compute/Server/RpcService.cs
@@ -220,7 +220,7 @@ namespace HordeServer.Server
 			}
 
 			// Create a new session
-			agent = await _agentService.CreateSessionAsync(agent, (AgentStatus)request.Status, request.Capabilities, request.Version, context.CancellationToken);
+			agent = await _agentService.CreateSessionAsync(agent, request.Capabilities, request.Version, context.CancellationToken);
 			if (agent == null)
 			{
 				throw new StructuredRpcException(StatusCode.NotFound, "Agent {AgentId} not found", agentId);
