Horde: Make it more clear when a step was canceled by Horde vs a user
#skipci

diff --git a/Engine/Source/Programs/Horde/HordeDashboard/src/backend/JobDetails.ts b/Engine/Source/Programs/Horde/HordeDashboard/src/backend/JobDetails.ts
--- a/Engine/Source/Programs/Horde/HordeDashboard/src/backend/JobDetails.ts
+++ b/Engine/Source/Programs/Horde/HordeDashboard/src/backend/JobDetails.ts
@@ -1034,7 +1034,7 @@ export const getStepSummaryMarkdown = (jobDetails: JobDetails, stepId: string):
             aborted = "The job was canceled";
             aborted += ` by ${jobDetails.jobdata?.abortedByUserInfo.name}.`;
         } else {
-            aborted = "The step was canceled";
+            aborted = "This step was canceled by Horde";
 
             if (step.error === JobStepError.TimedOut) {
                 aborted = "The step was canceled due to reaching the maximum run time limit";
diff --git a/Engine/Source/Programs/Horde/HordeDashboard/src/components/jobDetailsV2/JobDetailViewSteps.tsx b/Engine/Source/Programs/Horde/HordeDashboard/src/components/jobDetailsV2/JobDetailViewSteps.tsx
--- a/Engine/Source/Programs/Horde/HordeDashboard/src/components/jobDetailsV2/JobDetailViewSteps.tsx
+++ b/Engine/Source/Programs/Horde/HordeDashboard/src/components/jobDetailsV2/JobDetailViewSteps.tsx
@@ -703,7 +703,7 @@ export const getStepSummaryMarkdown = (jobDetails: JobDetailsV2, stepId: string)
          aborted = "The job was canceled";
          aborted += ` by ${jobData.abortedByUserInfo.name}.`;
       } else {
-         aborted = "The step was canceled";
+         aborted = "This step was canceled by Horde";
       }
       text.push(aborted);
    } else if (step.state === JobStepState.Skipped) {
diff --git a/Engine/Source/Programs/Horde/HordeDashboard/src/components/jobDetailsV2/JobDetailViewSummary.tsx b/Engine/Source/Programs/Horde/HordeDashboard/src/components/jobDetailsV2/JobDetailViewSummary.tsx
--- a/Engine/Source/Programs/Horde/HordeDashboard/src/components/jobDetailsV2/JobDetailViewSummary.tsx
+++ b/Engine/Source/Programs/Horde/HordeDashboard/src/components/jobDetailsV2/JobDetailViewSummary.tsx
@@ -95,6 +95,8 @@ export const SummaryPanel: React.FC<{ jobDetails: JobDetailsV2 }> = observer(({
       jobText += "was canceled";
       if (jobData.abortedByUserInfo) {
          jobText += ` by ${jobData.abortedByUserInfo.name}.`;
+      } else {
+         jobText += ` by Horde.`;
       }
    } else {
       jobText += `${jobData.state === JobState.Complete ? `completed ${getNiceTime(jobData.updateTime, false)}.` : "is currently running."}`;
diff --git a/Engine/Source/Programs/Horde/HordeDashboard/src/components/jobDetailsV2/StepDetailSummary.tsx b/Engine/Source/Programs/Horde/HordeDashboard/src/components/jobDetailsV2/StepDetailSummary.tsx
--- a/Engine/Source/Programs/Horde/HordeDashboard/src/components/jobDetailsV2/StepDetailSummary.tsx
+++ b/Engine/Source/Programs/Horde/HordeDashboard/src/components/jobDetailsV2/StepDetailSummary.tsx
@@ -112,7 +112,7 @@ const getStepSummaryMarkdown = (jobDetails: JobDetailsV2, stepId: string): strin
          aborted = "The job was canceled";
          aborted += ` by ${jobData?.abortedByUserInfo.name}.`;
       } else {
-         aborted = "The step was canceled";
+         aborted = "This step was canceled by Horde";
 
          if (step.error === JobStepError.TimedOut) {
             aborted = "The step was canceled due to reaching the maximum run time limit";
