Horde: Address additional artifact search feedback
#skipci

diff --git a/Engine/Source/Programs/Horde/HordeDashboard/src/components/artifacts/ArtifactsSearch.tsx b/Engine/Source/Programs/Horde/HordeDashboard/src/components/artifacts/ArtifactsSearch.tsx
--- a/Engine/Source/Programs/Horde/HordeDashboard/src/components/artifacts/ArtifactsSearch.tsx
+++ b/Engine/Source/Programs/Horde/HordeDashboard/src/components/artifacts/ArtifactsSearch.tsx
@@ -1,4 +1,4 @@
-import { ComboBox, DefaultButton, DetailsList, DetailsListLayoutMode, IComboBox, IComboBoxOption, IDetailsListProps, IconButton, Label, Modal, PrimaryButton, ScrollablePane, ScrollbarVisibility, SelectionMode, Spinner, SpinnerSize, Stack, Text, TextField } from "@fluentui/react";
+import { ComboBox, DefaultButton, DetailsList, DetailsListLayoutMode, DirectionalHint, IComboBox, IComboBoxOption, IContextualMenuProps, IDetailsListProps, IconButton, Label, Modal, PrimaryButton, ScrollablePane, ScrollbarVisibility, SelectionMode, Spinner, SpinnerSize, Stack, Text, TextField } from "@fluentui/react";
 import { useConst } from '@fluentui/react-hooks';
 import React, { useState } from "react";
 import { NavigateFunction, useNavigate } from "react-router-dom";
@@ -305,20 +305,34 @@ const ArtifactsList: React.FC<{ state: ArtifactQueryState, artifacts?: GetArtifa
 
          const stream = projectStore.streamById(item.streamId);
 
+         const downloadProps: IContextualMenuProps = {
+            items: [
+               {
+                  key: 'download_ugs',
+                  text: 'Download with UGS',
+                  onClick: () => {
+                     window.location.assign(`/api/v2/artifacts/${item.id}/download?format=ugs`);
+                  }
+               }
+            ],
+            directionalHint: DirectionalHint.bottomLeftEdge
+         };
+   
+
          return <Stack horizontal verticalAlign="center" verticalFill tokens={{ childrenGap: 24 }} styles={{ root: { backgroundColor: background, paddingLeft: 12, paddingRight: 12, paddingTop: 8, paddingBottom: 8 } }}>
-            <Stack style={{ width: 410, paddingLeft: 8 }}>
+            <Stack style={{ width: 360, paddingLeft: 8 }}>
                <Text style={{ fontWeight: 600 }}>{item.name}</Text>
             </Stack>
-            <Stack style={{ width: 240 }}>
+            <Stack style={{ width: 180 }}>
                <Text>{stream ? (stream.fullname ?? stream.name) : item.streamId}</Text>
             </Stack>
             <Stack style={{ width: 72 }}>
                <Text>{item.change}</Text>
             </Stack>
-            <Stack style={{ width: 160 }}>
+            <Stack style={{ width: 144 }}>
                <Text>{item.type}</Text>
             </Stack>
-            <Stack>
+            <Stack horizontal tokens={{childrenGap: 18}}>
                <DefaultButton style={{ width: 90 }} text="Browse" onClick={() => {
 
                   const key = item.keys.find(k => k.startsWith("job:") && k.indexOf("/step:") !== -1);
@@ -333,6 +347,14 @@ const ArtifactsList: React.FC<{ state: ArtifactQueryState, artifacts?: GetArtifa
                   state.updateSearch(navigate, false);
 
                }} />
+               <DefaultButton split text="Download Zip" menuProps={downloadProps}
+                  disabled={!item.id}
+                  style={{ fontFamily: "Horde Open Sans SemiBold", width: 100, padding: 0 }}
+                  onClick={() => {
+                     window.location.assign(`/api/v2/artifacts/${item.id}/download?format=zip`);
+                  }}
+               />
+
             </Stack>
          </Stack>
       }
@@ -378,7 +400,7 @@ export const FindArtifactsModal: React.FC<{ onClose: () => void }> = ({ onClose
 
    const navigate = useNavigate();
    const searchState = useConst(new ArtifactQueryState(new URLSearchParams(window.location.search)));
-   const [state, setState] = useState<{ searching?: boolean, artifacts?: GetArtifactResponse[] }>({});
+   const [state, setState] = useState<{ searching?: boolean, artifacts?: GetArtifactResponse[], noResults?: boolean }>({});
 
    const { hordeClasses } = getHordeStyling();
 
@@ -473,7 +495,7 @@ export const FindArtifactsModal: React.FC<{ onClose: () => void }> = ({ onClose
 
          streamChooserId++;
          setState({
-            searching: false, artifacts: artifacts
+            searching: false, artifacts: artifacts, noResults: !artifacts.length
          });
 
       } catch (reason) {
@@ -615,7 +637,7 @@ export const FindArtifactsModal: React.FC<{ onClose: () => void }> = ({ onClose
                   </Stack>
                   <Stack horizontal style={{ paddingTop: 12, paddingRight: 8 }}>
                      <Stack grow />
-                     <Stack horizontal tokens={{ childrenGap: 24 }}>
+                     <Stack horizontal tokens={{ childrenGap: 18 }}>
                         <Stack>
                            <DefaultButton disabled={!!state.searching} text="Reset" onClick={() => {
                               streamChooserId++;
@@ -632,7 +654,8 @@ export const FindArtifactsModal: React.FC<{ onClose: () => void }> = ({ onClose
                      </Stack>
                   </Stack>
                   <Stack key={`artifact_list_${streamChooserId}`} styles={{ root: { paddingTop: 12 } }}>
-                     <ArtifactsList state={searchState} artifacts={state.artifacts} sortBy={searchState.sort} />
+                     {!state.noResults && <ArtifactsList state={searchState} artifacts={state.artifacts} sortBy={searchState.sort} />}
+                     {state.noResults && <Stack grow horizontalAlign="center"><Text variant="mediumPlus">No Results Found</Text></Stack>}
                   </Stack>
                </Stack>
             </Stack>
