Horde: Surface when log data is missing in log event
#skipci

diff --git a/Engine/Source/Programs/Horde/HordeDashboard/src/components/ErrorPane.tsx b/Engine/Source/Programs/Horde/HordeDashboard/src/components/ErrorPane.tsx
--- a/Engine/Source/Programs/Horde/HordeDashboard/src/components/ErrorPane.tsx
+++ b/Engine/Source/Programs/Horde/HordeDashboard/src/components/ErrorPane.tsx
@@ -118,7 +118,8 @@ export const ErrorPane: React.FC<{ jobDetails: JobDetails; stepId: string; showE
 
       return (<Stack className={styles.itemCell} styles={{ root: { padding: 8, marginRight: 8 } }}><Stack className={item.severity === EventSeverity.Warning ? styles.gutterWarning : styles.gutter} styles={{ root: { padding: 0, margin: 0 } }}>
          <Stack styles={{ root: { paddingLeft: 14 } }}>
-            {lines}
+            {!!lines.length && lines}
+            {!lines.length && <Text>Missing Log Data</Text>}
          </Stack>
       </Stack>
       </Stack>
@@ -213,7 +214,8 @@ export const JobEventListPanel: React.FC<{ jobDetails: JobDetails, stepIds: stri
          <Link className="log-link" to={url}>
             <Stack className={styles.itemCell} styles={{ root: { padding: 8, paddingLeft: 24, marginRight: 8 } }}><Stack className={event.severity === EventSeverity.Warning ? styles.gutterWarning : styles.gutter} styles={{ root: { padding: 0, margin: 0 } }}>
                <Stack styles={{ root: { paddingLeft: 8 } }}>
-                  {lines}
+               {!!lines.length && lines}
+               {!lines.length && <Text>Missing Log Data</Text>}
                </Stack>
             </Stack>
             </Stack>
diff --git a/Engine/Source/Programs/Horde/HordeDashboard/src/components/IssueViewV2.tsx b/Engine/Source/Programs/Horde/HordeDashboard/src/components/IssueViewV2.tsx
--- a/Engine/Source/Programs/Horde/HordeDashboard/src/components/IssueViewV2.tsx
+++ b/Engine/Source/Programs/Horde/HordeDashboard/src/components/IssueViewV2.tsx
@@ -1436,7 +1436,8 @@ export const ErrorPane: React.FC<{ events?: GetLogEventResponse[]; onClose?: ()
 
       return (<Stack className={errorStyles.itemCell} style={{ padding: 8 }}><Stack className={item.severity === EventSeverity.Warning ? errorStyles.gutterWarning : errorStyles.gutter} styles={{ root: { padding: 0, margin: 0 } }}>
          <Stack styles={{ root: { paddingLeft: 14 } }}>
-            {lines}
+            {!!lines.length && lines}
+            {!lines.length && <Text>Missing Log Data</Text>}
          </Stack>
       </Stack>
       </Stack>);
diff --git a/Engine/Source/Programs/Horde/HordeDashboard/src/components/jobDetailsV2/StepErrorPanel.tsx b/Engine/Source/Programs/Horde/HordeDashboard/src/components/jobDetailsV2/StepErrorPanel.tsx
--- a/Engine/Source/Programs/Horde/HordeDashboard/src/components/jobDetailsV2/StepErrorPanel.tsx
+++ b/Engine/Source/Programs/Horde/HordeDashboard/src/components/jobDetailsV2/StepErrorPanel.tsx
@@ -225,7 +225,8 @@ const ErrorPane: React.FC<{ jobDetails: JobDetailsV2; view: StepSummaryErrorsVie
 
       return (<Stack className={styles.itemCell} styles={{ root: { padding: 8, marginRight: 8 } }}><Stack className={item.severity === EventSeverity.Warning ? styles.gutterWarning : styles.gutter} styles={{ root: { padding: 0, margin: 0 } }}>
          <Stack styles={{ root: { paddingLeft: 14 } }}>
-            {lines}
+            {!!lines.length && lines}
+            {!lines.length && <Text>Missing Log Data</Text>}
          </Stack>
       </Stack>
       </Stack>
