Horde: Add customizable landing pages
#skipci

diff --git a/Engine/Source/Programs/Horde/HordeDashboard/src/App.tsx b/Engine/Source/Programs/Horde/HordeDashboard/src/App.tsx
--- a/Engine/Source/Programs/Horde/HordeDashboard/src/App.tsx
+++ b/Engine/Source/Programs/Horde/HordeDashboard/src/App.tsx
@@ -116,7 +116,7 @@ const Main: React.FC = () => {
       console.log("Enabled Server Plugins:")
       enabledPlugins.forEach(p => {
          console.log(`    ${p.name} : ${p.version ?? ""}`)
-      })         
+      })
    } else {
       console.log("No Enabled Plugins")
    }
@@ -205,7 +205,11 @@ export default App;
 const HomeRedirect: React.FC = () => {
    if (window.location.pathname === "/" || !window.location.pathname) {
       if (dashboard.user?.dashboardFeatures?.showLandingPage) {
-         return <Navigate to="/docs/Landing.md" replace={true} />
+         let route = "/docs/Landing.md";
+         if (dashboard.user?.dashboardFeatures?.landingPageRoute?.length) {
+            route = dashboard.user?.dashboardFeatures?.landingPageRoute;
+         }
+         return <Navigate to={route} replace={true} />
       }
       return <Navigate to="/index" replace={true} />
    }
diff --git a/Engine/Source/Programs/Horde/HordeDashboard/src/backend/Api.ts b/Engine/Source/Programs/Horde/HordeDashboard/src/backend/Api.ts
--- a/Engine/Source/Programs/Horde/HordeDashboard/src/backend/Api.ts
+++ b/Engine/Source/Programs/Horde/HordeDashboard/src/backend/Api.ts
@@ -3313,6 +3313,9 @@ export type GetDashboardFeaturesResponse = {
 	/** Show the landing page by default */
 	showLandingPage?: boolean;
 
+	/** Custom landing page route to direct users to */
+	landingPageRoute?: string;
+
 	/** Enable CI functionality */
 	showCI?: boolean;
 
diff --git a/Engine/Source/Programs/Horde/HordeServer/Dashboard/DashboardConfig.cs b/Engine/Source/Programs/Horde/HordeServer/Dashboard/DashboardConfig.cs
--- a/Engine/Source/Programs/Horde/HordeServer/Dashboard/DashboardConfig.cs
+++ b/Engine/Source/Programs/Horde/HordeServer/Dashboard/DashboardConfig.cs
@@ -1,5 +1,6 @@
 // Copyright Epic Games, Inc. All Rights Reserved.
 
+using System;
 using System.Collections.Generic;
 using EpicGames.Core;
 using EpicGames.Horde.Common;
@@ -23,6 +24,11 @@ namespace HordeServer.Dashboard
 		/// </summary>
 		public bool ShowLandingPage { get; set; } = false;
 
+		/// <summary>
+		/// Custom landing page route to direct users to
+		/// </summary>
+		public string LandingPageRoute { get; set; } = String.Empty;
+
 		/// <summary>
 		/// Enable CI functionality
 		/// </summary>
diff --git a/Engine/Source/Programs/Horde/HordeServer/Users/UserController.cs b/Engine/Source/Programs/Horde/HordeServer/Users/UserController.cs
--- a/Engine/Source/Programs/Horde/HordeServer/Users/UserController.cs
+++ b/Engine/Source/Programs/Horde/HordeServer/Users/UserController.cs
@@ -74,6 +74,7 @@ namespace HordeServer.Users
 		{
 			GetDashboardFeaturesResponse response = new GetDashboardFeaturesResponse();
 			response.ShowLandingPage = globalConfig.Dashboard.ShowLandingPage;
+			response.LandingPageRoute = globalConfig.Dashboard.LandingPageRoute;
 			response.ShowCI = globalConfig.Dashboard.ShowCI;
 			response.ShowAgents = globalConfig.Dashboard.ShowAgents;
 			response.ShowAgentRegistration = globalConfig.Dashboard.ShowAgentRegistration;
diff --git a/Engine/Source/Programs/Shared/EpicGames.Horde/Users/UserMessages.cs b/Engine/Source/Programs/Shared/EpicGames.Horde/Users/UserMessages.cs
--- a/Engine/Source/Programs/Shared/EpicGames.Horde/Users/UserMessages.cs
+++ b/Engine/Source/Programs/Shared/EpicGames.Horde/Users/UserMessages.cs
@@ -176,6 +176,11 @@ namespace EpicGames.Horde.Users
 		/// </summary>
 		public bool ShowLandingPage { get; set; }
 
+		/// <summary>
+		/// Custom landing page route to direct users to
+		/// </summary>
+		public string LandingPageRoute { get; set; } = String.Empty;
+
 		/// <summary>
 		/// Enable CI functionality
 		/// </summary>
