Horde: Lock fluent to fix upstream regression
#skipci

diff --git a/Engine/Source/Programs/Horde/HordeDashboard/package.json b/Engine/Source/Programs/Horde/HordeDashboard/package.json
--- a/Engine/Source/Programs/Horde/HordeDashboard/package.json
+++ b/Engine/Source/Programs/Horde/HordeDashboard/package.json
@@ -13,7 +13,7 @@
   "dependencies": {
     "@blueprintjs/icons": "^3.13.0",
     "@datadog/browser-logs": "^4.30.0",
-    "@fluentui/react": "^8.106.2",
+    "@fluentui/react": "8.106.2",
     "@fluentui/react-charting": "5.2.0",
     "d3": "^7.6.1",
     "markdown-to-jsx": "7.4.7",
