Horde: Updating schema documentation

#rnx

diff --git a/Engine/Source/Programs/Horde/Docs/Config/Schema/Dashboard.md b/Engine/Source/Programs/Horde/Docs/Config/Schema/Dashboard.md
--- a/Engine/Source/Programs/Horde/Docs/Config/Schema/Dashboard.md
+++ b/Engine/Source/Programs/Horde/Docs/Config/Schema/Dashboard.md
@@ -7,6 +7,7 @@ Configuration for dashboard features
 Name | Description
 ---- | -----------
 `showLandingPage` | `boolean`<br>Navigate to the landing page by default
+`landingPageRoute` | `string`<br>Custom landing page route to direct users to
 `showCI` | `boolean`<br>Enable CI functionality
 `showAgents` | `boolean`<br>Whether to show functionality related to agents, pools, and utilization on the dashboard.
 `showAgentRegistration` | `boolean`<br>Whether to show the agent registration page. When using registration tokens from elsewhere this is not needed.
