Horde: Change IAgentCollection to return data through an async iterator

#rnx

diff --git a/Engine/Source/Programs/Horde/HordeAgent/HordeAgent.csproj b/Engine/Source/Programs/Horde/HordeAgent/HordeAgent.csproj
--- a/Engine/Source/Programs/Horde/HordeAgent/HordeAgent.csproj
+++ b/Engine/Source/Programs/Horde/HordeAgent/HordeAgent.csproj
@@ -115,8 +115,11 @@
     <Content Include="..\Plugins\Compute\HordeServer.Compute\Protos\horde\common\**\*.proto">
       <Link>Protos\horde\common\%(RecursiveDir)%(FileName)%(Extension)</Link>
     </Content>
-    <Compile Include="..\Plugins\Compute\HordeServer.Compute\Protos\horde\messages\**\*.cs">
-      <Link>Protos\horde\messages\%(RecursiveDir)%(FileName)%(Extension)</Link>
+    <Compile Include="..\Plugins\Compute\HordeServer.Compute\Protos\horde\messages\agent_capabilities.proto.cs">
+      <Link>Protos\horde\messages\agent_capabilities.proto.cs</Link>
+    </Compile>
+    <Compile Include="..\Plugins\Compute\HordeServer.Compute\Protos\horde\messages\lease.proto.cs">
+      <Link>Protos\horde\messages\lease.proto.cs</Link>
     </Compile>
     <Content Include="..\Plugins\Compute\HordeServer.Compute\Protos\horde\messages\**\*.proto">
       <Link>Protos\horde\messages\%(RecursiveDir)%(FileName)%(Extension)</Link>
diff --git a/Engine/Source/Programs/Horde/Plugins/Build/HordeServer.Build/Agents/AgentReportService.cs b/Engine/Source/Programs/Horde/Plugins/Build/HordeServer.Build/Agents/AgentReportService.cs
--- a/Engine/Source/Programs/Horde/Plugins/Build/HordeServer.Build/Agents/AgentReportService.cs
+++ b/Engine/Source/Programs/Horde/Plugins/Build/HordeServer.Build/Agents/AgentReportService.cs
@@ -86,8 +86,7 @@ namespace HordeServer.Agents
 			{
 				AgentReport report = new AgentReport();
 
-				IReadOnlyList<IAgent> agents = await _agentCollection.FindAsync(cancellationToken: cancellationToken);
-				foreach (IAgent agent in agents)
+				await foreach (IAgent agent in  _agentCollection.FindAsync(cancellationToken: cancellationToken))
 				{
 					if (agent.IsSessionValid(utcNow))
 					{
diff --git a/Engine/Source/Programs/Horde/Plugins/Compute/HordeServer.Compute/Agents/AgentCollection.cs b/Engine/Source/Programs/Horde/Plugins/Compute/HordeServer.Compute/Agents/AgentCollection.cs
--- a/Engine/Source/Programs/Horde/Plugins/Compute/HordeServer.Compute/Agents/AgentCollection.cs
+++ b/Engine/Source/Programs/Horde/Plugins/Compute/HordeServer.Compute/Agents/AgentCollection.cs
@@ -26,6 +26,7 @@ using MongoDB.Driver;
 using StackExchange.Redis;
 using Microsoft.Extensions.Hosting;
 using OpenTelemetry.Trace;
+using System.Runtime.CompilerServices;
 
 namespace HordeServer.Agents
 {
@@ -518,7 +519,7 @@ namespace HordeServer.Agents
 		}
 
 		/// <inheritdoc/>
-		public async Task<IReadOnlyList<IAgent>> FindAsync(PoolId? poolId, DateTime? modifiedAfter, string? property, AgentStatus? status, bool? enabled, bool includeDeleted, int? index, int? count, bool consistentRead, CancellationToken cancellationToken)
+		public async IAsyncEnumerable<IAgent> FindAsync(PoolId? poolId, DateTime? modifiedAfter, string? property, AgentStatus? status, bool? enabled, bool includeDeleted, bool consistentRead, [EnumeratorCancellation] CancellationToken cancellationToken)
 		{
 			FilterDefinitionBuilder<AgentDocument> filterBuilder = new FilterDefinitionBuilder<AgentDocument>();
 
@@ -554,19 +555,18 @@ namespace HordeServer.Agents
 			}
 
 			IMongoCollection<AgentDocument> collection = consistentRead ? _agentCollection : _agentCollection.WithReadPreference(ReadPreference.SecondaryPreferred);
-			IFindFluent<AgentDocument, AgentDocument> search = collection.Find(filter);
-			if (index != null)
-			{
-				search = search.Skip(index.Value);
-			}
-			if (count != null)
+
+			using (IAsyncCursor<AgentDocument> cursor = await collection.Find(filter).ToCursorAsync(cancellationToken))
 			{
-				search = search.Limit(count.Value);
+				while (await cursor.MoveNextAsync())
+				{
+					List<AgentDocument> documents = await PostLoadAsync(cursor.Current.ToList(), cancellationToken);
+					foreach (AgentDocument document in documents)
+					{
+						yield return CreateAgentObject(document);
+					}
+				}
 			}
-
-			List<AgentDocument> documents = await search.ToListAsync(cancellationToken);
-			documents = await PostLoadAsync(documents, cancellationToken);
-			return documents.ConvertAll(x => CreateAgentObject(x));
 		}
 
 		/// <inheritdoc/>
diff --git a/Engine/Source/Programs/Horde/Plugins/Compute/HordeServer.Compute/Agents/AgentService.cs b/Engine/Source/Programs/Horde/Plugins/Compute/HordeServer.Compute/Agents/AgentService.cs
--- a/Engine/Source/Programs/Horde/Plugins/Compute/HordeServer.Compute/Agents/AgentService.cs
+++ b/Engine/Source/Programs/Horde/Plugins/Compute/HordeServer.Compute/Agents/AgentService.cs
@@ -242,13 +242,11 @@ namespace HordeServer.Agents
 		/// <param name="modifiedAfter">If set, only returns agents modified after this time</param>
 		/// <param name="property">If set, only return agents matching this property</param>
 		/// <param name="includeDeleted">If set, include agents marked as deleted</param>
-		/// <param name="index">Index within the list of results</param>
-		/// <param name="count">Number of results to return</param>
 		/// <param name="cancellationToken">Cancellation token for the operation</param>
 		/// <returns>List of agents matching the given criteria</returns>
-		public Task<IReadOnlyList<IAgent>> FindAgentsAsync(PoolId? poolId, DateTime? modifiedAfter, string? property, bool includeDeleted, int? index, int? count, CancellationToken cancellationToken)
+		public IAsyncEnumerable<IAgent> FindAgentsAsync(PoolId? poolId, DateTime? modifiedAfter, string? property, bool includeDeleted, CancellationToken cancellationToken)
 		{
-			return Agents.FindAsync(poolId, modifiedAfter, property, null, null, includeDeleted, index, count, true, cancellationToken);
+			return Agents.FindAsync(poolId, modifiedAfter, property, null, null, includeDeleted, true, cancellationToken);
 		}
 
 		/// <summary>
@@ -832,9 +830,9 @@ namespace HordeServer.Agents
 		private async Task RefreshCachedAgentsAsync(CancellationToken cancellationToken = default)
 		{
 			using TelemetrySpan span = _tracer.StartActiveSpan($"{nameof(AgentService)}.{nameof(RefreshCachedAgentsAsync)}");
+
 			Dictionary<AgentId, IAgent> agents = new();
-			IReadOnlyList<IAgent> agentList = await Agents.FindAsync(consistentRead: false, cancellationToken: cancellationToken);
-			foreach (IAgent agent in agentList)
+			await foreach (IAgent agent in Agents.FindAsync(consistentRead: false, cancellationToken: cancellationToken))
 			{
 				agents[agent.Id] = agent;
 			}
diff --git a/Engine/Source/Programs/Horde/Plugins/Compute/HordeServer.Compute/Agents/Fleet/AwsAutoScalingLifecycle.cs b/Engine/Source/Programs/Horde/Plugins/Compute/HordeServer.Compute/Agents/Fleet/AwsAutoScalingLifecycle.cs
--- a/Engine/Source/Programs/Horde/Plugins/Compute/HordeServer.Compute/Agents/Fleet/AwsAutoScalingLifecycle.cs
+++ b/Engine/Source/Programs/Horde/Plugins/Compute/HordeServer.Compute/Agents/Fleet/AwsAutoScalingLifecycle.cs
@@ -246,14 +246,13 @@ public sealed class AwsAutoScalingLifecycleService : IHostedService, IAsyncDispo
 		{
 			string instanceIdProp = KnownPropertyNames.AwsInstanceId + "=" + e.Ec2InstanceId;
 
-			IReadOnlyList<IAgent> agentList = await _agentService.FindAgentsAsync(null, null, instanceIdProp, true, null, null, cancellationToken);
-			if (agentList.Count == 0)
+			IAgent? agent = await _agentService.FindAgentsAsync(null, null, instanceIdProp, true, cancellationToken).FirstOrDefaultAsync(cancellationToken);
+			if (agent == null)
 			{
 				_logger.LogWarning("Lifecycle action received but no agent with instance ID {InstanceId} found", e.Ec2InstanceId);
 				return false;
 			}
 
-			IAgent? agent = agentList[0];
 			await agent.TryUpdateAsync(new UpdateAgentOptions { RequestShutdown = true }, cancellationToken: cancellationToken);
 			await TrackAgentLifecycleAsync(agent.Id, e, cancellationToken);
 			return true;
@@ -298,8 +297,7 @@ public sealed class AwsAutoScalingLifecycleService : IHostedService, IAsyncDispo
 		}
 
 		List<string> validInstanceIds = new();
-		IReadOnlyList<IAgent> agents = await _agentService.FindAgentsAsync(null, null, null, true, null, null, cancellationToken);
-		foreach (IAgent agent in agents)
+		await foreach (IAgent agent in _agentService.FindAgentsAsync(null, null, null, true, cancellationToken))
 		{
 			if (IsAgentSuggestedByAsg(agent, out string? instanceId))
 			{
diff --git a/Engine/Source/Programs/Horde/Plugins/Compute/HordeServer.Compute/Agents/IAgentCollection.cs b/Engine/Source/Programs/Horde/Plugins/Compute/HordeServer.Compute/Agents/IAgentCollection.cs
--- a/Engine/Source/Programs/Horde/Plugins/Compute/HordeServer.Compute/Agents/IAgentCollection.cs
+++ b/Engine/Source/Programs/Horde/Plugins/Compute/HordeServer.Compute/Agents/IAgentCollection.cs
@@ -48,12 +48,10 @@ namespace HordeServer.Agents
 		/// <param name="status">Status to look for</param>
 		/// <param name="enabled">Enabled/disabled status to look for</param>
 		/// <param name="includeDeleted">Whether agents marked as deleted should be included</param>
-		/// <param name="index">Index of the first result</param>
-		/// <param name="count">Number of results to return</param>
 		/// <param name="consistentRead">If the database read should be made to the replica server</param>
 		/// <param name="cancellationToken">Cancellation token for the operation</param>
 		/// <returns>List of agents matching the given criteria</returns>
-		Task<IReadOnlyList<IAgent>> FindAsync(PoolId? poolId = null, DateTime? modifiedAfter = null, string? property = null, AgentStatus? status = null, bool? enabled = null, bool includeDeleted = false, int? index = null, int? count = null, bool consistentRead = true, CancellationToken cancellationToken = default);
+		IAsyncEnumerable<IAgent> FindAsync(PoolId? poolId = null, DateTime? modifiedAfter = null, string? property = null, AgentStatus? status = null, bool? enabled = null, bool includeDeleted = false, bool consistentRead = true, CancellationToken cancellationToken = default);
 
 		/// <summary>
 		/// Finds all active agent lease IDs
diff --git a/Engine/Source/Programs/Horde/Plugins/Compute/HordeServer.Compute/Agents/Utilization/UtilizationDataService.cs b/Engine/Source/Programs/Horde/Plugins/Compute/HordeServer.Compute/Agents/Utilization/UtilizationDataService.cs
--- a/Engine/Source/Programs/Horde/Plugins/Compute/HordeServer.Compute/Agents/Utilization/UtilizationDataService.cs
+++ b/Engine/Source/Programs/Horde/Plugins/Compute/HordeServer.Compute/Agents/Utilization/UtilizationDataService.cs
@@ -64,7 +64,7 @@ namespace HordeServer.Agents.Utilization
 			DateTime minTime = maxTime - count * interval;
 
 			// Query all the current data
-			IReadOnlyList<IAgent> agents = await _agentCollection.FindAsync(cancellationToken: cancellationToken);
+			IReadOnlyList<IAgent> agents = await _agentCollection.FindAsync(cancellationToken: cancellationToken).ToListAsync(cancellationToken);
 			IReadOnlyList<IPoolConfig> pools = await _poolCollection.GetConfigsAsync(cancellationToken);
 			IReadOnlyList<ILease> leases = await _leaseCollection.FindLeasesAsync(minTime: minTime, cancellationToken: cancellationToken);
 
diff --git a/Engine/Source/Programs/Horde/Plugins/Compute/HordeServer.Compute/Compute/ComputeService.cs b/Engine/Source/Programs/Horde/Plugins/Compute/HordeServer.Compute/Compute/ComputeService.cs
--- a/Engine/Source/Programs/Horde/Plugins/Compute/HordeServer.Compute/Compute/ComputeService.cs
+++ b/Engine/Source/Programs/Horde/Plugins/Compute/HordeServer.Compute/Compute/ComputeService.cs
@@ -394,8 +394,7 @@ namespace HordeServer.Compute
 			try
 			{
 				PoolId? poolId = arp.Requirements.Pool != null ? new PoolId(arp.Requirements.Pool) : null;
-				IReadOnlyList<IAgent> agents = await _agentCollection.FindAsync(poolId, consistentRead: false, cancellationToken: cancellationToken);
-				foreach (IAgent agent in agents)
+				await foreach (IAgent agent in _agentCollection.FindAsync(poolId, consistentRead: false, cancellationToken: cancellationToken))
 				{
 					Dictionary<string, int> assignedResources = new Dictionary<string, int>();
 
diff --git a/Engine/Source/Programs/Horde/Plugins/Compute/HordeServer.Compute/HordeServer.Compute.csproj b/Engine/Source/Programs/Horde/Plugins/Compute/HordeServer.Compute/HordeServer.Compute.csproj
--- a/Engine/Source/Programs/Horde/Plugins/Compute/HordeServer.Compute/HordeServer.Compute.csproj
+++ b/Engine/Source/Programs/Horde/Plugins/Compute/HordeServer.Compute/HordeServer.Compute.csproj
@@ -60,6 +60,7 @@
       <PrivateAssets>all</PrivateAssets>
       <IncludeAssets>runtime; build; native; contentfiles; analyzers; buildtransitive</IncludeAssets>
     </PackageReference>
+    <PackageReference Include="System.Linq.Async" Version="6.0.1" />
   </ItemGroup>
 
   <ItemGroup>
diff --git a/Engine/Source/Programs/Horde/Plugins/Compute/HordeServer.Compute/Server/PoolsServerController.cs b/Engine/Source/Programs/Horde/Plugins/Compute/HordeServer.Compute/Server/PoolsServerController.cs
--- a/Engine/Source/Programs/Horde/Plugins/Compute/HordeServer.Compute/Server/PoolsServerController.cs
+++ b/Engine/Source/Programs/Horde/Plugins/Compute/HordeServer.Compute/Server/PoolsServerController.cs
@@ -54,8 +54,7 @@ namespace HordeServer.Server
 
 				Dictionary<PoolId, int> poolIdToCount = new Dictionary<PoolId, int>();
 
-				IReadOnlyList<IAgent> agents = await agentCollection.FindAsync(cancellationToken: cancellationToken);
-				foreach (IAgent agent in agents)
+				await foreach (IAgent agent in agentCollection.FindAsync(cancellationToken: cancellationToken))
 				{
 					foreach (PoolId poolId in agent.Pools)
 					{
