Horde: Update release notes for deplpyment
#rnx
#skipci

diff --git a/Engine/Source/Programs/Horde/Docs/ReleaseNotes.md b/Engine/Source/Programs/Horde/Docs/ReleaseNotes.md
--- a/Engine/Source/Programs/Horde/Docs/ReleaseNotes.md
+++ b/Engine/Source/Programs/Horde/Docs/ReleaseNotes.md
@@ -1,5 +1,14 @@
 # Release Notes
 
+## 2024-08-22
+
+* Lock fluent to fix upstream regression (35716215)
+* Add customizable landing pages (35713203)
+* Surface when log data is missing in log event (35712179)
+* Address additional artifact search feedback (35710162)
+* Make it more clear when a step was canceled by Horde vs a user (35709294)
+* Add lock around invoking nftables CLI tool (35690511)
+
 ## 2024-08-19
 
 * Add ability to override agent's compute IP used for incoming connections (35621653)
