Horde: Improve OIDC configuration experience

- Add debug mode for better explaining why a JWT bearer token is rejected
- Validate required settings are set for auth mode OIDC/Okta
- Improve docs for OIDC settings

diff --git a/Engine/Source/Programs/Horde/HordeServer/Authentication/ExternalJwtAuthHandler.cs b/Engine/Source/Programs/Horde/HordeServer/Authentication/ExternalJwtAuthHandler.cs
--- a/Engine/Source/Programs/Horde/HordeServer/Authentication/ExternalJwtAuthHandler.cs
+++ b/Engine/Source/Programs/Horde/HordeServer/Authentication/ExternalJwtAuthHandler.cs
@@ -60,7 +60,18 @@ namespace HordeServer.Authentication
 
 			return Task.CompletedTask;
 		}
-
+		
+		private Task OnAuthenticationFailedAsync(AuthenticationFailedContext context)
+		{
+			if (_settings.OidcDebugMode)
+			{
+				IServiceProvider serviceProvider = context.HttpContext.RequestServices;
+				ILogger<ExternalJwtAuthHandler> logger = serviceProvider.GetRequiredService<ILogger<ExternalJwtAuthHandler>>();
+				logger.LogError(context.Exception, "JWT bearer auth failed. Auth header: {AuthHeader}", context.Request.Headers.Authorization.ToString());	
+			}
+			return Task.CompletedTask;
+		}
+		
 		/// <summary>
 		/// Callback when JWT bearer token is being validated
 		/// If user is not cached, it will look up additional info via /userinfo and cache it
@@ -141,7 +152,8 @@ namespace HordeServer.Authentication
 			{
 				options.Authority = _settings.OidcAuthority;
 				options.Audience = _settings.OidcAudience;
-				options.Events = new JwtBearerEvents() { OnMessageReceived = OnMessageReceivedAsync, OnTokenValidated = OnTokenValidatedAsync };
+				options.Events = new JwtBearerEvents() { OnAuthenticationFailed = OnAuthenticationFailedAsync, OnMessageReceived = OnMessageReceivedAsync, OnTokenValidated = OnTokenValidatedAsync };
+				options.RequireHttpsMetadata = !_settings.OidcDebugMode;
 
 				options.TokenValidationParameters.ValidAudience = _settings.OidcAudience;
 				options.TokenValidationParameters.RequireExpirationTime = true;
diff --git a/Engine/Source/Programs/Horde/HordeServer/Authentication/JwtAuthHandler.cs b/Engine/Source/Programs/Horde/HordeServer/Authentication/JwtAuthHandler.cs
--- a/Engine/Source/Programs/Horde/HordeServer/Authentication/JwtAuthHandler.cs
+++ b/Engine/Source/Programs/Horde/HordeServer/Authentication/JwtAuthHandler.cs
@@ -25,8 +25,8 @@ namespace HordeServer.Authentication
 
 		readonly AsyncCachedValue<IGlobals> _globals;
 
-		public JwtAuthHandler(ILoggerFactory logger, UrlEncoder encoder, GlobalsService globalsService, IOptionsMonitorCache<JwtBearerOptions> optionsCache)
-			: base(GetOptionsMonitor(optionsCache), logger, encoder)
+		public JwtAuthHandler(ILoggerFactory logger, UrlEncoder encoder, ServerSettings settings, GlobalsService globalsService, IOptionsMonitorCache<JwtBearerOptions> optionsCache)
+			: base(GetOptionsMonitor(settings, optionsCache), logger, encoder)
 		{
 			_globals = new AsyncCachedValue<IGlobals>(async ctx => await globalsService.GetAsync(ctx), TimeSpan.FromSeconds(30.0));
 		}
@@ -36,10 +36,13 @@ namespace HordeServer.Authentication
 			await _globals.DisposeAsync();
 		}
 
-		private static IOptionsMonitor<JwtBearerOptions> GetOptionsMonitor(IOptionsMonitorCache<JwtBearerOptions> optionsCache)
+		private static IOptionsMonitor<JwtBearerOptions> GetOptionsMonitor(ServerSettings settings, IOptionsMonitorCache<JwtBearerOptions> optionsCache)
 		{
-			ConfigureNamedOptions<JwtBearerOptions> namedOptions = new ConfigureNamedOptions<JwtBearerOptions>(AuthenticationScheme, options => { });
-			OptionsFactory<JwtBearerOptions> optionsFactory = new OptionsFactory<JwtBearerOptions>(new[] { namedOptions }, Array.Empty<IPostConfigureOptions<JwtBearerOptions>>());
+			ConfigureNamedOptions<JwtBearerOptions> namedOptions = new (AuthenticationScheme, options =>
+			{
+				options.RequireHttpsMetadata = !settings.OidcDebugMode;
+			});
+			OptionsFactory<JwtBearerOptions> optionsFactory = new (new[] { namedOptions }, Array.Empty<IPostConfigureOptions<JwtBearerOptions>>());
 			return new OptionsMonitor<JwtBearerOptions>(optionsFactory, Array.Empty<IOptionsChangeTokenSource<JwtBearerOptions>>(), optionsCache);
 		}
 
diff --git a/Engine/Source/Programs/Horde/HordeServer/Authentication/OidcAuthHandler.cs b/Engine/Source/Programs/Horde/HordeServer/Authentication/OidcAuthHandler.cs
--- a/Engine/Source/Programs/Horde/HordeServer/Authentication/OidcAuthHandler.cs
+++ b/Engine/Source/Programs/Horde/HordeServer/Authentication/OidcAuthHandler.cs
@@ -134,6 +134,7 @@ namespace HordeServer.Authentication
 			options.SaveTokens = true;
 			options.TokenValidationParameters.NameClaimType = "name";
 			options.ClaimActions.Add(new MapRolesClaimAction(settings));
+			options.RequireHttpsMetadata = !settings.OidcDebugMode;
 
 			handler(options);
 		}
diff --git a/Engine/Source/Programs/Horde/HordeServer/ServerSettings.cs b/Engine/Source/Programs/Horde/HordeServer/ServerSettings.cs
--- a/Engine/Source/Programs/Horde/HordeServer/ServerSettings.cs
+++ b/Engine/Source/Programs/Horde/HordeServer/ServerSettings.cs
@@ -202,10 +202,10 @@ namespace HordeServer
 		public string? ServerPrivateCert { get; set; }
 
 		/// <summary>
-		/// Issuer for tokens from the auth provider
+		/// Type of authentication (e.g anonymous, OIDC, built-in Horde accounts)
 		/// </summary>
 		public AuthMethod AuthMethod { get; set; } = AuthMethod.Anonymous;
-
+		
 		/// <summary>
 		/// Optional profile name to report through the /api/v1/server/auth endpoint. Allows sharing auth tokens between providers configured through
 		/// the same profile name in OidcToken.exe config files.
@@ -213,17 +213,17 @@ namespace HordeServer
 		public string? OidcProfileName { get; set; }
 
 		/// <summary>
-		/// Issuer for tokens from the auth provider
+		/// OpenID Connect (OIDC) authority URL (required when OIDC is enabled)
 		/// </summary>
 		public string? OidcAuthority { get; set; }
 
 		/// <summary>
-		/// Audience for validating externally issued tokens
+		/// Audience for validating externally issued tokens (required when OIDC is enabled)
 		/// </summary>
 		public string? OidcAudience { get; set; }
 
 		/// <summary>
-		/// Client id for the OIDC authority
+		/// Client ID for the OIDC authority (required when OIDC is enabled)
 		/// </summary>
 		public string? OidcClientId { get; set; }
 
@@ -239,12 +239,17 @@ namespace HordeServer
 
 		/// <summary>
 		/// Optional redirect url provided to OIDC login for external tools (typically to a local server)
+		/// Default value is the local web server started during signin by EpicGames.OIDC library
 		/// </summary>
-		public string[]? OidcLocalRedirectUrls { get; set; } =
-		{
-			"http://localhost:8749/ugs.client"
-		};
-
+		public string[]? OidcLocalRedirectUrls { get; set; } = ["http://localhost:8749/ugs.client"];
+		
+		/// <summary>
+		/// Debug mode for OIDC which logs reasons for why JWT tokens fail to authenticate
+		/// Also turns off HTTPS requirement for OIDC metadata fetching.
+		/// NOT FOR PRODUCTION USE!
+		/// </summary>
+		public bool OidcDebugMode { get; set; } = false;
+		
 		/// <summary>
 		/// OpenID Connect scopes to request when signing in
 		/// </summary>
diff --git a/Engine/Source/Programs/Horde/HordeServer/Startup.cs b/Engine/Source/Programs/Horde/HordeServer/Startup.cs
--- a/Engine/Source/Programs/Horde/HordeServer/Startup.cs
+++ b/Engine/Source/Programs/Horde/HordeServer/Startup.cs
@@ -453,7 +453,8 @@ namespace HordeServer
 
 			authBuilder.AddServiceAccounts(options => { });
 			schemes.Add(ServiceAccountAuthHandler.AuthenticationScheme);
-
+			
+			ValidateOidcSettings(settings);
 			switch (settings.AuthMethod)
 			{
 				case AuthMethod.Anonymous:
@@ -646,7 +647,28 @@ namespace HordeServer
 
 			OnAddHealthChecks(services);
 		}
-
+		
+		private static void ValidateOidcSettings(ServerSettings settings)
+		{
+			if (settings.AuthMethod is AuthMethod.OpenIdConnect or AuthMethod.Okta)
+			{
+				if (settings.OidcAuthority == null)
+				{
+					throw new ArgumentException($"Key '{nameof(ServerSettings.OidcAuthority)}' in server settings must be set when auth mode {settings.AuthMethod} is used");
+				}
+				
+				if (settings.OidcAudience == null)
+				{
+					throw new ArgumentException($"Key '{nameof(ServerSettings.OidcAudience)}' in server settings must be set when auth mode {settings.AuthMethod} is used");
+				}
+				
+				if (settings.OidcClientId == null)
+				{
+					throw new ArgumentException($"Key '{nameof(ServerSettings.OidcClientId)}' in server settings must be set when auth mode {settings.AuthMethod} is used");
+				}
+			}
+		}
+		
 		public static void ConfigureFormatters()
 		{
 			LogValueFormatter.RegisterTypeAnnotation<AgentId>("AgentId");
