Horde: Show permissions for ACL scopes on account page

Useful for debugging permissions in Horde.

diff --git a/Engine/Source/Programs/Horde/HordeServer.Shared/Acls/AclConfig.cs b/Engine/Source/Programs/Horde/HordeServer.Shared/Acls/AclConfig.cs
--- a/Engine/Source/Programs/Horde/HordeServer.Shared/Acls/AclConfig.cs
+++ b/Engine/Source/Programs/Horde/HordeServer.Shared/Acls/AclConfig.cs
@@ -1,6 +1,7 @@
 // Copyright Epic Games, Inc. All Rights Reserved.
 
 using System.ComponentModel.DataAnnotations;
+using System.Reflection;
 using System.Security.Claims;
 using System.Text.Json.Serialization;
 using EpicGames.Horde.Acls;
@@ -32,6 +33,12 @@ namespace HordeServer.Acls
 		/// </summary>
 		[JsonIgnore]
 		public AclScopeName[]? LegacyScopeNames { get; set; }
+		
+		/// <summary>
+		/// ACL actions associated with this config (for debugging purposes)
+		/// </summary>
+		[JsonIgnore]
+		public AclAction[]? Actions { get; set; }
 
 		/// <summary>
 		/// ACLs which are parented to this
@@ -129,19 +136,20 @@ namespace HordeServer.Acls
 		/// <summary>
 		/// Called after the config file has been read
 		/// </summary>
-		public void PostLoad(AclConfig parentScope, string scopeNameSuffix)
+		public void PostLoad(AclConfig parentScope, string scopeNameSuffix, AclAction[] actions)
 		{
-			PostLoad(parentScope, parentScope.ScopeName.Append(scopeNameSuffix));
+			PostLoad(parentScope, parentScope.ScopeName.Append(scopeNameSuffix), actions);
 		}
 
 		/// <summary>
 		/// Called after the config file has been read
 		/// </summary>
-		public void PostLoad(AclConfig? parentScope, AclScopeName scopeName)
+		public void PostLoad(AclConfig? parentScope, AclScopeName scopeName, AclAction[] actions)
 		{
 			Parent = parentScope;
 			ScopeName = scopeName;
 			Children = null;
+			Actions = actions;
 
 			if (parentScope == null)
 			{
@@ -220,6 +228,20 @@ namespace HordeServer.Acls
 				}
 			}
 		}
+		
+		/// <summary>
+		/// Get all AclActions declared for the types specified
+		/// </summary>
+		/// <param name="type">Array of struct/class types</param>
+		/// <returns>List of AclAction</returns>
+		public static AclAction[] GetActions(Type[] type)
+		{
+			return type
+				.SelectMany(x => x.GetProperties(BindingFlags.Static | BindingFlags.Public))
+				.Where(x => x.PropertyType == typeof(AclAction))
+				.Select(x => (AclAction)x.GetValue(null)!)
+				.ToArray();
+		}
 	}
 
 	/// <summary>
diff --git a/Engine/Source/Programs/Horde/HordeServer/Server/AccountController.cs b/Engine/Source/Programs/Horde/HordeServer/Server/AccountController.cs
--- a/Engine/Source/Programs/Horde/HordeServer/Server/AccountController.cs
+++ b/Engine/Source/Programs/Horde/HordeServer/Server/AccountController.cs
@@ -55,9 +55,11 @@ namespace HordeServer.Server
 		const string StyleSheet =
 			"body { font-family: 'Segoe UI', 'Roboto', arial, sans-serif; } " +
 			"p { margin:20px; font-size:13px; } " +
-			"h1 { margin:20px; font-size:32px; font-weight:200; } " +
+			"h1 { margin: 20px; font-size: 32px; font-weight: 200; } " +
+			"h2 { margin: 20px; font-size: 24px; font-weight: 200; } " +
 			"table { margin:10px 20px; } " +
-			"td { margin:5px; font-size:13px; }";
+			"th { margin: 5px; font-size: 13px; vertical-align: top; padding: 3px; text-align: left; }" +
+			"td { margin:5px; font-size:13px; vertical-align: top; padding: 3px; font-family: Consolas, Courier New, monospace; }";
 
 		readonly IUserCollection _users;
 		readonly IAccountCollection _hordeAccounts;
@@ -116,22 +118,62 @@ namespace HordeServer.Server
 					content.Append("<a href=\"/api/v1/admin/chainedjobtoken\">Get chained job token</a><br/>");
 					content.Append("</p>");
 				}
-				content.Append(CultureInfo.InvariantCulture, $"<p>Claims for {User.Identity?.Name}:");
-				content.Append("<table>");
+				content.AppendLine(CultureInfo.InvariantCulture, $"<h2>Claims for {User.Identity?.Name}:</h2>");
+				content.AppendLine("<table>");
+				content.AppendLine("<tr><th>Type</th><th>Value</th></tr>");
 				foreach (System.Security.Claims.Claim claim in User.Claims)
 				{
-					content.Append(CultureInfo.InvariantCulture, $"<tr><td>{claim.Type}</td><td>{claim.Value}</td></tr>");
+					content.AppendLine(CultureInfo.InvariantCulture, $"<tr><td>{claim.Type}</td><td>{claim.Value}</td></tr>");
 				}
-				content.Append("</table>");
-				content.Append("</p>");
+				content.AppendLine("</table>");
+				
+				bool addedAcl = false;
+				content.AppendLine("<h2>ACL Scopes</h2>");
+				content.AppendLine("<p>Only scopes with at least one authorized action are shown.</p>");
+				content.AppendLine("<table class=\"acl-scopes\">");
+				content.AppendLine("<tr><th>Scope</th><th>Permission</th><th>Authorized</th></tr>");
+				foreach (AclConfig aclConfig in _globalConfig.Value.AclScopes.Values)
+				{
+					List<AclAction> actions = aclConfig.Actions?.ToList() ?? [];
+					bool shouldIncludeAcl = actions.Any(x => aclConfig.Authorize(x, User));
+					List<string> actionTds = actions.Select(x =>
+					{
+						string color = aclConfig.Authorize(x, User) ? "#c8ffc8" : "#ffc8c8";
+						string authorized = aclConfig.Authorize(x, User) ? "Yes" : "No";
+						return $"<td>{x.Name}</td><td style=\"background-color: {color}; text-align: center;\">{authorized}</td>";
+					}).ToList();
+
+					if (actions.Count == 0)
+					{
+						actionTds.Add("<td>No actions</td><td></td>");
+					}
+					
+					if (shouldIncludeAcl)
+					{
+						addedAcl = true;
+						content.AppendLine($"<tr><td rowspan=\"{actionTds.Count + 1}\">{aclConfig.ScopeName}</td>{actionTds[0]}</tr>");
+						foreach (string td in actionTds)
+						{
+							content.AppendLine($"<tr>{td}</td>");
+						}
+						content.AppendLine("<tr style=\"border-bottom:1px solid black\"><td colspan=\"100%\"></td></tr>");
+						content.AppendLine("<tr><td colspan=\"100%\" style=\"display: block; height 20px;\"></td></tr>");
+					}
+				}
+				
+				if (!addedAcl)
+				{
+					content.AppendLine("<tr><td colspan=\"100%\">No scopes</td></tr>");
+				}
+				content.AppendLine("</table>");
 
-				content.Append(CultureInfo.InvariantCulture, $"<p>Built from Perforce</p>");
+				content.AppendLine("<p>Built from Perforce</p>");
 			}
 			else
 			{
-				content.Append("<p><a href=\"/account/login\"><b>Login</b></a></p>");
+				content.AppendLine("<p><a href=\"/account/login\"><b>Login</b></a></p>");
 			}
-			content.Append("</html>");
+			content.AppendLine("</html>");
 			return new ContentResult { ContentType = "text/html", StatusCode = (int)HttpStatusCode.OK, Content = content.ToString() };
 		}
 
diff --git a/Engine/Source/Programs/Horde/HordeServer/Server/GlobalConfig.cs b/Engine/Source/Programs/Horde/HordeServer/Server/GlobalConfig.cs
--- a/Engine/Source/Programs/Horde/HordeServer/Server/GlobalConfig.cs
+++ b/Engine/Source/Programs/Horde/HordeServer/Server/GlobalConfig.cs
@@ -8,10 +8,13 @@ using System.Security.Claims;
 using System.Text.Json.Serialization;
 using EpicGames.Core;
 using EpicGames.Horde.Acls;
+using HordeServer.Accounts;
 using HordeServer.Acls;
 using HordeServer.Configuration;
 using HordeServer.Dashboard;
 using HordeServer.Plugins;
+using HordeServer.Server.Notices;
+using HordeServer.ServiceAccounts;
 using HordeServer.Utilities;
 
 namespace HordeServer.Server
@@ -104,9 +107,17 @@ namespace HordeServer.Server
 		public void PostLoad(ServerSettings serverSettings, IReadOnlyList<ILoadedPlugin> loadedPlugins, IEnumerable<IDefaultAclModifier> defaultAclModifiers)
 		{
 			ServerSettings = serverSettings;
+			AclAction[] aclActions = AclConfig.GetActions(
+			[
+				typeof(AccountAclAction),
+				typeof(ServiceAccountAclAction),
+				typeof(NoticeAclAction),
+				typeof(ServerAclAction),
+				typeof(AdminAclAction),
+			]);
 
 			AclConfig defaultAcl = CreateRootAcl(defaultAclModifiers);
-			Acl.PostLoad(defaultAcl, defaultAcl.ScopeName);
+			Acl.PostLoad(defaultAcl, defaultAcl.ScopeName, aclActions);
 
 			// Ensure that all plugins have an entry in the global config so they can register their ACLs
 			foreach (ILoadedPlugin loadedPlugin in loadedPlugins)
@@ -142,7 +153,7 @@ namespace HordeServer.Server
 			}
 
 			AclConfig defaultAcl = defaultAclBuilder.Build();
-			defaultAcl.PostLoad(null, AclScopeName.Root);
+			defaultAcl.PostLoad(null, AclScopeName.Root, []);
 
 			return defaultAcl;
 		}
diff --git a/Engine/Source/Programs/Horde/Plugins/Analytics/HordeServer.Analytics/Telemetry/Metrics/MetricConfig.cs b/Engine/Source/Programs/Horde/Plugins/Analytics/HordeServer.Analytics/Telemetry/Metrics/MetricConfig.cs
--- a/Engine/Source/Programs/Horde/Plugins/Analytics/HordeServer.Analytics/Telemetry/Metrics/MetricConfig.cs
+++ b/Engine/Source/Programs/Horde/Plugins/Analytics/HordeServer.Analytics/Telemetry/Metrics/MetricConfig.cs
@@ -167,7 +167,7 @@ namespace HordeServer.Telemetry.Metrics
 		/// </summary>
 		public void PostLoad(AclConfig parentAcl)
 		{
-			Acl.PostLoad(parentAcl, $"telemetry-store:{Id}");
+			Acl.PostLoad(parentAcl, $"telemetry-store:{Id}", AclConfig.GetActions([typeof(TelemetryAclAction)])); 
 
 			_metricLookup.Clear();
 			foreach (MetricConfig metric in Metrics)
diff --git a/Engine/Source/Programs/Horde/Plugins/Build/HordeServer.Build/Projects/ProjectConfig.cs b/Engine/Source/Programs/Horde/Plugins/Build/HordeServer.Build/Projects/ProjectConfig.cs
--- a/Engine/Source/Programs/Horde/Plugins/Build/HordeServer.Build/Projects/ProjectConfig.cs
+++ b/Engine/Source/Programs/Horde/Plugins/Build/HordeServer.Build/Projects/ProjectConfig.cs
@@ -13,6 +13,8 @@ using HordeServer.Acls;
 using HordeServer.Agents.Pools;
 using HordeServer.Artifacts;
 using HordeServer.Configuration;
+using HordeServer.Devices;
+using HordeServer.Logs;
 using HordeServer.Streams;
 using HordeServer.Utilities;
 
@@ -128,7 +130,8 @@ namespace HordeServer.Projects
 		public void PostLoad(ProjectId id, AclConfig parentAcl, IEnumerable<ArtifactTypeConfig> parentArtifactTypes)
 		{
 			Id = id;
-			Acl.PostLoad(parentAcl, $"project:{Id}");
+			AclAction[] aclActions = AclConfig.GetActions([typeof(ArtifactAclAction), typeof(DeviceAclAction), typeof(LogAclAction), typeof(ProjectAclAction)]);
+			Acl.PostLoad(parentAcl, $"project:{Id}", aclActions);
 			Acl.LegacyScopeNames = new AclScopeName[] { parentAcl.ScopeName.Append($"p:{Id}") };
 
 			_artifactTypeLookup.Clear();
diff --git a/Engine/Source/Programs/Horde/Plugins/Build/HordeServer.Build/Streams/StreamConfig.cs b/Engine/Source/Programs/Horde/Plugins/Build/HordeServer.Build/Streams/StreamConfig.cs
--- a/Engine/Source/Programs/Horde/Plugins/Build/HordeServer.Build/Streams/StreamConfig.cs
+++ b/Engine/Source/Programs/Horde/Plugins/Build/HordeServer.Build/Streams/StreamConfig.cs
@@ -25,8 +25,11 @@ using HordeServer.Acls;
 using HordeServer.Artifacts;
 using HordeServer.Configuration;
 using HordeServer.Issues;
+using HordeServer.Jobs;
+using HordeServer.Jobs.Bisect;
 using HordeServer.Jobs.Graphs;
 using HordeServer.Jobs.Templates;
+using HordeServer.Notifications;
 using HordeServer.Projects;
 using HordeServer.Replicators;
 using HordeServer.Utilities;
@@ -264,8 +267,16 @@ namespace HordeServer.Streams
 		{
 			Id = id;
 			ProjectConfig = projectConfig;
-
-			Acl.PostLoad(projectConfig.Acl, $"stream:{Id}");
+			
+			AclAction[] aclActions = AclConfig.GetActions(
+				[
+					typeof(BisectTaskAclAction),
+					typeof(JobAclAction),
+					typeof(NotificationAclAction),
+					typeof(ReplicatorAclAction),
+					typeof(StreamAclAction)
+				]);
+			Acl.PostLoad(projectConfig.Acl, $"stream:{Id}", aclActions);
 			Acl.LegacyScopeNames = projectConfig.Acl.LegacyScopeNames?.Select(x => x.Append($"s:{Id}")).ToArray();
 
 			JobOptions.MergeDefaults(projectConfig.JobOptions);
@@ -904,7 +915,7 @@ namespace HordeServer.Streams
 				Id = new TemplateId(StringId.Sanitize(Name));
 			}
 
-			Acl.PostLoad(streamConfig.Acl, $"template:{Id}");
+			Acl.PostLoad(streamConfig.Acl, $"template:{Id}", []);
 			Acl.LegacyScopeNames = streamConfig.Acl.LegacyScopeNames?.Select(x => x.Append($"t:{Id}")).ToArray();
 		}
 	}
diff --git a/Engine/Source/Programs/Horde/Plugins/Compute/HordeServer.Compute/ComputeConfig.cs b/Engine/Source/Programs/Horde/Plugins/Compute/HordeServer.Compute/ComputeConfig.cs
--- a/Engine/Source/Programs/Horde/Plugins/Compute/HordeServer.Compute/ComputeConfig.cs
+++ b/Engine/Source/Programs/Horde/Plugins/Compute/HordeServer.Compute/ComputeConfig.cs
@@ -11,8 +11,14 @@ using EpicGames.Horde.Compute;
 using EpicGames.Horde.Tools;
 using EpicGames.Serialization;
 using HordeServer.Acls;
+using HordeServer.Agents;
+using HordeServer.Agents.Leases;
 using HordeServer.Agents.Pools;
+using HordeServer.Agents.Sessions;
+using HordeServer.Agents.Software;
+using HordeServer.Compute;
 using HordeServer.Configuration;
+using HordeServer.Logs;
 using HordeServer.Plugins;
 
 #pragma warning disable CA2227 // Change x to be read-only by removing the property setter
@@ -269,7 +275,17 @@ namespace HordeServer
 		/// <param name="parentAcl">The parent config instance</param>
 		public void PostLoad(AclConfig parentAcl)
 		{
-			Acl.PostLoad(parentAcl, $"compute:{Id}");
+			AclAction[] aclActions = AclConfig.GetActions(
+				[
+					typeof(AgentAclAction),
+					typeof(AgentSoftwareAclAction),
+					typeof(ComputeAclAction),
+					typeof(LeaseAclAction),
+					typeof(LogAclAction),
+					typeof(PoolAclAction),
+					typeof(SessionAclAction)
+				]);
+			Acl.PostLoad(parentAcl, $"compute:{Id}", aclActions);
 		}
 
 		/// <summary>
diff --git a/Engine/Source/Programs/Horde/Plugins/Secrets/HordeServer.Secrets/Secrets/SecretConfig.cs b/Engine/Source/Programs/Horde/Plugins/Secrets/HordeServer.Secrets/Secrets/SecretConfig.cs
--- a/Engine/Source/Programs/Horde/Plugins/Secrets/HordeServer.Secrets/Secrets/SecretConfig.cs
+++ b/Engine/Source/Programs/Horde/Plugins/Secrets/HordeServer.Secrets/Secrets/SecretConfig.cs
@@ -40,7 +40,7 @@ namespace HordeServer.Secrets
 		/// <param name="parentAcl">Parent acl object</param>
 		public void PostLoad(AclConfig parentAcl)
 		{
-			Acl.PostLoad(parentAcl, $"secret:{Id}");
+			Acl.PostLoad(parentAcl, $"secret:{Id}", AclConfig.GetActions([typeof(SecretAclAction)]));
 		}
 
 		/// <summary>
diff --git a/Engine/Source/Programs/Horde/Plugins/Storage/HordeServer.Storage/StorageConfig.cs b/Engine/Source/Programs/Horde/Plugins/Storage/HordeServer.Storage/StorageConfig.cs
--- a/Engine/Source/Programs/Horde/Plugins/Storage/HordeServer.Storage/StorageConfig.cs
+++ b/Engine/Source/Programs/Horde/Plugins/Storage/HordeServer.Storage/StorageConfig.cs
@@ -348,7 +348,7 @@ namespace HordeServer.Storage
 		/// </summary>
 		public void PostLoad(AclConfig parentAcl)
 		{
-			Acl.PostLoad(parentAcl, $"namespace:{Id}");
+			Acl.PostLoad(parentAcl, $"namespace:{Id}", AclConfig.GetActions([typeof(StorageAclAction)]));
 		}
 
 		/// <summary>
diff --git a/Engine/Source/Programs/Horde/Plugins/Tools/HordeServer.Tools/Tools/ToolConfig.cs b/Engine/Source/Programs/Horde/Plugins/Tools/HordeServer.Tools/Tools/ToolConfig.cs
--- a/Engine/Source/Programs/Horde/Plugins/Tools/HordeServer.Tools/Tools/ToolConfig.cs
+++ b/Engine/Source/Programs/Horde/Plugins/Tools/HordeServer.Tools/Tools/ToolConfig.cs
@@ -105,7 +105,7 @@ namespace HordeServer.Tools
 		/// <param name="parentAcl">Parent ACL object</param>
 		public void PostLoad(AclConfig parentAcl)
 		{
-			Acl.PostLoad(parentAcl, $"tool:{Id}");
+			Acl.PostLoad(parentAcl, $"tool:{Id}", AclConfig.GetActions([typeof(ToolAclAction)]));
 		}
 
 		/// <inheritdoc cref="AclConfig.Authorize(AclAction, ClaimsPrincipal)"/>
