Horde: Fix some static analysis warnings.

#rnx

diff --git a/Engine/Source/Programs/Horde/Plugins/Build/HordeServer.Build.Tests/Agents/AgentServiceTest.cs b/Engine/Source/Programs/Horde/Plugins/Build/HordeServer.Build.Tests/Agents/AgentServiceTest.cs
--- a/Engine/Source/Programs/Horde/Plugins/Build/HordeServer.Build.Tests/Agents/AgentServiceTest.cs
+++ b/Engine/Source/Programs/Horde/Plugins/Build/HordeServer.Build.Tests/Agents/AgentServiceTest.cs
@@ -88,12 +88,6 @@ public class AgentServiceTest : BuildTestSetup
 		Assert.AreEqual(ToUnixTime(expected), ToUnixTime(actual!.Value));
 	}
 
-	private static void AssertNotEqual(DateTime expected, DateTime? actual)
-	{
-		// When saved as MongoDB documents, some precision is lost. This compares only Unix seconds.
-		Assert.AreNotEqual(ToUnixTime(expected), ToUnixTime(actual!.Value));
-	}
-
 	[TestMethod]
 	[DataRow(AgentStatus.Ok)]
 	[DataRow(AgentStatus.Unhealthy)]
diff --git a/Engine/Source/Programs/Horde/Plugins/Build/HordeServer.Build.Tests/BuildTestSetup.cs b/Engine/Source/Programs/Horde/Plugins/Build/HordeServer.Build.Tests/BuildTestSetup.cs
--- a/Engine/Source/Programs/Horde/Plugins/Build/HordeServer.Build.Tests/BuildTestSetup.cs
+++ b/Engine/Source/Programs/Horde/Plugins/Build/HordeServer.Build.Tests/BuildTestSetup.cs
@@ -276,6 +276,7 @@ namespace HordeServer.Tests
 		/// <param name="awsInstanceId">AWS instance ID for the agent (will be set in properties)</param>
 		/// <param name="lease">A lease to assign the agent</param>
 		/// <param name="ephemeral">Whether the agent is ephemeral</param>
+		/// <param name="status">Initial status for the agent</param>
 		/// <returns>A new agent</returns>
 		public async Task<IAgent> CreateAgentAsync(
 			PoolId? poolId,
diff --git a/Engine/Source/Programs/Horde/Plugins/Build/HordeServer.Build.Tests/Jobs/JobServiceTests.cs b/Engine/Source/Programs/Horde/Plugins/Build/HordeServer.Build.Tests/Jobs/JobServiceTests.cs
--- a/Engine/Source/Programs/Horde/Plugins/Build/HordeServer.Build.Tests/Jobs/JobServiceTests.cs
+++ b/Engine/Source/Programs/Horde/Plugins/Build/HordeServer.Build.Tests/Jobs/JobServiceTests.cs
@@ -2,7 +2,6 @@
 
 using System.Text;
 using EpicGames.Core;
-using EpicGames.Horde.Agents;
 using EpicGames.Horde.Agents.Pools;
 using EpicGames.Horde.Commits;
 using EpicGames.Horde.Jobs;
diff --git a/Engine/Source/Programs/Horde/Plugins/Build/HordeServer.Build.Tests/Jobs/JobTaskSourceTests.cs b/Engine/Source/Programs/Horde/Plugins/Build/HordeServer.Build.Tests/Jobs/JobTaskSourceTests.cs
--- a/Engine/Source/Programs/Horde/Plugins/Build/HordeServer.Build.Tests/Jobs/JobTaskSourceTests.cs
+++ b/Engine/Source/Programs/Horde/Plugins/Build/HordeServer.Build.Tests/Jobs/JobTaskSourceTests.cs
@@ -1,6 +1,5 @@
 // Copyright Epic Games, Inc. All Rights Reserved.
 
-using EpicGames.Horde.Agents;
 using EpicGames.Horde.Agents.Pools;
 using EpicGames.Horde.Commits;
 using EpicGames.Horde.Jobs;
diff --git a/Engine/Source/Programs/Horde/Plugins/Compute/HordeServer.Compute/Agents/AgentCollection.cs b/Engine/Source/Programs/Horde/Plugins/Compute/HordeServer.Compute/Agents/AgentCollection.cs
--- a/Engine/Source/Programs/Horde/Plugins/Compute/HordeServer.Compute/Agents/AgentCollection.cs
+++ b/Engine/Source/Programs/Horde/Plugins/Compute/HordeServer.Compute/Agents/AgentCollection.cs
@@ -558,7 +558,7 @@ namespace HordeServer.Agents
 
 			using (IAsyncCursor<AgentDocument> cursor = await collection.Find(filter).ToCursorAsync(cancellationToken))
 			{
-				while (await cursor.MoveNextAsync())
+				while (await cursor.MoveNextAsync(cancellationToken))
 				{
 					List<AgentDocument> documents = await PostLoadAsync(cursor.Current.ToList(), cancellationToken);
 					foreach (AgentDocument document in documents)
