Horde: Fix invalid dep injection of ServerSettings in JwtHandler

diff --git a/Engine/Source/Programs/Horde/HordeServer/Authentication/JwtAuthHandler.cs b/Engine/Source/Programs/Horde/HordeServer/Authentication/JwtAuthHandler.cs
--- a/Engine/Source/Programs/Horde/HordeServer/Authentication/JwtAuthHandler.cs
+++ b/Engine/Source/Programs/Horde/HordeServer/Authentication/JwtAuthHandler.cs
@@ -25,8 +25,8 @@ namespace HordeServer.Authentication
 
 		readonly AsyncCachedValue<IGlobals> _globals;
 
-		public JwtAuthHandler(ILoggerFactory logger, UrlEncoder encoder, ServerSettings settings, GlobalsService globalsService, IOptionsMonitorCache<JwtBearerOptions> optionsCache)
-			: base(GetOptionsMonitor(settings, optionsCache), logger, encoder)
+		public JwtAuthHandler(ILoggerFactory logger, UrlEncoder encoder, IOptionsMonitor<ServerSettings> settings, GlobalsService globalsService, IOptionsMonitorCache<JwtBearerOptions> optionsCache)
+			: base(GetOptionsMonitor(settings.CurrentValue, optionsCache), logger, encoder)
 		{
 			_globals = new AsyncCachedValue<IGlobals>(async ctx => await globalsService.GetAsync(ctx), TimeSpan.FromSeconds(30.0));
 		}
