Horde: Add REST API endpoint for listing ACL permissions for current user

diff --git a/Engine/Source/Programs/Horde/HordeServer/Server/AccountController.cs b/Engine/Source/Programs/Horde/HordeServer/Server/AccountController.cs
--- a/Engine/Source/Programs/Horde/HordeServer/Server/AccountController.cs
+++ b/Engine/Source/Programs/Horde/HordeServer/Server/AccountController.cs
@@ -11,6 +11,7 @@ using System.Threading.Tasks;
 using EpicGames.Horde.Accounts;
 using EpicGames.Horde.Acls;
 using EpicGames.Horde.Server;
+using EpicGames.Horde.Users;
 using HordeServer.Accounts;
 using HordeServer.Acls;
 using HordeServer.Authentication;
@@ -127,41 +128,42 @@ namespace HordeServer.Server
 				}
 				content.AppendLine("</table>");
 				
-				bool addedAcl = false;
-				content.AppendLine("<h2>ACL Scopes</h2>");
+				content.AppendLine("<h2>ACL Permissions</h2>");
 				content.AppendLine("<p>Only scopes with at least one authorized action are shown.</p>");
 				content.AppendLine("<table class=\"acl-scopes\">");
-				content.AppendLine("<tr><th>Scope</th><th>Permission</th><th>Authorized</th></tr>");
-				foreach (AclConfig aclConfig in _globalConfig.Value.AclScopes.Values)
+				content.AppendLine("<tr><th>Scope</th><th>Action</th><th>Authorized</th></tr>");
+				
+				List<UserAclPermission> aclPermissions = UserController.GetUserAclPermissions(_globalConfig.Value, User);
+				Dictionary<string, List<UserAclPermission>> scopeToPermissions = aclPermissions
+					.GroupBy(permission => permission.Scope)
+					.ToDictionary(group => group.Key, group => group.ToList());
+				
+				foreach ((string scope, List<UserAclPermission> permissions) in scopeToPermissions)
 				{
-					List<AclAction> actions = aclConfig.Actions?.ToList() ?? [];
-					bool shouldIncludeAcl = actions.Any(x => aclConfig.Authorize(x, User));
-					List<string> actionTds = actions.Select(x =>
-					{
-						string color = aclConfig.Authorize(x, User) ? "#c8ffc8" : "#ffc8c8";
-						string authorized = aclConfig.Authorize(x, User) ? "Yes" : "No";
-						return $"<td>{x.Name}</td><td style=\"background-color: {color}; text-align: center;\">{authorized}</td>";
-					}).ToList();
-
-					if (actions.Count == 0)
-					{
-						actionTds.Add("<td>No actions</td><td></td>");
-					}
-					
-					if (shouldIncludeAcl)
+					if (permissions.Count > 0)
 					{
-						addedAcl = true;
-						content.AppendLine($"<tr><td rowspan=\"{actionTds.Count + 1}\">{aclConfig.ScopeName}</td>{actionTds[0]}</tr>");
-						foreach (string td in actionTds)
+						for (int i = 0; i < permissions.Count; i++)
 						{
-							content.AppendLine($"<tr>{td}</td>");
+							UserAclPermission uap = permissions[i];
+							if (i == 0)
+							{
+								content.Append("<tr>");
+								content.AppendLine($"<tr><td rowspan=\"{permissions.Count + 1}\">{scope}</td>");
+							}
+							else
+							{
+								string color = uap.IsAuthorized ? "#c8ffc8" : "#ffc8c8";
+								string authorized = uap.IsAuthorized ? "Yes" : "No";
+								content.AppendLine($"<td>{uap.Action}</td><td style=\"background-color: {color}; text-align: center;\">{authorized}</td>");
+							}
+							content.Append("</tr>");
 						}
-						content.AppendLine("<tr style=\"border-bottom:1px solid black\"><td colspan=\"100%\"></td></tr>");
-						content.AppendLine("<tr><td colspan=\"100%\" style=\"display: block; height 20px;\"></td></tr>");
 					}
+					content.AppendLine("<tr style=\"border-bottom:1px solid black\"><td colspan=\"100%\"></td></tr>");
+					content.AppendLine("<tr><td colspan=\"100%\" style=\"display: block; height 20px;\"></td></tr>");
 				}
 				
-				if (!addedAcl)
+				if (scopeToPermissions.Count == 0)
 				{
 					content.AppendLine("<tr><td colspan=\"100%\">No scopes</td></tr>");
 				}
diff --git a/Engine/Source/Programs/Horde/HordeServer/Users/UserController.cs b/Engine/Source/Programs/Horde/HordeServer/Users/UserController.cs
--- a/Engine/Source/Programs/Horde/HordeServer/Users/UserController.cs
+++ b/Engine/Source/Programs/Horde/HordeServer/Users/UserController.cs
@@ -9,6 +9,7 @@ using EpicGames.Horde.Jobs;
 using EpicGames.Horde.Jobs.Bisect;
 using EpicGames.Horde.Users;
 using HordeServer.Accounts;
+using HordeServer.Acls;
 using HordeServer.Plugins;
 using HordeServer.Server;
 using HordeServer.Server.Notices;
@@ -119,5 +120,42 @@ namespace HordeServer.Users
 		{
 			return User.Claims.Select(x => new { x.Type, x.Value }).ToArray();
 		}
+		
+		/// <summary>
+		/// Gets ACL permissions for the current user
+		/// </summary>
+		[HttpGet]
+		[Route("/api/v1/user/acl-permissions")]
+		public ActionResult<GetUserAclPermissionsResponse> GetUserAclPermissions()
+		{
+			return new GetUserAclPermissionsResponse(GetUserAclPermissions(_globalConfig.Value, User));
+		}
+		
+		/// <summary>
+		/// Get a list of ACL permissions for a user
+		/// A scope is only included if the user is authorized for at least one action inside of it
+		/// </summary>
+		/// <param name="globalConfig">Global config</param>
+		/// <param name="user">User to inspect</param>
+		/// <returns></returns>
+		public static List<UserAclPermission> GetUserAclPermissions(GlobalConfig globalConfig, ClaimsPrincipal user)
+		{
+			List<UserAclPermission> result = [];
+			foreach (AclConfig aclConfig in globalConfig.AclScopes.Values)
+			{
+				string scope = aclConfig.ScopeName.ToString();
+				List<UserAclPermission> permissions = (aclConfig.Actions ?? [])
+					.Select(x => new UserAclPermission(scope, x.Name, aclConfig.Authorize(x, user)))
+					.ToList();
+				
+				bool hasPermission = permissions.Exists(x => x.IsAuthorized);
+				if (hasPermission)
+				{
+					result.AddRange(permissions);
+				}
+			}
+			
+			return result;
+		}
 	}
 }
diff --git a/Engine/Source/Programs/Shared/EpicGames.Horde/Users/UserMessages.cs b/Engine/Source/Programs/Shared/EpicGames.Horde/Users/UserMessages.cs
--- a/Engine/Source/Programs/Shared/EpicGames.Horde/Users/UserMessages.cs
+++ b/Engine/Source/Programs/Shared/EpicGames.Horde/Users/UserMessages.cs
@@ -122,6 +122,56 @@ namespace EpicGames.Horde.Users
 			Value = value;
 		}
 	}
+	
+	/// <summary>
+	/// Resolved permissions for an action in a given ACL scope
+	/// </summary>
+	public class UserAclPermission
+	{
+		/// <summary>
+		/// Scope name
+		/// </summary>
+		public string Scope { get; set; }
+		
+		/// <summary>
+		/// Action name
+		/// </summary>
+		public string Action { get; set; }
+		
+		/// <summary>
+		/// Whether action is authorized in given scope
+		/// </summary>
+		public bool IsAuthorized { get; set; }
+		
+		/// <summary>
+		/// Constructor
+		/// </summary>
+		public UserAclPermission(string scope, string action, bool isAuthorized)
+		{
+			Scope = scope;
+			Action = action;
+			IsAuthorized = isAuthorized;
+		}
+	}
+	
+	/// <summary>
+	/// Resolved permissions for ACL scopes for a given user
+	/// </summary>
+	public class GetUserAclPermissionsResponse
+	{
+		/// <summary>
+		/// List of ACL permissions
+		/// </summary>
+		public List<UserAclPermission> AclPermissions { get; set; }
+
+		/// <summary>
+		/// Constructor
+		/// </summary>
+		public GetUserAclPermissionsResponse(List<UserAclPermission> aclPermissions)
+		{
+			AclPermissions = aclPermissions;
+		}
+	}
 
 	/// <summary>
 	/// Job template settings for the current user
