Horde: Agent enrollment improvements
#skipci

diff --git a/Engine/Source/Programs/Horde/HordeDashboard/src/components/agents/AgentRequestsView.tsx b/Engine/Source/Programs/Horde/HordeDashboard/src/components/agents/AgentRequestsView.tsx
--- a/Engine/Source/Programs/Horde/HordeDashboard/src/components/agents/AgentRequestsView.tsx
+++ b/Engine/Source/Programs/Horde/HordeDashboard/src/components/agents/AgentRequestsView.tsx
@@ -11,10 +11,13 @@ import { getHordeStyling } from "../../styles/Styles";
 import { Breadcrumbs } from "../Breadcrumbs";
 import { TopNav } from "../TopNav";
 import ErrorHandler from "../ErrorHandler";
+import { useNavigate } from "react-router-dom";
+import moment from "moment";
+import { agentStore } from "horde/backend/AgentStore";
 
 class AgentRequestsHandler extends PollBase {
 
-   constructor(pollTime = 5000) {
+   constructor(pollTime = 2500) {
 
       super(pollTime);
 
@@ -31,11 +34,16 @@ class AgentRequestsHandler extends PollBase {
    async poll(): Promise<void> {
 
       try {
-
          const requests = await backend.getAgentRegistrationRequests();
-         this.requests = requests.agents;
-         this.initial = false;
-         this.setUpdated();
+         if (requests.agents.length) {
+            this.requests = requests.agents;
+            this.selectedAgents = this.requests;
+            this.selection.setItems(this.selectedAgents);
+            this.selection.setRangeSelected(0, this.selectedAgents.length, true, false);
+            this.initial = false;
+            this.stop();
+            this.setUpdated();
+         }
 
       } catch (err) {
 
@@ -60,8 +68,9 @@ const handler = new AgentRequestsHandler();
 
 const AgentsPanel: React.FC = observer(() => {
 
-   const [confirmRegister, setConfirmRegister] = useState(false);
+   const [confirmRegister, setConfirmRegister] = useState(true);
    const [submitting, setSubmitting] = useState(false);
+   const navigate = useNavigate();
 
    useEffect(() => {
 
@@ -82,13 +91,23 @@ const AgentsPanel: React.FC = observer(() => {
 
       const agents = handler.selectedAgents.map(a => { return { key: a.key } })
 
+      const queryAgents = async (queryTime: string) => {
+
+         const nagents = await backend.getAgents({ modifiedAfter: queryTime });
+         if (nagents.length) {
+            await agentStore.update(false);
+            navigate("/agents");
+         } else {
+            setTimeout(queryAgents, 3000);
+         }
+      }
+
       try {
          setSubmitting(true);
+         const queryTime = moment.utc().toISOString();
          await backend.registerAgents({ agents: agents });
-         setSubmitting(false);
-         handler.stop();
-         handler.clear();
-         handler.start();
+         setTimeout(queryAgents, 3000);
+         
       } catch (reason) {
          console.error(reason);
 
@@ -121,12 +140,18 @@ const AgentsPanel: React.FC = observer(() => {
 
    const { hordeClasses } = getHordeStyling();
 
+   let title = handler.selectedAgents.length > 1 ? "Enroll Agents" : "Enroll Agent";
+   if (handler.selectedAgents.length == 1) {
+      title += " " + handler.selectedAgents[0].hostName
+   }
+   let agentText = handler.selectedAgents.length > 1 ? "these agents" : "this agent";
+
    return <Stack>
       {submitting && <Modal isOpen={true} isBlocking={true} topOffsetFixed={true} styles={{ main: { padding: 8, width: 400, hasBeenOpened: false, top: "120px", position: "absolute" } }} >
          <Stack style={{ paddingTop: 32 }}>
             <Stack tokens={{ childrenGap: 24 }} styles={{ root: { padding: 8 } }}>
                <Stack horizontalAlign="center">
-                  <Text variant="mediumPlus">Please wait...</Text>
+                  <Text variant="mediumPlus">Enrolling, please wait...</Text>
                </Stack>
                <Stack verticalAlign="center" style={{ paddingBottom: 32 }}>
                   <Spinner size={SpinnerSize.large} />
@@ -134,15 +159,15 @@ const AgentsPanel: React.FC = observer(() => {
             </Stack>
          </Stack>
       </Modal>}
-      {confirmRegister &&
+      {confirmRegister && handler.selectedAgents.length > 0 &&
          <Dialog
             hidden={false}
             onDismiss={() => setConfirmRegister(false)}
             minWidth={612}
             dialogContentProps={{
                type: DialogType.normal,
-               title: `Enroll Agents`,
-               subText: `Confirm enrolling agents: ${handler.selectedAgents.map(a => a.hostName).join(", ")}`
+               title: title,
+               subText: `Do you want to enroll ${agentText}? This will allow ${agentText} to take on work assigned to it by the server: ${handler.selectedAgents.map(a => a.hostName).join(", ")}`
             }}
             modalProps={{ isBlocking: true, topOffsetFixed: true, styles: { main: { padding: 8, width: 400, hasBeenOpened: false, top: "120px", position: "absolute" } } }} >
             <Stack style={{ height: "18px" }} />
@@ -154,11 +179,16 @@ const AgentsPanel: React.FC = observer(() => {
       }
       {<Stack style={{ paddingBottom: 12 }}>
          <Stack verticalAlign="center">
-            {!!requests.length && !handler.initial && <Stack horizontalAlign="end">
+            {!!requests.length && <Stack horizontalAlign="end">
                <PrimaryButton disabled={!handler.selectedAgents.length} styles={{ root: { fontFamily: "Horde Open Sans SemiBold" } }} onClick={() => setConfirmRegister(true)}>Enroll Agents</PrimaryButton>
             </Stack>}
-            {!requests.length && !handler.initial && <Stack horizontalAlign="center">
-               <Text variant="mediumPlus">No Agent Enrollment Requests Found</Text>
+            {!requests.length && <Stack horizontalAlign="center" tokens={{ childrenGap: 12 }}>
+               <Stack>
+                  <Text variant="mediumPlus">No enrollment requests found, waiting for agents to connect</Text>
+               </Stack>
+               <Stack>
+                  <Spinner size={SpinnerSize.large} />
+               </Stack>
             </Stack>}
          </Stack>
       </Stack>}
