Horde: Store session state in Redis rather than MongoDB.

diff --git a/Engine/Source/Programs/Horde/Plugins/Build/HordeServer.Build.Tests/Agents/AgentServiceTest.cs b/Engine/Source/Programs/Horde/Plugins/Build/HordeServer.Build.Tests/Agents/AgentServiceTest.cs
--- a/Engine/Source/Programs/Horde/Plugins/Build/HordeServer.Build.Tests/Agents/AgentServiceTest.cs
+++ b/Engine/Source/Programs/Horde/Plugins/Build/HordeServer.Build.Tests/Agents/AgentServiceTest.cs
@@ -11,6 +11,7 @@ using HordeServer.Auditing;
 using HordeServer.Jobs;
 using HordeServer.Utilities;
 using Microsoft.AspNetCore.Mvc;
+using Microsoft.Extensions.DependencyInjection;
 
 namespace HordeServer.Tests.Agents;
 
@@ -165,13 +166,13 @@ public class AgentServiceTest : BuildTestSetup
 
 		// Let background task run for purging outdated sessions, which will terminate session for our agent
 		await Clock.AdvanceAsync(TimeSpan.FromHours(1));
-		await ((AgentCollection)AgentCollection).TickSharedAsync(CancellationToken.None);
+		await ServiceProvider.GetRequiredService<AgentCollection>().TickSharedAsync(CancellationToken.None);
 
 		// Ephemeral agent is marked as deleted once its session is terminated
 		Assert.IsTrue((await AgentService.GetAgentAsync(agent.Id))!.Deleted);
 
 		await Clock.AdvanceAsync(TimeSpan.FromHours(25));
-		await ((AgentCollection)AgentCollection).TickSharedAsync(CancellationToken.None);
+		await ServiceProvider.GetRequiredService<AgentCollection>().TickSharedAsync(CancellationToken.None);
 
 		// Once more time has passed, the ephemeral agent marked as deleted is removed from database
 		Assert.IsNull(await AgentService.GetAgentAsync(agent.Id));
diff --git a/Engine/Source/Programs/Horde/Plugins/Build/HordeServer.Build.Tests/Compute/ComputeServiceTest.cs b/Engine/Source/Programs/Horde/Plugins/Build/HordeServer.Build.Tests/Compute/ComputeServiceTest.cs
--- a/Engine/Source/Programs/Horde/Plugins/Build/HordeServer.Build.Tests/Compute/ComputeServiceTest.cs
+++ b/Engine/Source/Programs/Horde/Plugins/Build/HordeServer.Build.Tests/Compute/ComputeServiceTest.cs
@@ -380,7 +380,7 @@ namespace HordeServer.Tests.Compute
 		{
 			IOptionsMonitor<ComputeConfig> computeConfig = ServiceProvider.GetRequiredService<IOptionsMonitor<ComputeConfig>>();
 			StaticComputeConfig ss = new() { ComputeTunnelAddress = tunnelAddress };
-			ComputeService cs = new(AgentCollection, LogCollection, AgentService, AgentRelayService, GetRedisServiceSingleton(),
+			ComputeService cs = new(AgentCollection, ServiceProvider.GetRequiredService<IAgentScheduler>(), LogCollection, AgentService, AgentRelayService, GetRedisServiceSingleton(),
 				new TestOptionsMonitor<StaticComputeConfig>(ss), computeConfig, Clock, Tracer, Meter,
 				NullLogger<ComputeService>.Instance);
 			await CreateAgentAsync(new PoolId(pool), properties: ["ComputeIp=11.0.0.1", "ComputePort=5000"]);
diff --git a/Engine/Source/Programs/Horde/Plugins/Build/HordeServer.Build.Tests/Server/RpcServiceTest.cs b/Engine/Source/Programs/Horde/Plugins/Build/HordeServer.Build.Tests/Server/RpcServiceTest.cs
--- a/Engine/Source/Programs/Horde/Plugins/Build/HordeServer.Build.Tests/Server/RpcServiceTest.cs
+++ b/Engine/Source/Programs/Horde/Plugins/Build/HordeServer.Build.Tests/Server/RpcServiceTest.cs
@@ -302,7 +302,7 @@ namespace HordeServer.Tests.Server
 			res = await RpcService.CreateSession(req, _adminContext);
 
 			agent = (await AgentService.GetAgentAsync(new AgentId(res.AgentId)))!;
-			CollectionAssert.AreEquivalent(new List<PoolId> { new("fooPool"), new("barPool"), new("bazPool") }, agent.Pools.ToList());
+			CollectionAssert.AreEquivalent(new List<PoolId> { new("bazPool") }, agent.Pools.ToList());
 		}
 
 		[TestMethod]
diff --git a/Engine/Source/Programs/Horde/Plugins/Compute/HordeServer.Compute.Tests/Agents/AgentCollectionTests.cs b/Engine/Source/Programs/Horde/Plugins/Compute/HordeServer.Compute.Tests/Agents/AgentCollectionTests.cs
--- a/Engine/Source/Programs/Horde/Plugins/Compute/HordeServer.Compute.Tests/Agents/AgentCollectionTests.cs
+++ b/Engine/Source/Programs/Horde/Plugins/Compute/HordeServer.Compute.Tests/Agents/AgentCollectionTests.cs
@@ -21,28 +21,44 @@ public class AgentCollectionTests : ComputeTestSetup
 
 	public AgentCollectionTests()
 	{
-		_agent = AgentCollection.AddAsync(new AgentId("test"), ephemeral: true, enrollmentKey: "").Result;
+		_agent = null!;
 		_lease1 = new CreateLeaseOptions(LeaseId.Parse("aaaaaaaaaaaaaaaaaaaaaaaa"), null, "lease", null, null, null, null, false, new Empty());
 		_lease2 = new CreateLeaseOptions(LeaseId.Parse("bbbbbbbbbbbbbbbbbbbbbbbb"), null, "lease", null, null, null, null, false, new Empty());
 		_leaseWithParent3 = new CreateLeaseOptions(LeaseId.Parse("cccccccccccccccccccccccc"), _lease1.Id, "leaseWithParent3", null, null, null, null, false, new Empty());
 		_leaseWithParent4 = new CreateLeaseOptions(LeaseId.Parse("dddddddddddddddddddddddd"), _lease1.Id, "leaseWithParent4", null, null, null, null, false, new Empty());
 	}
 
+	[TestInitialize]
+	public async Task SetupAsync()
+	{
+		_agent = await AgentCollection.AddAsync(new AgentId("test"), ephemeral: true, enrollmentKey: "");
+	}
+
 	[TestMethod]
 	public async Task AddLeaseAsync()
 	{
-		await _agent.TryCreateLeaseAsync(_lease1);
+		IAgent? agent = await _agent.TryCreateSessionAsync(new CreateSessionOptions(new RpcAgentCapabilities(), Array.Empty<PoolId>(), "foo"));
+		Assert.IsNotNull(agent);
+		Assert.IsNotNull(agent.SessionId);
+
+		agent = await agent.TryCreateLeaseAsync(_lease1);
+		Assert.IsNotNull(agent);
+		Assert.AreEqual(1, agent.Leases.Count);
 		await UpdateAgentAsync();
 
-		await _agent.TryCreateLeaseAsync(_leaseWithParent3);
+		agent = await agent.TryCreateLeaseAsync(_leaseWithParent3);
+		Assert.IsNotNull(agent);
+		Assert.AreEqual(2, agent.Leases.Count);
 		await UpdateAgentAsync();
 
-		await _agent.TryCreateLeaseAsync(_leaseWithParent4);
+		agent = await agent.TryCreateLeaseAsync(_leaseWithParent4);
+		Assert.IsNotNull(agent);
+		Assert.AreEqual(3, agent.Leases.Count);
 		await UpdateAgentAsync();
 
-		List<LeaseId> leases = await AgentCollection.FindActiveLeaseIdsAsync();
+		LeaseId[] leases = await AgentScheduler.FindActiveLeaseIdsAsync();
 
-		Assert.AreEqual(3, leases.Count);
+		Assert.AreEqual(3, leases.Length);
 		Assert.IsTrue(leases.Contains(_lease1.Id));
 		Assert.IsTrue(leases.Contains(_leaseWithParent3.Id));
 		Assert.IsTrue(leases.Contains(_leaseWithParent4.Id));
@@ -60,9 +76,9 @@ public class AgentCollectionTests : ComputeTestSetup
 		await _agent.TryCreateSessionAsync(new CreateSessionOptions(
 			new RpcAgentCapabilities(), new List<PoolId>(), null));
 
-		List<LeaseId> leases = await AgentCollection.FindActiveLeaseIdsAsync();
+		LeaseId[] leases = await AgentScheduler.FindActiveLeaseIdsAsync();
 
-		Assert.AreEqual(0, leases.Count);
+		Assert.AreEqual(0, leases.Length);
 	}
 
 	[TestMethod]
@@ -89,26 +105,27 @@ public class AgentCollectionTests : ComputeTestSetup
 			Leases = new List<RpcLease>()
 		});
 
-		List<LeaseId> leases = await AgentCollection.FindActiveLeaseIdsAsync();
+		LeaseId[] leases = await AgentScheduler.FindActiveLeaseIdsAsync();
 
-		Assert.AreEqual(0, leases.Count);
+		Assert.AreEqual(0, leases.Length);
 	}
 
 	[TestMethod]
 	public async Task UpdateSession_WithNewLeasesAsync()
 	{
-		await _agent.TryCreateLeaseAsync(_lease1);
-		await UpdateAgentAsync();
+		IAgent? agent = await _agent.TryCreateSessionAsync(new CreateSessionOptions(new RpcAgentCapabilities(), Array.Empty<PoolId>(), "foo"));
+		Assert.IsNotNull(agent);
+		Assert.IsNotNull(agent.SessionId);
 
-		IAgent? newAgent = await _agent.TryCreateLeaseAsync(_lease1);
-		Assert.IsNotNull(newAgent);
+		agent = await agent.TryCreateLeaseAsync(_lease1);
+		Assert.IsNotNull(agent);
 
-		newAgent = await _agent.TryCreateLeaseAsync(_lease2);
-		Assert.IsNotNull(newAgent);
+		agent = await agent.TryCreateLeaseAsync(_lease2);
+		Assert.IsNotNull(agent);
 
-		List<LeaseId> leases = await AgentCollection.FindActiveLeaseIdsAsync();
+		LeaseId[] leases = await AgentScheduler.FindActiveLeaseIdsAsync();
 
-		Assert.AreEqual(2, leases.Count);
+		Assert.AreEqual(2, leases.Length);
 		Assert.IsTrue(leases.Contains(_lease1.Id));
 		Assert.IsTrue(leases.Contains(_lease2.Id));
 	}
@@ -116,13 +133,17 @@ public class AgentCollectionTests : ComputeTestSetup
 	[TestMethod]
 	public async Task UpdateSession_WithOneLeaseRemovedAsync()
 	{
-		await _agent.TryCreateLeaseAsync(_lease1);
-		await UpdateAgentAsync();
+		IAgent? agent = await _agent.TryCreateSessionAsync(new CreateSessionOptions(new RpcAgentCapabilities(), Array.Empty<PoolId>(), "foo"));
+		Assert.IsNotNull(agent);
+		Assert.IsNotNull(agent.SessionId);
 
-		await _agent.TryCreateLeaseAsync(_lease2);
-		await UpdateAgentAsync();
+		agent = await agent.TryCreateLeaseAsync(_lease1);
+		Assert.IsNotNull(agent);
 
-		await _agent.TryUpdateSessionAsync(new UpdateSessionOptions
+		agent = await agent.TryCreateLeaseAsync(_lease2);
+		Assert.IsNotNull(agent);
+
+		agent = await agent.TryUpdateSessionAsync(new UpdateSessionOptions
 		{
 			Leases = new List<RpcLease>
 			{
@@ -130,30 +151,106 @@ public class AgentCollectionTests : ComputeTestSetup
 				new RpcLease{ Id = _lease2.Id, State = RpcLeaseState.Active }
 			}
 		});
-		await UpdateAgentAsync();
+		Assert.IsNotNull(agent);
 
-		await _agent.TryUpdateSessionAsync(new UpdateSessionOptions { Leases = new List<RpcLease> { new RpcLease { Id = _lease1.Id } } });
+		agent = await agent.TryUpdateSessionAsync(new UpdateSessionOptions
+		{
+			Leases = new List<RpcLease>
+			{
+				new RpcLease { Id = _lease1.Id }
+			}
+		});
+		Assert.IsNotNull(agent);
 
-		List<LeaseId> leases = await AgentCollection.FindActiveLeaseIdsAsync();
+		LeaseId[] leases = await AgentScheduler.FindActiveLeaseIdsAsync();
 
-		Assert.AreEqual(1, leases.Count);
+		Assert.AreEqual(1, leases.Length);
 		Assert.IsTrue(leases.Contains(_lease1.Id));
 	}
 
 	[TestMethod]
 	public async Task CancelLeaseAsync()
 	{
-		await _agent.TryCreateLeaseAsync(_lease1);
-		await UpdateAgentAsync();
+		IAgent? agent = _agent;
 
-		await _agent.TryCreateLeaseAsync(_lease2);
-		await UpdateAgentAsync();
+		agent = await agent.TryCreateSessionAsync(new CreateSessionOptions(new RpcAgentCapabilities(), Array.Empty<PoolId>(), "foo"));
+		Assert.IsNotNull(agent);
+		Assert.IsNotNull(agent.SessionId);
+
+		agent = await agent.TryCreateLeaseAsync(_lease1);
+		Assert.IsNotNull(agent);
+		Assert.AreEqual(1, agent.Leases.Count);
 
-		await _agent.TryCancelLeaseAsync(0);
+		agent = await agent.TryCreateLeaseAsync(_lease2);
+		Assert.IsNotNull(agent);
+		Assert.AreEqual(2, agent.Leases.Count);
 
-		List<LeaseId> leases = await AgentCollection.FindActiveLeaseIdsAsync();
+		agent = await agent.TryCancelLeaseAsync(0);
+		Assert.IsNotNull(agent);
+		Assert.AreEqual(1, agent.Leases.Count);
 
-		Assert.AreEqual(1, leases.Count);
+		LeaseId[] leases = await AgentScheduler.FindActiveLeaseIdsAsync();
+		Assert.AreEqual(1, leases.Length);
+		Assert.IsTrue(leases.Contains(_lease2.Id));
+	}
+
+	[TestMethod]
+	public async Task CancelLease2Async()
+	{
+		IAgent? agent = _agent;
+
+		agent = await agent.TryCreateSessionAsync(new CreateSessionOptions(new RpcAgentCapabilities(), Array.Empty<PoolId>(), "foo"));
+		Assert.IsNotNull(agent);
+		Assert.IsNotNull(agent.SessionId);
+
+		agent = await agent.TryCreateLeaseAsync(_lease1);
+		Assert.IsNotNull(agent);
+		Assert.AreEqual(1, agent.Leases.Count);
+
+		agent = await agent.TryCreateLeaseAsync(_lease2);
+		Assert.IsNotNull(agent);
+		Assert.AreEqual(2, agent.Leases.Count);
+		Assert.IsTrue(agent.Leases.Any(x => x.Id == _lease1.Id && x.State == LeaseState.Pending));
+		Assert.IsTrue(agent.Leases.Any(x => x.Id == _lease2.Id && x.State == LeaseState.Pending));
+
+		LeaseId[] leases = await AgentScheduler.FindActiveLeaseIdsAsync();
+		Assert.AreEqual(2, leases.Length);
+		Assert.IsTrue(leases.Contains(_lease1.Id));
+		Assert.IsTrue(leases.Contains(_lease2.Id));
+
+		agent = await agent.TryUpdateSessionAsync(new UpdateSessionOptions
+		{
+			Leases = new List<RpcLease>
+			{
+				new RpcLease { Id = _lease1.Id, State = RpcLeaseState.Active },
+				new RpcLease { Id = _lease2.Id, State = RpcLeaseState.Active }
+			}
+		});
+		Assert.IsNotNull(agent);
+		Assert.AreEqual(2, agent.Leases.Count);
+		Assert.IsTrue(agent.Leases.Any(x => x.Id == _lease1.Id && x.State == LeaseState.Active));
+		Assert.IsTrue(agent.Leases.Any(x => x.Id == _lease2.Id && x.State == LeaseState.Active));
+
+		agent = await agent.TryCancelLeaseAsync(0);
+		Assert.IsNotNull(agent);
+		Assert.AreEqual(2, agent.Leases.Count);
+		Assert.IsTrue(agent.Leases.Any(x => x.Id == _lease1.Id && x.State == LeaseState.Cancelled));
+		Assert.IsTrue(agent.Leases.Any(x => x.Id == _lease2.Id && x.State == LeaseState.Active));
+
+		agent = await agent.TryUpdateSessionAsync(new UpdateSessionOptions
+		{
+			Leases = new List<RpcLease>
+			{
+				new RpcLease { Id = _lease1.Id, State = RpcLeaseState.Completed },
+				new RpcLease { Id = _lease2.Id, State = RpcLeaseState.Active }
+			}
+		});
+		Assert.IsNotNull(agent);
+		Assert.AreEqual(1, agent.Leases.Count);
+		Assert.IsTrue(agent.Leases.Any(x => x.Id == _lease2.Id && x.State == LeaseState.Active));
+
+		leases = await AgentScheduler.FindActiveLeaseIdsAsync();
+		Assert.AreEqual(1, leases.Length);
 		Assert.IsTrue(leases.Contains(_lease2.Id));
 	}
 
@@ -168,32 +265,40 @@ public class AgentCollectionTests : ComputeTestSetup
 
 		await _agent.TryTerminateSessionAsync();
 
-		List<LeaseId> leases = await AgentCollection.FindActiveLeaseIdsAsync();
+		LeaseId[] leases = await AgentScheduler.FindActiveLeaseIdsAsync();
 
-		Assert.AreEqual(0, leases.Count);
+		Assert.AreEqual(0, leases.Length);
 	}
 
 	[TestMethod]
 	public async Task GetChildLeaseIdsAsync()
 	{
-		await _agent.TryCreateLeaseAsync(_lease1);
-		await UpdateAgentAsync();
+		IAgent? agent = _agent;
 
-		await _agent.TryCreateLeaseAsync(_leaseWithParent3);
-		await UpdateAgentAsync();
+		agent = await agent.TryCreateSessionAsync(new CreateSessionOptions(new RpcAgentCapabilities(), Array.Empty<PoolId>(), "foo"));
+		Assert.IsNotNull(agent);
+		Assert.IsNotNull(agent.SessionId);
 
-		await _agent.TryCreateLeaseAsync(_leaseWithParent4);
-		await UpdateAgentAsync();
+		agent = await agent.TryCreateLeaseAsync(_lease1);
+		Assert.IsNotNull(agent);
+
+		agent = await agent.TryCreateLeaseAsync(_leaseWithParent3);
+		Assert.IsNotNull(agent);
+
+		agent = await agent.TryCreateLeaseAsync(_leaseWithParent4);
+		Assert.IsNotNull(agent);
 
-		List<LeaseId> leases = await AgentCollection.GetChildLeaseIdsAsync(_leaseWithParent3.ParentId!.Value);
+		LeaseId[] leases = await AgentScheduler.GetChildLeaseIdsAsync(_leaseWithParent3.ParentId!.Value);
 
-		Assert.AreEqual(2, leases.Count);
+		Assert.AreEqual(2, leases.Length);
 		Assert.IsTrue(leases.Contains(_leaseWithParent3.Id));
 		Assert.IsTrue(leases.Contains(_leaseWithParent4.Id));
 	}
 
 	private async Task UpdateAgentAsync()
 	{
-		_agent = (await AgentCollection.GetAsync(_agent.Id))!;
+		IAgent? agent = await AgentCollection.GetAsync(_agent.Id);
+		Assert.IsNotNull(agent);
+		_agent = agent;
 	}
 }
\ No newline at end of file
diff --git a/Engine/Source/Programs/Horde/Plugins/Compute/HordeServer.Compute.Tests/Agents/AgentSchedulerTests.cs b/Engine/Source/Programs/Horde/Plugins/Compute/HordeServer.Compute.Tests/Agents/AgentSchedulerTests.cs
new file mode 100644
--- /dev/null
+++ b/Engine/Source/Programs/Horde/Plugins/Compute/HordeServer.Compute.Tests/Agents/AgentSchedulerTests.cs
@@ -0,0 +1,83 @@
+﻿// Copyright Epic Games, Inc. All Rights Reserved.
+
+using EpicGames.Core;
+using EpicGames.Horde.Agents;
+using HordeCommon;
+using HordeCommon.Rpc;
+using HordeCommon.Rpc.Messages;
+using HordeServer.Agents.Sessions;
+using Microsoft.Extensions.DependencyInjection;
+
+namespace HordeServer.Tests.Agents
+{
+	[TestClass]
+	public class AgentSchedulerTests : ComputeTestSetup
+	{
+		[TestMethod]
+		public async Task CreateSessionAsync()
+		{
+			IClock clock = ServiceProvider.GetRequiredService<IClock>();
+
+			RpcAgentCapabilities caps = new RpcAgentCapabilities();
+			caps.Properties.Add("foo=bar");
+
+			RpcSession? session = await AgentScheduler.TryCreateSessionAsync(new AgentId("foo"), SessionIdUtils.GenerateNewId(), caps);
+			Assert.IsNotNull(session);
+
+			Assert.AreEqual(session.AgentId, new AgentId("foo"));
+			Assert.AreEqual(0, session.Leases.Count);
+			Assert.IsTrue(session.ExpiryTime > clock.UtcNow);
+
+			RpcSession? session2 = await AgentScheduler.TryGetSessionAsync(session.SessionId);
+			Assert.AreEqual(session, session2);
+
+			RpcAgentCapabilities? caps2 = await AgentScheduler.TryGetSessionCapabilitiesAsync(session);
+			Assert.AreEqual(caps2, caps);
+		}
+
+		[TestMethod]
+		public async Task UpdateCapabilitesAsync()
+		{
+			// Set and get the initial capabilities
+			RpcAgentCapabilities caps = new RpcAgentCapabilities();
+			caps.Properties.Add("foo=bar");
+
+			RpcSession? session = await AgentScheduler.TryCreateSessionAsync(new AgentId("foo"), SessionIdUtils.GenerateNewId(), caps);
+			Assert.IsNotNull(session);
+
+			RpcAgentCapabilities? caps2 = await AgentScheduler.TryGetSessionCapabilitiesAsync(session);
+			Assert.AreEqual(caps, caps2);
+
+			// Update them
+			RpcAgentCapabilities caps3 = new RpcAgentCapabilities(caps);
+			caps3.Properties.Add("bar=baz");
+
+			RpcSession? session2 = await AgentScheduler.TryUpdateSessionAsync(session, newCapabilities: caps3);
+			Assert.IsNotNull(session2);
+
+			RpcAgentCapabilities? caps4 = await AgentScheduler.TryGetSessionCapabilitiesAsync(session);
+			Assert.IsNull(caps4);
+
+			RpcAgentCapabilities? caps5 = await AgentScheduler.TryGetSessionCapabilitiesAsync(session2);
+			Assert.AreEqual(caps3, caps5);
+		}
+
+		[TestMethod]
+		public async Task TerminateSessionAsync()
+		{
+			FakeClock clock = ServiceProvider.GetRequiredService<FakeClock>();
+
+			RpcSession? session = await AgentScheduler.TryCreateSessionAsync(new AgentId("foo"), SessionIdUtils.GenerateNewId(), new RpcAgentCapabilities());
+			Assert.IsNotNull(session);
+
+			await clock.AdvanceAsync(TimeSpan.FromSeconds(30.0));
+
+			RpcSession? session2 = await AgentScheduler.TryUpdateSessionAsync(session, new RpcSession(session) { Status = RpcAgentStatus.Stopped });
+			Assert.IsNotNull(session2);
+			Assert.AreEqual(clock.UtcNow, session2.ExpiryTime);
+
+			RpcSession? session3 = await AgentScheduler.TryGetSessionAsync(session.SessionId);
+			Assert.IsNull(session3);
+		}
+	}
+}
diff --git a/Engine/Source/Programs/Horde/Plugins/Compute/HordeServer.Compute.Tests/ComputeTestSetup.cs b/Engine/Source/Programs/Horde/Plugins/Compute/HordeServer.Compute.Tests/ComputeTestSetup.cs
--- a/Engine/Source/Programs/Horde/Plugins/Compute/HordeServer.Compute.Tests/ComputeTestSetup.cs
+++ b/Engine/Source/Programs/Horde/Plugins/Compute/HordeServer.Compute.Tests/ComputeTestSetup.cs
@@ -12,6 +12,7 @@ namespace HordeServer.Tests
 	{
 		public ComputeService ComputeService => ServiceProvider.GetRequiredService<ComputeService>();
 		public IAgentCollection AgentCollection => ServiceProvider.GetRequiredService<IAgentCollection>();
+		public IAgentScheduler AgentScheduler => ServiceProvider.GetRequiredService<IAgentScheduler>();
 		public StorageService StorageService => ServiceProvider.GetRequiredService<StorageService>();
 		public ILogCollection LogCollection => ServiceProvider.GetRequiredService<ILogCollection>();
 
diff --git a/Engine/Source/Programs/Horde/Plugins/Compute/HordeServer.Compute/Agents/AgentCollection.cs b/Engine/Source/Programs/Horde/Plugins/Compute/HordeServer.Compute/Agents/AgentCollection.cs
--- a/Engine/Source/Programs/Horde/Plugins/Compute/HordeServer.Compute/Agents/AgentCollection.cs
+++ b/Engine/Source/Programs/Horde/Plugins/Compute/HordeServer.Compute/Agents/AgentCollection.cs
@@ -1,16 +1,17 @@
 // Copyright Epic Games, Inc. All Rights Reserved.
 
-using System.Collections.ObjectModel;
 using System.Collections.Concurrent;
+using System.Collections.ObjectModel;
 using System.Diagnostics.CodeAnalysis;
 using System.Linq.Expressions;
+using System.Runtime.CompilerServices;
 using EpicGames.Core;
 using EpicGames.Horde.Agents;
 using EpicGames.Horde.Agents.Leases;
 using EpicGames.Horde.Agents.Pools;
 using EpicGames.Horde.Agents.Sessions;
-using EpicGames.Redis;
 using Google.Protobuf;
+using Google.Protobuf.Collections;
 using Google.Protobuf.WellKnownTypes;
 using HordeCommon.Rpc;
 using HordeCommon.Rpc.Messages;
@@ -19,34 +20,38 @@ using HordeServer.Agents.Leases;
 using HordeServer.Agents.Sessions;
 using HordeServer.Auditing;
 using HordeServer.Server;
+using HordeServer.Utilities;
+using Microsoft.Extensions.Hosting;
 using Microsoft.Extensions.Logging;
 using MongoDB.Bson;
 using MongoDB.Bson.Serialization.Attributes;
 using MongoDB.Driver;
-using StackExchange.Redis;
-using Microsoft.Extensions.Hosting;
 using OpenTelemetry.Trace;
-using System.Runtime.CompilerServices;
+using StackExchange.Redis;
 
 namespace HordeServer.Agents
 {
 	/// <summary>
 	/// Collection of agent documents
 	/// </summary>
-	public sealed class AgentCollection : IAgentCollection, IHostedService, IAsyncDisposable
+	public sealed class AgentCollection : IAgentCollection, IHostedService
 	{
 		class Agent : IAgent
 		{
 			static IReadOnlyList<string> DefaultProperties { get; } = new List<string>();
 			static IReadOnlyDictionary<string, int> DefaultResources { get; } = new Dictionary<string, int>();
 
+			public AgentDocument Document => _document;
+			public RpcSession? Session => _session;
+
 			readonly AgentCollection _collection;
 			readonly AgentDocument _document;
+			readonly RpcSession? _session;
 
-			AgentId IAgent.Id => _document.Id;
-			SessionId? IAgent.SessionId => _document.SessionId;
-			DateTime? IAgent.SessionExpiresAt => _document.SessionExpiresAt;
-			AgentStatus IAgent.Status => _document.Status;
+			public AgentId Id => _document.Id;
+			SessionId? IAgent.SessionId => _session?.SessionId;
+			DateTime? IAgent.SessionExpiresAt => _session?.ExpiryTime;
+			AgentStatus IAgent.Status => (AgentStatus?)_session?.Status ?? AgentStatus.Stopped;
 			DateTime? IAgent.LastOnlineTime => _document.LastOnlineTime;
 			bool IAgent.Enabled => _document.Enabled;
 			bool IAgent.Ephemeral => _document.Ephemeral;
@@ -70,87 +75,65 @@ namespace HordeServer.Agents
 			IReadOnlyList<AgentWorkspaceInfo> IAgent.Workspaces => _document.Workspaces;
 			DateTime IAgent.LastConformTime => _document.LastConformTime;
 			int? IAgent.ConformAttemptCount => _document.ConformAttemptCount;
-			static readonly IReadOnlyList<AgentLease> s_emptyLeases = new List<AgentLease>();
-			IReadOnlyList<IAgentLease> IAgent.Leases => _document.Leases ?? s_emptyLeases;
+			IReadOnlyList<IAgentLease> IAgent.Leases => (IReadOnlyList<IAgentLease>?)_session?.Leases ?? Array.Empty<IAgentLease>();
 			string IAgent.EnrollmentKey => _document.EnrollmentKey;
 			DateTime IAgent.UpdateTime => _document.UpdateTime;
 			uint IAgent.UpdateIndex => _document.UpdateIndex;
 
-			public Agent(AgentCollection collection, AgentDocument document)
+			public Agent(AgentCollection collection, AgentDocument document, RpcSession? session)
 			{
 				_collection = collection;
 				_document = document;
+				_session = session;
 			}
 
+			/// <inheritdoc/>
 			public async Task<ISession?> GetSessionAsync(SessionId sessionId, CancellationToken cancellationToken = default)
-			{
-				return await _collection.GetSessionAsync(_document.Id, sessionId, cancellationToken);
-			}
+				=> await _collection.GetSessionAsync(_document.Id, sessionId, cancellationToken);
 
+			/// <inheritdoc/>
 			public async Task<IReadOnlyList<ISession>> FindSessionsAsync(DateTime? startTime, DateTime? finishTime, int index, int count, CancellationToken cancellationToken = default)
-			{
-				return await _collection.FindSessionsAsync(_document.Id, startTime, finishTime, index, count, cancellationToken);
-			}
+				=> await _collection.FindSessionsAsync(_document.Id, startTime, finishTime, index, count, cancellationToken);
 
+			/// <inheritdoc/>
 			public async Task<IAgent?> TryCreateLeaseAsync(CreateLeaseOptions options, CancellationToken cancellationToken = default)
-			{
-				AgentDocument? newDocument = await _collection.TryCreateLeaseAsync(_document, options, cancellationToken);
-				return _collection.CreateAgentObject(newDocument);
-			}
+				=> await _collection.TryCreateLeaseAsync(this, options, cancellationToken);
 
+			/// <inheritdoc/>
 			public async Task<IAgent?> TryCancelLeaseAsync(int leaseIdx, CancellationToken cancellationToken = default)
-			{
-				AgentDocument? newDocument = await _collection.TryCancelLeaseAsync(_document, leaseIdx, cancellationToken);
-				return _collection.CreateAgentObject(newDocument);
-			}
+				=> await _collection.TryCancelLeaseAsync(this, leaseIdx, cancellationToken);
 
+			/// <inheritdoc/>
 			public async Task<IAgent?> TryDeleteAsync(CancellationToken cancellationToken = default)
-			{
-				AgentDocument? newDocument = await _collection.TryDeleteAsync(_document, cancellationToken);
-				return _collection.CreateAgentObject(newDocument);
-			}
+				=> await _collection.TryDeleteAsync(this, cancellationToken);
 
+			/// <inheritdoc/>
 			public async Task<IAgent?> TryResetAsync(bool ephemeral, string enrollmentKey, CancellationToken cancellationToken = default)
-			{
-				AgentDocument? newDocument = await _collection.TryResetAsync(_document, ephemeral, enrollmentKey, cancellationToken);
-				return _collection.CreateAgentObject(newDocument);
-			}
+				=> await _collection.TryResetAsync(this, ephemeral, enrollmentKey, cancellationToken);
 
+			/// <inheritdoc/>
 			public async Task<IAgent?> TryCreateSessionAsync(CreateSessionOptions options, CancellationToken cancellationToken = default)
-			{
-				AgentDocument? newDocument = await _collection.TryCreateSessionAsync(_document, options, cancellationToken);
-				return _collection.CreateAgentObject(newDocument);
-			}
+				=> await _collection.TryCreateSessionAsync(this, options, cancellationToken);
 
+			/// <inheritdoc/>
 			public async Task<IAgent?> TryTerminateSessionAsync(CancellationToken cancellationToken = default)
-			{
-				AgentDocument? newDocument = await _collection.TryTerminateSessionAsync(_document, cancellationToken);
-				return _collection.CreateAgentObject(newDocument);
-			}
+				=> await _collection.TryTerminateSessionAsync(this, cancellationToken);
 
+			/// <inheritdoc/>
 			public async Task<IAgent?> TryUpdateSessionAsync(UpdateSessionOptions options, CancellationToken cancellationToken = default)
-			{
-				AgentDocument? newDocument = await _collection.TryUpdateSessionAsync(_document, options, cancellationToken);
-				return _collection.CreateAgentObject(newDocument);
-			}
+				=> await _collection.TryUpdateSessionAsync(this, options, cancellationToken);
 
+			/// <inheritdoc/>
 			public async Task<IAgent?> TryUpdateAsync(UpdateAgentOptions options, CancellationToken cancellationToken = default)
-			{
-				AgentDocument? newDocument = await _collection.TryUpdateSettingsAsync(_document, options, cancellationToken);
-				return _collection.CreateAgentObject(newDocument);
-			}
+				=> await _collection.TryUpdateSettingsAsync(this, options, cancellationToken);
 
+			/// <inheritdoc/>
 			public async Task<IAgent?> TryUpdateWorkspacesAsync(List<AgentWorkspaceInfo> workspaces, bool requestConform, CancellationToken cancellationToken = default)
-			{
-				AgentDocument? newDocument = await _collection.TryUpdateWorkspacesAsync(_document, workspaces, requestConform, cancellationToken);
-				return _collection.CreateAgentObject(newDocument);
-			}
+				=> await _collection.TryUpdateWorkspacesAsync(this, workspaces, requestConform, cancellationToken);
 
+			/// <inheritdoc/>
 			public async Task<IAgent?> WaitForUpdateAsync(CancellationToken cancellationToken = default)
-			{
-				AgentDocument? newDocument = await _collection.WaitForUpdateAsync(_document, cancellationToken);
-				return _collection.CreateAgentObject(newDocument);
-			}
+				=> await _collection.WaitForSessionUpdateAsync(this, cancellationToken);
 		}
 
 		/// <summary>
@@ -162,9 +145,11 @@ namespace HordeServer.Agents
 			public AgentId Id { get; set; }
 
 			public SessionId? SessionId { get; set; }
-			public DateTime? SessionExpiresAt { get; set; }
 
-			public AgentStatus Status { get; set; }
+			[BsonIgnoreIfNull]
+			[Obsolete("Session state now stored in Redis")]
+			public AgentStatus? Status { get; set; }
+
 			public DateTime? LastOnlineTime { get; set; }
 
 			[BsonRequired]
@@ -218,7 +203,10 @@ namespace HordeServer.Agents
 			[BsonIgnoreIfNull]
 			public int? ConformAttemptCount { get; set; }
 
+			[BsonIgnoreIfNull]
+			[Obsolete("Session state now stored in Redis")]
 			public List<AgentLease>? Leases { get; set; }
+
 			public DateTime UpdateTime { get; set; }
 			public uint UpdateIndex { get; set; }
 			public string EnrollmentKey { get; set; } = String.Empty;
@@ -237,7 +225,6 @@ namespace HordeServer.Agents
 				Id = id;
 				Ephemeral = ephemeral;
 				EnrollmentKey = enrollmentKey;
-				Status = AgentStatus.Stopped;
 			}
 		}
 
@@ -276,24 +263,21 @@ namespace HordeServer.Agents
 		readonly IMongoCollection<AgentDocument> _agentCollection;
 		readonly IMongoCollection<SessionDocument> _sessionCollection;
 		readonly ILeaseCollection _leaseCollection;
+		readonly IAgentScheduler _scheduler;
 		readonly IAuditLog<AgentId> _auditLog;
-		readonly IRedisService _redisService;
 		readonly IClock _clock;
-		readonly RedisChannel<AgentId> _updateEventChannel;
 		readonly ITicker _sharedTicker;
-		readonly ConcurrentDictionary<AgentId, TaskCompletionSource> _agentIdToTcs = new ConcurrentDictionary<AgentId, TaskCompletionSource>();
-		IAsyncDisposable? _updateEventSubscription;
+		readonly ConcurrentDictionary<SessionId, TaskCompletionSource> _sessionIdToTcs = new ConcurrentDictionary<SessionId, TaskCompletionSource>();
 		readonly Tracer _tracer;
 		readonly ILogger _logger;
 
 		/// <summary>
 		/// Constructor
 		/// </summary>
-		public AgentCollection(IMongoService mongoService, IRedisService redisService, ILeaseCollection leaseCollection, IClock clock, IAuditLog<AgentId> auditLog, Tracer tracer, ILogger<AgentCollection> logger)
+		public AgentCollection(IMongoService mongoService, ILeaseCollection leaseCollection, IAgentScheduler scheduler, IClock clock, IAuditLog<AgentId> auditLog, Tracer tracer, ILogger<AgentCollection> logger)
 		{
 			List<MongoIndex<AgentDocument>> agentIndexes = new List<MongoIndex<AgentDocument>>();
 			agentIndexes.Add(keys => keys.Ascending(x => x.Deleted).Ascending(x => x.Id).Ascending(x => x.Pools));
-			agentIndexes.Add(keys => keys.Ascending(x => x.SessionExpiresAt), sparse: true);
 			_agentCollection = mongoService.GetCollection<AgentDocument>("Agents", agentIndexes);
 
 			List<MongoIndex<SessionDocument>> sessionIndexes = new List<MongoIndex<SessionDocument>>();
@@ -301,10 +285,9 @@ namespace HordeServer.Agents
 			sessionIndexes.Add(keys => keys.Ascending(x => x.FinishTime));
 			_sessionCollection = mongoService.GetCollection<SessionDocument>("Sessions", sessionIndexes);
 
-			_redisService = redisService;
 			_leaseCollection = leaseCollection;
+			_scheduler = scheduler;
 			_clock = clock;
-			_updateEventChannel = new RedisChannel<AgentId>(RedisChannel.Literal("agents/notify"));
 			_sharedTicker = clock.AddSharedTicker($"{nameof(AgentCollection)}.{nameof(TickSharedAsync)}", TimeSpan.FromSeconds(30.0), TickSharedAsync, logger);
 			_auditLog = auditLog;
 			_tracer = tracer;
@@ -315,17 +298,12 @@ namespace HordeServer.Agents
 		public async ValueTask DisposeAsync()
 		{
 			await _sharedTicker.DisposeAsync();
-			if (_updateEventSubscription != null)
-			{
-				await _updateEventSubscription.DisposeAsync();
-				_updateEventSubscription = null;
-			}
 		}
 
 		/// <inheritdoc/>
 		public async Task StartAsync(CancellationToken cancellationToken)
 		{
-			_updateEventSubscription = await _redisService.GetDatabase().Multiplexer.SubscribeAsync(_updateEventChannel, OnAgentUpdate);
+			_scheduler.SessionUpdated += OnSessionUpdate;
 			await _sharedTicker.StartAsync();
 		}
 
@@ -333,11 +311,7 @@ namespace HordeServer.Agents
 		public async Task StopAsync(CancellationToken cancellationToken)
 		{
 			await _sharedTicker.StopAsync();
-			if (_updateEventSubscription != null)
-			{
-				await _updateEventSubscription.DisposeAsync();
-				_updateEventSubscription = null;
-			}
+			_scheduler.SessionUpdated -= OnSessionUpdate;
 		}
 
 		internal async ValueTask TickSharedAsync(CancellationToken stoppingToken)
@@ -353,30 +327,18 @@ namespace HordeServer.Agents
 			using TelemetrySpan span = _tracer.StartActiveSpan($"{nameof(AgentService)}.{nameof(TerminateExpiredSessionsAsync)}");
 
 			int c = 0;
-			while (!cancellationToken.IsCancellationRequested)
+			await foreach (RpcSession session in _scheduler.FindExpiredSessionsAsync(cancellationToken))
 			{
-				// Find all the agents which are ready to be expired
-				const int MaxAgents = 100;
-				DateTime utcNow = _clock.UtcNow;
-
-				FilterDefinition<AgentDocument> filter = Builders<AgentDocument>.Filter.Exists(x => x.SessionExpiresAt) & Builders<AgentDocument>.Filter.Lt(x => x.SessionExpiresAt, utcNow);
-
-				List<AgentDocument> expiredAgents = await _agentCollection.Find(filter).ToListAsync(cancellationToken);
-				expiredAgents = await PostLoadAsync(expiredAgents, cancellationToken);
-
-				// Transition each agent to being offline
-				foreach (AgentDocument expiredAgent in expiredAgents)
+				AgentDocument? document = await GetDocumentAsync(session.AgentId, cancellationToken);
+				if (document != null)
 				{
-					cancellationToken.ThrowIfCancellationRequested();
-					_logger.LogDebug("Terminating session {SessionId} for agent {Agent}", expiredAgent.SessionId, expiredAgent.Id);
-					await TryTerminateSessionAsync(expiredAgent, cancellationToken);
-				}
-				c += expiredAgents.Count;
-
-				// Try again if we didn't fetch everything
-				if (expiredAgents.Count < MaxAgents)
-				{
-					break;
+					Agent? agent = await CreateAgentObjectAsync(document, session, cancellationToken);
+					if (agent != null)
+					{
+						_logger.LogDebug("Terminating session {SessionId} for agent {Agent}", session.SessionId, session.AgentId);
+						await TryUpdateSessionAsync(agent, new UpdateSessionOptions { Status = AgentStatus.Stopped }, cancellationToken);
+						c++;
+					}
 				}
 			}
 			span.SetAttribute("NumAgentsTerminated", c);
@@ -410,6 +372,7 @@ namespace HordeServer.Agents
 			span.SetAttribute("NumAgentsDeleted", numDeleted);
 		}
 
+#pragma warning disable CS0618
 		async ValueTask<AgentDocument?> PostLoadAsync(AgentDocument? document, CancellationToken cancellationToken)
 		{
 			while (document != null)
@@ -425,6 +388,17 @@ namespace HordeServer.Agents
 
 					newDocument = await TryUpdateAsync(document, updateDefinition, cancellationToken);
 				}
+				else if (document.DocumentVersion == 1)
+				{
+					if (await MoveLeasesToSchedulerAsync(document, cancellationToken))
+					{
+						UpdateDefinition<AgentDocument> updateDefinition = Builders<AgentDocument>.Update
+							.Unset(x => x.Leases)
+							.Set(x => x.DocumentVersion, 2);
+
+						newDocument = await TryUpdateAsync(document, updateDefinition, cancellationToken);
+					}
+				}
 				else
 				{
 					break;
@@ -434,6 +408,56 @@ namespace HordeServer.Agents
 			return document;
 		}
 
+		async Task<bool> MoveLeasesToSchedulerAsync(AgentDocument document, CancellationToken cancellationToken)
+		{
+			if (document.SessionId != null)
+			{
+				// Create the session with this id, or get the existing session document
+				RpcSession? session = await _scheduler.TryCreateSessionAsync(document.Id, document.SessionId.Value, new RpcAgentCapabilities(), cancellationToken);
+				if (session == null)
+				{
+					session = await _scheduler.TryGetSessionAsync(document.SessionId.Value, cancellationToken);
+					if (session == null)
+					{
+						return false;
+					}
+				}
+
+				// Update the session with the leases in the document
+				RpcSession newSession = new RpcSession(session);
+				if (document.Status != null)
+				{
+					newSession.Status = (RpcAgentStatus)document.Status.Value;
+				}
+				if (document.Leases != null)
+				{
+					foreach (AgentLease agentLease in document.Leases)
+					{
+						if (!newSession.Leases.Any(x => x.Id == agentLease.Id))
+						{
+							RpcSessionLease sessionLease = new RpcSessionLease();
+							sessionLease.Id = agentLease.Id;
+							sessionLease.ParentId = agentLease.ParentId;
+							sessionLease.Payload = (agentLease.Payload == null) ? null : Any.Parser.ParseFrom(agentLease.Payload);
+							sessionLease.State = (RpcLeaseState)agentLease.State;
+							sessionLease.Resources.Add(agentLease.Resources ?? Enumerable.Empty<KeyValuePair<string, int>>());
+							sessionLease.Exclusive = agentLease.Exclusive;
+							session.Leases.Add(sessionLease);
+						}
+					}
+				}
+
+				// Apply the updates
+				if (session != newSession && await _scheduler.TryUpdateSessionAsync(session, newSession, cancellationToken: cancellationToken) == null)
+				{
+					return false;
+				}
+			}
+			return true;
+		}
+
+#pragma warning restore CS0618
+
 		async ValueTask<List<AgentDocument>> PostLoadAsync(List<AgentDocument> documents, CancellationToken cancellationToken)
 		{
 			List<AgentDocument> newDocuments = new List<AgentDocument>();
@@ -449,7 +473,7 @@ namespace HordeServer.Agents
 		}
 
 		[return: NotNullIfNotNull("document")]
-		Agent? CreateAgentObject(AgentDocument? document)
+		Agent? CreateAgentObject(AgentDocument? document, RpcSession? session)
 		{
 			if (document == null)
 			{
@@ -457,50 +481,135 @@ namespace HordeServer.Agents
 			}
 			else
 			{
-				return new Agent(this, document);
+				return new Agent(this, document, session);
+			}
+		}
+
+		ValueTask<Agent?> CreateAgentObjectAsync(AgentDocument? document, CancellationToken cancellationToken)
+			=> CreateAgentObjectAsync(document, null, cancellationToken);
+
+		async ValueTask<Agent?> CreateAgentObjectAsync(AgentDocument? document, RpcSession? session, CancellationToken cancellationToken)
+		{
+			if (document == null)
+			{
+				return null;
+			}
+
+			AgentId agentId = document.Id;
+			for (; ; )
+			{
+				// If the document doesn't have a session, we don't need to fetch it
+				if (document.SessionId == null && session == null)
+				{
+					return new Agent(this, document, null);
+				}
+				if (document.SessionId != null && session != null && session.SessionId == document.SessionId.Value)
+				{
+					return new Agent(this, document, session);
+				}
+
+				// Reconcile the differences
+				for (; ; )
+				{
+					// Fetch the document again
+					document = await GetDocumentAsync(agentId, cancellationToken);
+					if (document == null)
+					{
+						return null;
+					}
+					if (document.SessionId == null)
+					{
+						session = null;
+						break;
+					}
+
+					// Fetch the session
+					session = await _scheduler.TryGetSessionAsync(document.SessionId.Value, cancellationToken);
+					if (session != null)
+					{
+						break;
+					}
+
+					// If there's no matching session, try to to remove it from the document. We can expect the document to be authoritative here, since we always create the session before
+					// updating the document; if it no longer exists, the document is invalid.
+					document = await TryUpdateAsync(document, Builders<AgentDocument>.Update.Unset(x => x.SessionId).Set(x => x.LastOnlineTime, _clock.UtcNow), cancellationToken);
+					if (document != null)
+					{
+						break;
+					}
+				}
 			}
 		}
 
 		/// <inheritdoc/>
 		public async Task<IAgent> AddAsync(AgentId id, bool ephemeral, string enrollmentKey, CancellationToken cancellationToken)
 		{
-			AgentDocument agent = new AgentDocument(id, ephemeral, enrollmentKey);
-			agent.LastOnlineTime = _clock.UtcNow;
+			AgentDocument document = new AgentDocument(id, ephemeral, enrollmentKey);
+			document.LastOnlineTime = _clock.UtcNow;
+			document.DocumentVersion = 2;
 
-			await _agentCollection.InsertOneAsync(agent, null, cancellationToken);
-			return CreateAgentObject(agent);
+			await _agentCollection.InsertOneAsync(document, null, cancellationToken);
+			return new Agent(this, document, null);
 		}
 
 		/// <inheritdoc/>
-		async Task<AgentDocument?> TryResetAsync(AgentDocument agent, bool ephemeral, string enrollmentKey, CancellationToken cancellationToken = default)
+		async Task<Agent?> TryResetAsync(Agent agent, bool ephemeral, string enrollmentKey, CancellationToken cancellationToken = default)
 		{
-			AgentDocument agentDocument = (AgentDocument)agent;
+			// Don't allow resetting while a session is active
+			if (agent.Document.SessionId != null)
+			{
+				return null;
+			}
 
+			// Try to update the document
 			UpdateDefinition<AgentDocument> update = Builders<AgentDocument>.Update
 				.Set(x => x.Ephemeral, ephemeral)
 				.Set(x => x.EnrollmentKey, enrollmentKey)
-				.Unset(x => x.Deleted)
-				.Unset(x => x.SessionId);
+				.Unset(x => x.Deleted);
 
-			return await TryUpdateAsync(agentDocument, update, cancellationToken);
+			AgentDocument? newDocument = await TryUpdateAsync(agent.Document, update, cancellationToken);
+			if (newDocument == null)
+			{
+				return null;
+			}
+
+			// Create the agent instance
+			return new Agent(this, newDocument, null);
 		}
 
 		/// <inheritdoc/>
-		async Task<AgentDocument?> TryDeleteAsync(AgentDocument agent, CancellationToken cancellationToken)
+		async Task<Agent?> TryDeleteAsync(Agent agent, CancellationToken cancellationToken)
 		{
+			if (agent.Document.Deleted)
+			{
+				return null;
+			}
+
+			// Terminate the current session
+			if (agent.Session != null)
+			{
+				Agent? newAgent = await TryTerminateSessionAsync(agent, cancellationToken);
+				if (newAgent == null)
+				{
+					return null;
+				}
+				agent = newAgent;
+			}
+
+			// Update the document
 			UpdateDefinition<AgentDocument> update = Builders<AgentDocument>.Update
 				.Set(x => x.Deleted, true)
-				.Set(x => x.EnrollmentKey, "")
-				.Unset(x => x.SessionId);
+				.Set(x => x.EnrollmentKey, "");
 
-			return await TryUpdateAsync(agent, update, cancellationToken);
+			AgentDocument? newDocument = await TryUpdateAsync(agent.Document, update, cancellationToken);
+			return await CreateAgentObjectAsync(newDocument, cancellationToken);
 		}
 
 		/// <inheritdoc/>
 		public async Task<IAgent?> GetAsync(AgentId agentId, CancellationToken cancellationToken)
 		{
 			AgentDocument? document = await GetDocumentAsync(agentId, cancellationToken);
-			return CreateAgentObject(document);
+			return await CreateAgentObjectAsync(document, cancellationToken);
 		}
 
 		async Task<AgentDocument?> GetDocumentAsync(AgentId agentId, CancellationToken cancellationToken)
@@ -515,7 +624,18 @@ namespace HordeServer.Agents
 		{
 			List<AgentDocument> documents = await _agentCollection.Find(p => agentIds.Contains(p.Id)).ToListAsync(cancellationToken);
 			documents = await PostLoadAsync(documents, cancellationToken);
-			return documents.ConvertAll(x => CreateAgentObject(x));
+
+			List<IAgent> agents = new List<IAgent>();
+			foreach (AgentDocument document in documents)
+			{
+				IAgent? agent = await CreateAgentObjectAsync(document, cancellationToken);
+				if (agent != null)
+				{
+					agents.Add(agent);
+				}
+			}
+
+			return agents;
 		}
 
 		/// <inheritdoc/>
@@ -544,11 +664,6 @@ namespace HordeServer.Agents
 				filter &= filterBuilder.AnyEq(x => x.Properties, property);
 			}
 
-			if (status != null)
-			{
-				filter &= filterBuilder.Eq(x => x.Status, status.Value);
-			}
-
 			if (enabled != null)
 			{
 				filter &= filterBuilder.Eq(x => x.Enabled, enabled.Value);
@@ -563,26 +678,16 @@ namespace HordeServer.Agents
 					List<AgentDocument> documents = await PostLoadAsync(cursor.Current.ToList(), cancellationToken);
 					foreach (AgentDocument document in documents)
 					{
-						yield return CreateAgentObject(document);
+						IAgent? agent = await CreateAgentObjectAsync(document, cancellationToken);
+						if (agent != null && (status == null || agent.Status == status.Value))
+						{
+							yield return agent;
+						}
 					}
 				}
 			}
 		}
 
-		/// <inheritdoc/>
-		public async Task<List<LeaseId>> FindActiveLeaseIdsAsync(CancellationToken cancellationToken)
-		{
-			RedisValue[] activeLeaseIds = await _redisService.GetDatabase().SetMembersAsync(RedisKeyActiveLeaseIds());
-			return activeLeaseIds.Select(x => LeaseId.Parse(x.ToString())).ToList();
-		}
-
-		/// <inheritdoc/>
-		public async Task<List<LeaseId>> GetChildLeaseIdsAsync(LeaseId id, CancellationToken cancellationToken)
-		{
-			RedisValue[] childIds = await _redisService.GetDatabase().SetMembersAsync(RedisKeyLeaseChildren(id));
-			return childIds.Select(x => LeaseId.Parse(x.ToString())).ToList();
-		}
-
 		/// <inheritdoc/>
 		async Task<SessionDocument?> GetSessionAsync(AgentId agentId, SessionId sessionId, CancellationToken cancellationToken = default)
 		{
@@ -625,18 +730,11 @@ namespace HordeServer.Agents
 			Expression<Func<AgentDocument, bool>> filter = x => x.Id == current.Id && x.UpdateIndex == prevUpdateIndex;
 			UpdateDefinition<AgentDocument> updateWithIndex = update.Set(x => x.UpdateIndex, current.UpdateIndex).Set(x => x.UpdateTime, current.UpdateTime);
 
-			AgentDocument? newDocument = await _agentCollection.FindOneAndUpdateAsync<AgentDocument>(filter, updateWithIndex, new FindOneAndUpdateOptions<AgentDocument, AgentDocument> { ReturnDocument = ReturnDocument.After }, cancellationToken);
-			if (newDocument == null)
-			{
-				return null;
-			}
-
-			PublishUpdateEvent(newDocument.Id);
-			return newDocument;
+			return await _agentCollection.FindOneAndUpdateAsync<AgentDocument>(filter, updateWithIndex, new FindOneAndUpdateOptions<AgentDocument, AgentDocument> { ReturnDocument = ReturnDocument.After }, cancellationToken);
 		}
 
 		/// <inheritdoc/>
-		async Task<AgentDocument?> TryUpdateSettingsAsync(AgentDocument agent, UpdateAgentOptions options, CancellationToken cancellationToken)
+		async Task<Agent?> TryUpdateSettingsAsync(Agent agent, UpdateAgentOptions options, CancellationToken cancellationToken)
 		{
 			// Update the database
 			UpdateDefinitionBuilder<AgentDocument> updateBuilder = new UpdateDefinitionBuilder<AgentDocument>();
@@ -644,7 +742,7 @@ namespace HordeServer.Agents
 			List<UpdateDefinition<AgentDocument>> updates = new List<UpdateDefinition<AgentDocument>>();
 			if (options.ExplicitPools != null)
 			{
-				List<PoolId> pools = CreatePoolsList(agent.DynamicPools, options.ExplicitPools, agent.Properties);
+				List<PoolId> pools = CreatePoolsList(agent.Document.DynamicPools, options.ExplicitPools, agent.Document.Properties);
 				updates.Add(updateBuilder.Set(x => x.Pools, pools));
 
 				List<PoolId> explicitPools = CreatePoolsList(options.ExplicitPools).ToList();
@@ -707,69 +805,92 @@ namespace HordeServer.Agents
 			}
 
 			// Apply the update
-			return await TryUpdateAsync(agent, updateBuilder.Combine(updates), cancellationToken);
+			AgentDocument? newDocument = await TryUpdateAsync(agent.Document, updateBuilder.Combine(updates), cancellationToken);
+			return await CreateAgentObjectAsync(newDocument, agent.Session, cancellationToken);
 		}
 
 		/// <inheritdoc/>
-		async Task<AgentDocument?> TryUpdateSessionAsync(AgentDocument agent, UpdateSessionOptions options, CancellationToken cancellationToken)
+		async Task<Agent?> TryUpdateSessionAsync(Agent agent, UpdateSessionOptions options, CancellationToken cancellationToken)
 		{
-			// Create an update definition for the agent
-			UpdateDefinitionBuilder<AgentDocument> updateBuilder = Builders<AgentDocument>.Update;
-			List<UpdateDefinition<AgentDocument>> updates = new List<UpdateDefinition<AgentDocument>>();
+			AgentId agentId = agent.Document.Id;
+			if (agent.Session == null)
+			{
+				return null;
+			}
 
-			if (options.Status != null && agent.Status != options.Status.Value)
+			// Try to update the pools and capabilities first
+			AgentDocument? newDocument = agent.Document;
 			{
-				updates.Add(updateBuilder.Set(x => x.Status, options.Status.Value));
+				UpdateDefinitionBuilder<AgentDocument> updateBuilder = Builders<AgentDocument>.Update;
+				List<UpdateDefinition<AgentDocument>> updates = new List<UpdateDefinition<AgentDocument>>();
 
-				if (options.Status == AgentStatus.Stopped)
+				List<string> newProperties = newDocument.Properties ?? new List<string>();
+				if (options.Capabilities != null)
 				{
-					updates.Add(updateBuilder.Set(x => x.LastOnlineTime, _clock.UtcNow));
+					options.Capabilities.Flatten(out newProperties, out Dictionary<string, int> newResources);
+
+					if (!ResourcesEqual(newResources, newDocument.Resources))
+					{
+						updates.Add(updateBuilder.Set(x => x.Resources, newResources));
+					}
 				}
-				else
+				if (options.DynamicPools != null)
 				{
-					updates.Add(updateBuilder.Unset(x => x.LastOnlineTime));
+					List<PoolId> dynamicPools = CreatePoolsList(options.DynamicPools).ToList();
+					if (!Enumerable.SequenceEqual(dynamicPools, newDocument.DynamicPools))
+					{
+						updates.Add(updateBuilder.Set(x => x.DynamicPools, dynamicPools));
+					}
 				}
-			}
-			if (options.SessionExpiresAt != null)
-			{
-				updates.Add(updateBuilder.Set(x => x.SessionExpiresAt, options.SessionExpiresAt.Value));
-			}
 
-			List<string> newProperties = agent.Properties ?? new List<string>();
-			if (options.Capabilities != null)
-			{
-				options.Capabilities.Flatten(out newProperties, out Dictionary<string, int> newResources);
+				List<PoolId> pools = CreatePoolsList(options.DynamicPools ?? newDocument.DynamicPools, newDocument.ExplicitPools, newProperties);
+				if (!Enumerable.SequenceEqual(pools, newDocument.Pools))
+				{
+					updates.Add(updateBuilder.Set(x => x.Pools, pools));
+				}
+
+				SetPoolProperties(newProperties, pools);
 
-				if (!ResourcesEqual(newResources, agent.Resources))
+				if (!Enumerable.SequenceEqual(newDocument.Properties ?? Enumerable.Empty<string>(), newProperties, StringComparer.Ordinal))
 				{
-					updates.Add(updateBuilder.Set(x => x.Resources, newResources));
+					updates.Add(updateBuilder.Set(x => x.Properties, newProperties));
 				}
-			}
-			if (options.DynamicPools != null)
-			{
-				List<PoolId> dynamicPools = CreatePoolsList(options.DynamicPools).ToList();
-				if (!Enumerable.SequenceEqual(dynamicPools, agent.DynamicPools))
+
+				// Apply the updates
+				if (updates.Count > 0)
 				{
-					updates.Add(updateBuilder.Set(x => x.DynamicPools, dynamicPools));
+					newDocument = await TryUpdateAsync(newDocument, updateBuilder.Combine(updates), cancellationToken);
+					if (newDocument == null)
+					{
+						return null;
+					}
 				}
 			}
+
+			// Apply the updates to the session
+			RpcSession? newSession = new RpcSession(agent.Session);
+			if (options.Status != null)
+			{
+				newSession.Status = (RpcAgentStatus)options.Status;
+			}
+			if (options.Capabilities != null)
+			{
+				newSession.CapabilitiesHash = IoHash.Compute(options.Capabilities.ToByteArray()).ToString();
+			}
 			if (options.Leases != null)
 			{
-				// Flag for whether the leases array should be updated
-				bool updateLeases = false;
 				Dictionary<LeaseId, RpcLease> idToRemoteLease = options.Leases.ToDictionary(x => x.Id, x => x);
 
 				// Remove any completed leases from the agent
-				List<AgentLease> leases = agent.Leases?.ConvertAll(x => new AgentLease(x)) ?? new List<AgentLease>();
+				RepeatedField<RpcSessionLease> leases = newSession.Leases;
 				for (int idx = 0; idx < leases.Count; idx++)
 				{
-					AgentLease lease = leases[idx];
+					RpcSessionLease lease = leases[idx];
 					if (!idToRemoteLease.TryGetValue(lease.Id, out RpcLease? remoteLease))
 					{
-						if (lease.State != LeaseState.Pending)
+						if (lease.State != RpcLeaseState.Pending)
 						{
 							leases.RemoveAt(idx--);
-							updateLeases = true;
 						}
 					}
 					else
@@ -777,71 +898,65 @@ namespace HordeServer.Agents
 						if (remoteLease.State == RpcLeaseState.Cancelled || remoteLease.State == RpcLeaseState.Completed)
 						{
 							leases.RemoveAt(idx--);
-							updateLeases = true;
 						}
-						else if (remoteLease.State == RpcLeaseState.Active && lease.State == LeaseState.Pending)
+						else if (newSession.Status == RpcAgentStatus.Stopping || newSession.Status == RpcAgentStatus.Busy)
 						{
-							lease.State = LeaseState.Active;
-							updateLeases = true;
+							lease.State = RpcLeaseState.Cancelled;
 						}
-					}
-				}
-
-				// If the agent is stopping, cancel all the leases. Clear out the current session once it's complete.
-				AgentStatus status = options.Status ?? agent.Status;
-				if (status == AgentStatus.Stopping || status == AgentStatus.Busy)
-				{
-					foreach (AgentLease lease in leases)
-					{
-						if (lease.State != LeaseState.Cancelled)
+						else if (remoteLease.State == RpcLeaseState.Active && lease.State == RpcLeaseState.Pending)
 						{
-							lease.State = LeaseState.Cancelled;
-							updateLeases = true;
+							lease.State = RpcLeaseState.Active;
 						}
 					}
 				}
-
-				if (updateLeases)
-				{
-					updates.Add(updateBuilder.Set(x => x.Leases, leases));
-				}
 			}
 
-			// Update the pools
-			List<PoolId> pools = CreatePoolsList(options.DynamicPools ?? agent.DynamicPools, agent.ExplicitPools, newProperties);
-			if (!Enumerable.SequenceEqual(pools, agent.Pools))
+			// Apply the updates in redis
+			newSession = await _scheduler.TryUpdateSessionAsync(agent.Session, newSession, options.Capabilities, cancellationToken);
+			if (newSession == null)
 			{
-				updates.Add(updateBuilder.Set(x => x.Pools, pools));
+				return null;
 			}
 
-			SetPoolProperties(newProperties, pools);
-
-			if (!Enumerable.SequenceEqual(agent.Properties ?? Enumerable.Empty<string>(), newProperties, StringComparer.Ordinal))
+			// If the session is stopped, remove it from the agent document
+			if (agent.Session.Status != RpcAgentStatus.Stopped && newSession.Status == RpcAgentStatus.Stopped)
 			{
-				updates.Add(updateBuilder.Set(x => x.Properties, newProperties));
-			}
+				// Get the session expiry time. We'll use this to update all the other documents.
+				DateTime expiryTime = newSession.ExpiryTime;
 
-			// If there are no new updates, return immediately. This is important for preventing UpdateSession calls from returning immediately.
-			if (updates.Count == 0)
-			{
-				return agent;
-			}
+				// Update the agent document
+				while (newDocument != null && newDocument.SessionId == newSession.SessionId)
+				{
+					UpdateDefinition<AgentDocument> update = new BsonDocument();
+					update = update.Unset(x => x.SessionId);
+					update = update.Set(x => x.LastOnlineTime, expiryTime);
 
-			// Update the agent, and try to create new lease documents if we succeed
-			AgentDocument? newAgent = await TryUpdateAsync(agent, updateBuilder.Combine(updates), cancellationToken);
-			if (newAgent == null)
-			{
-				return null;
+					if (agent.Document.Ephemeral)
+					{
+						update = update.Set(x => x.Deleted, true);
+					}
+
+					newDocument = await TryUpdateAsync(agent.Document, update, cancellationToken);
+					if (newDocument != null)
+					{
+						break;
+					}
+
+					newDocument = await GetDocumentAsync(agent.Id, cancellationToken);
+				}
+
+				// Update the session document
+				ILogger agentLogger = GetLogger(agent.Id);
+				agentLogger.LogInformation("Terminated session {SessionId}", newSession.SessionId);
+				await _sessionCollection.UpdateOneAsync(x => x.Id == newSession.SessionId, Builders<SessionDocument>.Update.Set(x => x.FinishTime, expiryTime), cancellationToken: CancellationToken.None);
 			}
 
 			// Remove any complete leases
-			if (agent.Leases != null)
+			if (newSession.Leases != null)
 			{
-				HashSet<LeaseId> newLeaseIds = new HashSet<LeaseId>(newAgent.Leases?.Select(x => x.Id) ?? Enumerable.Empty<LeaseId>());
-				foreach (AgentLease oldLease in agent.Leases.Where(x => !newLeaseIds.Contains(x.Id)))
+				HashSet<LeaseId> newLeaseIds = new HashSet<LeaseId>(newSession.Leases?.Select(x => x.Id) ?? Enumerable.Empty<LeaseId>());
+				foreach (RpcSessionLease oldLease in agent.Session.Leases.Where(x => !newLeaseIds.Contains(x.Id)))
 				{
-					await RemoveActiveLeaseAsync(oldLease);
-
 					RpcLease? rpcLease = options.Leases?.FirstOrDefault(x => x.Id == oldLease.Id);
 					if (rpcLease != null && rpcLease.State == RpcLeaseState.Completed)
 					{
@@ -854,7 +969,7 @@ namespace HordeServer.Agents
 				}
 			}
 
-			return newAgent;
+			return await CreateAgentObjectAsync(newDocument, newSession, cancellationToken);
 		}
 
 		static void SetPoolProperties(List<string> properties, List<PoolId> pools)
@@ -918,7 +1033,7 @@ namespace HordeServer.Agents
 		}
 
 		/// <inheritdoc/>
-		async Task<AgentDocument?> TryUpdateWorkspacesAsync(AgentDocument agent, List<AgentWorkspaceInfo> workspaces, bool requestConform, CancellationToken cancellationToken)
+		async Task<Agent?> TryUpdateWorkspacesAsync(Agent agent, List<AgentWorkspaceInfo> workspaces, bool requestConform, CancellationToken cancellationToken)
 		{
 			DateTime lastConformTime = DateTime.UtcNow;
 
@@ -933,188 +1048,143 @@ namespace HordeServer.Agents
 			}
 
 			// Update the agent
-			return await TryUpdateAsync(agent, update, cancellationToken);
+			AgentDocument? newDocument = await TryUpdateAsync(agent.Document, update, cancellationToken);
+			return await CreateAgentObjectAsync(newDocument, agent.Session, cancellationToken);
 		}
 
 		/// <inheritdoc/>
-		async Task<AgentDocument?> TryCreateSessionAsync(AgentDocument agent, CreateSessionOptions options, CancellationToken cancellationToken)
+		async Task<Agent?> TryCreateSessionAsync(Agent agent, CreateSessionOptions options, CancellationToken cancellationToken)
 		{
-			options.Capabilities.Flatten(out List<string> newProperties, out Dictionary<string, int> newResources);
-
-			List<PoolId> newDynamicPools = new(options.DynamicPools);
-			List<PoolId> newPools = CreatePoolsList(agent.ExplicitPools, newDynamicPools, newProperties);
-			SetPoolProperties(newProperties, newPools);
-
-			// Create a new session document
-			SessionDocument newSession = new SessionDocument(SessionIdUtils.GenerateNewId(), agent.Id, _clock.UtcNow, options.Version);
-			await _sessionCollection.InsertOneAsync(newSession, null, cancellationToken);
-
-			try
+			// Fail if there's already an active session
+			if (agent.Session != null)
 			{
-				// Reset the agent to use the new session
-				UpdateDefinitionBuilder<AgentDocument> updateBuilder = Builders<AgentDocument>.Update;
-
-				List<UpdateDefinition<AgentDocument>> updates = new List<UpdateDefinition<AgentDocument>>();
-				updates.Add(updateBuilder.Set(x => x.SessionId, newSession.Id));
-				updates.Add(updateBuilder.Set(x => x.SessionExpiresAt, newSession.StartTime + AgentService.SessionExpiryTime));
-				updates.Add(updateBuilder.Set(x => x.Status, AgentStatus.Ok));
-				updates.Add(updateBuilder.Unset(x => x.LastOnlineTime));
-				updates.Add(updateBuilder.Unset(x => x.Leases));
-				updates.Add(updateBuilder.Unset(x => x.Deleted));
-				updates.Add(updateBuilder.Set(x => x.Properties, newProperties));
-				updates.Add(updateBuilder.Set(x => x.Resources, newResources));
-				updates.Add(updateBuilder.Set(x => x.Pools, newPools));
-				updates.Add(updateBuilder.Set(x => x.DynamicPools, newDynamicPools));
-				updates.Add(updateBuilder.Set(x => x.Version, options.Version));
-				updates.Add(updateBuilder.Unset(x => x.RequestRestart));
-				updates.Add(updateBuilder.Unset(x => x.RequestShutdown));
-				updates.Add(updateBuilder.Unset(x => x.RequestForceRestart));
-				updates.Add(updateBuilder.Set(x => x.LastShutdownReason, "Unexpected"));
-
-				if (String.Equals(options.Version, agent.LastUpgradeVersion, StringComparison.Ordinal))
-				{
-					updates.Add(updateBuilder.Unset(x => x.UpgradeAttemptCount));
-				}
+				return null;
+			}
 
-				foreach (AgentLease agentLease in agent.Leases ?? Enumerable.Empty<AgentLease>())
-				{
-					await RemoveActiveLeaseAsync(agentLease);
-				}
+			// Create the new settings for the agent
+			AgentId agentId = agent.Id;
+			SessionId sessionId = SessionIdUtils.GenerateNewId();
 
-				// Apply the update
-				AgentDocument? newAgent = await TryUpdateAsync(agent, updateBuilder.Combine(updates), cancellationToken);
-				if (newAgent == null)
-				{
-					await _sessionCollection.DeleteOneAsync(x => x.Id == newSession.Id, CancellationToken.None);
-					return null;
-				}
+			options.Capabilities.Flatten(out List<string> newProperties, out Dictionary<string, int> newResources);
 
-				return newAgent;
-			}
-			catch
-			{
-				await _sessionCollection.DeleteOneAsync(x => x.Id == newSession.Id, CancellationToken.None);
-				throw;
-			}
-		}
+			List<PoolId> newDynamicPools = new(options.DynamicPools);
+			List<PoolId> newPools = CreatePoolsList(agent.Document.ExplicitPools, newDynamicPools, newProperties);
+			SetPoolProperties(newProperties, newPools);
 
-		/// <inheritdoc/>
-		async Task<AgentDocument?> TryTerminateSessionAsync(AgentDocument agent, CancellationToken cancellationToken)
-		{
-			// Get the time that the session finishes at
-			DateTime finishTime = _clock.UtcNow;
-			if (agent.SessionExpiresAt.HasValue && agent.SessionExpiresAt.Value < finishTime)
+			// Attempt to create the session in Redis
+			RpcSession? newSession = await _scheduler.TryCreateSessionAsync(agentId, sessionId, options.Capabilities, cancellationToken);
+			if (newSession == null)
 			{
-				finishTime = agent.SessionExpiresAt.Value;
+				return null;
 			}
 
-			UpdateDefinition<AgentDocument> update = new BsonDocument();
-
-			update = update.Unset(x => x.SessionId);
-			update = update.Unset(x => x.SessionExpiresAt);
-			update = update.Unset(x => x.Leases);
-			update = update.Set(x => x.Status, AgentStatus.Stopped);
-			update = update.Set(x => x.LastOnlineTime, finishTime);
+			// Update the agent with the new session id
+			UpdateDefinition<AgentDocument> update = Builders<AgentDocument>.Update
+				.Set(x => x.SessionId, sessionId)
+				.Unset(x => x.LastOnlineTime)
+				.Unset(x => x.Deleted)
+				.Set(x => x.Properties, newProperties)
+				.Set(x => x.Resources, newResources)
+				.Set(x => x.Pools, newPools)
+				.Set(x => x.DynamicPools, newDynamicPools)
+				.Set(x => x.Version, options.Version)
+				.Unset(x => x.RequestRestart)
+				.Unset(x => x.RequestShutdown)
+				.Unset(x => x.RequestForceRestart)
+				.Set(x => x.LastShutdownReason, "Unexpected");
 
-			if (agent.Ephemeral)
+			if (String.Equals(options.Version, agent.Document.LastUpgradeVersion, StringComparison.Ordinal))
 			{
-				update = update.Set(x => x.Deleted, true);
+				update = update.Unset(x => x.UpgradeAttemptCount);
 			}
 
-			AgentDocument? newAgent = await TryUpdateAsync(agent, update, cancellationToken);
-			if (newAgent == null)
+			AgentDocument? newDocument = await TryUpdateAsync(agent.Document, update, cancellationToken);
+			if (newDocument == null)
 			{
+				// If this fails, stop the session we create. It will expire naturally anyway, but we can avoid the delay on anything being requeued by just ending it now.
+				await _scheduler.TryUpdateSessionAsync(newSession, new RpcSession(newSession) { Status = RpcAgentStatus.Stopped }, cancellationToken: cancellationToken);
 				return null;
 			}
 
-			// Update the session document
-			ILogger agentLogger = GetLogger(agent.Id);
-			agentLogger.LogInformation("Terminated session {SessionId}", agent.SessionId);
-			await _sessionCollection.UpdateOneAsync(x => x.Id == agent.SessionId, Builders<SessionDocument>.Update.Set(x => x.FinishTime, finishTime), cancellationToken: CancellationToken.None);
-
-			// Remove any outstanding leases
-			if (agent.Leases != null)
-			{
-				foreach (AgentLease agentLease in agent.Leases)
-				{
-					agentLogger.LogInformation("Removing lease {LeaseId} during session terminate...", agentLease.Id);
-
-					// Remove it from the active list used for autoscaling
-					await RemoveActiveLeaseAsync(agentLease);
-
-					// Update the lease
-					await _leaseCollection.TrySetOutcomeAsync(agentLease.Id, finishTime, LeaseOutcome.Failed, null, CancellationToken.None);
-				}
-			}
+			// Create a new session document
+			SessionDocument newSessionDocument = new SessionDocument(sessionId, agentId, _clock.UtcNow, options.Version);
+			await _sessionCollection.InsertOneAsync(newSessionDocument, null, cancellationToken);
 
-			return newAgent;
+			// Update the agent is updated with the session
+			return new Agent(this, newDocument, newSession);
 		}
 
-		private static string RedisKeyActiveLeaseIds() => $"agent/active-lease-id";
-		private static string RedisKeyLeaseChildren(LeaseId parentId) => $"agent/lease-children/{parentId.ToString()}";
-
 		/// <inheritdoc/>
-		async Task<AgentDocument?> TryCreateLeaseAsync(AgentDocument agent, CreateLeaseOptions options, CancellationToken cancellationToken)
+		async Task<Agent?> TryTerminateSessionAsync(Agent agent, CancellationToken cancellationToken)
 		{
-			AgentLease newLease = new AgentLease(options);
+			return await TryUpdateSessionAsync(agent, new UpdateSessionOptions(AgentStatus.Stopped), cancellationToken);
+		}
 
-			List<AgentLease> leases = new List<AgentLease>();
-			if (agent.Leases != null)
+		/// <inheritdoc/>
+		async Task<Agent?> TryCreateLeaseAsync(Agent agent, CreateLeaseOptions options, CancellationToken cancellationToken)
+		{
+			if (agent.Session == null)
 			{
-				leases.AddRange(agent.Leases);
+				return null;
 			}
-			leases.Add(newLease);
 
-			List<UpdateDefinition<AgentDocument>> updates = new List<UpdateDefinition<AgentDocument>>();
-			updates.Add(Builders<AgentDocument>.Update.Set(x => x.Leases, leases));
-			GetNewLeaseUpdates(agent, newLease.Payload, updates);
+			AgentId agentId = agent.Document.Id;
+
+			RpcSessionLease newLease = new RpcSessionLease();
+			newLease.Id = options.Id;
+			newLease.State = RpcLeaseState.Pending;
+			newLease.Payload = Any.Pack(options.Payload);
+			newLease.ParentId = options.ParentId;
+			newLease.Exclusive = options.Exclusive;
+			newLease.Resources.Add(options.Resources ?? Enumerable.Empty<KeyValuePair<string, int>>());
 
-			UpdateDefinition<AgentDocument> update = Builders<AgentDocument>.Update.Combine(updates);
+			RpcSession? newSession = new RpcSession(agent.Session);
+			newSession.Leases.Add(newLease);
 
-			AgentDocument? updatedDoc = await TryUpdateAsync(agent, update, cancellationToken);
-			if (updatedDoc == null)
+			newSession = await _scheduler.TryUpdateSessionAsync(agent.Session, newSession, cancellationToken: cancellationToken);
+			if (newSession == null)
 			{
 				return null;
 			}
 
-			await AddActiveLeaseAsync(newLease);
-
 			try
 			{
 				DateTime startTime = _clock.UtcNow;
-				await _leaseCollection.AddAsync(options.Id, options.ParentId, options.Name, agent.Id, agent.SessionId!.Value, options.StreamId, options.PoolId, options.LogId, startTime, Any.Pack(options.Payload).ToByteArray(), cancellationToken);
+				await _leaseCollection.AddAsync(options.Id, options.ParentId, options.Name, agentId, newSession.SessionId, options.StreamId, options.PoolId, options.LogId, startTime, Any.Pack(options.Payload).ToByteArray(), cancellationToken);
 			}
 			catch (Exception ex)
 			{
-				_logger.LogError(ex, "Unable to create lease {LeaseId} for agent {AgentId}; lease already exists?", options.Id, agent.Id);
+				_logger.LogError(ex, "Unable to create lease {LeaseId} for agent {AgentId}; lease already exists?", options.Id, agentId);
 			}
 
-			return updatedDoc;
-		}
-
-		private async Task AddActiveLeaseAsync(AgentLease lease)
-		{
-			IDatabase redis = _redisService.GetDatabase();
+			// Update the agent state
+			AgentDocument? document = agent.Document;
+			while (document != null)
+			{
+				List<UpdateDefinition<AgentDocument>> updates = new List<UpdateDefinition<AgentDocument>>();
+				GetNewLeaseUpdates(document, options.Payload.ToByteArray(), updates);
 
-			await redis.SetAddAsync(RedisKeyActiveLeaseIds(), lease.Id.ToString());
-			await redis.KeyExpireAsync(RedisKeyActiveLeaseIds(), TimeSpan.FromHours(36));
+				if (updates.Count == 0)
+				{
+					break;
+				}
 
-			if (lease.ParentId != null)
-			{
-				await redis.SetAddAsync(RedisKeyLeaseChildren(lease.ParentId.Value), lease.Id.ToString());
-				await redis.KeyExpireAsync(RedisKeyLeaseChildren(lease.ParentId.Value), TimeSpan.FromHours(36));
-			}
-		}
+				UpdateDefinition<AgentDocument> update = Builders<AgentDocument>.Update.Combine(updates);
 
-		private async Task RemoveActiveLeaseAsync(AgentLease lease)
-		{
-			IDatabase redis = _redisService.GetDatabase();
-			await redis.SetRemoveAsync(RedisKeyActiveLeaseIds(), lease.Id.ToString());
+				document = await TryUpdateAsync(document, update, cancellationToken);
+				if (document != null)
+				{
+					break;
+				}
 
-			if (lease.ParentId != null)
-			{
-				await redis.SetRemoveAsync(RedisKeyLeaseChildren(lease.ParentId.Value), lease.Id.ToString());
+				document = await GetDocumentAsync(agentId, cancellationToken);
+				if (document == null)
+				{
+					return null;
+				}
 			}
+
+			return await CreateAgentObjectAsync(document, newSession, cancellationToken);
 		}
 
 		static void GetNewLeaseUpdates(AgentDocument agent, ReadOnlySpan<byte> payloadData, List<UpdateDefinition<AgentDocument>> updates)
@@ -1144,55 +1214,71 @@ namespace HordeServer.Agents
 		}
 
 		/// <inheritdoc/>
-		async Task<AgentDocument?> TryCancelLeaseAsync(AgentDocument agent, int leaseIdx, CancellationToken cancellationToken)
+		async Task<Agent?> TryCancelLeaseAsync(Agent agent, int leaseIdx, CancellationToken cancellationToken)
 		{
-			UpdateDefinition<AgentDocument> update = Builders<AgentDocument>.Update.Set(x => x.Leases![leaseIdx].State, LeaseState.Cancelled);
-			AgentDocument? newAgent = await TryUpdateAsync(agent, update, cancellationToken);
-			if (newAgent != null)
+			if (agent.Session == null)
 			{
-				PublishUpdateEvent(agent.Id);
-				if (agent.Leases != null && leaseIdx < agent.Leases.Count)
-				{
-					await RemoveActiveLeaseAsync(agent.Leases[leaseIdx]);
-				}
+				return null;
 			}
-			return newAgent;
+
+			RpcSession? newSession = new RpcSession(agent.Session);
+			if (newSession.Leases[leaseIdx].State == RpcLeaseState.Pending)
+			{
+				newSession.Leases.RemoveAt(leaseIdx);
+			}
+			else
+			{
+				newSession.Leases[leaseIdx].State = RpcLeaseState.Cancelled;
+			}
+
+			newSession = await _scheduler.TryUpdateSessionAsync(agent.Session, newSession, null, cancellationToken);
+			if (newSession == null)
+			{
+				return null;
+			}
+
+			return await CreateAgentObjectAsync(agent.Document, newSession, cancellationToken);
 		}
 
-		/// <inheritdoc/>
-		void PublishUpdateEvent(AgentId agentId)
+		void OnSessionUpdate(SessionId sessionId)
 		{
-			_ = _redisService.GetDatabase().PublishAsync(_updateEventChannel, agentId, CommandFlags.FireAndForget);
+			if (_sessionIdToTcs.TryGetValue(sessionId, out TaskCompletionSource? tcs))
+			{
+				tcs.TrySetResult();
+			}
 		}
 
-		void OnAgentUpdate(AgentId agentId)
+		async Task<Agent?> WaitForSessionUpdateAsync(Agent agent, CancellationToken cancellationToken)
 		{
-			if (_agentIdToTcs.TryGetValue(agentId, out TaskCompletionSource? tcs))
+			if (agent.Session == null)
 			{
-				tcs.TrySetResult();
+				return null;
 			}
+
+			RpcSession? newSession = await WaitForSessionUpdateAsync(agent.Session, cancellationToken);
+			return await CreateAgentObjectAsync(agent.Document, newSession, cancellationToken);
 		}
 
-		async Task<AgentDocument?> WaitForUpdateAsync(AgentDocument agent, CancellationToken cancellationToken)
+		async Task<RpcSession?> WaitForSessionUpdateAsync(RpcSession session, CancellationToken cancellationToken)
 		{
 			TaskCompletionSource newTcs = new TaskCompletionSource();
 			try
 			{
-				TaskCompletionSource tcs = _agentIdToTcs.GetOrAdd(agent.Id, newTcs);
+				TaskCompletionSource tcs = _sessionIdToTcs.GetOrAdd(session.SessionId, newTcs);
 
-				AgentDocument? newAgent = await GetDocumentAsync(agent.Id, cancellationToken);
-				if (newAgent != null && newAgent.UpdateIndex == agent.UpdateIndex)
+				RpcSession? newSession = await _scheduler.TryGetSessionAsync(session.SessionId, cancellationToken);
+				if (newSession != null && newSession.ExpiryTicks == session.ExpiryTicks)
 				{
 					await tcs.Task.WaitAsync(cancellationToken);
-					newAgent = await GetDocumentAsync(agent.Id, cancellationToken);
+					newSession = await _scheduler.TryGetSessionAsync(session.SessionId, cancellationToken);
 				}
 
-				return newAgent;
+				return newSession;
 			}
 			finally
 			{
 				newTcs.TrySetResult();
-				_agentIdToTcs.TryRemove(new KeyValuePair<AgentId, TaskCompletionSource>(agent.Id, newTcs));
+				_sessionIdToTcs.TryRemove(new KeyValuePair<SessionId, TaskCompletionSource>(session.SessionId, newTcs));
 			}
 		}
 
diff --git a/Engine/Source/Programs/Horde/Plugins/Compute/HordeServer.Compute/Agents/AgentScheduler.cs b/Engine/Source/Programs/Horde/Plugins/Compute/HordeServer.Compute/Agents/AgentScheduler.cs
new file mode 100644
--- /dev/null
+++ b/Engine/Source/Programs/Horde/Plugins/Compute/HordeServer.Compute/Agents/AgentScheduler.cs
@@ -0,0 +1,305 @@
+// Copyright Epic Games, Inc. All Rights Reserved.
+
+using System.Runtime.CompilerServices;
+using EpicGames.Core;
+using EpicGames.Horde.Agents;
+using EpicGames.Horde.Agents.Leases;
+using EpicGames.Horde.Agents.Sessions;
+using EpicGames.Redis;
+using Google.Protobuf;
+using HordeCommon.Rpc;
+using HordeCommon.Rpc.Messages;
+using HordeServer.Agents.Sessions;
+using HordeServer.Server;
+using Microsoft.Extensions.Hosting;
+using StackExchange.Redis;
+
+namespace HordeServer.Agents
+{
+	class AgentScheduler : IAgentScheduler, IHostedService, IAsyncDisposable
+	{
+#pragma warning disable CA1822 // Make property static
+		static class Keys
+		{
+			public static RedisHashKey<AgentId, SessionId> Agents { get; }
+				= new($"compute:agents");
+
+			public static SessionCollectionKeys Sessions { get; }
+				= new();
+
+			public static LeaseCollectionKeys Leases { get; }
+				= new();
+		}
+
+		record struct SessionCollectionKeys
+		{
+			public RedisHashKey<SessionId, long> ExpiryTimes
+				=> new($"compute:sessions:update-ticks");
+
+			public SessionKeys this[SessionId sessionId]
+				=> new(sessionId);
+		}
+
+		record struct SessionKeys(SessionId SessionId)
+		{
+			public RedisStringKey<RpcAgentCapabilities> Capabilities
+				=> new($"compute:sessions:{SessionId}:caps");
+
+			public RedisStringKey<RpcSession> State
+				=> new($"compute:sessions:{SessionId}:state");
+		}
+
+		record struct LeaseCollectionKeys
+		{
+			public RedisSetKey<LeaseId> Current
+				=> new($"compute:leases:active");
+
+			public LeaseKeys this[LeaseId leaseId]
+				=> new(leaseId);
+		}
+
+		record struct LeaseKeys(LeaseId LeaseId)
+		{
+			public RedisSetKey<LeaseId> Children
+				=> new($"compute:leases:{LeaseId}:children");
+		}
+
+#pragma warning restore CA1822
+
+		public event Action<SessionId>? SessionUpdated;
+
+		readonly IRedisService _redisService;
+		readonly IClock _clock;
+
+		static readonly RedisChannel<SessionId> s_sessionUpdateChannel = new RedisChannel<SessionId>(RedisChannel.Literal("compute:sessions:update"));
+
+		IAsyncDisposable? _updateEventSubscription;
+
+		/// <summary>
+		/// Constructor
+		/// </summary>
+		public AgentScheduler(IRedisService redisService, IClock clock)
+		{
+			_redisService = redisService;
+			_clock = clock;
+		}
+
+		/// <inheritdoc/>
+		public async ValueTask DisposeAsync()
+		{
+			if (_updateEventSubscription != null)
+			{
+				await _updateEventSubscription.DisposeAsync();
+				_updateEventSubscription = null;
+			}
+		}
+
+		void OnSessionUpdate(SessionId sessionId)
+			=> SessionUpdated?.Invoke(sessionId);
+
+		/// <inheritdoc/>
+		public async Task StartAsync(CancellationToken cancellationToken)
+		{
+			_updateEventSubscription = await _redisService.GetDatabase().Multiplexer.SubscribeAsync(s_sessionUpdateChannel, OnSessionUpdate);
+		}
+
+		/// <inheritdoc/>
+		public async Task StopAsync(CancellationToken cancellationToken)
+		{
+			if (_updateEventSubscription != null)
+			{
+				await _updateEventSubscription.DisposeAsync();
+				_updateEventSubscription = null;
+			}
+		}
+
+		/// <inheritdoc/>
+		public async Task<RpcSession?> TryCreateSessionAsync(AgentId agentId, SessionId sessionId, RpcAgentCapabilities capabilities, CancellationToken cancellationToken = default)
+		{
+			DateTime utcNow = _clock.UtcNow;
+
+			byte[] capabilitiesData = capabilities.ToByteArray();
+			IoHash capabilitiesHash = IoHash.Compute(capabilitiesData);
+
+			RpcSession newSession = new RpcSession();
+			newSession.AgentId = agentId;
+			newSession.SessionId = sessionId;
+			newSession.CapabilitiesHash = capabilitiesHash.ToString();
+			newSession.ExpiryTicks = (utcNow + RpcSession.ExpireAfterTime).Ticks;
+			newSession.Status = RpcAgentStatus.Ok;
+
+			ITransaction transaction = _redisService.GetDatabase().CreateTransaction();
+			transaction.AddCondition(Keys.Sessions.ExpiryTimes.HashNotExists(sessionId));
+			_ = transaction.HashSetAsync(Keys.Sessions.ExpiryTimes, sessionId, newSession.ExpiryTicks, flags: CommandFlags.FireAndForget);
+			_ = transaction.StringSetAsync(Keys.Sessions[sessionId].Capabilities.Inner, capabilitiesData, flags: CommandFlags.FireAndForget);
+			_ = transaction.StringSetAsync(Keys.Sessions[sessionId].State, newSession, flags: CommandFlags.FireAndForget);
+
+			if (!await transaction.ExecuteAsync().WaitAsync(cancellationToken))
+			{
+				return null;
+			}
+
+			return newSession;
+		}
+
+		/// <inheritdoc/>
+		public async Task<RpcSession?> TryGetSessionAsync(AgentId agentId, CancellationToken cancellationToken = default)
+		{
+			IDatabase database = _redisService.GetDatabase();
+
+			SessionId sessionId = await database.HashGetAsync(Keys.Agents, agentId).WaitAsync(cancellationToken);
+			if (sessionId.Id.IsEmpty)
+			{
+				return null;
+			}
+
+			return await TryGetSessionAsync(sessionId, cancellationToken);
+		}
+
+		/// <inheritdoc/>
+		public async Task<RpcSession?> TryGetSessionAsync(SessionId sessionId, CancellationToken cancellationToken = default)
+		{
+			IDatabase database = _redisService.GetDatabase();
+			return await database.StringGetAsync(Keys.Sessions[sessionId].State);
+		}
+
+		/// <inheritdoc/>
+		public async Task<RpcAgentCapabilities?> TryGetSessionCapabilitiesAsync(RpcSession session, CancellationToken cancellationToken = default)
+		{
+			IDatabase database = _redisService.GetDatabase();
+
+			byte[]? data = (byte[]?)await database.StringGetAsync(Keys.Sessions[session.SessionId].Capabilities.Inner);
+			if (data == null)
+			{
+				return null;
+			}
+			if (!session.CapabilitiesHash.Equals(IoHash.Compute(data).ToString(), StringComparison.Ordinal))
+			{
+				return null;
+			}
+
+			return RpcAgentCapabilities.Parser.ParseFrom(data);
+		}
+
+		/// <inheritdoc/>
+		public async Task<RpcSession?> TryUpdateSessionAsync(RpcSession session, RpcSession? newSession, RpcAgentCapabilities? newCapabilities, CancellationToken cancellationToken)
+		{
+			SessionId sessionId = session.SessionId;
+
+			// Copy the session document if a new one was not specified, and update the last modified time
+			newSession ??= new RpcSession(session);
+
+			// Start building the update transaction
+			ITransaction transaction = _redisService.GetDatabase().CreateTransaction();
+			transaction.AddCondition(Keys.Sessions.ExpiryTimes.HashEqual(newSession.SessionId, session.ExpiryTicks));
+
+			// Handle agents transitioning to the stopping state
+			if (newSession.Status == RpcAgentStatus.Stopped)
+			{
+				// TODO: remove all leases
+				newSession.ExpiryTicks = Math.Min(newSession.ExpiryTicks, _clock.UtcNow.Ticks);
+
+				_ = transaction.HashDeleteAsync(Keys.Agents, newSession.AgentId, flags: CommandFlags.FireAndForget);
+				_ = transaction.HashDeleteAsync(Keys.Sessions.ExpiryTimes, newSession.SessionId, flags: CommandFlags.FireAndForget);
+				_ = transaction.KeyDeleteAsync(Keys.Sessions[sessionId].Capabilities, flags: CommandFlags.FireAndForget);
+				_ = transaction.KeyDeleteAsync(Keys.Sessions[sessionId].State, flags: CommandFlags.FireAndForget);
+			}
+			else
+			{
+				// Extend the expiry time, ensuring it only ever increases
+				newSession.ExpiryTicks = Math.Max(session.ExpiryTicks + 1, (_clock.UtcNow + RpcSession.ExpireAfterTime).Ticks);
+				_ = transaction.HashSetAsync(Keys.Sessions.ExpiryTimes, session.SessionId, newSession.ExpiryTicks, flags: CommandFlags.FireAndForget);
+
+				// Update the capabilities
+				if (newCapabilities != null)
+				{
+					byte[] capabilitiesData = newCapabilities.ToByteArray();
+					IoHash capabilitiesHash = IoHash.Compute(capabilitiesData);
+
+					newSession.CapabilitiesHash = capabilitiesHash.ToString();
+					_ = transaction.StringSetAsync(Keys.Sessions[session.SessionId].Capabilities.Inner, capabilitiesData, flags: CommandFlags.FireAndForget);
+				}
+
+				// Add any new leases to the global state
+				foreach (RpcSessionLease newLease in newSession.Leases)
+				{
+					if (!session.Leases.Any(x => x.Id == newLease.Id))
+					{
+						_ = transaction.SetAddAsync(Keys.Leases.Current, newLease.Id, flags: CommandFlags.FireAndForget);
+						if (newLease.ParentId != null)
+						{
+							_ = transaction.SetAddAsync(Keys.Leases[newLease.ParentId.Value].Children, newLease.Id, flags: CommandFlags.FireAndForget);
+						}
+					}
+				}
+
+				// Update the session state
+				_ = transaction.StringSetAsync(Keys.Sessions[session.SessionId].State, newSession, flags: CommandFlags.FireAndForget);
+			}
+
+			// Remove any complete leases from the global state
+			foreach (RpcSessionLease oldLease in session.Leases)
+			{
+				if (!newSession.Leases.Any(x => x.Id == oldLease.Id))
+				{
+					_ = transaction.SetRemoveAsync(Keys.Leases.Current, oldLease.Id, flags: CommandFlags.FireAndForget);
+					if (oldLease.ParentId != null)
+					{
+						_ = transaction.SetRemoveAsync(Keys.Leases[oldLease.ParentId.Value].Children, oldLease.Id, flags: CommandFlags.FireAndForget);
+					}
+				}
+			}
+
+			// Execute the transaction
+			if (!await transaction.ExecuteAsync().WaitAsync(cancellationToken))
+			{
+				return null;
+			}
+
+			// Notify watchers that the session state has changed
+			_ = _redisService.GetDatabase().PublishAsync(s_sessionUpdateChannel, sessionId, CommandFlags.FireAndForget);
+			return newSession;
+		}
+
+		/// <inheritdoc/>
+		public async IAsyncEnumerable<RpcSession> FindExpiredSessionsAsync([EnumeratorCancellation] CancellationToken cancellationToken = default)
+		{
+			DateTime utcNow = _clock.UtcNow;
+			IDatabase database = _redisService.GetDatabase();
+
+			HashEntry<SessionId, long>[] entries = await database.HashGetAllAsync(Keys.Sessions.ExpiryTimes);
+			foreach ((SessionId sessionId, long expiryTicks) in entries)
+			{
+				if (expiryTicks < utcNow.Ticks)
+				{
+					RpcSession? session = await TryGetSessionAsync(sessionId, cancellationToken);
+					if (session != null && session.ExpiryTicks == expiryTicks)
+					{
+						yield return session;
+					}
+				}
+			}
+		}
+
+		/// <inheritdoc/>
+		public async Task<LeaseId[]> FindActiveLeaseIdsAsync(CancellationToken cancellationToken = default)
+		{
+			IDatabase database = _redisService.GetDatabase();
+			return await database.SetMembersAsync(Keys.Leases.Current).WaitAsync(cancellationToken);
+		}
+
+		/// <inheritdoc/>
+		public async Task<int> GetChildLeaseCountAsync(LeaseId id, CancellationToken cancellationToken = default)
+		{
+			IDatabase database = _redisService.GetDatabase();
+			return (int)await database.SetLengthAsync(Keys.Leases[id].Children).WaitAsync(cancellationToken);
+		}
+
+		/// <inheritdoc/>
+		public async Task<LeaseId[]> GetChildLeaseIdsAsync(LeaseId id, CancellationToken cancellationToken = default)
+		{
+			IDatabase database = _redisService.GetDatabase();
+			return await database.SetMembersAsync(Keys.Leases[id].Children).WaitAsync(cancellationToken);
+		}
+	}
+}
diff --git a/Engine/Source/Programs/Horde/Plugins/Compute/HordeServer.Compute/Agents/AgentService.cs b/Engine/Source/Programs/Horde/Plugins/Compute/HordeServer.Compute/Agents/AgentService.cs
--- a/Engine/Source/Programs/Horde/Plugins/Compute/HordeServer.Compute/Agents/AgentService.cs
+++ b/Engine/Source/Programs/Horde/Plugins/Compute/HordeServer.Compute/Agents/AgentService.cs
@@ -654,20 +654,11 @@ namespace HordeServer.Agents
 					throw new InvalidOperationException("Session has already expired");
 				}
 
-				// Extend the current session time if we're within a time period of the current time expiring. This reduces
-				// unnecessary DB writes a little, but it also allows us to skip the update and jump into a long poll state
-				// if there's still time on the session left.
-				DateTime? sessionExpiresAt = null;
-				if (!agent.SessionExpiresAt.HasValue || utcNow > (agent.SessionExpiresAt - SessionExpiryTime) + SessionRenewTime)
-				{
-					sessionExpiresAt = utcNow + SessionExpiryTime;
-				}
-
 				// Get the new dynamic pools for the agent
 				List<PoolId> dynamicPools = await GetDynamicPoolsAsync(agent, cancellationToken);
 
 				// Update the agent, and try to create new lease documents if we succeed
-				IAgent? newAgent = await agent.TryUpdateSessionAsync(new UpdateSessionOptions(status, sessionExpiresAt, capabilities, dynamicPools, newLeases), cancellationToken);
+				IAgent? newAgent = await agent.TryUpdateSessionAsync(new UpdateSessionOptions(status, capabilities, dynamicPools, newLeases), cancellationToken);
 				if (newAgent != null)
 				{
 					agent = newAgent;
diff --git a/Engine/Source/Programs/Horde/Plugins/Compute/HordeServer.Compute/Agents/IAgent.cs b/Engine/Source/Programs/Horde/Plugins/Compute/HordeServer.Compute/Agents/IAgent.cs
--- a/Engine/Source/Programs/Horde/Plugins/Compute/HordeServer.Compute/Agents/IAgent.cs
+++ b/Engine/Source/Programs/Horde/Plugins/Compute/HordeServer.Compute/Agents/IAgent.cs
@@ -312,11 +312,10 @@ namespace HordeServer.Agents
 	/// Options for updating a new agent session
 	/// </summary>
 	/// <param name="Status">New status of the agent</param>
-	/// <param name="SessionExpiresAt">New expiry time for the current session</param>
 	/// <param name="Capabilities">Capbilities for the session</param>
 	/// <param name="DynamicPools">New list of dynamic pools for the agent</param>
 	/// <param name="Leases">New set of leases</param>
-	public record class UpdateSessionOptions(AgentStatus? Status = null, DateTime? SessionExpiresAt = null, RpcAgentCapabilities? Capabilities = null, IReadOnlyList<PoolId>? DynamicPools = null, IEnumerable<RpcLease>? Leases = null);
+	public record class UpdateSessionOptions(AgentStatus? Status = null, RpcAgentCapabilities? Capabilities = null, IReadOnlyList<PoolId>? DynamicPools = null, IEnumerable<RpcLease>? Leases = null);
 
 	/// <summary>
 	/// Extension methods for IAgent
diff --git a/Engine/Source/Programs/Horde/Plugins/Compute/HordeServer.Compute/Agents/IAgentCollection.cs b/Engine/Source/Programs/Horde/Plugins/Compute/HordeServer.Compute/Agents/IAgentCollection.cs
--- a/Engine/Source/Programs/Horde/Plugins/Compute/HordeServer.Compute/Agents/IAgentCollection.cs
+++ b/Engine/Source/Programs/Horde/Plugins/Compute/HordeServer.Compute/Agents/IAgentCollection.cs
@@ -1,7 +1,6 @@
 // Copyright Epic Games, Inc. All Rights Reserved.
 
 using EpicGames.Horde.Agents;
-using EpicGames.Horde.Agents.Leases;
 using EpicGames.Horde.Agents.Pools;
 using HordeServer.Auditing;
 
@@ -53,21 +52,6 @@ namespace HordeServer.Agents
 		/// <returns>List of agents matching the given criteria</returns>
 		IAsyncEnumerable<IAgent> FindAsync(PoolId? poolId = null, DateTime? modifiedAfter = null, string? property = null, AgentStatus? status = null, bool? enabled = null, bool includeDeleted = false, bool consistentRead = true, CancellationToken cancellationToken = default);
 
-		/// <summary>
-		/// Finds all active agent lease IDs
-		/// </summary>
-		/// <param name="cancellationToken">Cancellation token for the operation</param>
-		/// <returns>List of agent lease IDs</returns>
-		Task<List<LeaseId>> FindActiveLeaseIdsAsync(CancellationToken cancellationToken = default);
-
-		/// <summary>
-		/// Get all child lease IDs
-		/// </summary>
-		/// <param name="id">Lease ID</param>
-		/// <param name="cancellationToken">Cancellation token for the operation</param>
-		/// <returns>List of agent lease IDs</returns>
-		Task<List<LeaseId>> GetChildLeaseIdsAsync(LeaseId id, CancellationToken cancellationToken = default);
-
 		/// <summary>
 		/// Gets the log channel for an agent
 		/// </summary>
diff --git a/Engine/Source/Programs/Horde/Plugins/Compute/HordeServer.Compute/Agents/IAgentScheduler.cs b/Engine/Source/Programs/Horde/Plugins/Compute/HordeServer.Compute/Agents/IAgentScheduler.cs
new file mode 100644
--- /dev/null
+++ b/Engine/Source/Programs/Horde/Plugins/Compute/HordeServer.Compute/Agents/IAgentScheduler.cs
@@ -0,0 +1,93 @@
+// Copyright Epic Games, Inc. All Rights Reserved.
+
+using EpicGames.Horde.Agents;
+using EpicGames.Horde.Agents.Leases;
+using EpicGames.Horde.Agents.Sessions;
+using HordeCommon.Rpc.Messages;
+
+namespace HordeServer.Agents
+{
+	/// <summary>
+	/// Interface for a scheduler that manages agent sessions and leases.
+	/// </summary>
+	public interface IAgentScheduler
+	{
+		#region Sessions
+
+		/// <summary>
+		/// Event raised whenever a session is updated (locally or remotely)
+		/// </summary>
+		event Action<SessionId>? SessionUpdated;
+
+		/// <summary>
+		/// Attempts to create a new session for the given agent. Fails if the agent is already executing a session.
+		/// </summary>
+		/// <param name="agentId">The agent id</param>
+		/// <param name="sessionId">Identifier for the new session</param>
+		/// <param name="capabilities">Capabilities for the agent</param>
+		/// <param name="cancellationToken">Cancellation token for the operation</param>
+		/// <returns>New session object</returns>
+		Task<RpcSession?> TryCreateSessionAsync(AgentId agentId, SessionId sessionId, RpcAgentCapabilities capabilities, CancellationToken cancellationToken = default);
+
+		/// <summary>
+		/// Attempt to get the latest state of the given session
+		/// </summary>
+		/// <param name="sessionId">The session identifier</param>
+		/// <param name="cancellationToken">Cancellation token for the operation</param>
+		/// <returns>Session object, or null if the session is not active</returns>
+		Task<RpcSession?> TryGetSessionAsync(SessionId sessionId, CancellationToken cancellationToken = default);
+
+		/// <summary>
+		/// Attempt to get capabilities for the given session
+		/// </summary>
+		/// <param name="session">The session object</param>
+		/// <param name="cancellationToken">Cancellation token for the operation</param>
+		/// <returns>Matching capabilities object for the session</returns>
+		Task<RpcAgentCapabilities?> TryGetSessionCapabilitiesAsync(RpcSession session, CancellationToken cancellationToken = default);
+
+		/// <summary>
+		/// Attempts to create a new session for the given agent. Fails if the agent is already executing a session.
+		/// </summary>
+		/// <param name="session">Previous session state</param>
+		/// <param name="newSession">New session stats</param>
+		/// <param name="newCapabilities">New capabilities for the agent</param>
+		/// <param name="cancellationToken">Cancellation token for the operation</param>
+		/// <returns>New session object</returns>
+		Task<RpcSession?> TryUpdateSessionAsync(RpcSession session, RpcSession? newSession = null, RpcAgentCapabilities? newCapabilities = null, CancellationToken cancellationToken = default);
+
+		#endregion
+		#region Stats
+
+		/// <summary>
+		/// Enumerates all the sessions which have expired
+		/// </summary>
+		/// <param name="cancellationToken">Cancellation token for the operation</param>
+		/// <returns>List of expired sessions</returns>
+		IAsyncEnumerable<RpcSession> FindExpiredSessionsAsync(CancellationToken cancellationToken = default);
+
+		/// <summary>
+		/// Finds all active agent lease IDs
+		/// </summary>
+		/// <param name="cancellationToken">Cancellation token for the operation</param>
+		/// <returns>List of agent lease IDs</returns>
+		Task<LeaseId[]> FindActiveLeaseIdsAsync(CancellationToken cancellationToken = default);
+
+		/// <summary>
+		/// Find how many child leases a particular lease has
+		/// </summary>
+		/// <param name="id">Lease ID</param>
+		/// <param name="cancellationToken">Cancellation token for the operation</param>
+		/// <returns>List of agent lease IDs</returns>
+		Task<int> GetChildLeaseCountAsync(LeaseId id, CancellationToken cancellationToken = default);
+
+		/// <summary>
+		/// Get all child lease IDs
+		/// </summary>
+		/// <param name="id">Lease ID</param>
+		/// <param name="cancellationToken">Cancellation token for the operation</param>
+		/// <returns>List of agent lease IDs</returns>
+		Task<LeaseId[]> GetChildLeaseIdsAsync(LeaseId id, CancellationToken cancellationToken = default);
+
+		#endregion
+	}
+}
diff --git a/Engine/Source/Programs/Horde/Plugins/Compute/HordeServer.Compute/Compute/ComputeService.cs b/Engine/Source/Programs/Horde/Plugins/Compute/HordeServer.Compute/Compute/ComputeService.cs
--- a/Engine/Source/Programs/Horde/Plugins/Compute/HordeServer.Compute/Compute/ComputeService.cs
+++ b/Engine/Source/Programs/Horde/Plugins/Compute/HordeServer.Compute/Compute/ComputeService.cs
@@ -167,6 +167,7 @@ namespace HordeServer.Compute
 		public event ResourceNeedEvent? OnResourceNeedsUpdated;
 
 		readonly IAgentCollection _agentCollection;
+		readonly IAgentScheduler _agentScheduler;
 		readonly ILogCollection _logCollection;
 		readonly AgentService _agentService;
 		readonly AgentRelayService _agentRelayService;
@@ -189,6 +190,7 @@ namespace HordeServer.Compute
 		/// </summary>
 		public ComputeService(
 			IAgentCollection agentCollection,
+			IAgentScheduler agentScheduler,
 			ILogCollection logCollection,
 			AgentService agentService,
 			AgentRelayService agentRelayService,
@@ -201,6 +203,7 @@ namespace HordeServer.Compute
 			ILogger<ComputeService> logger)
 		{
 			_agentCollection = agentCollection;
+			_agentScheduler = agentScheduler;
 			_logCollection = logCollection;
 			_agentService = agentService;
 			_agentRelayService = agentRelayService;
@@ -687,8 +690,7 @@ namespace HordeServer.Compute
 				return null;
 			}
 
-			List<LeaseId> childLeaseIds = await _agentCollection.GetChildLeaseIdsAsync(parentLeaseId.Value, cancellationToken);
-			return childLeaseIds.Count;
+			return await _agentScheduler.GetChildLeaseCountAsync(parentLeaseId.Value, cancellationToken);
 		}
 
 		private async Task<ComputeResource?> TryAssignAsync(AllocateResourceParams arp, IAgent agent, ComputeTask computeTask, LeaseId leaseId)
diff --git a/Engine/Source/Programs/Horde/Plugins/Compute/HordeServer.Compute/ComputePlugin.cs b/Engine/Source/Programs/Horde/Plugins/Compute/HordeServer.Compute/ComputePlugin.cs
--- a/Engine/Source/Programs/Horde/Plugins/Compute/HordeServer.Compute/ComputePlugin.cs
+++ b/Engine/Source/Programs/Horde/Plugins/Compute/HordeServer.Compute/ComputePlugin.cs
@@ -70,6 +70,10 @@ namespace HordeServer
 			services.AddSingleton<IAgentCollection>(sp => sp.GetRequiredService<AgentCollection>());
 			services.AddHostedService(sp => sp.GetRequiredService<AgentCollection>());
 
+			services.AddSingleton<AgentScheduler>();
+			services.AddSingleton<IAgentScheduler>(sp => sp.GetRequiredService<AgentScheduler>());
+			services.AddHostedService(sp => sp.GetRequiredService<AgentScheduler>());
+
 			services.AddSingleton<LeaseCollection>();
 			services.AddSingleton<ILeaseCollection>(sp => sp.GetRequiredService<LeaseCollection>());
 
diff --git a/Engine/Source/Programs/Horde/Plugins/Compute/HordeServer.Compute/HordeServer.Compute.csproj b/Engine/Source/Programs/Horde/Plugins/Compute/HordeServer.Compute/HordeServer.Compute.csproj
--- a/Engine/Source/Programs/Horde/Plugins/Compute/HordeServer.Compute/HordeServer.Compute.csproj
+++ b/Engine/Source/Programs/Horde/Plugins/Compute/HordeServer.Compute/HordeServer.Compute.csproj
@@ -30,6 +30,7 @@
     <Protobuf Include="Protos\horde\messages\agent_requirements.proto" ProtoRoot="Protos" />
     <Protobuf Include="Protos\horde\messages\agent_workspace.proto" ProtoRoot="Protos" Access="Public" />
     <Protobuf Include="Protos\horde\messages\lease.proto" ProtoRoot="Protos" />
+    <Protobuf Include="Protos\horde\messages\session.proto" ProtoRoot="Protos" />
     <Protobuf Include="Protos\horde\tasks\compute_task.proto" ProtoRoot="Protos" Access="Public" />
     <Protobuf Include="Protos\horde\tasks\conform_task.proto" ProtoRoot="Protos" Access="Public" />
     <Protobuf Include="Protos\horde\tasks\job_task.proto" ProtoRoot="Protos" />
diff --git a/Engine/Source/Programs/Horde/Plugins/Compute/HordeServer.Compute/Protos/horde/messages/session.proto.cs b/Engine/Source/Programs/Horde/Plugins/Compute/HordeServer.Compute/Protos/horde/messages/session.proto.cs
new file mode 100644
--- /dev/null
+++ b/Engine/Source/Programs/Horde/Plugins/Compute/HordeServer.Compute/Protos/horde/messages/session.proto.cs
@@ -0,0 +1,50 @@
+// Copyright Epic Games, Inc. All Rights Reserved.
+
+using EpicGames.Horde.Agents;
+using EpicGames.Horde.Agents.Leases;
+using EpicGames.Horde.Agents.Sessions;
+using HordeServer.Agents;
+
+#pragma warning disable CS1591
+
+namespace HordeCommon.Rpc.Messages
+{
+	partial class RpcSession
+	{
+		public static TimeSpan ExpireAfterTime { get; } = TimeSpan.FromMinutes(2.0);
+
+		public AgentId AgentId
+		{
+			get => new EpicGames.Horde.Agents.AgentId(AgentIdValue);
+			set => AgentIdValue = value.ToString();
+		}
+
+		public SessionId SessionId
+		{
+			get => EpicGames.Horde.Agents.Sessions.SessionId.Parse(SessionIdValue);
+			set => SessionIdValue = value.ToString();
+		}
+
+		public DateTime ExpiryTime
+			=> new DateTime(ExpiryTicks, DateTimeKind.Utc);
+	}
+
+	partial class RpcSessionLease : IAgentLease
+	{
+		public LeaseId Id
+		{
+			get => LeaseId.Parse(IdValue);
+			set => IdValue = value.ToString();
+		}
+
+		public LeaseId? ParentId
+		{
+			get => String.IsNullOrEmpty(ParentIdValue) ? null : LeaseId.Parse(ParentIdValue);
+			set => ParentIdValue = value?.ToString() ?? String.Empty;
+		}
+
+		IReadOnlyDictionary<string, int>? IAgentLease.Resources => this.Resources;
+
+		LeaseState IAgentLease.State => (LeaseState)State;
+	}
+}
