Horde: Remove unused using directive.

#rnx

diff --git a/Engine/Source/Programs/Horde/Plugins/Compute/HordeServer.Compute/Agents/AgentScheduler.cs b/Engine/Source/Programs/Horde/Plugins/Compute/HordeServer.Compute/Agents/AgentScheduler.cs
--- a/Engine/Source/Programs/Horde/Plugins/Compute/HordeServer.Compute/Agents/AgentScheduler.cs
+++ b/Engine/Source/Programs/Horde/Plugins/Compute/HordeServer.Compute/Agents/AgentScheduler.cs
@@ -9,7 +9,6 @@ using EpicGames.Redis;
 using Google.Protobuf;
 using HordeCommon.Rpc;
 using HordeCommon.Rpc.Messages;
-using HordeServer.Agents.Sessions;
 using HordeServer.Server;
 using Microsoft.Extensions.Hosting;
 using StackExchange.Redis;
