Horde: Use a consistent value for the expiry time of sessions.

#rnx

diff --git a/Engine/Source/Programs/Horde/Plugins/Compute/HordeServer.Compute/Agents/AgentService.cs b/Engine/Source/Programs/Horde/Plugins/Compute/HordeServer.Compute/Agents/AgentService.cs
--- a/Engine/Source/Programs/Horde/Plugins/Compute/HordeServer.Compute/Agents/AgentService.cs
+++ b/Engine/Source/Programs/Horde/Plugins/Compute/HordeServer.Compute/Agents/AgentService.cs
@@ -48,7 +48,7 @@ namespace HordeServer.Agents
 		/// <summary>
 		/// Maximum time between updates for an agent to be considered online
 		/// </summary>
-		public static readonly TimeSpan SessionExpiryTime = TimeSpan.FromMinutes(5);
+		public static readonly TimeSpan SessionExpiryTime = RpcSession.ExpireAfterTime;
 
 		/// <summary>
 		/// Time before a session expires that we will poll until
diff --git a/Engine/Source/Programs/Horde/Plugins/Compute/HordeServer.Compute/Protos/horde/messages/session.proto.cs b/Engine/Source/Programs/Horde/Plugins/Compute/HordeServer.Compute/Protos/horde/messages/session.proto.cs
--- a/Engine/Source/Programs/Horde/Plugins/Compute/HordeServer.Compute/Protos/horde/messages/session.proto.cs
+++ b/Engine/Source/Programs/Horde/Plugins/Compute/HordeServer.Compute/Protos/horde/messages/session.proto.cs
@@ -11,7 +11,7 @@ namespace HordeCommon.Rpc.Messages
 {
 	partial class RpcSession
 	{
-		public static TimeSpan ExpireAfterTime { get; } = TimeSpan.FromMinutes(2.0);
+		public static TimeSpan ExpireAfterTime { get; } = TimeSpan.FromMinutes(5.0);
 
 		public AgentId AgentId
 		{
