Horde: Remove agent lookup from scheduler.

#rnx

diff --git a/Engine/Source/Programs/Horde/Plugins/Compute/HordeServer.Compute/Agents/AgentScheduler.cs b/Engine/Source/Programs/Horde/Plugins/Compute/HordeServer.Compute/Agents/AgentScheduler.cs
--- a/Engine/Source/Programs/Horde/Plugins/Compute/HordeServer.Compute/Agents/AgentScheduler.cs
+++ b/Engine/Source/Programs/Horde/Plugins/Compute/HordeServer.Compute/Agents/AgentScheduler.cs
@@ -20,9 +20,6 @@ namespace HordeServer.Agents
 #pragma warning disable CA1822 // Make property static
 		static class Keys
 		{
-			public static RedisHashKey<AgentId, SessionId> Agents { get; }
-				= new($"compute:agents");
-
 			public static SessionCollectionKeys Sessions { get; }
 				= new();
 
@@ -141,20 +138,6 @@ namespace HordeServer.Agents
 			return newSession;
 		}
 
-		/// <inheritdoc/>
-		public async Task<RpcSession?> TryGetSessionAsync(AgentId agentId, CancellationToken cancellationToken = default)
-		{
-			IDatabase database = _redisService.GetDatabase();
-
-			SessionId sessionId = await database.HashGetAsync(Keys.Agents, agentId).WaitAsync(cancellationToken);
-			if (sessionId.Id.IsEmpty)
-			{
-				return null;
-			}
-
-			return await TryGetSessionAsync(sessionId, cancellationToken);
-		}
-
 		/// <inheritdoc/>
 		public async Task<RpcSession?> TryGetSessionAsync(SessionId sessionId, CancellationToken cancellationToken = default)
 		{
@@ -198,7 +181,6 @@ namespace HordeServer.Agents
 				// TODO: remove all leases
 				newSession.ExpiryTicks = Math.Min(newSession.ExpiryTicks, _clock.UtcNow.Ticks);
 
-				_ = transaction.HashDeleteAsync(Keys.Agents, newSession.AgentId, flags: CommandFlags.FireAndForget);
 				_ = transaction.HashDeleteAsync(Keys.Sessions.ExpiryTimes, newSession.SessionId, flags: CommandFlags.FireAndForget);
 				_ = transaction.KeyDeleteAsync(Keys.Sessions[sessionId].Capabilities, flags: CommandFlags.FireAndForget);
 				_ = transaction.KeyDeleteAsync(Keys.Sessions[sessionId].State, flags: CommandFlags.FireAndForget);
