Horde: Reduce number of update notifications from scheduler when ticking sessions.

#rnx

diff --git a/Engine/Source/Programs/Horde/Plugins/Compute/HordeServer.Compute/Agents/AgentScheduler.cs b/Engine/Source/Programs/Horde/Plugins/Compute/HordeServer.Compute/Agents/AgentScheduler.cs
--- a/Engine/Source/Programs/Horde/Plugins/Compute/HordeServer.Compute/Agents/AgentScheduler.cs
+++ b/Engine/Source/Programs/Horde/Plugins/Compute/HordeServer.Compute/Agents/AgentScheduler.cs
@@ -195,10 +195,13 @@ namespace HordeServer.Agents
 				if (newCapabilities != null)
 				{
 					byte[] capabilitiesData = newCapabilities.ToByteArray();
-					IoHash capabilitiesHash = IoHash.Compute(capabilitiesData);
+					string capabilitiesHash = IoHash.Compute(capabilitiesData).ToString();
 
-					newSession.CapabilitiesHash = capabilitiesHash.ToString();
-					_ = transaction.StringSetAsync(Keys.Sessions[session.SessionId].Capabilities.Inner, capabilitiesData, flags: CommandFlags.FireAndForget);
+					if (capabilitiesHash != newSession.CapabilitiesHash)
+					{
+						newSession.CapabilitiesHash = capabilitiesHash;
+						_ = transaction.StringSetAsync(Keys.Sessions[session.SessionId].Capabilities.Inner, capabilitiesData, flags: CommandFlags.FireAndForget);
+					}
 				}
 
 				// Add any new leases to the global state
@@ -238,7 +241,10 @@ namespace HordeServer.Agents
 			}
 
 			// Notify watchers that the session state has changed
-			_ = _redisService.GetDatabase().PublishAsync(s_sessionUpdateChannel, sessionId, CommandFlags.FireAndForget);
+			if (newSession != null || newCapabilities != null)
+			{
+				_ = _redisService.GetDatabase().PublishAsync(s_sessionUpdateChannel, sessionId, CommandFlags.FireAndForget);
+			}
 			return newSession;
 		}
 
