Horde: Additional logging for session expiry.

#rnx

diff --git a/Engine/Source/Programs/Horde/Plugins/Compute/HordeServer.Compute/Agents/AgentCollection.cs b/Engine/Source/Programs/Horde/Plugins/Compute/HordeServer.Compute/Agents/AgentCollection.cs
--- a/Engine/Source/Programs/Horde/Plugins/Compute/HordeServer.Compute/Agents/AgentCollection.cs
+++ b/Engine/Source/Programs/Horde/Plugins/Compute/HordeServer.Compute/Agents/AgentCollection.cs
@@ -333,9 +333,9 @@ namespace HordeServer.Agents
 				if (document != null)
 				{
 					Agent? agent = await CreateAgentObjectAsync(document, session, cancellationToken);
-					if (agent != null)
+					if (agent != null && agent.Session != null && agent.Session.ExpiryTicks == session.ExpiryTicks)
 					{
-						_logger.LogDebug("Terminating session {SessionId} for agent {Agent}", session.SessionId, session.AgentId);
+						_logger.LogDebug("Terminating session {SessionId} for agent {Agent} (expiry time: {Time})", session.SessionId, session.AgentId, agent.Session?.ExpiryTime);
 						await TryUpdateSessionAsync(agent, new UpdateSessionOptions { Status = AgentStatus.Stopped }, cancellationToken);
 						c++;
 					}
diff --git a/Engine/Source/Programs/Horde/Plugins/Compute/HordeServer.Compute/Agents/AgentScheduler.cs b/Engine/Source/Programs/Horde/Plugins/Compute/HordeServer.Compute/Agents/AgentScheduler.cs
--- a/Engine/Source/Programs/Horde/Plugins/Compute/HordeServer.Compute/Agents/AgentScheduler.cs
+++ b/Engine/Source/Programs/Horde/Plugins/Compute/HordeServer.Compute/Agents/AgentScheduler.cs
@@ -11,6 +11,7 @@ using HordeCommon.Rpc;
 using HordeCommon.Rpc.Messages;
 using HordeServer.Server;
 using Microsoft.Extensions.Hosting;
+using Microsoft.Extensions.Logging;
 using StackExchange.Redis;
 
 namespace HordeServer.Agents
@@ -66,6 +67,7 @@ namespace HordeServer.Agents
 
 		readonly IRedisService _redisService;
 		readonly IClock _clock;
+		readonly ILogger _logger;
 
 		static readonly RedisChannel<SessionId> s_sessionUpdateChannel = new RedisChannel<SessionId>(RedisChannel.Literal("compute:sessions:update"));
 
@@ -74,10 +76,11 @@ namespace HordeServer.Agents
 		/// <summary>
 		/// Constructor
 		/// </summary>
-		public AgentScheduler(IRedisService redisService, IClock clock)
+		public AgentScheduler(IRedisService redisService, IClock clock, ILogger<AgentScheduler> logger)
 		{
 			_redisService = redisService;
 			_clock = clock;
+			_logger = logger;
 		}
 
 		/// <inheritdoc/>
@@ -240,6 +243,9 @@ namespace HordeServer.Agents
 				return null;
 			}
 
+			// Trace the new expiry time for debugging
+			_logger.LogDebug("Updated session {SessionId} expiry time to {ExpiryTime}", sessionId, newSession.ExpiryTime);
+
 			// Notify watchers that the session state has changed
 			if (newSession != null || newCapabilities != null)
 			{
