Horde: Explicitly check for modified fields before sending session update notifications.

#rnx

diff --git a/Engine/Source/Programs/Horde/Plugins/Compute/HordeServer.Compute/Agents/AgentCollection.cs b/Engine/Source/Programs/Horde/Plugins/Compute/HordeServer.Compute/Agents/AgentCollection.cs
--- a/Engine/Source/Programs/Horde/Plugins/Compute/HordeServer.Compute/Agents/AgentCollection.cs
+++ b/Engine/Source/Programs/Horde/Plugins/Compute/HordeServer.Compute/Agents/AgentCollection.cs
@@ -1244,6 +1244,7 @@ namespace HordeServer.Agents
 		{
 			if (_sessionIdToTcs.TryGetValue(sessionId, out TaskCompletionSource? tcs))
 			{
+				_logger.LogDebug("Session {SessionId} was updated.");
 				tcs.TrySetResult();
 			}
 		}
diff --git a/Engine/Source/Programs/Horde/Plugins/Compute/HordeServer.Compute/Agents/AgentScheduler.cs b/Engine/Source/Programs/Horde/Plugins/Compute/HordeServer.Compute/Agents/AgentScheduler.cs
--- a/Engine/Source/Programs/Horde/Plugins/Compute/HordeServer.Compute/Agents/AgentScheduler.cs
+++ b/Engine/Source/Programs/Horde/Plugins/Compute/HordeServer.Compute/Agents/AgentScheduler.cs
@@ -247,7 +247,7 @@ namespace HordeServer.Agents
 			_logger.LogDebug("Updated session {SessionId} expiry time to {ExpiryTime}", sessionId, newSession.ExpiryTime);
 
 			// Notify watchers that the session state has changed
-			if (newSession != null || newCapabilities != null)
+			if (newSession.Status != session.Status || !newSession.CapabilitiesHash.Equals(session.CapabilitiesHash, StringComparison.Ordinal) || !newSession.Leases.Equals(session.Leases))
 			{
 				_ = _redisService.GetDatabase().PublishAsync(s_sessionUpdateChannel, sessionId, CommandFlags.FireAndForget);
 			}
diff --git a/Engine/Source/Programs/Horde/Plugins/Compute/HordeServer.Compute/Agents/AgentService.cs b/Engine/Source/Programs/Horde/Plugins/Compute/HordeServer.Compute/Agents/AgentService.cs
--- a/Engine/Source/Programs/Horde/Plugins/Compute/HordeServer.Compute/Agents/AgentService.cs
+++ b/Engine/Source/Programs/Horde/Plugins/Compute/HordeServer.Compute/Agents/AgentService.cs
@@ -613,6 +613,7 @@ namespace HordeServer.Agents
 			agent = await UpdateSessionAsync(agent, sessionId, status, capabilities, newLeases, cancellationToken);
 			if (agent != null && agent.UpdateIndex == updateIndex && (agent.Leases.Count > 0 || agent.Status != AgentStatus.Stopping))
 			{
+				_logger.LogDebug("Waiting for lease update on agent {AgentId} session {SessionId}", agent.Id, sessionId);
 				agent = await WaitForLeaseAsync(agent, newLeases, cancellationToken);
 			}
 			return agent;
