Horde: Fix static analysis warning.

#rnx

diff --git a/Engine/Source/Programs/Horde/Plugins/Compute/HordeServer.Compute/Agents/AgentCollection.cs b/Engine/Source/Programs/Horde/Plugins/Compute/HordeServer.Compute/Agents/AgentCollection.cs
--- a/Engine/Source/Programs/Horde/Plugins/Compute/HordeServer.Compute/Agents/AgentCollection.cs
+++ b/Engine/Source/Programs/Horde/Plugins/Compute/HordeServer.Compute/Agents/AgentCollection.cs
@@ -1244,7 +1244,7 @@ namespace HordeServer.Agents
 		{
 			if (_sessionIdToTcs.TryGetValue(sessionId, out TaskCompletionSource? tcs))
 			{
-				_logger.LogDebug("Session {SessionId} was updated.");
+				_logger.LogDebug("Session {SessionId} was updated.", sessionId);
 				tcs.TrySetResult();
 			}
 		}
