Horde: Fix exception creating agent responses.

#rnx

diff --git a/Engine/Source/Programs/Horde/Plugins/Compute/HordeServer.Compute/Agents/AgentsController.cs b/Engine/Source/Programs/Horde/Plugins/Compute/HordeServer.Compute/Agents/AgentsController.cs
--- a/Engine/Source/Programs/Horde/Plugins/Compute/HordeServer.Compute/Agents/AgentsController.cs
+++ b/Engine/Source/Programs/Horde/Plugins/Compute/HordeServer.Compute/Agents/AgentsController.cs
@@ -141,7 +141,7 @@ namespace HordeServer.Agents
 			}
 
 			List<GetAgentLeaseResponse> leases = new List<GetAgentLeaseResponse>();
-			foreach (AgentLease agentLease in agent.Leases)
+			foreach (IAgentLease agentLease in agent.Leases)
 			{
 				try
 				{
