Horde: Rename Horde.Agent.TrayApp project to HordeTrayApp.

#rnx

diff --git a/Engine/Source/Programs/Horde/BuildHorde.xml b/Engine/Source/Programs/Horde/BuildHorde.xml
--- a/Engine/Source/Programs/Horde/BuildHorde.xml
+++ b/Engine/Source/Programs/Horde/BuildHorde.xml
@@ -148,7 +148,7 @@
 			<Sign Files="$(ServerToolsStagingDir_HordeAgentWin64)/JobDriver/JobDriver.exe;$(ServerToolsStagingDir_HordeAgentWin64)/JobDriver/JobDriver.dll;$(ServerToolsStagingDir_HordeAgentWin64)/JobDriver/EpicGames.*.dll" If="'$(SignOutputs)' == 'true'"/>
 
 			<!-- Build the tray app -->
-			<DotNet Arguments="publish &quot;$(RootDir)/Engine/Source/Programs/Horde/Horde.Agent.TrayApp/Horde.Agent.TrayApp.csproj&quot; --output &quot;$(ServerToolsStagingDir_HordeAgentWin64)/TrayApp&quot; --runtime win7-x64 --self-contained $(VersionArguments)"/>
+			<DotNet Arguments="publish &quot;$(RootDir)/Engine/Source/Programs/Horde/HordeTrayApp/HordeTrayApp.csproj&quot; --output &quot;$(ServerToolsStagingDir_HordeAgentWin64)/TrayApp&quot; --runtime win7-x64 --self-contained $(VersionArguments)"/>
 			<Sign Files="$(ServerToolsStagingDir_HordeAgentWin64)/TrayApp/Horde.Agent.TrayApp.exe;$(ServerToolsStagingDir_HordeAgentWin64)/TrayApp/EpicGames.*.dll" If="'$(SignOutputs)' == 'true'"/>
 
 			<Delete Files="$(ServerToolsStagingDir_HordeAgentWin64)/appsettings.Local.json"/>
diff --git a/Engine/Source/Programs/Horde/Horde.Agent.TrayApp/Forms/IdleForm.cs b/Engine/Source/Programs/Horde/Horde.Agent.TrayApp/Forms/IdleForm.cs
deleted file mode 100644
diff --git a/Engine/Source/Programs/Horde/Horde.Agent.TrayApp/Horde.Agent.TrayApp.csproj b/Engine/Source/Programs/Horde/Horde.Agent.TrayApp/Horde.Agent.TrayApp.csproj
deleted file mode 100644
diff --git a/Engine/Source/Programs/Horde/Horde.Agent.TrayApp/Program.cs b/Engine/Source/Programs/Horde/Horde.Agent.TrayApp/Program.cs
deleted file mode 100644
diff --git a/Engine/Source/Programs/Horde/Horde.Agent.TrayApp/Settings.cs b/Engine/Source/Programs/Horde/Horde.Agent.TrayApp/Settings.cs
deleted file mode 100644
diff --git a/Engine/Source/Programs/Horde/HordeTrayApp/Forms/IdleForm.cs b/Engine/Source/Programs/Horde/HordeTrayApp/Forms/IdleForm.cs
new file mode 100644
--- /dev/null
+++ b/Engine/Source/Programs/Horde/HordeTrayApp/Forms/IdleForm.cs
@@ -0,0 +1,72 @@
+// Copyright Epic Games, Inc. All Rights Reserved.
+
+namespace Horde.Agent.TrayApp.Forms
+{
+	partial class IdleForm : Form
+	{
+		readonly Dictionary<string, ListViewItem> _items = new Dictionary<string, ListViewItem>(StringComparer.Ordinal);
+
+		public IdleForm()
+		{
+			InitializeComponent();
+
+			Rectangle rect = Screen.GetWorkingArea(this);
+			Location = new Point(rect.X + rect.Width - Size.Width - 20, rect.Y + rect.Height - Size.Height - 20);
+
+			float scale = DeviceDpi / 96.0f;
+			nameColumnHeader.Width = (int)(nameColumnHeader.Width * scale);
+			valueColumnHeader.Width = (int)(valueColumnHeader.Width * scale);
+			minValueColumnHeader.Width = (int)(minValueColumnHeader.Width * scale);
+			resultColumnHeader.Width = (int)(resultColumnHeader.Width * scale);
+		}
+
+		public void TickStats(bool enabled, int stateChangeTime, int stateChangeMaxTime, IEnumerable<IdleStat> stats)
+		{
+			BeginInvoke(() => TickStatsMainThread(enabled, stateChangeTime, stateChangeMaxTime, stats));
+		}
+
+		void TickStatsMainThread(bool enabled, int stateChangeTime, int stateChangeMaxTime, IEnumerable<IdleStat> stats)
+		{
+			if (IsDisposed)
+			{
+				return;
+			}
+
+			string activeText = enabled ? "Enabled" : "Paused";
+			if (stateChangeTime == 0)
+			{
+				stateTextBox.Text = $"{activeText}";
+			}
+			else
+			{
+				stateTextBox.Text = $"{activeText} (state change in {Math.Max(stateChangeMaxTime - stateChangeTime, 0)}s)";
+			}
+
+			HashSet<string> removeItems = new HashSet<string>(_items.Keys, StringComparer.Ordinal);
+			foreach (IdleStat stat in stats)
+			{
+				ListViewItem? item;
+				if (!_items.TryGetValue(stat.Name, out item))
+				{
+					item = statsListView.Items.Add(stat.Name);
+					item.SubItems.Add("");
+					item.SubItems.Add("");
+					item.SubItems.Add("");
+					_items.Add(stat.Name, item);
+				}
+
+				item.SubItems[1].Text = stat.Value.ToString();
+				item.SubItems[2].Text = stat.MinValue.ToString();
+				item.SubItems[3].Text = (stat.Value < stat.MinValue) ? "" : "Idle";
+
+				removeItems.Remove(stat.Name);
+			}
+
+			foreach (string removeItem in removeItems)
+			{
+				statsListView.Items.Remove(_items[removeItem]);
+				_items.Remove(removeItem);
+			}
+		}
+	}
+}
diff --git a/Engine/Source/Programs/Horde/HordeTrayApp/HordeTrayApp.csproj b/Engine/Source/Programs/Horde/HordeTrayApp/HordeTrayApp.csproj
new file mode 100644
--- /dev/null
+++ b/Engine/Source/Programs/Horde/HordeTrayApp/HordeTrayApp.csproj
@@ -0,0 +1,45 @@
+﻿<Project Sdk="Microsoft.NET.Sdk">
+
+  <PropertyGroup>
+    <OutputType>WinExe</OutputType>
+    <TargetFramework>net6.0-windows</TargetFramework>
+    <Nullable>enable</Nullable>
+    <ImplicitUsings>enable</ImplicitUsings>
+    <UseWindowsForms>true</UseWindowsForms>
+    <NeutralLanguage>en-US</NeutralLanguage>
+    <Configurations>Debug;Release;Analyze</Configurations>
+	<AssemblyName>Horde.Agent.TrayApp</AssemblyName>
+  </PropertyGroup>
+
+  <PropertyGroup Condition=" '$(Configuration)|$(Platform)' == 'Analyze|AnyCPU' ">
+    <RunAnalyzersDuringBuild>True</RunAnalyzersDuringBuild>
+    <TreatWarningsAsErrors>false</TreatWarningsAsErrors>
+    <EnforceCodeStyleInBuild>true</EnforceCodeStyleInBuild>
+    <GenerateDocumentationFile>True</GenerateDocumentationFile>
+    <DocumentationFile></DocumentationFile>
+  </PropertyGroup>
+
+  <ItemGroup>
+    <Compile Include="..\HordeAgent\IpcMessages.cs" Link="IpcMessages.cs" />
+  </ItemGroup>
+
+  <ItemGroup>
+    <ProjectReference Include="..\..\Shared\EpicGames.Core\EpicGames.Core.csproj" />
+  </ItemGroup>
+
+  <ItemGroup>
+    <Compile Update="Properties\Resources.Designer.cs">
+      <DesignTime>True</DesignTime>
+      <AutoGen>True</AutoGen>
+      <DependentUpon>Resources.resx</DependentUpon>
+    </Compile>
+  </ItemGroup>
+
+  <ItemGroup>
+    <EmbeddedResource Update="Properties\Resources.resx">
+      <Generator>ResXFileCodeGenerator</Generator>
+      <LastGenOutput>Resources.Designer.cs</LastGenOutput>
+    </EmbeddedResource>
+  </ItemGroup>
+
+</Project>
\ No newline at end of file
diff --git a/Engine/Source/Programs/Horde/HordeTrayApp/Program.cs b/Engine/Source/Programs/Horde/HordeTrayApp/Program.cs
new file mode 100644
--- /dev/null
+++ b/Engine/Source/Programs/Horde/HordeTrayApp/Program.cs
@@ -0,0 +1,687 @@
+// Copyright Epic Games, Inc. All Rights Reserved.
+
+using System.Diagnostics;
+using System.IO.Pipes;
+using System.Reflection;
+using System.Runtime.InteropServices;
+using System.Text.Json;
+using EpicGames.Core;
+using HordeAgent;
+using Horde.Agent.TrayApp.Forms;
+using Horde.Agent.TrayApp.Properties;
+using Microsoft.Win32;
+
+namespace Horde.Agent.TrayApp
+{
+	record struct IdleStat(string Name, long Value, long MinValue);
+
+	static class Program
+	{
+		const string MutexName = "Horde.Agent.TrayApp-Mutex";
+		const string EventName = "Horde.Agent.TrayApp-Exit";
+
+		[STAThread]
+		public static int Main(string[] args)
+		{
+			using EventWaitHandle closeEvent = new EventWaitHandle(false, EventResetMode.AutoReset, EventName);
+			using SingleInstanceMutex mutex = new SingleInstanceMutex(MutexName);
+
+			if (args.Any(x => x.Equals("-close", StringComparison.OrdinalIgnoreCase)))
+			{
+				closeEvent.Set();
+				return mutex.Wait(5000) ? 0 : 1;
+			}
+
+			if (!mutex.Wait(0))
+			{
+				return 1;
+			}
+
+			if (args.Any(x => x.Equals("-shadowcopy", StringComparison.OrdinalIgnoreCase)))
+			{
+				DirectoryReference? localAppData = DirectoryReference.GetSpecialFolder(Environment.SpecialFolder.LocalApplicationData);
+				if (localAppData != null)
+				{
+					FileReference sourceExe = new FileReference(Assembly.GetExecutingAssembly().Location).ChangeExtension(".exe");
+					if (FileReference.Exists(sourceExe) && !sourceExe.IsUnderDirectory(localAppData))
+					{
+						DirectoryReference sourceDir = sourceExe.Directory;
+						DirectoryReference targetDir = DirectoryReference.Combine(localAppData, "Epic Games", "HordeTrayApp");
+						try
+						{
+							DirectoryReference.CreateDirectory(targetDir);
+							FileUtils.ForceDeleteDirectoryContents(targetDir);
+
+							CopyFiles(sourceDir, targetDir);
+							mutex.Release();
+
+							FileReference targetExe = FileReference.Combine(targetDir, sourceExe.MakeRelativeTo(sourceDir)).ChangeExtension(".exe");
+							using Process process = Process.Start(targetExe.FullName, args);
+
+							return 0;
+						}
+						catch (Exception ex)
+						{
+							MessageBox.Show($"Unable to copy app to temp location. Error:\n\n{ex}");
+							return 1;
+						}
+					}
+				}
+			}
+
+			MainAsync(closeEvent).GetAwaiter().GetResult();
+			return 0;
+		}
+
+		static void CopyFiles(DirectoryReference sourceDir, DirectoryReference targetDir)
+		{
+			DirectoryReference.CreateDirectory(targetDir);
+			foreach (DirectoryReference sourceSubDir in DirectoryReference.EnumerateDirectories(sourceDir))
+			{
+				DirectoryReference targetSubDir = DirectoryReference.Combine(targetDir, sourceSubDir.GetDirectoryName());
+				CopyFiles(sourceSubDir, targetSubDir);
+			}
+			foreach (FileReference sourceFile in DirectoryReference.EnumerateFiles(sourceDir))
+			{
+				FileReference targetFile = FileReference.Combine(targetDir, sourceFile.GetFileName());
+				FileReference.Copy(sourceFile, targetFile, true);
+			}
+		}
+
+		static async Task MainAsync(EventWaitHandle closeEvent)
+		{
+			ApplicationConfiguration.Initialize();
+
+			await using (CustomApplicationContext appContext = new CustomApplicationContext(closeEvent))
+			{
+				Application.Run(appContext);
+			}
+		}
+	}
+
+	class SingleInstanceMutex : IDisposable
+	{
+		readonly Mutex _mutex;
+		bool _locked;
+
+		public SingleInstanceMutex(string name)
+		{
+			_mutex = new Mutex(true, name);
+		}
+
+		public void Release()
+		{
+			if (_locked)
+			{
+				_mutex.ReleaseMutex();
+				_locked = false;
+			}
+		}
+
+		public bool Wait(int timeout)
+		{
+			if (!_locked)
+			{
+				try
+				{
+					_locked = _mutex.WaitOne(timeout);
+				}
+				catch (AbandonedMutexException)
+				{
+					_locked = true;
+				}
+			}
+			return _locked;
+		}
+
+		public void Dispose()
+		{
+			Release();
+			_mutex.Dispose();
+		}
+	}
+
+	class CustomApplicationContext : ApplicationContext, IAsyncDisposable
+	{
+		readonly NotifyIcon _trayIcon;
+		readonly BackgroundTask _clientTask;
+		readonly BackgroundTask _tickPauseStateTask;
+		readonly BackgroundTask _waitForExitTask;
+		readonly Control _mainThreadInvokeTarget;
+
+		readonly ToolStripMenuItem _enrollMenuItem;
+		readonly ToolStripMenuItem _statusEnabled;
+		readonly ToolStripMenuItem _statusDisabled;
+		readonly ToolStripMenuItem _statusWhenIdle;
+
+		readonly Settings _settings;
+
+		AgentSettingsMessage? _agentSettings;
+		IdleForm? _idleForm;
+		bool _disposed;
+
+		void EnrollWithServer()
+		{
+			Uri? serverUrl = _agentSettings?.ServerUrl;
+			if (serverUrl != null)
+			{
+				Process.Start(new ProcessStartInfo(new Uri(serverUrl, "agents/registration").ToString()) { UseShellExecute = true });
+			}
+		}
+
+		public CustomApplicationContext(EventWaitHandle eventHandle)
+		{
+			_settings = LoadSettings();
+
+			_enrollMenuItem = new ToolStripMenuItem("Enroll with Server...");
+			_enrollMenuItem.Click += (s, e) => EnrollWithServer();
+
+			_statusEnabled = new ToolStripMenuItem("Enabled");
+			_statusEnabled.Click += (s, e) => SetUserStatus(UserStatus.Enabled);
+
+			_statusDisabled = new ToolStripMenuItem("Disabled");
+			_statusDisabled.Click += (s, e) => SetUserStatus(UserStatus.Disabled);
+
+			_statusWhenIdle = new ToolStripMenuItem("When Idle");
+			_statusWhenIdle.Click += (s, e) => SetUserStatus(UserStatus.WhenIdle);
+
+			ToolStripMenuItem showStatsMenuItem = new ToolStripMenuItem("Stats...");
+			showStatsMenuItem.Click += Status_Stats_OnClick;
+
+			ToolStripMenuItem statusMenuItem = new ToolStripMenuItem("Status");
+			statusMenuItem.DropDownItems.Add(_statusEnabled);
+			statusMenuItem.DropDownItems.Add(_statusDisabled);
+			statusMenuItem.DropDownItems.Add(_statusWhenIdle);
+			statusMenuItem.DropDownItems.Add(new ToolStripSeparator());
+			statusMenuItem.DropDownItems.Add(showStatsMenuItem);
+
+			ToolStripMenuItem logsMenuItem = new ToolStripMenuItem("Open logs dir");
+			logsMenuItem.Click += OnOpenLogs;
+
+			ToolStripMenuItem exitMenuItem = new ToolStripMenuItem("Exit");
+			exitMenuItem.Click += OnExit;
+
+			ContextMenuStrip menu = new ContextMenuStrip();
+			menu.Items.Add(_enrollMenuItem);
+			menu.Items.Add(statusMenuItem);
+			menu.Items.Add(new ToolStripSeparator());
+			menu.Items.Add(logsMenuItem);
+			menu.Items.Add(new ToolStripSeparator());
+			menu.Items.Add(exitMenuItem);
+
+			_mainThreadInvokeTarget = new Control();
+			_mainThreadInvokeTarget.CreateControl();
+
+			_trayIcon = new NotifyIcon()
+			{
+				Icon = Resources.StatusNormal,
+				ContextMenuStrip = menu,
+				Visible = true
+			};
+			_trayIcon.Click += TrayIcon_Click;
+
+			_clientTask = BackgroundTask.StartNew(StatusTaskAsync);
+			_tickPauseStateTask = BackgroundTask.StartNew(ctx => TickPauseStateAsync(ctx));
+			_waitForExitTask = BackgroundTask.StartNew(ctx => WaitForExitAsync(eventHandle, ctx));
+		}
+
+		private static Settings LoadSettings()
+		{
+			Settings? result = null;
+
+			DirectoryReference? settingsRoot = DirectoryReference.GetSpecialFolder(Environment.SpecialFolder.CommonApplicationData);
+			if (settingsRoot != null)
+			{
+				FileReference settingsPath = FileReference.Combine(settingsRoot, "Epic", "Horde", "TrayApp", "Settings.json");
+				if (FileReference.Exists(settingsPath))
+				{
+					try
+					{
+						using FileStream stream = FileReference.Open(settingsPath, FileMode.Open, FileAccess.Read);
+						result = JsonSerializer.Deserialize<Settings>(stream);
+					}
+					catch (Exception)
+					{
+					}
+				}
+				else
+				{
+					// File not found, create a file containing the default settings
+					DirectoryReference.CreateDirectory(settingsPath.Directory);
+					using FileStream stream = FileReference.Open(settingsPath, FileMode.OpenOrCreate, FileAccess.Write);
+					JsonSerializer.Serialize(stream, new Settings(), new JsonSerializerOptions() { WriteIndented = true });
+				}
+			}
+
+			return result ?? new Settings();
+		}
+
+		private void TrayIcon_Click(object? sender, EventArgs e)
+		{
+			UserStatus status = GetUserStatus();
+			_statusEnabled.Checked = (status == UserStatus.Enabled);
+			_statusDisabled.Checked = (status == UserStatus.Disabled);
+			_statusWhenIdle.Checked = (status == UserStatus.WhenIdle);
+		}
+
+		protected override void Dispose(bool disposing)
+		{
+			base.Dispose(disposing);
+
+			if (disposing)
+			{
+				if (_idleForm != null)
+				{
+					_idleForm.Dispose();
+					_idleForm = null;
+				}
+
+				_enrollMenuItem.Dispose();
+				_statusEnabled.Dispose();
+				_statusDisabled.Dispose();
+				_statusWhenIdle.Dispose();
+
+				_mainThreadInvokeTarget.Dispose();
+				_trayIcon.Dispose();
+				_disposed = true;
+			}
+		}
+
+		public async ValueTask DisposeAsync()
+		{
+			await _waitForExitTask.DisposeAsync();
+			await _tickPauseStateTask.DisposeAsync();
+			await _clientTask.DisposeAsync();
+
+			Dispose();
+			GC.SuppressFinalize(this);
+		}
+
+		private void OnOpenLogs(object? sender, EventArgs e)
+		{
+			DirectoryReference? programDataDir = DirectoryReference.GetSpecialFolder(Environment.SpecialFolder.CommonApplicationData);
+			if (programDataDir != null)
+			{
+				DirectoryReference logsDir = DirectoryReference.Combine(programDataDir, "Epic", "Horde", "Agent");
+				if (DirectoryReference.Exists(logsDir))
+				{
+					Process.Start(new ProcessStartInfo { FileName = logsDir.FullName, UseShellExecute = true });
+				}
+				else
+				{
+					MessageBox.Show("Unable to open logs dir " + logsDir.FullName, "Horde Tray App", MessageBoxButtons.OK, MessageBoxIcon.Error);
+				}
+			}
+		}
+
+		enum UserStatus
+		{
+			Enabled = 0,
+			Disabled = 1,
+			WhenIdle = 2,
+		}
+
+		const string RegistryKey = "HKEY_CURRENT_USER\\Software\\Epic Games\\Horde\\TrayApp";
+		const string RegistryStatusValue = "Status";
+
+		private static UserStatus GetUserStatus()
+		{
+			return (UserStatus)((Registry.GetValue(RegistryKey, RegistryStatusValue, null) as int?) ?? 0);
+		}
+
+		private void SetUserStatus(UserStatus status)
+		{
+			Registry.SetValue(RegistryKey, RegistryStatusValue, (int)status);
+			_statusChangedEvent.Set();
+		}
+
+		private void Status_Stats_OnClick(object? sender, EventArgs e)
+		{
+			if (_idleForm == null)
+			{
+				_idleForm = new IdleForm();
+				_idleForm.FormClosed += (s, e) =>
+				{
+					_idleForm.Dispose();
+					_idleForm = null;
+				};
+				_idleForm.Show();
+			}
+		}
+
+		private void OnExit(object? sender, EventArgs e)
+		{
+			ExitThread();
+		}
+
+		void SetStatus(AgentStatusMessage status)
+		{
+			_mainThreadInvokeTarget.BeginInvoke(() => SetStatus_MainThread(status));
+		}
+
+		void SetStatus_MainThread(AgentStatusMessage status)
+		{
+			if (!_disposed)
+			{
+				if (!status.Healthy)
+				{
+					_trayIcon.Icon = Resources.StatusError;
+					_trayIcon.Text = String.IsNullOrEmpty(status.Detail) ? "Error. Check logs." : (status.Detail.Length > 100) ? status.Detail.Substring(0, 100) : status.Detail;
+				}
+				else if (status.NumLeases > 0)
+				{
+					_trayIcon.Icon = Resources.StatusBusy;
+					_trayIcon.Text = (status.NumLeases == 1) ? "Currently handling 1 lease" : $"Currently handling {status.NumLeases} leases";
+				}
+				else if (_enabled)
+				{
+					_trayIcon.Icon = Resources.StatusNormal;
+					_trayIcon.Text = "Agent is operating normally";
+				}
+				else
+				{
+					_trayIcon.Icon = Resources.StatusPaused;
+					_trayIcon.Text = "Agent is paused";
+				}
+			}
+		}
+
+		async Task StatusTaskAsync(CancellationToken cancellationToken)
+		{
+			SetStatus(AgentStatusMessage.Starting);
+			for (; ; )
+			{
+				try
+				{
+					await PollForStatusUpdatesAsync(cancellationToken);
+				}
+				catch (OperationCanceledException) when (cancellationToken.IsCancellationRequested)
+				{
+					break;
+				}
+				catch
+				{
+					SetStatus(new AgentStatusMessage(false, 0, "Unable to connect to Agent. Check the service is running."));
+					await Task.Delay(TimeSpan.FromSeconds(5.0), cancellationToken);
+				}
+			}
+		}
+
+		async Task WaitForExitAsync(EventWaitHandle eventHandle, CancellationToken cancellationToken)
+		{
+			await eventHandle.WaitOneAsync(cancellationToken);
+			_mainThreadInvokeTarget.BeginInvoke(() => Exit_MainThread());
+		}
+
+#pragma warning disable IDE1006
+		[StructLayout(LayoutKind.Sequential)]
+		struct LASTINPUTINFO
+		{
+			public int cbSize;
+			public uint dwTime;
+		}
+
+		[DllImport("user32.dll")]
+		static extern bool GetLastInputInfo(ref LASTINPUTINFO plii);
+
+		[DllImport("kernel32.dll")]
+		static extern uint GetTickCount();
+
+		[StructLayout(LayoutKind.Sequential)]
+		struct FILETIME
+		{
+			public uint dwLowDateTime;
+			public uint dwHighDateTime;
+
+			public readonly ulong Total => dwLowDateTime | ((ulong)dwHighDateTime << 32);
+		};
+
+		[StructLayout(LayoutKind.Sequential)]
+		struct MEMORYSTATUSEX
+		{
+			public int dwLength;
+			public uint dwMemoryLoad;
+			public ulong ullTotalPhys;
+			public ulong ullAvailPhys;
+			public ulong ullTotalPageFile;
+			public ulong ullAvailPageFile;
+			public ulong ullTotalVirtual;
+			public ulong ullAvailVirtual;
+			public ulong ullAvailExtendedVirtual;
+		}
+
+		[DllImport("kernel32.dll", CharSet = CharSet.Auto, SetLastError = true)]
+		static extern bool GlobalMemoryStatusEx(ref MEMORYSTATUSEX lpBuffer);
+
+		[DllImport("kernel32.dll", SetLastError = true)]
+		static extern bool GetSystemTimes(out FILETIME lpIdleTime, out FILETIME lpKernelTime, out FILETIME lpUserTime);
+#pragma warning restore IDE1006
+
+		bool _enabled;
+		readonly AsyncEvent _statusChangedEvent = new AsyncEvent();
+		readonly AsyncEvent _enabledChangedEvent = new AsyncEvent();
+
+		async Task TickPauseStateAsync(CancellationToken cancellationToken)
+		{
+			await using BackgroundTask cpuStatsTask = BackgroundTask.StartNew(ctx => TickCpuStatsAsync(ctx));
+			await using BackgroundTask criticalProcessTask = BackgroundTask.StartNew(ctx => TickCriticalProcessAsync(ctx));
+
+			TimeSpan pollInterval = TimeSpan.FromSeconds(0.25);
+
+			Stopwatch stateChangeTimer = Stopwatch.StartNew();
+			while (!cancellationToken.IsCancellationRequested)
+			{
+				Task statusChangedTask = _statusChangedEvent.Task;
+
+				UserStatus userStatus = GetUserStatus();
+				if (userStatus == UserStatus.Enabled)
+				{
+					if (!_enabled)
+					{
+						_enabled = true;
+						_enabledChangedEvent.Set();
+					}
+				}
+				else if (userStatus == UserStatus.Disabled)
+				{
+					if (_enabled)
+					{
+						_enabled = false;
+						_enabledChangedEvent.Set();
+					}
+				}
+
+				DateTime utcNow = DateTime.UtcNow;
+				IEnumerable<IdleStat> idleStats = GetIdleStats();
+
+				bool idle = idleStats.All(x => x.Value >= x.MinValue);
+				if (idle == _enabled)
+				{
+					stateChangeTimer.Restart();
+				}
+
+				const int WakeTimeSecs = 2;
+				const int IdleTimeSecs = 30;
+				int stateChangeTime = (int)stateChangeTimer.Elapsed.TotalSeconds;
+				int stateChangeMaxTime = _enabled ? WakeTimeSecs : IdleTimeSecs;
+				_idleForm?.TickStats(_enabled, stateChangeTime, stateChangeMaxTime, idleStats);
+
+				if (userStatus == UserStatus.WhenIdle && stateChangeTime >= stateChangeMaxTime)
+				{
+					_enabled ^= true;
+					_enabledChangedEvent.Set();
+					stateChangeTimer.Restart();
+				}
+
+				await Task.WhenAny(statusChangedTask, Task.Delay(pollInterval, cancellationToken));
+			}
+		}
+
+		IEnumerable<IdleStat> GetIdleStats()
+		{
+			// Check there has been no input for a while
+			LASTINPUTINFO lastInputInfo = new LASTINPUTINFO();
+			lastInputInfo.cbSize = Marshal.SizeOf<LASTINPUTINFO>();
+
+			if (GetLastInputInfo(ref lastInputInfo))
+			{
+				yield return new IdleStat("LastInputTime", (GetTickCount() - lastInputInfo.dwTime) / 1000, _settings.Idle.MinIdleTimeSecs);
+			}
+
+			// Check that no critical processes are running
+			if (_settings.Idle.CriticalProcesses.Any())
+			{
+				yield return new IdleStat("CriticalProcCount", -_idleCriticalProcessCount, 0);
+			}
+
+			// Only look at memory/CPU usage if we're not paused; executing jobs will increase them
+			if (!_enabled)
+			{
+				// Check the CPU usage doesn't exceed the limit
+				yield return new IdleStat("IdleCpuPct", _idleCpuPct, _settings.Idle.MinIdleCpuPct);
+
+				// Check there's enough available virtual memory 
+				MEMORYSTATUSEX memoryStatus = new MEMORYSTATUSEX();
+				memoryStatus.dwLength = Marshal.SizeOf<MEMORYSTATUSEX>();
+
+				if (GlobalMemoryStatusEx(ref memoryStatus))
+				{
+					yield return new IdleStat("VirtualMemMb", (long)(memoryStatus.ullAvailPhys + memoryStatus.ullAvailPageFile) / (1024 * 1024), _settings.Idle.MinFreeVirtualMemMb);
+				}
+			}
+		}
+
+		int _idleCpuPct = 0;
+		int _idleCriticalProcessCount = 0;
+
+		async Task TickCpuStatsAsync(CancellationToken cancellationToken)
+		{
+			const int NumSamples = 10;
+			TimeSpan sampleInterval = TimeSpan.FromSeconds(0.2);
+			(ulong IdleTime, ulong TotalTime)[] samples = new (ulong IdleTime, ulong TotalTime)[NumSamples];
+
+			int sampleIdx = 0;
+			for (; ; )
+			{
+				if (GetSystemTimes(out FILETIME idleTime, out FILETIME kernelTime, out FILETIME userTime))
+				{
+					(ulong prevIdleTime, ulong prevTotalTime) = samples[sampleIdx];
+					(ulong nextIdleTime, ulong nextTotalTime) = (idleTime.Total, kernelTime.Total + userTime.Total);
+
+					samples[sampleIdx] = (nextIdleTime, nextTotalTime);
+					sampleIdx = (sampleIdx + 1) % NumSamples;
+
+					if (prevTotalTime > 0 && nextTotalTime > prevTotalTime)
+					{
+						_idleCpuPct = (int)(((nextIdleTime - prevIdleTime) * 100) / (nextTotalTime - prevTotalTime));
+					}
+				}
+				await Task.Delay(sampleInterval, cancellationToken);
+			}
+		}
+
+		async Task TickCriticalProcessAsync(CancellationToken cancellationToken)
+		{
+			TimeSpan sampleInterval = TimeSpan.FromSeconds(1.0);
+
+			for (; ; )
+			{
+				try
+				{
+					if (RuntimeInformation.IsOSPlatform(OSPlatform.Windows) && _settings.Idle.CriticalProcesses.Any())
+					{
+						IEnumerable<int> hordeProcessIds = Process.GetProcessesByName("HordeAgent").Select(x => x.Id);
+						IEnumerable<Process> criticalProcesses = _settings.Idle.CriticalProcesses
+							.Select(x => Path.GetFileNameWithoutExtension(x).ToUpperInvariant())
+							.Distinct()
+							.SelectMany(x => Process.GetProcessesByName(x))
+							.Where(x => !x.HasExited);
+
+						// Ignore processes that are descendants of HordeAgent
+						if (hordeProcessIds.Any() && criticalProcesses.Any())
+						{
+							criticalProcesses = criticalProcesses
+								.Where(x => !ProcessUtils.GetAncestorProcesses(x)
+								.Select(x => x.Id).Intersect(hordeProcessIds).Any())
+								.Where(x => !x.HasExited);
+						}
+
+						_idleCriticalProcessCount = criticalProcesses.Count();
+					}
+				}
+				catch (InvalidOperationException)
+				{
+					// If a process stops running Process.Id will throw an exception
+				}
+				await Task.Delay(sampleInterval, cancellationToken);
+			}
+		}
+
+		async Task PollForStatusUpdatesAsync(CancellationToken cancellationToken)
+		{
+			AgentMessageBuffer message = new AgentMessageBuffer();
+			using (NamedPipeClientStream pipeClient = new NamedPipeClientStream(".", AgentMessagePipe.PipeName, PipeDirection.InOut))
+			{
+				SetStatus(new AgentStatusMessage(false, 0, "Connecting to agent..."));
+				await pipeClient.ConnectAsync(cancellationToken);
+
+				SetStatus(new AgentStatusMessage(false, 0, "Waiting for status update."));
+				for (; ; )
+				{
+					Task idleChangeTask = _enabledChangedEvent.Task;
+
+					bool enabled = _enabled;
+					message.Set(AgentMessageType.SetEnabledRequest, new AgentEnabledMessage(enabled));
+					await message.SendAsync(pipeClient, cancellationToken);
+
+					if (_agentSettings == null)
+					{
+						message.Set(AgentMessageType.GetSettingsRequest);
+						await message.SendAsync(pipeClient, cancellationToken);
+
+						if (!await message.TryReadAsync(pipeClient, cancellationToken))
+						{
+							break;
+						}
+						if (message.Type == AgentMessageType.GetSettingsResponse)
+						{
+							_agentSettings = message.Parse<AgentSettingsMessage>();
+						}
+					}
+
+					message.Set(AgentMessageType.GetStatusRequest);
+					await message.SendAsync(pipeClient, cancellationToken);
+
+					if (!await message.TryReadAsync(pipeClient, cancellationToken))
+					{
+						break;
+					}
+
+					switch (message.Type)
+					{
+						case AgentMessageType.GetStatusResponse:
+							AgentStatusMessage status = message.Parse<AgentStatusMessage>();
+							SetStatus(status);
+							break;
+					}
+
+					await Task.WhenAny(idleChangeTask, Task.Delay(TimeSpan.FromSeconds(5.0), cancellationToken));
+				}
+			}
+		}
+
+		void Exit_MainThread()
+		{
+			if (!_disposed)
+			{
+				_trayIcon.Visible = false;
+				Application.Exit();
+			}
+		}
+
+		void Exit(object sender, EventArgs e)
+		{
+			Exit_MainThread();
+		}
+	}
+}
diff --git a/Engine/Source/Programs/Horde/HordeTrayApp/Settings.cs b/Engine/Source/Programs/Horde/HordeTrayApp/Settings.cs
new file mode 100644
--- /dev/null
+++ b/Engine/Source/Programs/Horde/HordeTrayApp/Settings.cs
@@ -0,0 +1,22 @@
+// Copyright Epic Games, Inc. All Rights Reserved.
+
+#pragma warning disable CS1591
+
+namespace Horde.Agent.TrayApp
+{
+	[Serializable]
+	public class Settings
+	{
+		public IdleSettings Idle { get; set; } = new();
+	}
+
+	[Serializable]
+	public class IdleSettings
+	{
+		public int MinIdleTimeSecs { get; set; } = 2;
+		public int MinIdleCpuPct { get; set; } = 70;
+		public int MinFreeVirtualMemMb { get; set; } = 256;
+
+		public string[] CriticalProcesses { get; set; } = Array.Empty<string>();
+	}
+}
