Horde: Updating schema documentation

#rnx

diff --git a/Engine/Source/Programs/Horde/Docs/Deployment/ServerSettings.md b/Engine/Source/Programs/Horde/Docs/Deployment/ServerSettings.md
--- a/Engine/Source/Programs/Horde/Docs/Deployment/ServerSettings.md
+++ b/Engine/Source/Programs/Horde/Docs/Deployment/ServerSettings.md
@@ -22,14 +22,15 @@ Name | Description
 `databaseReadOnlyMode` | `boolean`<br>
 `shutdownMemoryThreshold` | `integer`<br>Shutdown the current server process if memory usage reaches this threshold (specified in MB)<br>Usually set to 80-90% of available memory to avoid CLR heap using all of it. If a memory leak was to occur, it's usually better to restart the process rather than to let the GC work harder and harder trying to recoup memory.<br>Should only be used when multiple server processes are running behind a load balancer and one can be safely restarted automatically by the underlying process handler (Docker, Kubernetes, AWS ECS, Supervisor etc). The shutdown behaves similar to receiving a SIGTERM and will wait for outstanding requests to finish.
 `serverPrivateCert` | `string`<br>Optional PFX certificate to use for encrypting agent SSL traffic. This can be a self-signed certificate, as long as it's trusted by agents.
-`authMethod` | [AuthMethod](#authmethod-enum)<br>Issuer for tokens from the auth provider
+`authMethod` | [AuthMethod](#authmethod-enum)<br>Type of authentication (e.g anonymous, OIDC, built-in Horde accounts)
 `oidcProfileName` | `string`<br>Optional profile name to report through the /api/v1/server/auth endpoint. Allows sharing auth tokens between providers configured through the same profile name in OidcToken.exe config files.
-`oidcAuthority` | `string`<br>Issuer for tokens from the auth provider
-`oidcAudience` | `string`<br>Audience for validating externally issued tokens
-`oidcClientId` | `string`<br>Client id for the OIDC authority
+`oidcAuthority` | `string`<br>OpenID Connect (OIDC) authority URL (required when OIDC is enabled)
+`oidcAudience` | `string`<br>Audience for validating externally issued tokens (required when OIDC is enabled)
+`oidcClientId` | `string`<br>Client ID for the OIDC authority (required when OIDC is enabled)
 `oidcClientSecret` | `string`<br>Client secret for the OIDC authority
 `oidcSigninRedirect` | `string`<br>Optional redirect url provided to OIDC login
-`oidcLocalRedirectUrls` | `string[]`<br>Optional redirect url provided to OIDC login for external tools (typically to a local server)
+`oidcLocalRedirectUrls` | `string[]`<br>Optional redirect url provided to OIDC login for external tools (typically to a local server) Default value is the local web server started during signin by EpicGames.OIDC library
+`oidcDebugMode` | `boolean`<br>Debug mode for OIDC which logs reasons for why JWT tokens fail to authenticate Also turns off HTTPS requirement for OIDC metadata fetching. NOT FOR PRODUCTION USE!
 `oidcRequestedScopes` | `string[]`<br>OpenID Connect scopes to request when signing in
 `oidcClaimNameMapping` | `string[]`<br>List of fields in /userinfo endpoint to try map to the standard name claim (see System.Security.Claims.ClaimTypes.Name)
 `oidcClaimEmailMapping` | `string[]`<br>List of fields in /userinfo endpoint to try map to the standard email claim (see System.Security.Claims.ClaimTypes.Email)
