Horde: Allow creating long-lived filters for sessions matching certain agent requirements, making it more efficient to find agents matching certain conditions.

#rnx

diff --git a/Engine/Source/Programs/Horde/Plugins/Compute/HordeServer.Compute.Tests/Agents/AgentSchedulerTests.cs b/Engine/Source/Programs/Horde/Plugins/Compute/HordeServer.Compute.Tests/Agents/AgentSchedulerTests.cs
--- a/Engine/Source/Programs/Horde/Plugins/Compute/HordeServer.Compute.Tests/Agents/AgentSchedulerTests.cs
+++ b/Engine/Source/Programs/Horde/Plugins/Compute/HordeServer.Compute.Tests/Agents/AgentSchedulerTests.cs
@@ -2,9 +2,11 @@
 
 using EpicGames.Core;
 using EpicGames.Horde.Agents;
+using EpicGames.Horde.Agents.Sessions;
 using HordeCommon;
 using HordeCommon.Rpc;
 using HordeCommon.Rpc.Messages;
+using HordeServer.Agents;
 using HordeServer.Agents.Sessions;
 using Microsoft.Extensions.DependencyInjection;
 
@@ -79,5 +81,133 @@ namespace HordeServer.Tests.Agents
 			RpcSession? session3 = await AgentScheduler.TryGetSessionAsync(session.SessionId);
 			Assert.IsNull(session3);
 		}
+
+		[TestMethod]
+		public async Task AddSessionsToNewQueueTestAsync()
+		{
+			AgentScheduler scheduler = (AgentScheduler)AgentScheduler;
+			await scheduler.StartAsync(default);
+
+			// First session
+			for (int idx = 0; idx < 3; idx++)
+			{
+				RpcAgentCapabilities caps1 = new RpcAgentCapabilities();
+				caps1.Properties.Add("Platform=Windows");
+
+				RpcSession? session1 = await AgentScheduler.TryCreateSessionAsync(new AgentId("foo"), SessionIdUtils.GenerateNewId(), caps1);
+				Assert.IsNotNull(session1);
+			}
+
+			// Second session
+			for (int idx = 0; idx < 5; idx++)
+			{
+				RpcAgentCapabilities caps2 = new RpcAgentCapabilities();
+				caps2.Properties.Add("Platform=MacOS");
+
+				RpcSession? session2 = await AgentScheduler.TryCreateSessionAsync(new AgentId("bar"), SessionIdUtils.GenerateNewId(), caps2);
+				Assert.IsNotNull(session2);
+			}
+
+			// Create the first queue
+			RpcAgentRequirements reqs1 = new RpcAgentRequirements();
+			reqs1.Properties.Add("Platform=Windows");
+
+			IoHash queueHash1 = await scheduler.CreateFilterAsync(reqs1);
+
+			// Create the second queue
+			RpcAgentRequirements reqs2 = new RpcAgentRequirements();
+			reqs2.Properties.Add("Platform=MacOS");
+
+			IoHash queueHash2 = await scheduler.CreateFilterAsync(reqs2);
+
+			// Update the cached queues
+			await scheduler.ForceUpdateFiltersAsync(CancellationToken.None);
+
+			SessionId[] sessionIds1 = await scheduler.GetAllFilteredSessionsAsync(queueHash1);
+			Assert.AreEqual(3, sessionIds1.Length);
+
+			SessionId[] sessionIds2 = await scheduler.GetAllFilteredSessionsAsync(queueHash2);
+			Assert.AreEqual(5, sessionIds2.Length);
+		}
+
+		[TestMethod]
+		public async Task AddSessionsToExistingQueueTestAsync()
+		{
+			FakeClock clock = ServiceProvider.GetRequiredService<FakeClock>();
+			AgentScheduler scheduler = (AgentScheduler)AgentScheduler;
+			await scheduler.StartAsync(default);
+
+			// Create the first queue
+			RpcAgentRequirements reqs1 = new RpcAgentRequirements();
+			reqs1.Properties.Add("Platform=Windows");
+
+			IoHash queueHash1 = await scheduler.CreateFilterAsync(reqs1);
+
+			// Create the second queue
+			RpcAgentRequirements reqs2 = new RpcAgentRequirements();
+			reqs2.Properties.Add("Platform=MacOS");
+
+			IoHash queueHash2 = await scheduler.CreateFilterAsync(reqs2);
+
+			// Update the cached queues
+			await clock.AdvanceAsync(TimeSpan.FromMinutes(1.0));
+
+			// First session
+			for (int idx = 0; idx < 3; idx++)
+			{
+				RpcAgentCapabilities caps1 = new RpcAgentCapabilities();
+				caps1.Properties.Add("Platform=Windows");
+
+				RpcSession? session1 = await AgentScheduler.TryCreateSessionAsync(new AgentId("foo"), SessionIdUtils.GenerateNewId(), caps1);
+				Assert.IsNotNull(session1);
+			}
+
+			// Second session
+			for (int idx = 0; idx < 5; idx++)
+			{
+				RpcAgentCapabilities caps2 = new RpcAgentCapabilities();
+				caps2.Properties.Add("Platform=MacOS");
+
+				RpcSession? session2 = await AgentScheduler.TryCreateSessionAsync(new AgentId("bar"), SessionIdUtils.GenerateNewId(), caps2);
+				Assert.IsNotNull(session2);
+			}
+
+			//
+			SessionId[] sessionIds1 = await scheduler.GetAllFilteredSessionsAsync(queueHash1);
+			Assert.AreEqual(3, sessionIds1.Length);
+
+			SessionId[] sessionIds2 = await scheduler.GetAllFilteredSessionsAsync(queueHash2);
+			Assert.AreEqual(5, sessionIds2.Length);
+		}
+
+		[TestMethod]
+		public async Task ExpireQueuesTestAsync()
+		{
+			FakeClock clock = ServiceProvider.GetRequiredService<FakeClock>();
+			AgentScheduler scheduler = (AgentScheduler)AgentScheduler;
+			await scheduler.StartAsync(default);
+
+			// Create the first queue
+			RpcAgentRequirements reqs1 = new RpcAgentRequirements();
+			reqs1.Properties.Add("Platform=Windows");
+
+			IoHash reqsHash1 = await scheduler.CreateFilterAsync(reqs1);
+			Assert.AreEqual(1, (await scheduler.GetFiltersAsync()).Length);
+
+			// Check that the queue exists
+			await clock.AdvanceAsync(HordeServer.Agents.AgentScheduler.ExpireFiltersTime / 2);
+			await scheduler.TouchFilterAsync(reqsHash1);
+			Assert.AreEqual(1, (await scheduler.GetFiltersAsync()).Length);
+
+			await clock.AdvanceAsync(HordeServer.Agents.AgentScheduler.ExpireFiltersTime / 2);
+			await scheduler.TouchFilterAsync(reqsHash1);
+			Assert.AreEqual(1, (await scheduler.GetFiltersAsync()).Length);
+
+			await clock.AdvanceAsync(HordeServer.Agents.AgentScheduler.ExpireFiltersTime / 2);
+			Assert.AreEqual(1, (await scheduler.GetFiltersAsync()).Length);
+
+			await clock.AdvanceAsync(HordeServer.Agents.AgentScheduler.ExpireFiltersTime);
+			Assert.AreEqual(0, (await scheduler.GetFiltersAsync()).Length);
+		}
 	}
 }
diff --git a/Engine/Source/Programs/Horde/Plugins/Compute/HordeServer.Compute/Agents/AgentScheduler.cs b/Engine/Source/Programs/Horde/Plugins/Compute/HordeServer.Compute/Agents/AgentScheduler.cs
--- a/Engine/Source/Programs/Horde/Plugins/Compute/HordeServer.Compute/Agents/AgentScheduler.cs
+++ b/Engine/Source/Programs/Horde/Plugins/Compute/HordeServer.Compute/Agents/AgentScheduler.cs
@@ -10,6 +10,7 @@ using Google.Protobuf;
 using HordeCommon.Rpc;
 using HordeCommon.Rpc.Messages;
 using HordeServer.Server;
+using Microsoft.Extensions.Caching.Memory;
 using Microsoft.Extensions.Hosting;
 using Microsoft.Extensions.Logging;
 using StackExchange.Redis;
@@ -18,9 +19,17 @@ namespace HordeServer.Agents
 {
 	class AgentScheduler : IAgentScheduler, IHostedService, IAsyncDisposable
 	{
+		/// <summary>
+		/// Time after which we expire filters
+		/// </summary>
+		public static TimeSpan ExpireFiltersTime { get; } = TimeSpan.FromMinutes(30.0);
+
 #pragma warning disable CA1822 // Make property static
 		static class Keys
 		{
+			public static FilterCollectionKeys Filters { get; }
+				= new();
+
 			public static SessionCollectionKeys Sessions { get; }
 				= new();
 
@@ -28,6 +37,30 @@ namespace HordeServer.Agents
 				= new();
 		}
 
+		record struct FilterCollectionKeys
+		{
+			public RedisSortedSetKey<IoHash> Created
+				=> new($"compute:filters:created");
+
+			public RedisSortedSetKey<IoHash> Deleted
+				=> new($"compute:filters:deleted");
+
+			public RedisHashKey<IoHash, long> Touched
+				=> new($"compute:filters:touched");
+
+			public FilterKeys this[IoHash requirementsHash]
+				=> new(requirementsHash);
+		}
+
+		record struct FilterKeys(IoHash RequirementsHash)
+		{
+			public RedisStringKey<RpcAgentRequirements> Requirements
+				=> new($"compute:filters:{RequirementsHash}:reqs");
+
+			public RedisSetKey<SessionId> Sessions
+				=> new($"compute:filters:{RequirementsHash}:sessions");
+		}
+
 		record struct SessionCollectionKeys
 		{
 			public RedisHashKey<SessionId, long> ExpiryTimes
@@ -42,6 +75,9 @@ namespace HordeServer.Agents
 			public RedisStringKey<RpcAgentCapabilities> Capabilities
 				=> new($"compute:sessions:{SessionId}:caps");
 
+			public RedisSetKey<IoHash> Filters
+				=> new($"compute:sessions:{SessionId}:filters");
+
 			public RedisStringKey<RpcSession> State
 				=> new($"compute:sessions:{SessionId}:state");
 		}
@@ -67,10 +103,16 @@ namespace HordeServer.Agents
 
 		readonly IRedisService _redisService;
 		readonly IClock _clock;
+		readonly ITicker _updateCachedFiltersTicker;
+		readonly ITicker _expireFiltersTicker;
 		readonly ILogger _logger;
+		readonly MemoryCache _memoryCache;
 
 		static readonly RedisChannel<SessionId> s_sessionUpdateChannel = new RedisChannel<SessionId>(RedisChannel.Literal("compute:sessions:update"));
 
+		record class CachedFilters(long Ticks, SortedSetEntry<IoHash>[] Created, SortedSetEntry<IoHash>[] Deleted);
+		CachedFilters _cachedFilters = new CachedFilters(0, Array.Empty<SortedSetEntry<IoHash>>(), Array.Empty<SortedSetEntry<IoHash>>());
+
 		IAsyncDisposable? _updateEventSubscription;
 
 		/// <summary>
@@ -80,12 +122,20 @@ namespace HordeServer.Agents
 		{
 			_redisService = redisService;
 			_clock = clock;
+			_updateCachedFiltersTicker = clock.AddTicker<AgentScheduler>(TimeSpan.FromSeconds(30.0), UpdateCachedFiltersTickAsync, logger);
+			_expireFiltersTicker = clock.AddSharedTicker($"{nameof(AgentScheduler)}.{nameof(ExpireFiltersAsync)}", TimeSpan.FromMinutes(5.0), ExpireFiltersAsync, logger);
+			_memoryCache = new MemoryCache(new MemoryCacheOptions());
 			_logger = logger;
 		}
 
 		/// <inheritdoc/>
 		public async ValueTask DisposeAsync()
 		{
+			_memoryCache.Dispose();
+
+			await _updateCachedFiltersTicker.DisposeAsync();
+			await _expireFiltersTicker.DisposeAsync();
+
 			if (_updateEventSubscription != null)
 			{
 				await _updateEventSubscription.DisposeAsync();
@@ -100,11 +150,17 @@ namespace HordeServer.Agents
 		public async Task StartAsync(CancellationToken cancellationToken)
 		{
 			_updateEventSubscription = await _redisService.GetDatabase().Multiplexer.SubscribeAsync(s_sessionUpdateChannel, OnSessionUpdate);
+
+			await _updateCachedFiltersTicker.StartAsync();
+			await _expireFiltersTicker.StartAsync();
 		}
 
 		/// <inheritdoc/>
 		public async Task StopAsync(CancellationToken cancellationToken)
 		{
+			await _expireFiltersTicker.StopAsync();
+			await _updateCachedFiltersTicker.StopAsync();
+
 			if (_updateEventSubscription != null)
 			{
 				await _updateEventSubscription.DisposeAsync();
@@ -112,6 +168,8 @@ namespace HordeServer.Agents
 			}
 		}
 
+		#region Sessions
+
 		/// <inheritdoc/>
 		public async Task<RpcSession?> TryCreateSessionAsync(AgentId agentId, SessionId sessionId, RpcAgentCapabilities capabilities, CancellationToken cancellationToken = default)
 		{
@@ -133,6 +191,10 @@ namespace HordeServer.Agents
 			_ = transaction.StringSetAsync(Keys.Sessions[sessionId].Capabilities.Inner, capabilitiesData, flags: CommandFlags.FireAndForget);
 			_ = transaction.StringSetAsync(Keys.Sessions[sessionId].State, newSession, flags: CommandFlags.FireAndForget);
 
+			if (!await UpdateFiltersAsync(_cachedFilters, transaction, newSession, capabilities, cancellationToken))
+			{
+				return null;
+			}
 			if (!await transaction.ExecuteAsync().WaitAsync(cancellationToken))
 			{
 				return null;
@@ -190,8 +252,10 @@ namespace HordeServer.Agents
 			}
 			else
 			{
+				DateTime utcNow = _clock.UtcNow;
+
 				// Extend the expiry time, ensuring it only ever increases
-				newSession.ExpiryTicks = Math.Max(session.ExpiryTicks + 1, (_clock.UtcNow + RpcSession.ExpireAfterTime).Ticks);
+				newSession.ExpiryTicks = Math.Max(session.ExpiryTicks + 1, (utcNow + RpcSession.ExpireAfterTime).Ticks);
 				_ = transaction.HashSetAsync(Keys.Sessions.ExpiryTimes, session.SessionId, newSession.ExpiryTicks, flags: CommandFlags.FireAndForget);
 
 				// Update the capabilities
@@ -220,6 +284,12 @@ namespace HordeServer.Agents
 					}
 				}
 
+				// Update the filters that this session belongs to
+				if (!await UpdateFiltersAsync(_cachedFilters, transaction, newSession, newCapabilities, cancellationToken))
+				{
+					return null;
+				}
+
 				// Update the session state
 				_ = transaction.StringSetAsync(Keys.Sessions[session.SessionId].State, newSession, flags: CommandFlags.FireAndForget);
 			}
@@ -274,6 +344,230 @@ namespace HordeServer.Agents
 			}
 		}
 
+		#endregion
+		#region Filters
+
+		/// <inheritdoc/>
+		public async Task<IoHash> CreateFilterAsync(RpcAgentRequirements requirements, CancellationToken cancellationToken = default)
+		{
+			byte[] requirementsData = requirements.ToByteArray();
+			IoHash requirementsHash = IoHash.Compute(requirementsData);
+
+			for (; ; )
+			{
+				IDatabase database = _redisService.GetDatabase();
+				DateTime utcNow = _clock.UtcNow;
+
+				ITransaction transaction = database.CreateTransaction();
+				transaction.AddCondition(Keys.Filters.Touched.HashNotExists(requirementsHash));
+				_ = transaction.HashSetAsync(Keys.Filters.Touched, requirementsHash, utcNow.Ticks);
+				_ = transaction.StringSetAsync(Keys.Filters[requirementsHash].Requirements.Inner, requirementsData, flags: CommandFlags.FireAndForget);
+				_ = transaction.SortedSetAddAsync(Keys.Filters.Created, requirementsHash, utcNow.Ticks, flags: CommandFlags.FireAndForget);
+				_ = transaction.SortedSetRemoveAsync(Keys.Filters.Deleted, requirementsHash, flags: CommandFlags.FireAndForget);
+
+				if (await transaction.ExecuteAsync().WaitAsync(cancellationToken))
+				{
+					break;
+				}
+
+				transaction = database.CreateTransaction();
+				transaction.AddCondition(Keys.Filters.Touched.HashNotExists(requirementsHash));
+				_ = transaction.HashSetAsync(Keys.Filters.Touched, requirementsHash, utcNow.Ticks, flags: CommandFlags.FireAndForget);
+
+				if (await transaction.ExecuteAsync().WaitAsync(cancellationToken))
+				{
+					break;
+				}
+			}
+
+			return requirementsHash;
+		}
+
+		/// <summary>
+		/// Update all the cached filters
+		/// </summary>
+		async ValueTask UpdateCachedFiltersTickAsync(CancellationToken cancellationToken)
+		{
+			IDatabase database = _redisService.GetDatabase();
+			DateTime utcNow = _clock.UtcNow;
+
+			ITransaction transaction = database.CreateTransaction();
+			Task<SortedSetEntry<IoHash>[]> created = transaction.SortedSetRangeByRankWithScoresAsync(Keys.Filters.Created);
+			Task<SortedSetEntry<IoHash>[]> deleted = transaction.SortedSetRangeByRankWithScoresAsync(Keys.Filters.Deleted);
+
+			await transaction.ExecuteAsync().WaitAsync(cancellationToken);
+			_cachedFilters = new CachedFilters(utcNow.Ticks, await created, await deleted);
+		}
+
+		/// <summary>
+		/// Expire all the filters older than the standard expiry time
+		/// </summary>
+		async ValueTask ExpireFiltersAsync(CancellationToken cancellationToken)
+		{
+			IDatabase database = _redisService.GetDatabase();
+
+			DateTime expiryTime = _clock.UtcNow - ExpireFiltersTime;
+
+			HashEntry<IoHash, long>[] entries = await database.HashGetAllAsync(Keys.Filters.Touched).WaitAsync(cancellationToken);
+			foreach ((IoHash requirementsHash, long ticks) in entries)
+			{
+				cancellationToken.ThrowIfCancellationRequested();
+				if (ticks < expiryTime.Ticks)
+				{
+					ITransaction transaction = database.CreateTransaction();
+					transaction.AddCondition(Keys.Filters.Touched.HashEqual(requirementsHash, ticks));
+
+					_ = transaction.HashDeleteAsync(Keys.Filters.Touched, requirementsHash, flags: CommandFlags.FireAndForget);
+					_ = transaction.SortedSetRemoveAsync(Keys.Filters.Created, requirementsHash, flags: CommandFlags.FireAndForget);
+					_ = transaction.SortedSetAddAsync(Keys.Filters.Deleted, requirementsHash, _clock.UtcNow.Ticks, flags: CommandFlags.FireAndForget);
+					_ = transaction.KeyDeleteAsync(Keys.Filters[requirementsHash].Requirements, flags: CommandFlags.FireAndForget);
+					_ = transaction.KeyDeleteAsync(Keys.Filters[requirementsHash].Sessions, flags: CommandFlags.FireAndForget);
+
+					_ = transaction.ExecuteAsync(CommandFlags.FireAndForget);
+				}
+			}
+		}
+
+		/// <inheritdoc/>
+		public async Task TouchFilterAsync(IoHash requirementsHash, CancellationToken cancellationToken = default)
+		{
+			IDatabase database = _redisService.GetDatabase();
+			DateTime utcNow = _clock.UtcNow;
+			await database.HashSetAsync(Keys.Filters.Touched, requirementsHash, utcNow.Ticks);
+		}
+
+		/// <inheritdoc/>
+		public async Task<IoHash[]> GetFiltersAsync(CancellationToken cancellationToken = default)
+		{
+			IDatabase database = _redisService.GetDatabase();
+			return await database.HashKeysAsync(Keys.Filters.Touched);
+		}
+
+		public async Task ForceUpdateFiltersAsync(CancellationToken cancellationToken)
+		{
+			IDatabase database = _redisService.GetDatabase();
+
+			await UpdateCachedFiltersTickAsync(cancellationToken);
+			CachedFilters cachedFilters = _cachedFilters;
+
+			SessionId[] sessionIds = await database.HashKeysAsync(Keys.Sessions.ExpiryTimes);
+			foreach (SessionId sessionId in sessionIds)
+			{
+				RpcSession? session = await TryGetSessionAsync(sessionId, cancellationToken);
+				if (session != null)
+				{
+					ITransaction transaction = database.CreateTransaction();
+					await UpdateFiltersAsync(cachedFilters, transaction, session, null, cancellationToken);
+					await transaction.ExecuteAsync().WaitAsync(cancellationToken);
+				}
+			}
+		}
+
+		async Task<bool> UpdateFiltersAsync(CachedFilters cachedFilters, ITransaction transaction, RpcSession session, RpcAgentCapabilities? capabilities, CancellationToken cancellationToken)
+		{
+			long minFilterTime = Math.Max(0, session.LastFilterUpdateTicks - TimeSpan.FromSeconds(30.0).Ticks);
+
+			// Add the session to any new matching filters
+			SortedSetEntry<IoHash>[] createdFilters = cachedFilters.Created;
+
+			int createdIdx = createdFilters.BinarySearch(new SortedSetEntry<IoHash>(IoHash.Zero, minFilterTime));
+			if (createdIdx < 0)
+			{
+				createdIdx = ~createdIdx;
+			}
+
+			if (createdIdx < createdFilters.Length)
+			{
+				capabilities ??= await TryGetSessionCapabilitiesAsync(session, cancellationToken);
+				if (capabilities == null)
+				{
+					return false;
+				}
+
+				for (; createdIdx < createdFilters.Length; createdIdx++)
+				{
+					IoHash requirementsHash = createdFilters[createdIdx].Element;
+
+					RpcAgentRequirements? requirements = await TryGetFilterRequirementsAsync(requirementsHash, cancellationToken);
+					if (requirements != null && MeetsRequirements(capabilities, requirements))
+					{
+						transaction.AddCondition(Keys.Filters.Touched.HashExists(requirementsHash));
+						_ = transaction.SetAddAsync(Keys.Filters[requirementsHash].Sessions, session.SessionId, flags: CommandFlags.FireAndForget);
+						_ = transaction.SetAddAsync(Keys.Sessions[session.SessionId].Filters, requirementsHash, flags: CommandFlags.FireAndForget);
+					}
+				}
+			}
+
+			// Remove the session from any removed filters
+			SortedSetEntry<IoHash>[] deletedFilters = cachedFilters.Deleted;
+
+			int deletedIdx = deletedFilters.BinarySearch(new SortedSetEntry<IoHash>(IoHash.Zero, minFilterTime));
+			if (deletedIdx < 0)
+			{
+				deletedIdx = ~deletedIdx;
+			}
+
+			for (; deletedIdx < deletedFilters.Length; deletedIdx++)
+			{
+				IoHash queueHash = deletedFilters[deletedIdx].Element;
+				transaction.AddCondition(Keys.Filters.Touched.HashNotExists(queueHash));
+				_ = transaction.SetRemoveAsync(Keys.Sessions[session.SessionId].Filters, queueHash, flags: CommandFlags.FireAndForget);
+			}
+
+			// Update the last queue time
+			session.LastFilterUpdateTicks = cachedFilters.Ticks;
+			return true;
+		}
+
+		static bool MeetsRequirements(RpcAgentCapabilities capabilities, RpcAgentRequirements requirements)
+		{
+			foreach (string property in requirements.Properties)
+			{
+				if (!capabilities.Properties.Contains(property))
+				{
+					return false;
+				}
+			}
+			return true;
+		}
+
+		void AddRequirementsToCache(IoHash requirementsHash, RpcAgentRequirements requirements)
+		{
+			using (ICacheEntry entry = _memoryCache.CreateEntry(requirementsHash))
+			{
+				entry.SetSlidingExpiration(TimeSpan.FromMinutes(30.0));
+				entry.SetValue(requirements);
+			}
+		}
+
+		/// <inheritdoc/>
+		public async ValueTask<RpcAgentRequirements?> TryGetFilterRequirementsAsync(IoHash requirementsHash, CancellationToken cancellationToken)
+		{
+			RpcAgentRequirements? requirements;
+			if (!_memoryCache.TryGetValue(requirementsHash, out requirements))
+			{
+				IDatabase database = _redisService.GetDatabase();
+				requirements = await database.StringGetAsync(Keys.Filters[requirementsHash].Requirements);
+
+				if (requirements != null)
+				{
+					AddRequirementsToCache(requirementsHash, requirements);
+				}
+			}
+			return requirements;
+		}
+
+		/// <inheritdoc/>
+		public async Task<SessionId[]> GetAllFilteredSessionsAsync(IoHash requirementsHash, CancellationToken cancellationToken = default)
+		{
+			IDatabase database = _redisService.GetDatabase();
+			_ = database.HashSetAsync(Keys.Filters.Touched, requirementsHash, _clock.UtcNow.Ticks, flags: CommandFlags.FireAndForget);
+			return await database.SetMembersAsync(Keys.Filters[requirementsHash].Sessions).WaitAsync(cancellationToken);
+		}
+
+		#endregion
+		#region Stats
+
 		/// <inheritdoc/>
 		public async Task<LeaseId[]> FindActiveLeaseIdsAsync(CancellationToken cancellationToken = default)
 		{
@@ -294,5 +588,7 @@ namespace HordeServer.Agents
 			IDatabase database = _redisService.GetDatabase();
 			return await database.SetMembersAsync(Keys.Leases[id].Children).WaitAsync(cancellationToken);
 		}
+
+		#endregion
 	}
 }
diff --git a/Engine/Source/Programs/Horde/Plugins/Compute/HordeServer.Compute/Agents/IAgentScheduler.cs b/Engine/Source/Programs/Horde/Plugins/Compute/HordeServer.Compute/Agents/IAgentScheduler.cs
--- a/Engine/Source/Programs/Horde/Plugins/Compute/HordeServer.Compute/Agents/IAgentScheduler.cs
+++ b/Engine/Source/Programs/Horde/Plugins/Compute/HordeServer.Compute/Agents/IAgentScheduler.cs
@@ -1,5 +1,6 @@
 // Copyright Epic Games, Inc. All Rights Reserved.
 
+using EpicGames.Core;
 using EpicGames.Horde.Agents;
 using EpicGames.Horde.Agents.Leases;
 using EpicGames.Horde.Agents.Sessions;
@@ -55,6 +56,46 @@ namespace HordeServer.Agents
 		/// <returns>New session object</returns>
 		Task<RpcSession?> TryUpdateSessionAsync(RpcSession session, RpcSession? newSession = null, RpcAgentCapabilities? newCapabilities = null, CancellationToken cancellationToken = default);
 
+		#endregion
+		#region Filters
+
+		/// <summary>
+		/// Create a new filter for agents matching  with the given requirements
+		/// </summary>
+		/// <param name="requirements">Requirements for the filter</param>
+		/// <param name="cancellationToken">Cancellation token for the operation</param>
+		Task<IoHash> CreateFilterAsync(RpcAgentRequirements requirements, CancellationToken cancellationToken = default);
+
+		/// <summary>
+		/// Update the timestamp on the given filter
+		/// </summary>
+		/// <param name="requirementsHash">Hash of the requirements for the filter to remove</param>
+		/// <param name="cancellationToken">Cancellation token for the operation</param>
+		Task TouchFilterAsync(IoHash requirementsHash, CancellationToken cancellationToken = default);
+
+		/// <summary>
+		/// Gets the hash of all current filters
+		/// </summary>
+		/// <param name="cancellationToken">Cancellation token for the operation</param>
+		/// <returns>Array of filter hashes</returns>
+		Task<IoHash[]> GetFiltersAsync(CancellationToken cancellationToken = default);
+
+		/// <summary>
+		/// Gets the requirements for a particular filter
+		/// </summary>
+		/// <param name="requirementsHash">Hash of the requirements object</param>
+		/// <param name="cancellationToken">Cancellation token for the operation</param>
+		/// <returns>Requirements object</returns>
+		ValueTask<RpcAgentRequirements?> TryGetFilterRequirementsAsync(IoHash requirementsHash, CancellationToken cancellationToken);
+
+		/// <summary>
+		/// Finds all sessions currently included in a filter
+		/// </summary>
+		/// <param name="requirementsHash">Hash of the requirements object</param>
+		/// <param name="cancellationToken">Cancellation token for the operation</param>
+		/// <returns>Array of sessions</returns>
+		Task<SessionId[]> GetAllFilteredSessionsAsync(IoHash requirementsHash, CancellationToken cancellationToken = default);
+
 		#endregion
 		#region Stats
 
diff --git a/Engine/Source/Programs/Shared/EpicGames.Redis/RedisSortedSetKey.cs b/Engine/Source/Programs/Shared/EpicGames.Redis/RedisSortedSetKey.cs
--- a/Engine/Source/Programs/Shared/EpicGames.Redis/RedisSortedSetKey.cs
+++ b/Engine/Source/Programs/Shared/EpicGames.Redis/RedisSortedSetKey.cs
@@ -25,7 +25,7 @@ namespace EpicGames.Redis
 	/// Typed implementation of <see cref="SortedSetEntry"/>
 	/// </summary>
 	/// <typeparam name="T">The element type</typeparam>
-	public readonly struct SortedSetEntry<T>
+	public readonly struct SortedSetEntry<T> : IEquatable<SortedSetEntry<T>>, IComparable, IComparable<SortedSetEntry<T>>
 	{
 		/// <summary>
 		/// Accessor for the element type
@@ -75,6 +75,50 @@ namespace EpicGames.Redis
 			outElement = Element;
 			outScore = Score;
 		}
+
+		/// <inheritdoc/>
+		public override bool Equals(object? obj) 
+			=> obj is SortedSetEntry ssObj && Equals(ssObj);
+
+		/// <inheritdoc/>
+		public override int GetHashCode()
+			=> HashCode.Combine(Score, Element);
+
+		/// <inheritdoc/>
+		public bool Equals(SortedSetEntry<T> other) 
+			=> Score == other.Score && Equals(Element, other.Element);
+
+		/// <inheritdoc/>
+		public int CompareTo(SortedSetEntry<T> other) 
+			=> Score.CompareTo(other.Score);
+
+		/// <inheritdoc/>
+		public int CompareTo(object? obj) 
+			=> obj is SortedSetEntry<T> ssObj ? CompareTo(ssObj) : -1;
+
+		/// <inheritdoc/>
+		public static bool operator ==(SortedSetEntry<T> a, SortedSetEntry<T> b)
+			=> a.Equals(b);
+
+		/// <inheritdoc/>
+		public static bool operator !=(SortedSetEntry<T> a, SortedSetEntry<T> b)
+			=> !a.Equals(b);
+
+		/// <inheritdoc/>
+		public static bool operator <=(SortedSetEntry<T> a, SortedSetEntry<T> b)
+			=> a.CompareTo(b) <= 0;
+
+		/// <inheritdoc/>
+		public static bool operator <(SortedSetEntry<T> a, SortedSetEntry<T> b)
+			=> a.CompareTo(b) < 0;
+
+		/// <inheritdoc/>
+		public static bool operator >=(SortedSetEntry<T> a, SortedSetEntry<T> b)
+			=> a.CompareTo(b) >= 0;
+
+		/// <inheritdoc/>
+		public static bool operator >(SortedSetEntry<T> a, SortedSetEntry<T> b)
+			=> a.CompareTo(b) > 0;
 	}
 
 	/// <summary>
@@ -191,7 +235,7 @@ namespace EpicGames.Redis
 		#region SortedSetRangeByRankWithScoresAsync
 
 		/// <inheritdoc cref="IDatabaseAsync.SortedSetRangeByRankWithScoresAsync(RedisKey, Int64, Int64, Order, CommandFlags)"/>
-		public static async Task<SortedSetEntry<TElement>[]> SortedSetRangeByRankWithScoresAsync<TElement>(this IDatabaseAsync target, RedisSortedSetKey<TElement> key, long start, long stop = -1, Order order = Order.Ascending, CommandFlags flags = CommandFlags.None)
+		public static async Task<SortedSetEntry<TElement>[]> SortedSetRangeByRankWithScoresAsync<TElement>(this IDatabaseAsync target, RedisSortedSetKey<TElement> key, long start = 0, long stop = -1, Order order = Order.Ascending, CommandFlags flags = CommandFlags.None)
 		{
 			SortedSetEntry[] values = await target.SortedSetRangeByRankWithScoresAsync(key.Inner, start, stop, order, flags);
 			return Array.ConvertAll(values, x => new SortedSetEntry<TElement>(x));
