Horde: Deprecate devices in agent capabilities and requirements.

We've only exposed a flat list of properties and resources in agents for a while now, so the RWAPI definitions for devices make the API confusing. Now that we support filtering against requirements, it's more appropriate to be explicit around the semantics of these fields.

#rnx

diff --git a/Engine/Source/Programs/Horde/HordeAgent/Commands/Utilities/CapsCommand.cs b/Engine/Source/Programs/Horde/HordeAgent/Commands/Utilities/CapsCommand.cs
--- a/Engine/Source/Programs/Horde/HordeAgent/Commands/Utilities/CapsCommand.cs
+++ b/Engine/Source/Programs/Horde/HordeAgent/Commands/Utilities/CapsCommand.cs
@@ -27,14 +27,21 @@ namespace HordeAgent.Commands.Utilities
 		public override async Task<int> ExecuteAsync(ILogger logger)
 		{
 			RpcAgentCapabilities capabilities = await _capabilitiesService.GetCapabilitiesAsync(null);
-			foreach (RpcDeviceCapabilities device in capabilities.Devices)
+
+			logger.LogInformation("Properties:");
+			foreach (string property in capabilities.Properties)
+			{
+				logger.LogInformation("  {Property}", property);
+			}
+
+			logger.LogInformation("");
+
+			logger.LogInformation("Resources:");
+			foreach ((string name, int value) in capabilities.Resources)
 			{
-				logger.LogInformation("Device: {Name}", device.Handle);
-				foreach (string property in device.Properties)
-				{
-					logger.LogInformation("  {Property}", property);
-				}
+				logger.LogInformation("  {Name}={Value}", name, value);
 			}
+
 			return 0;
 		}
 	}
diff --git a/Engine/Source/Programs/Horde/HordeAgent/Services/CapabilitiesService.cs b/Engine/Source/Programs/Horde/HordeAgent/Services/CapabilitiesService.cs
--- a/Engine/Source/Programs/Horde/HordeAgent/Services/CapabilitiesService.cs
+++ b/Engine/Source/Programs/Horde/HordeAgent/Services/CapabilitiesService.cs
@@ -83,19 +83,17 @@ namespace HordeAgent.Services
 			ILogger logger = _logger;
 
 			// Create the primary device
-			RpcDeviceCapabilities primaryDevice = new RpcDeviceCapabilities();
-			primaryDevice.Handle = KnownPropertyNames.PrimaryDeviceHandle;
+			RpcAgentCapabilities capabilities = new RpcAgentCapabilities();
 
-			List<RpcDeviceCapabilities> otherDevices = new List<RpcDeviceCapabilities>();
 			if (RuntimeInformation.IsOSPlatform(OSPlatform.Windows))
 			{
 				using IScope winScope = GlobalTracer.Instance.BuildSpan("GetWindowsCapabilities").StartActive();
-				primaryDevice.Properties.Add("Platform=Win64");
-				primaryDevice.Properties.Add("PlatformGroup=Windows");
-				primaryDevice.Properties.Add("PlatformGroup=Microsoft");
-				primaryDevice.Properties.Add("PlatformGroup=Desktop");
+				capabilities.Properties.Add("Platform=Win64");
+				capabilities.Properties.Add("PlatformGroup=Windows");
+				capabilities.Properties.Add("PlatformGroup=Microsoft");
+				capabilities.Properties.Add("PlatformGroup=Desktop");
 
-				primaryDevice.Properties.Add("OSFamily=Windows");
+				capabilities.Properties.Add("OSFamily=Windows");
 
 				// WMI doesn't work currently on ARM64, due to reliance on a NET Framework assembly.
 				using (CimSession session = CimSession.Create(null))
@@ -113,11 +111,11 @@ namespace HordeAgent.Services
 								string name = property.Name;
 								if (name.Equals("Caption", StringComparison.OrdinalIgnoreCase))
 								{
-									primaryDevice.Properties.Add($"OSDistribution={property.Value}");
+									capabilities.Properties.Add($"OSDistribution={property.Value}");
 								}
 								else if (name.Equals("Version", StringComparison.OrdinalIgnoreCase))
 								{
-									primaryDevice.Properties.Add($"OSKernelVersion={property.Value}");
+									capabilities.Properties.Add($"OSKernelVersion={property.Value}");
 								}
 							}
 						}
@@ -160,7 +158,7 @@ namespace HordeAgent.Services
 							}
 						}
 						
-						AddCpuInfo(primaryDevice, cpuNameToCount, totalLogicalCores, totalPhysicalCores);
+						AddCpuInfo(capabilities, cpuNameToCount, totalLogicalCores, totalPhysicalCores);
 					}
 					
 					{
@@ -177,8 +175,8 @@ namespace HordeAgent.Services
 								}
 							}
 						}
-						
-						primaryDevice.Properties.Add($"RAM={totalCapacity / (1024 * 1024 * 1024)}");
+
+						capabilities.Resources.Add(KnownResourceNames.Ram, (int)(totalCapacity / (1024 * 1024 * 1024)));
 					}
 					
 					{
@@ -205,11 +203,11 @@ namespace HordeAgent.Services
 							if (name != null)
 							{
 								string prefix = $"GPU-{++index}";
-								primaryDevice.Properties.Add($"{prefix}-Name={name}");
+								capabilities.Properties.Add($"{prefix}-Name={name}");
 								
 								if (driverVersion != null)
 								{
-									primaryDevice.Properties.Add($"{prefix}-DriverVersion={driverVersion}");
+									capabilities.Properties.Add($"{prefix}-DriverVersion={driverVersion}");
 								}
 							}
 						}
@@ -219,35 +217,35 @@ namespace HordeAgent.Services
 				// Add EC2 properties if needed
 				if (_settings.EnableAwsEc2Support)
 				{
-					await AddAwsPropertiesAsync(primaryDevice.Properties, logger);
+					await AddAwsPropertiesAsync(capabilities.Properties, logger);
 				}
 
 				// Add session information
-				primaryDevice.Properties.Add($"User={Environment.UserName}");
-				primaryDevice.Properties.Add($"Domain={Environment.UserDomainName}");
-				primaryDevice.Properties.Add($"Interactive={Environment.UserInteractive}");
-				primaryDevice.Properties.Add($"Elevated={IsUserAdministrator()}");
+				capabilities.Properties.Add($"User={Environment.UserName}");
+				capabilities.Properties.Add($"Domain={Environment.UserDomainName}");
+				capabilities.Properties.Add($"Interactive={Environment.UserInteractive}");
+				capabilities.Properties.Add($"Elevated={IsUserAdministrator()}");
 			}
 			else if (RuntimeInformation.IsOSPlatform(OSPlatform.Linux))
 			{
 				using IScope _ = GlobalTracer.Instance.BuildSpan("GetLinuxCapabilities").StartActive();
-				primaryDevice.Properties.Add("Platform=Linux");
-				primaryDevice.Properties.Add("PlatformGroup=Linux");
-				primaryDevice.Properties.Add("PlatformGroup=Unix");
-				primaryDevice.Properties.Add("PlatformGroup=Desktop");
+				capabilities.Properties.Add("Platform=Linux");
+				capabilities.Properties.Add("PlatformGroup=Linux");
+				capabilities.Properties.Add("PlatformGroup=Unix");
+				capabilities.Properties.Add("PlatformGroup=Desktop");
 
-				primaryDevice.Properties.Add("OSFamily=Linux");
-				primaryDevice.Properties.Add("OSVersion=Linux");
+				capabilities.Properties.Add("OSFamily=Linux");
+				capabilities.Properties.Add("OSVersion=Linux");
 
 				// Add EC2 properties if needed
 				if (_settings.EnableAwsEc2Support)
 				{
-					await AddAwsPropertiesAsync(primaryDevice.Properties, logger);
+					await AddAwsPropertiesAsync(capabilities.Properties, logger);
 				}
 				
 				if (_settings.WineExecutablePath != null)
 				{
-					primaryDevice.Properties.Add($"{KnownPropertyNames.WineEnabled}=true");
+					capabilities.Properties.Add($"{KnownPropertyNames.WineEnabled}=true");
 				}
 				
 				// Parse the CPU info
@@ -284,7 +282,7 @@ namespace HordeAgent.Services
 						}
 					}
 
-					AddCpuInfo(primaryDevice, nameToCount, logicalCores.Count, physicalCores.Count);
+					AddCpuInfo(capabilities, nameToCount, logicalCores.Count, physicalCores.Count);
 				}
 
 				// Parse the RAM info
@@ -295,25 +293,25 @@ namespace HordeAgent.Services
 					if (match.Success)
 					{
 						long totalCapacity = Int64.Parse(match.Groups[1].Value) * 1024;
-						primaryDevice.Properties.Add($"RAM={totalCapacity / (1024 * 1024 * 1024)}");
+						capabilities.Resources.Add(KnownResourceNames.Ram, (int)(totalCapacity / (1024 * 1024 * 1024)));
 					}
 				}
 
 				// Add session information
-				primaryDevice.Properties.Add($"User={Environment.UserName}");
-				primaryDevice.Properties.Add($"Domain={Environment.UserDomainName}");
-				primaryDevice.Properties.Add($"Interactive={Environment.UserInteractive}");
-				primaryDevice.Properties.Add($"Elevated={IsUserAdministrator()}");
+				capabilities.Properties.Add($"User={Environment.UserName}");
+				capabilities.Properties.Add($"Domain={Environment.UserDomainName}");
+				capabilities.Properties.Add($"Interactive={Environment.UserInteractive}");
+				capabilities.Properties.Add($"Elevated={IsUserAdministrator()}");
 			}
 			else if (RuntimeInformation.IsOSPlatform(OSPlatform.OSX))
 			{
 				using IScope _ = GlobalTracer.Instance.BuildSpan("GetMacCapabilities").StartActive();
-				primaryDevice.Properties.Add("Platform=Mac");
-				primaryDevice.Properties.Add("PlatformGroup=Apple");
-				primaryDevice.Properties.Add("PlatformGroup=Desktop");
+				capabilities.Properties.Add("Platform=Mac");
+				capabilities.Properties.Add("PlatformGroup=Apple");
+				capabilities.Properties.Add("PlatformGroup=Desktop");
 
-				primaryDevice.Properties.Add("OSFamily=MacOS");
-				primaryDevice.Properties.Add("OSVersion=MacOS");
+				capabilities.Properties.Add("OSFamily=MacOS");
+				capabilities.Properties.Add("OSVersion=MacOS");
 
 				string output;
 				using (Process process = new Process())
@@ -338,7 +336,7 @@ namespace HordeAgent.Services
 					XmlNode? model = hardwareNode.SelectSingleNode("key[. = 'machine_model']/following-sibling::string");
 					if (model != null)
 					{
-						primaryDevice.Properties.Add($"Model={model.InnerText}");
+						capabilities.Properties.Add($"Model={model.InnerText}");
 					}
 
 					XmlNode? cpuTypeNode = hardwareNode.SelectSingleNode("key[. = 'cpu_type']/following-sibling::string");
@@ -346,24 +344,24 @@ namespace HordeAgent.Services
 					XmlNode? cpuPackagesNode = hardwareNode.SelectSingleNode("key[. = 'packages']/following-sibling::integer");
 					if (cpuTypeNode != null && cpuSpeedNode != null && cpuPackagesNode != null)
 					{
-						primaryDevice.Properties.Add((cpuPackagesNode.InnerText != "1") ? $"CPU={cpuPackagesNode.InnerText} x {cpuTypeNode.InnerText} @ {cpuSpeedNode.InnerText}" : $"CPU={cpuTypeNode.InnerText} @ {cpuSpeedNode.InnerText}");
+						capabilities.Properties.Add((cpuPackagesNode.InnerText != "1") ? $"CPU={cpuPackagesNode.InnerText} x {cpuTypeNode.InnerText} @ {cpuSpeedNode.InnerText}" : $"CPU={cpuTypeNode.InnerText} @ {cpuSpeedNode.InnerText}");
 					}
 
-					primaryDevice.Properties.Add($"LogicalCores={Environment.ProcessorCount}");
+					capabilities.Resources.Add(KnownResourceNames.LogicalCores, Environment.ProcessorCount);
 
 					XmlNode? cpuCountNode = hardwareNode.SelectSingleNode("key[. = 'number_processors']/following-sibling::integer");
 					if (cpuCountNode != null)
 					{
-						primaryDevice.Properties.Add($"PhysicalCores={cpuCountNode.InnerText}");
+						capabilities.Properties.Add($"PhysicalCores={cpuCountNode.InnerText}");
 					}
 
 					XmlNode? memoryNode = hardwareNode.SelectSingleNode("key[. = 'physical_memory']/following-sibling::string");
 					if (memoryNode != null)
 					{
 						string[] parts = memoryNode.InnerText.Split(new char[] { ' ' }, StringSplitOptions.RemoveEmptyEntries);
-						if (parts.Length == 2 && parts[1] == "GB")
+						if (parts.Length == 2 && parts[1] == "GB" && Int32.TryParse(parts[0], out int ramGb))
 						{
-							primaryDevice.Properties.Add($"RAM={parts[0]}");
+							capabilities.Resources.Add(KnownResourceNames.Ram, ramGb);
 						}
 					}
 				}
@@ -374,13 +372,13 @@ namespace HordeAgent.Services
 					XmlNode? osVersionNode = softwareNode.SelectSingleNode("key[. = 'os_version']/following-sibling::string");
 					if (osVersionNode != null)
 					{
-						primaryDevice.Properties.Add($"OSDistribution={osVersionNode.InnerText}");
+						capabilities.Properties.Add($"OSDistribution={osVersionNode.InnerText}");
 					}
 
 					XmlNode? kernelVersionNode = softwareNode.SelectSingleNode("key[. = 'kernel_version']/following-sibling::string");
 					if (kernelVersionNode != null)
 					{
-						primaryDevice.Properties.Add($"OSKernelVersion={kernelVersionNode.InnerText}");
+						capabilities.Properties.Add($"OSKernelVersion={kernelVersionNode.InnerText}");
 					}
 				}
 			}
@@ -395,11 +393,11 @@ namespace HordeAgent.Services
 				{
 					if (address.AddressFamily == System.Net.Sockets.AddressFamily.InterNetwork)
 					{
-						primaryDevice.Properties.Add($"Ipv4={address}");
+						capabilities.Properties.Add($"Ipv4={address}");
 					}
 					else if (address.AddressFamily == System.Net.Sockets.AddressFamily.InterNetworkV6)
 					{
-						primaryDevice.Properties.Add($"Ipv6={address}");
+						capabilities.Properties.Add($"Ipv6={address}");
 					}
 				}
 			}
@@ -426,13 +424,13 @@ namespace HordeAgent.Services
 			// Add the compute configuration
 			if (ip != null && _settings.ComputePort != 0)
 			{
-				primaryDevice.Properties.Add($"ComputeIp={ip}");
-				primaryDevice.Properties.Add($"ComputePort={_settings.ComputePort}");
+				capabilities.Properties.Add($"ComputeIp={ip}");
+				capabilities.Properties.Add($"ComputePort={_settings.ComputePort}");
 			}
 
 			// Get the time that the machine booted
-			primaryDevice.Properties.Add($"BootTime={s_bootTime}");
-			primaryDevice.Properties.Add($"StartTime={s_startTime}");
+			capabilities.Properties.Add($"BootTime={s_bootTime}");
+			capabilities.Properties.Add($"StartTime={s_startTime}");
 
 			// Add information about the current session
 			if (workingDir != null)
@@ -453,39 +451,34 @@ namespace HordeAgent.Services
 					try
 					{
 						DriveInfo info = new DriveInfo(driveName);
-						primaryDevice.Properties.Add($"{KnownPropertyNames.DiskFreeSpace}={info.AvailableFreeSpace}");
-						primaryDevice.Properties.Add($"{KnownPropertyNames.DiskTotalSize}={info.TotalSize}");
+						capabilities.Properties.Add($"{KnownPropertyNames.DiskFreeSpace}={info.AvailableFreeSpace}");
+						capabilities.Properties.Add($"{KnownPropertyNames.DiskTotalSize}={info.TotalSize}");
 					}
 					catch (Exception ex)
 					{
 						logger.LogWarning(ex, "Unable to query disk info for path '{DriveName}'", driveName);
 					}
 				}
-				primaryDevice.Properties.Add($"WorkingDir={workingDir}");
+				capabilities.Properties.Add($"WorkingDir={workingDir}");
 			}
 
 			// Add any horde. env vars for custom properties.
 			IEnumerable<string> envVars = Environment.GetEnvironmentVariables().Keys.Cast<string>();
 			foreach (string envVar in envVars.Where(x => x.StartsWith("horde.", StringComparison.InvariantCultureIgnoreCase)))
 			{
-				primaryDevice.Properties.Add($"{envVar}={Environment.GetEnvironmentVariable(envVar)}");
+				capabilities.Properties.Add($"{envVar}={Environment.GetEnvironmentVariable(envVar)}");
 			}
 
-			// Create the worker
-			RpcAgentCapabilities agent = new();
-			agent.Devices.Add(primaryDevice);
-			agent.Devices.AddRange(otherDevices);
-
 			// Add the max supported compute protocol version
-			agent.Properties.Add($"ComputeProtocol={(int)ComputeProtocol.Latest}");
+			capabilities.Properties.Add($"ComputeProtocol={(int)ComputeProtocol.Latest}");
 
 			// Whether the agent is packaged as a self-contained .NET app
 			// Used during the transition period over from multi-platform, non-self-contained agent packages.
-			agent.Properties.Add($"SelfContained={AgentApp.IsSelfContained}");
+			capabilities.Properties.Add($"SelfContained={AgentApp.IsSelfContained}");
 
 			// Add any additional properties from the config file
-			agent.Properties.AddRange(_settings.Properties.Select(kvp => $"{kvp.Key}={kvp.Value}"));
-			return agent;
+			capabilities.Properties.AddRange(_settings.Properties.Select(kvp => $"{kvp.Key}={kvp.Value}"));
+			return capabilities;
 		}
 
 		/// <summary>
@@ -542,21 +535,21 @@ namespace HordeAgent.Services
 			}
 		}
 
-		static void AddCpuInfo(RpcDeviceCapabilities primaryDevice, Dictionary<string, int> nameToCount, int numLogicalCores, int numPhysicalCores)
+		static void AddCpuInfo(RpcAgentCapabilities capabilities, Dictionary<string, int> nameToCount, int numLogicalCores, int numPhysicalCores)
 		{
 			if (nameToCount.Count > 0)
 			{
-				primaryDevice.Properties.Add("CPU=" + String.Join(", ", nameToCount.Select(x => (x.Value > 1) ? $"{x.Key} x {x.Value}" : x.Key)));
+				capabilities.Properties.Add("CPU=" + String.Join(", ", nameToCount.Select(x => (x.Value > 1) ? $"{x.Key} x {x.Value}" : x.Key)));
 			}
 
 			if (numLogicalCores > 0)
 			{
-				primaryDevice.Properties.Add($"LogicalCores={numLogicalCores}");
+				capabilities.Resources.Add("LogicalCores", numLogicalCores);
 			}
 
 			if (numPhysicalCores > 0)
 			{
-				primaryDevice.Properties.Add($"PhysicalCores={numPhysicalCores}");
+				capabilities.Properties.Add($"PhysicalCores={numPhysicalCores}");
 			}
 		}
 
diff --git a/Engine/Source/Programs/Horde/HordeAgent/Services/SessionFactory.cs b/Engine/Source/Programs/Horde/HordeAgent/Services/SessionFactory.cs
--- a/Engine/Source/Programs/Horde/HordeAgent/Services/SessionFactory.cs
+++ b/Engine/Source/Programs/Horde/HordeAgent/Services/SessionFactory.cs
@@ -117,18 +117,18 @@ namespace HordeAgent.Services
 			RpcAgentCapabilities capabilities = await capabilitiesService.GetCapabilitiesAsync(workingDir);
 			if (capabilities.Properties.Count > 0)
 			{
-				logger.LogInformation("Global:");
+				logger.LogInformation("Properties:");
 				foreach (string property in capabilities.Properties)
 				{
 					logger.LogInformation("  {AgentProperty}", property);
 				}
 			}
-			foreach (RpcDeviceCapabilities device in capabilities.Devices)
+			if (capabilities.Resources.Count > 0)
 			{
-				logger.LogInformation("{DeviceName} Device:", device.Handle);
-				foreach (string property in device.Properties)
+				logger.LogInformation("Resources:");
+				foreach ((string name, int value) in capabilities.Resources)
 				{
-					logger.LogInformation("   {DeviceProperty}", property);
+					logger.LogInformation("  {Name}={Value}", name, value);
 				}
 			}
 
@@ -313,7 +313,7 @@ namespace HordeAgent.Services
 
 		static string? GetProperty(RpcAgentCapabilities capabilities, string name)
 		{
-			foreach (string property in capabilities.Devices[0].Properties)
+			foreach (string property in capabilities.Properties)
 			{
 				if (property.Length > name.Length && property[name.Length] == '=' && property.StartsWith(name, StringComparison.OrdinalIgnoreCase))
 				{
diff --git a/Engine/Source/Programs/Horde/Plugins/Build/HordeServer.Build.Tests/Server/RpcServiceTest.cs b/Engine/Source/Programs/Horde/Plugins/Build/HordeServer.Build.Tests/Server/RpcServiceTest.cs
--- a/Engine/Source/Programs/Horde/Plugins/Build/HordeServer.Build.Tests/Server/RpcServiceTest.cs
+++ b/Engine/Source/Programs/Horde/Plugins/Build/HordeServer.Build.Tests/Server/RpcServiceTest.cs
@@ -315,6 +315,8 @@ namespace HordeServer.Tests.Server
 			Assert.IsTrue(agent.Properties.Contains("fooKey=barValue"));
 		}
 
+#pragma warning disable CA1041
+#pragma warning disable CS0612
 		[TestMethod]
 		public async Task PropertiesFromDeviceCapabilitiesAsync()
 		{
@@ -329,11 +331,13 @@ namespace HordeServer.Tests.Server
 		public async Task KnownPropertiesAreSetAsResourcesAsync()
 		{
 			RpcCreateSessionRequest req = new() { Id = new AgentId("bogusAgentName").ToString(), Capabilities = new RpcAgentCapabilities() };
-			req.Capabilities.Devices.Add(new RpcDeviceCapabilities { Handle = "someHandle", Properties = { $"{KnownPropertyNames.LogicalCores}=10" } });
+			req.Capabilities.Devices.Add(new RpcDeviceCapabilities { Handle = "someHandle", Properties = { $"{KnownResourceNames.LogicalCores}=10" } });
 			RpcCreateSessionResponse res = await RpcService.CreateSession(req, _adminContext);
 			IAgent agent = (await AgentService.GetAgentAsync(new AgentId(res.AgentId)))!;
-			Assert.AreEqual(10, agent.Resources[KnownPropertyNames.LogicalCores]);
+			Assert.AreEqual(10, agent.Resources[KnownResourceNames.LogicalCores]);
 		}
+#pragma warning restore CS0612
+#pragma warning restore CA1041
 
 		[TestMethod]
 		public async Task UpdateSessionTestAsync()
diff --git a/Engine/Source/Programs/Horde/Plugins/Compute/HordeServer.Compute/Agents/AgentCollection.cs b/Engine/Source/Programs/Horde/Plugins/Compute/HordeServer.Compute/Agents/AgentCollection.cs
--- a/Engine/Source/Programs/Horde/Plugins/Compute/HordeServer.Compute/Agents/AgentCollection.cs
+++ b/Engine/Source/Programs/Horde/Plugins/Compute/HordeServer.Compute/Agents/AgentCollection.cs
@@ -827,11 +827,11 @@ namespace HordeServer.Agents
 				List<string> newProperties = newDocument.Properties ?? new List<string>();
 				if (options.Capabilities != null)
 				{
-					options.Capabilities.Flatten(out newProperties, out Dictionary<string, int> newResources);
+					newProperties = new List<string>(options.Capabilities.Properties);
 
-					if (!ResourcesEqual(newResources, newDocument.Resources))
+					if (!ResourcesEqual(options.Capabilities.Resources, newDocument.Resources))
 					{
-						updates.Add(updateBuilder.Set(x => x.Resources, newResources));
+						updates.Add(updateBuilder.Set(x => x.Resources, options.Capabilities.Resources.ToDictionary()));
 					}
 				}
 				if (options.DynamicPools != null)
@@ -1065,14 +1065,19 @@ namespace HordeServer.Agents
 			AgentId agentId = agent.Id;
 			SessionId sessionId = SessionIdUtils.GenerateNewId();
 
-			options.Capabilities.Flatten(out List<string> newProperties, out Dictionary<string, int> newResources);
+			List<string> newProperties = options.Capabilities.Properties.ToList();
+			Dictionary<string, int> newResources = options.Capabilities.Resources.ToDictionary();
 
 			List<PoolId> newDynamicPools = new(options.DynamicPools);
 			List<PoolId> newPools = CreatePoolsList(agent.Document.ExplicitPools, newDynamicPools, newProperties);
 			SetPoolProperties(newProperties, newPools);
 
+			RpcAgentCapabilities capabilities = new RpcAgentCapabilities();
+			capabilities.Properties.Add(newProperties);
+			capabilities.Resources.Add(options.Capabilities.Resources);
+
 			// Attempt to create the session in Redis
-			RpcSession? newSession = await _scheduler.TryCreateSessionAsync(agentId, sessionId, options.Capabilities, cancellationToken);
+			RpcSession? newSession = await _scheduler.TryCreateSessionAsync(agentId, sessionId, capabilities, cancellationToken);
 			if (newSession == null)
 			{
 				return null;
diff --git a/Engine/Source/Programs/Horde/Plugins/Compute/HordeServer.Compute/Agents/Fleet/ComputeQueueAwsMetricStrategy.cs b/Engine/Source/Programs/Horde/Plugins/Compute/HordeServer.Compute/Agents/Fleet/ComputeQueueAwsMetricStrategy.cs
--- a/Engine/Source/Programs/Horde/Plugins/Compute/HordeServer.Compute/Agents/Fleet/ComputeQueueAwsMetricStrategy.cs
+++ b/Engine/Source/Programs/Horde/Plugins/Compute/HordeServer.Compute/Agents/Fleet/ComputeQueueAwsMetricStrategy.cs
@@ -98,7 +98,7 @@ namespace HordeServer.Agents.Fleet
 
 			Dictionary<string, List<MetricDatum>> metricsPerCloudWatchNamespace = new();
 			int numAgents = agents.Count;
-			int totalCpuCores = agents.Select(x => x.Resources.TryGetValue(KnownPropertyNames.LogicalCores, out int numCpuCores) ? numCpuCores : 0).Sum();
+			int totalCpuCores = agents.Select(x => x.Resources.TryGetValue(KnownResourceNames.LogicalCores, out int numCpuCores) ? numCpuCores : 0).Sum();
 			int numQueuedComputeTasks = await _computeTaskSource.GetNumQueuedTasksForPoolAsync(new ClusterId(_settings.ComputeClusterId), pool);
 
 			double numQueuedTasksPerAgent = numQueuedComputeTasks / (double)Math.Max(numAgents, 1);
diff --git a/Engine/Source/Programs/Horde/Plugins/Compute/HordeServer.Compute/Protos/horde/messages/agent_capabilities.proto.cs b/Engine/Source/Programs/Horde/Plugins/Compute/HordeServer.Compute/Protos/horde/messages/agent_capabilities.proto.cs
--- a/Engine/Source/Programs/Horde/Plugins/Compute/HordeServer.Compute/Protos/horde/messages/agent_capabilities.proto.cs
+++ b/Engine/Source/Programs/Horde/Plugins/Compute/HordeServer.Compute/Protos/horde/messages/agent_capabilities.proto.cs
@@ -13,41 +13,47 @@ namespace HordeCommon.Rpc.Messages
 			Properties.AddRange(properties);
 		}
 
-		public void Flatten(out List<string> properties, out Dictionary<string, int> resources)
+		public RpcAgentCapabilities MergeDevices()
 		{
-			properties = new List<string>();
-			resources = new Dictionary<string, int>();
-
-			properties.AddRange(Properties);
-
-			if (Devices.Count <= 0)
+#pragma warning disable CS0612 // Type or member is obsolete
+			if (Devices.Count == 0)
 			{
-				return;
+				return this;
 			}
-
-			RpcDeviceCapabilities device = Devices[0];
-			if (device.Properties == null)
+			else
 			{
-				return;
-			}
+				RpcAgentCapabilities other = new RpcAgentCapabilities();
+				other.Properties.Add(Properties);
+				other.Resources.Add(Resources);
 
-			properties.AddRange(device.Properties);
-			properties.Sort(StringComparer.OrdinalIgnoreCase);
+				if (Devices.Count > 0)
+				{
+					RpcDeviceCapabilities? primaryDevice = Devices[0];
+					other.Properties.Add(primaryDevice.Properties);
+					other.Resources.Add(primaryDevice.Resources);
+				}
 
-			CopyPropertyToResource(KnownPropertyNames.LogicalCores, properties, resources);
-			CopyPropertyToResource(KnownPropertyNames.Ram, properties, resources);
+				other.CopyPropertyToResource(KnownResourceNames.LogicalCores);
+				other.CopyPropertyToResource(KnownResourceNames.Ram);
+
+				return other;
+			}
+#pragma warning restore CS0612 // Type or member is obsolete
 		}
 
-		static void CopyPropertyToResource(string name, List<string> properties, Dictionary<string, int> resources)
+		void CopyPropertyToResource(string name)
 		{
-			foreach (string property in properties)
+			if (!Resources.ContainsKey(name))
 			{
-				if (property.Length > name.Length && property.StartsWith(name, StringComparison.OrdinalIgnoreCase) && property[name.Length] == '=')
+				foreach (string property in Properties)
 				{
-					int value;
-					if (Int32.TryParse(property.AsSpan(name.Length + 1), out value))
+					if (property.Length > name.Length && property.StartsWith(name, StringComparison.OrdinalIgnoreCase) && property[name.Length] == '=')
 					{
-						resources[name] = value;
+						int value;
+						if (Int32.TryParse(property.AsSpan(name.Length + 1), out value))
+						{
+							Resources.Add(name, value);
+						}
 					}
 				}
 			}
diff --git a/Engine/Source/Programs/Horde/Plugins/Compute/HordeServer.Compute/Server/RpcService.cs b/Engine/Source/Programs/Horde/Plugins/Compute/HordeServer.Compute/Server/RpcService.cs
--- a/Engine/Source/Programs/Horde/Plugins/Compute/HordeServer.Compute/Server/RpcService.cs
+++ b/Engine/Source/Programs/Horde/Plugins/Compute/HordeServer.Compute/Server/RpcService.cs
@@ -220,7 +220,7 @@ namespace HordeServer.Server
 			}
 
 			// Create a new session
-			agent = await _agentService.CreateSessionAsync(agent, request.Capabilities, request.Version, context.CancellationToken);
+			agent = await _agentService.CreateSessionAsync(agent, request.Capabilities.MergeDevices(), request.Version, context.CancellationToken);
 			if (agent == null)
 			{
 				throw new StructuredRpcException(StatusCode.NotFound, "Agent {AgentId} not found", agentId);
@@ -285,7 +285,7 @@ namespace HordeServer.Server
 					// Update the session
 					try
 					{
-						agent = await _agentService.UpdateSessionWithWaitAsync(agent, sessionId, (AgentStatus)request.Status, request.Capabilities, request.Leases, cancellationSource.Token);
+						agent = await _agentService.UpdateSessionWithWaitAsync(agent, sessionId, (AgentStatus)request.Status, request.Capabilities?.MergeDevices(), request.Leases, cancellationSource.Token);
 					}
 					catch (OperationCanceledException)
 					{
diff --git a/Engine/Source/Programs/Shared/EpicGames.Horde/Agents/KnownPropertyNames.cs b/Engine/Source/Programs/Shared/EpicGames.Horde/Agents/KnownPropertyNames.cs
--- a/Engine/Source/Programs/Shared/EpicGames.Horde/Agents/KnownPropertyNames.cs
+++ b/Engine/Source/Programs/Shared/EpicGames.Horde/Agents/KnownPropertyNames.cs
@@ -36,16 +36,6 @@ public static class KnownPropertyNames
 	/// Pools requested by the agent to join when registering with server
 	/// </summary>
 	public const string RequestedPools = "RequestedPools";
-
-	/// <summary>
-	/// Number of logical cores
-	/// </summary>
-	public const string LogicalCores = "LogicalCores";
-
-	/// <summary>
-	/// Amount of RAM, in GB
-	/// </summary>
-	public const string Ram = "RAM";
 	
 	/// <summary>
 	/// The total size of storage space on drive, in bytes
diff --git a/Engine/Source/Programs/Shared/EpicGames.Horde/Agents/KnownResourceNames.cs b/Engine/Source/Programs/Shared/EpicGames.Horde/Agents/KnownResourceNames.cs
new file mode 100644
--- /dev/null
+++ b/Engine/Source/Programs/Shared/EpicGames.Horde/Agents/KnownResourceNames.cs
@@ -0,0 +1,20 @@
+// Copyright Epic Games, Inc. All Rights Reserved.
+
+namespace EpicGames.Horde.Agents
+{
+	/// <summary>
+	/// Standard resource names
+	/// </summary>
+	public static class KnownResourceNames
+	{
+		/// <summary>
+		/// Number of logical cores
+		/// </summary>
+		public const string LogicalCores = "LogicalCores";
+
+		/// <summary>
+		/// Amount of RAM, in GB
+		/// </summary>
+		public const string Ram = "RAM";
+	}
+}
