Horde: Add a cache for agent capabilities, which includes a hash table of available properties for faster lookup.

#rnx

diff --git a/Engine/Source/Programs/Horde/Plugins/Compute/HordeServer.Compute.Tests/Agents/AgentSchedulerTests.cs b/Engine/Source/Programs/Horde/Plugins/Compute/HordeServer.Compute.Tests/Agents/AgentSchedulerTests.cs
--- a/Engine/Source/Programs/Horde/Plugins/Compute/HordeServer.Compute.Tests/Agents/AgentSchedulerTests.cs
+++ b/Engine/Source/Programs/Horde/Plugins/Compute/HordeServer.Compute.Tests/Agents/AgentSchedulerTests.cs
@@ -33,7 +33,7 @@ namespace HordeServer.Tests.Agents
 			RpcSession? session2 = await AgentScheduler.TryGetSessionAsync(session.SessionId);
 			Assert.AreEqual(session, session2);
 
-			RpcAgentCapabilities? caps2 = await AgentScheduler.TryGetSessionCapabilitiesAsync(session);
+			RpcAgentCapabilities? caps2 = await AgentScheduler.TryGetCapabilitiesAsync(session);
 			Assert.AreEqual(caps2, caps);
 		}
 
@@ -47,7 +47,7 @@ namespace HordeServer.Tests.Agents
 			RpcSession? session = await AgentScheduler.TryCreateSessionAsync(new AgentId("foo"), SessionIdUtils.GenerateNewId(), caps);
 			Assert.IsNotNull(session);
 
-			RpcAgentCapabilities? caps2 = await AgentScheduler.TryGetSessionCapabilitiesAsync(session);
+			RpcAgentCapabilities? caps2 = await AgentScheduler.TryGetCapabilitiesAsync(session);
 			Assert.AreEqual(caps, caps2);
 
 			// Update them
@@ -57,10 +57,7 @@ namespace HordeServer.Tests.Agents
 			RpcSession? session2 = await AgentScheduler.TryUpdateSessionAsync(session, newCapabilities: caps3);
 			Assert.IsNotNull(session2);
 
-			RpcAgentCapabilities? caps4 = await AgentScheduler.TryGetSessionCapabilitiesAsync(session);
-			Assert.IsNull(caps4);
-
-			RpcAgentCapabilities? caps5 = await AgentScheduler.TryGetSessionCapabilitiesAsync(session2);
+			RpcAgentCapabilities? caps5 = await AgentScheduler.TryGetCapabilitiesAsync(session2);
 			Assert.AreEqual(caps3, caps5);
 		}
 
diff --git a/Engine/Source/Programs/Horde/Plugins/Compute/HordeServer.Compute/Agents/AgentScheduler.cs b/Engine/Source/Programs/Horde/Plugins/Compute/HordeServer.Compute/Agents/AgentScheduler.cs
--- a/Engine/Source/Programs/Horde/Plugins/Compute/HordeServer.Compute/Agents/AgentScheduler.cs
+++ b/Engine/Source/Programs/Horde/Plugins/Compute/HordeServer.Compute/Agents/AgentScheduler.cs
@@ -210,22 +210,48 @@ namespace HordeServer.Agents
 			return await database.StringGetAsync(Keys.Sessions[sessionId].State);
 		}
 
-		/// <inheritdoc/>
-		public async Task<RpcAgentCapabilities?> TryGetSessionCapabilitiesAsync(RpcSession session, CancellationToken cancellationToken = default)
+		record class CapabilitiesLookup(IReadOnlySet<string> Properties, IReadOnlyDictionary<string, int> Resources)
 		{
-			IDatabase database = _redisService.GetDatabase();
+			public CapabilitiesLookup(RpcAgentCapabilities capabilities)
+				: this(new HashSet<string>(capabilities.Properties, StringComparer.Ordinal), new Dictionary<string, int>(capabilities.Resources, StringComparer.Ordinal)) { }
+		}
 
-			byte[]? data = (byte[]?)await database.StringGetAsync(Keys.Sessions[session.SessionId].Capabilities.Inner);
-			if (data == null)
-			{
-				return null;
-			}
-			if (!session.CapabilitiesHash.Equals(IoHash.Compute(data).ToString(), StringComparison.Ordinal))
+		record class CapabilitiesCacheEntry(string Hash, CapabilitiesLookup Capabilities, RpcAgentCapabilities Message);
+
+		async ValueTask<CapabilitiesCacheEntry?> TryGetCachedCapabilitiesAsync(RpcSession session, CancellationToken cancellationToken = default)
+		{
+			CapabilitiesCacheEntry? cachedCapabilities;
+			if (!_memoryCache.TryGetValue(session.SessionId, out cachedCapabilities) || cachedCapabilities == null || cachedCapabilities.Hash != session.CapabilitiesHash)
 			{
-				return null;
+				IDatabase database = _redisService.GetDatabase();
+
+				byte[]? data = (byte[]?)await database.StringGetAsync(Keys.Sessions[session.SessionId].Capabilities.Inner);
+				if (data == null)
+				{
+					return null;
+				}
+				if (!session.CapabilitiesHash.Equals(IoHash.Compute(data).ToString(), StringComparison.Ordinal))
+				{
+					return null;
+				}
+
+				RpcAgentCapabilities capabilities = RpcAgentCapabilities.Parser.ParseFrom(data);
+				cachedCapabilities = new CapabilitiesCacheEntry(session.CapabilitiesHash, new CapabilitiesLookup(capabilities), capabilities);
+
+				using (ICacheEntry entry = _memoryCache.CreateEntry(session.SessionId))
+				{
+					entry.SetSlidingExpiration(TimeSpan.FromMinutes(30.0));
+					entry.Value = cachedCapabilities;
+				}
 			}
+			return cachedCapabilities;
+		}
 
-			return RpcAgentCapabilities.Parser.ParseFrom(data);
+		/// <inheritdoc/>
+		public async Task<RpcAgentCapabilities?> TryGetCapabilitiesAsync(RpcSession session, CancellationToken cancellationToken = default)
+		{
+			CapabilitiesCacheEntry? cachedCapabilities = await TryGetCachedCapabilitiesAsync(session, cancellationToken);
+			return cachedCapabilities?.Message;
 		}
 
 		/// <inheritdoc/>
@@ -463,9 +489,9 @@ namespace HordeServer.Agents
 			}
 		}
 
-		async Task<bool> UpdateFiltersAsync(CachedFilters cachedFilters, ITransaction transaction, RpcSession session, RpcAgentCapabilities? capabilities, CancellationToken cancellationToken)
+		async Task<bool> UpdateFiltersAsync(CachedFilters cachedFilters, ITransaction transaction, RpcSession newSession, RpcAgentCapabilities? newCapabilities, CancellationToken cancellationToken)
 		{
-			long minFilterTime = Math.Max(0, session.LastFilterUpdateTicks - TimeSpan.FromSeconds(30.0).Ticks);
+			long minFilterTime = Math.Max(0, newSession.LastFilterUpdateTicks - TimeSpan.FromSeconds(30.0).Ticks);
 
 			// Add the session to any new matching filters
 			SortedSetEntry<IoHash>[] createdFilters = cachedFilters.Created;
@@ -478,7 +504,16 @@ namespace HordeServer.Agents
 
 			if (createdIdx < createdFilters.Length)
 			{
-				capabilities ??= await TryGetSessionCapabilitiesAsync(session, cancellationToken);
+				CapabilitiesLookup? capabilities;
+				if (newCapabilities != null)
+				{
+					capabilities = new CapabilitiesLookup(newCapabilities);
+				}
+				else
+				{
+					capabilities = (await TryGetCachedCapabilitiesAsync(newSession, cancellationToken))?.Capabilities;
+				}
+
 				if (capabilities == null)
 				{
 					return false;
@@ -492,8 +527,8 @@ namespace HordeServer.Agents
 					if (requirements != null && MeetsRequirements(capabilities, requirements))
 					{
 						transaction.AddCondition(Keys.Filters.Touched.HashExists(requirementsHash));
-						_ = transaction.SetAddAsync(Keys.Filters[requirementsHash].Sessions, session.SessionId, flags: CommandFlags.FireAndForget);
-						_ = transaction.SetAddAsync(Keys.Sessions[session.SessionId].Filters, requirementsHash, flags: CommandFlags.FireAndForget);
+						_ = transaction.SetAddAsync(Keys.Filters[requirementsHash].Sessions, newSession.SessionId, flags: CommandFlags.FireAndForget);
+						_ = transaction.SetAddAsync(Keys.Sessions[newSession.SessionId].Filters, requirementsHash, flags: CommandFlags.FireAndForget);
 					}
 				}
 			}
@@ -511,15 +546,15 @@ namespace HordeServer.Agents
 			{
 				IoHash queueHash = deletedFilters[deletedIdx].Element;
 				transaction.AddCondition(Keys.Filters.Touched.HashNotExists(queueHash));
-				_ = transaction.SetRemoveAsync(Keys.Sessions[session.SessionId].Filters, queueHash, flags: CommandFlags.FireAndForget);
+				_ = transaction.SetRemoveAsync(Keys.Sessions[newSession.SessionId].Filters, queueHash, flags: CommandFlags.FireAndForget);
 			}
 
 			// Update the last queue time
-			session.LastFilterUpdateTicks = cachedFilters.Ticks;
+			newSession.LastFilterUpdateTicks = cachedFilters.Ticks;
 			return true;
 		}
 
-		static bool MeetsRequirements(RpcAgentCapabilities capabilities, RpcAgentRequirements requirements)
+		static bool MeetsRequirements(CapabilitiesLookup capabilities, RpcAgentRequirements requirements)
 		{
 			foreach (string property in requirements.Properties)
 			{
diff --git a/Engine/Source/Programs/Horde/Plugins/Compute/HordeServer.Compute/Agents/IAgentScheduler.cs b/Engine/Source/Programs/Horde/Plugins/Compute/HordeServer.Compute/Agents/IAgentScheduler.cs
--- a/Engine/Source/Programs/Horde/Plugins/Compute/HordeServer.Compute/Agents/IAgentScheduler.cs
+++ b/Engine/Source/Programs/Horde/Plugins/Compute/HordeServer.Compute/Agents/IAgentScheduler.cs
@@ -44,7 +44,7 @@ namespace HordeServer.Agents
 		/// <param name="session">The session object</param>
 		/// <param name="cancellationToken">Cancellation token for the operation</param>
 		/// <returns>Matching capabilities object for the session</returns>
-		Task<RpcAgentCapabilities?> TryGetSessionCapabilitiesAsync(RpcSession session, CancellationToken cancellationToken = default);
+		Task<RpcAgentCapabilities?> TryGetCapabilitiesAsync(RpcSession session, CancellationToken cancellationToken = default);
 
 		/// <summary>
 		/// Attempts to create a new session for the given agent. Fails if the agent is already executing a session.
