Horde: Store the agent requirements when we update the cached filter states.

#rnx

diff --git a/Engine/Source/Programs/Horde/Plugins/Compute/HordeServer.Compute/Agents/AgentScheduler.cs b/Engine/Source/Programs/Horde/Plugins/Compute/HordeServer.Compute/Agents/AgentScheduler.cs
--- a/Engine/Source/Programs/Horde/Plugins/Compute/HordeServer.Compute/Agents/AgentScheduler.cs
+++ b/Engine/Source/Programs/Horde/Plugins/Compute/HordeServer.Compute/Agents/AgentScheduler.cs
@@ -110,8 +110,10 @@ namespace HordeServer.Agents
 
 		static readonly RedisChannel<SessionId> s_sessionUpdateChannel = new RedisChannel<SessionId>(RedisChannel.Literal("compute:sessions:update"));
 
-		record class CachedFilters(long Ticks, SortedSetEntry<IoHash>[] Created, SortedSetEntry<IoHash>[] Deleted);
-		CachedFilters _cachedFilters = new CachedFilters(0, Array.Empty<SortedSetEntry<IoHash>>(), Array.Empty<SortedSetEntry<IoHash>>());
+		record class CreatedFilter(IoHash Hash, RpcAgentRequirements Requirements, double Score);
+
+		record class CachedFilters(long Ticks, CreatedFilter[] Created, SortedSetEntry<IoHash>[] Deleted);
+		CachedFilters _cachedFilters = new CachedFilters(0, Array.Empty<CreatedFilter>(), Array.Empty<SortedSetEntry<IoHash>>());
 
 		IAsyncDisposable? _updateEventSubscription;
 
@@ -414,15 +416,43 @@ namespace HordeServer.Agents
 		/// </summary>
 		async ValueTask UpdateCachedFiltersTickAsync(CancellationToken cancellationToken)
 		{
-			IDatabase database = _redisService.GetDatabase();
+			for (; ; )
+			{
+				if (await TryUpdateCachedFiltersAsync(cancellationToken))
+				{
+					break;
+				}
+			}
+		}
+
+		async Task<bool> TryUpdateCachedFiltersAsync(CancellationToken cancellationToken)
+		{
 			DateTime utcNow = _clock.UtcNow;
+			IDatabase database = _redisService.GetDatabase();
 
 			ITransaction transaction = database.CreateTransaction();
-			Task<SortedSetEntry<IoHash>[]> created = transaction.SortedSetRangeByRankWithScoresAsync(Keys.Filters.Created);
-			Task<SortedSetEntry<IoHash>[]> deleted = transaction.SortedSetRangeByRankWithScoresAsync(Keys.Filters.Deleted);
-
+			Task<SortedSetEntry<IoHash>[]> createdEntriesTask = transaction.SortedSetRangeByRankWithScoresAsync(Keys.Filters.Created);
+			Task<SortedSetEntry<IoHash>[]> deletedEntriesTask = transaction.SortedSetRangeByRankWithScoresAsync(Keys.Filters.Deleted);
 			await transaction.ExecuteAsync().WaitAsync(cancellationToken);
-			_cachedFilters = new CachedFilters(utcNow.Ticks, await created, await deleted);
+
+			SortedSetEntry<IoHash>[] createdEntries = await createdEntriesTask;
+			CreatedFilter[] createdFilters = new CreatedFilter[createdEntries.Length];
+
+			for (int idx = 0; idx < createdEntries.Length; idx++)
+			{
+				SortedSetEntry<IoHash> createdEntry = createdEntries[idx];
+
+				RpcAgentRequirements? requirements = await TryGetFilterRequirementsAsync(createdEntry.Element, cancellationToken);
+				if (requirements == null)
+				{
+					return false;
+				}
+
+				createdFilters[idx] = new CreatedFilter(createdEntry.Element, requirements, createdEntry.Score);
+			}
+
+			_cachedFilters = new CachedFilters(utcNow.Ticks, createdFilters, await deletedEntriesTask);
+			return true;
 		}
 
 		/// <summary>
@@ -494,9 +524,9 @@ namespace HordeServer.Agents
 			long minFilterTime = Math.Max(0, newSession.LastFilterUpdateTicks - TimeSpan.FromSeconds(30.0).Ticks);
 
 			// Add the session to any new matching filters
-			SortedSetEntry<IoHash>[] createdFilters = cachedFilters.Created;
+			CreatedFilter[] createdFilters = cachedFilters.Created;
 
-			int createdIdx = createdFilters.BinarySearch(new SortedSetEntry<IoHash>(IoHash.Zero, minFilterTime));
+			int createdIdx = createdFilters.BinarySearch(x => x.Score, minFilterTime);
 			if (createdIdx < 0)
 			{
 				createdIdx = ~createdIdx;
@@ -521,14 +551,12 @@ namespace HordeServer.Agents
 
 				for (; createdIdx < createdFilters.Length; createdIdx++)
 				{
-					IoHash requirementsHash = createdFilters[createdIdx].Element;
-
-					RpcAgentRequirements? requirements = await TryGetFilterRequirementsAsync(requirementsHash, cancellationToken);
-					if (requirements != null && MeetsRequirements(capabilities, requirements))
+					CreatedFilter createdFilter = createdFilters[createdIdx];
+					if (MeetsRequirements(capabilities, createdFilter.Requirements))
 					{
-						transaction.AddCondition(Keys.Filters.Touched.HashExists(requirementsHash));
-						_ = transaction.SetAddAsync(Keys.Filters[requirementsHash].Sessions, newSession.SessionId, flags: CommandFlags.FireAndForget);
-						_ = transaction.SetAddAsync(Keys.Sessions[newSession.SessionId].Filters, requirementsHash, flags: CommandFlags.FireAndForget);
+						transaction.AddCondition(Keys.Filters.Touched.HashExists(createdFilter.Hash));
+						_ = transaction.SetAddAsync(Keys.Filters[createdFilter.Hash].Sessions, newSession.SessionId, flags: CommandFlags.FireAndForget);
+						_ = transaction.SetAddAsync(Keys.Sessions[newSession.SessionId].Filters, createdFilter.Hash, flags: CommandFlags.FireAndForget);
 					}
 				}
 			}
