Horde: Document use of ssl: prefix for connecting to Perforce servers

#skipci

diff --git a/Engine/Source/Programs/Horde/Docs/Deployment/Integrations/Perforce.md b/Engine/Source/Programs/Horde/Docs/Deployment/Integrations/Perforce.md
--- a/Engine/Source/Programs/Horde/Docs/Deployment/Integrations/Perforce.md
+++ b/Engine/Source/Programs/Horde/Docs/Deployment/Integrations/Perforce.md
@@ -31,6 +31,8 @@ There are several configurables for each cluster:
 
 * `Name`: Used to reference the cluster from a stream in the CI system.
 * `Servers`: Each server supports several settings of its own:
+  * `ServerAndPort`: The server part may be a DNS entry with multiple records, in which case it will be actively load balanced.
+    If `ssl:` prefix is used, ensure P4 server's fingerprint/certificate is trusted (see docs for [p4 trust](https://www.perforce.com/manuals/cmdref/Content/CmdRef/p4_trust.html))
   * `ResolveDns`: If true, the given DNS name is resolved to find a concrete list of servers to be used. This allows
   IT/infrastructure teams to add and remove servers to a cluster without reconfiguring Horde.
   * `Properties`: Specifies properties that the agent must have to select this server.
diff --git a/Engine/Source/Programs/Horde/HordeServer/Defaults/globals.json b/Engine/Source/Programs/Horde/HordeServer/Defaults/globals.json
--- a/Engine/Source/Programs/Horde/HordeServer/Defaults/globals.json
+++ b/Engine/Source/Programs/Horde/HordeServer/Defaults/globals.json
@@ -39,7 +39,8 @@
 //			"servers": 
 //			[
 //				{
-//					// DNS name and port for the server
+//					// DNS name and port for the server.
+//					// If connecting using "ssl:" prefix, ensure P4 server's fingerprint/certificate is trusted (see docs for "p4 trust")
 //					"serverAndPort": "perforce:1666"
 //				}
 //			],
diff --git a/Engine/Source/Programs/Horde/Plugins/Build/HordeServer.Build/BuildConfig.cs b/Engine/Source/Programs/Horde/Plugins/Build/HordeServer.Build/BuildConfig.cs
--- a/Engine/Source/Programs/Horde/Plugins/Build/HordeServer.Build/BuildConfig.cs
+++ b/Engine/Source/Programs/Horde/Plugins/Build/HordeServer.Build/BuildConfig.cs
@@ -356,6 +356,7 @@ namespace HordeServer
 	{
 		/// <summary>
 		/// The server and port. The server may be a DNS entry with multiple records, in which case it will be actively load balanced.
+		/// If "ssl:" prefix is used, ensure P4 server's fingerprint/certificate is trusted (see docs for "p4 trust")
 		/// </summary>
 		public string ServerAndPort { get; set; } = "perforce:1666";
 
@@ -375,12 +376,12 @@ namespace HordeServer
 		public int MaxConformCount { get; set; }
 
 		/// <summary>
-		/// Optional condition for a machine to be eligable to use this server
+		/// Optional condition for a machine to be eligible to use this server
 		/// </summary>
 		public Condition? Condition { get; set; }
 
 		/// <summary>
-		/// List of properties for an agent to be eligable to use this server
+		/// List of properties for an agent to be eligible to use this server
 		/// </summary>
 		public List<string>? Properties { get; set; }
 	}
