Horde: Add agent settings for CPU count and multiplier

Only provided as hinting to workloads, which can chose to respect these. The initial use-case is letting UBA limit number of CPUs in use, much how maxcpu/mulcpu works for UBT.
Later on these values could configure job objects (on Windows) for proper OS-enforced CPU limiting.

diff --git a/Engine/Source/Programs/Horde/HordeAgent/AgentSettings.cs b/Engine/Source/Programs/Horde/HordeAgent/AgentSettings.cs
--- a/Engine/Source/Programs/Horde/HordeAgent/AgentSettings.cs
+++ b/Engine/Source/Programs/Horde/HordeAgent/AgentSettings.cs
@@ -192,6 +192,19 @@ namespace HordeAgent
 		/// Maximum size of the bundle cache, in megabytes.
 		/// </summary>
 		public long BundleCacheSize { get; set; } = 1024;
+		
+		/// <summary>
+		/// Maximum number of logical CPU cores workloads should use
+		/// Currently this is only provided as a hint and requires leases to respect this value as it's set via an env variable (UE_HORDE_CPU_COUNT).
+		/// </summary>
+		public int CpuCount { get; set; } = Environment.ProcessorCount;
+		
+		/// <summary>
+		/// CPU core multiplier applied to CPU core count setting
+		/// For example, 32 CPU cores and a multiplier of 0.5 results in max 16 CPU usage.
+		/// <see cref="CpuCount" />
+		/// </summary>
+		public double CpuMultiplier { get; set; } = 1.0;
 
 		/// <summary>
 		/// Key/value properties in addition to those set internally by the agent
diff --git a/Engine/Source/Programs/Horde/HordeAgent/Leases/Handlers/ComputeHandler.cs b/Engine/Source/Programs/Horde/HordeAgent/Leases/Handlers/ComputeHandler.cs
--- a/Engine/Source/Programs/Horde/HordeAgent/Leases/Handlers/ComputeHandler.cs
+++ b/Engine/Source/Programs/Horde/HordeAgent/Leases/Handlers/ComputeHandler.cs
@@ -1,5 +1,6 @@
 // Copyright Epic Games, Inc. All Rights Reserved.
 
+using System.Globalization;
 using System.Net.Sockets;
 using EpicGames.Core;
 using EpicGames.Horde.Agents.Leases;
@@ -77,14 +78,17 @@ namespace HordeAgent.Leases.Handlers
 							try
 							{
 								DirectoryReference.CreateDirectory(sandboxDir);
-
 								DirectoryReference sharedDir = DirectoryReference.Combine(session.WorkingDir, "Saved");
 								DirectoryReference.CreateDirectory(sharedDir);
 
-								Dictionary<string, string?> newEnvVars = new Dictionary<string, string?>();
-								newEnvVars["UE_HORDE_SHARED_DIR"] = sharedDir.FullName;
-								newEnvVars["UE_HORDE_TERMINATION_SIGNAL_FILE"] = _settings.GetTerminationSignalFile().FullName;
-
+								Dictionary<string, string?> newEnvVars = new ()
+								{
+									["UE_HORDE_SHARED_DIR"] = sharedDir.FullName,
+									["UE_HORDE_TERMINATION_SIGNAL_FILE"] = _settings.GetTerminationSignalFile().FullName,
+									["UE_HORDE_CPU_COUNT"] = Convert.ToString(_settings.CpuCount),
+									["UE_HORDE_CPU_MULTIPLIER"] = Convert.ToString(_settings.CpuMultiplier, CultureInfo.InvariantCulture),
+								};
+								
 								AgentMessageHandler worker = new AgentMessageHandler(sandboxDir, newEnvVars, false, _settings.WineExecutablePath, _settings.ContainerEngineExecutablePath, logger);
 								await worker.RunAsync(socket, cts.Token);
 								await socket.CloseAsync(cts.Token);
