Horde: Add a typed accessor for the capabilities hash on session objects.

#rnx

diff --git a/Engine/Source/Programs/Horde/Plugins/Compute/HordeServer.Compute/Agents/AgentCollection.cs b/Engine/Source/Programs/Horde/Plugins/Compute/HordeServer.Compute/Agents/AgentCollection.cs
--- a/Engine/Source/Programs/Horde/Plugins/Compute/HordeServer.Compute/Agents/AgentCollection.cs
+++ b/Engine/Source/Programs/Horde/Plugins/Compute/HordeServer.Compute/Agents/AgentCollection.cs
@@ -873,10 +873,6 @@ namespace HordeServer.Agents
 			{
 				newSession.Status = (RpcAgentStatus)options.Status;
 			}
-			if (options.Capabilities != null)
-			{
-				newSession.CapabilitiesHash = IoHash.Compute(options.Capabilities.ToByteArray()).ToString();
-			}
 			if (options.Leases != null)
 			{
 				Dictionary<LeaseId, RpcLease> idToRemoteLease = options.Leases.ToDictionary(x => x.Id, x => x);
diff --git a/Engine/Source/Programs/Horde/Plugins/Compute/HordeServer.Compute/Agents/AgentScheduler.cs b/Engine/Source/Programs/Horde/Plugins/Compute/HordeServer.Compute/Agents/AgentScheduler.cs
--- a/Engine/Source/Programs/Horde/Plugins/Compute/HordeServer.Compute/Agents/AgentScheduler.cs
+++ b/Engine/Source/Programs/Horde/Plugins/Compute/HordeServer.Compute/Agents/AgentScheduler.cs
@@ -183,7 +183,7 @@ namespace HordeServer.Agents
 			RpcSession newSession = new RpcSession();
 			newSession.AgentId = agentId;
 			newSession.SessionId = sessionId;
-			newSession.CapabilitiesHash = capabilitiesHash.ToString();
+			newSession.CapabilitiesHash = capabilitiesHash;
 			newSession.ExpiryTicks = (utcNow + RpcSession.ExpireAfterTime).Ticks;
 			newSession.Status = RpcAgentStatus.Ok;
 
@@ -218,7 +218,7 @@ namespace HordeServer.Agents
 				: this(new HashSet<string>(capabilities.Properties, StringComparer.Ordinal), new Dictionary<string, int>(capabilities.Resources, StringComparer.Ordinal)) { }
 		}
 
-		record class CapabilitiesCacheEntry(string Hash, CapabilitiesLookup Capabilities, RpcAgentCapabilities Message);
+		record class CapabilitiesCacheEntry(IoHash Hash, CapabilitiesLookup Capabilities, RpcAgentCapabilities Message);
 
 		async ValueTask<CapabilitiesCacheEntry?> TryGetCachedCapabilitiesAsync(RpcSession session, CancellationToken cancellationToken = default)
 		{
@@ -232,7 +232,9 @@ namespace HordeServer.Agents
 				{
 					return null;
 				}
-				if (!session.CapabilitiesHash.Equals(IoHash.Compute(data).ToString(), StringComparison.Ordinal))
+
+				IoHash hash = IoHash.Compute(data);
+				if (session.CapabilitiesHash != hash)
 				{
 					return null;
 				}
@@ -290,7 +292,7 @@ namespace HordeServer.Agents
 				if (newCapabilities != null)
 				{
 					byte[] capabilitiesData = newCapabilities.ToByteArray();
-					string capabilitiesHash = IoHash.Compute(capabilitiesData).ToString();
+					IoHash capabilitiesHash = IoHash.Compute(capabilitiesData);
 
 					if (capabilitiesHash != newSession.CapabilitiesHash)
 					{
@@ -345,7 +347,7 @@ namespace HordeServer.Agents
 			_logger.LogDebug("Updated session {SessionId} expiry time to {ExpiryTime}", sessionId, newSession.ExpiryTime);
 
 			// Notify watchers that the session state has changed
-			if (newSession.Status != session.Status || !newSession.CapabilitiesHash.Equals(session.CapabilitiesHash, StringComparison.Ordinal) || !newSession.Leases.Equals(session.Leases))
+			if (newSession.Status != session.Status || newSession.CapabilitiesHash != session.CapabilitiesHash || !newSession.Leases.Equals(session.Leases))
 			{
 				_ = _redisService.GetDatabase().PublishAsync(s_sessionUpdateChannel, sessionId, CommandFlags.FireAndForget);
 			}
diff --git a/Engine/Source/Programs/Horde/Plugins/Compute/HordeServer.Compute/Protos/horde/messages/session.proto.cs b/Engine/Source/Programs/Horde/Plugins/Compute/HordeServer.Compute/Protos/horde/messages/session.proto.cs
--- a/Engine/Source/Programs/Horde/Plugins/Compute/HordeServer.Compute/Protos/horde/messages/session.proto.cs
+++ b/Engine/Source/Programs/Horde/Plugins/Compute/HordeServer.Compute/Protos/horde/messages/session.proto.cs
@@ -1,5 +1,6 @@
 // Copyright Epic Games, Inc. All Rights Reserved.
 
+using EpicGames.Core;
 using EpicGames.Horde.Agents;
 using EpicGames.Horde.Agents.Leases;
 using EpicGames.Horde.Agents.Sessions;
@@ -25,6 +26,12 @@ namespace HordeCommon.Rpc.Messages
 			set => SessionIdValue = value.ToString();
 		}
 
+		public IoHash CapabilitiesHash
+		{
+			get => IoHash.Parse(CapabilitiesHashValue);
+			set => CapabilitiesHashValue = value.ToString();
+		}
+
 		public DateTime ExpiryTime
 			=> new DateTime(ExpiryTicks, DateTimeKind.Utc);
 	}
