Horde: Keep track of the agents which are currently able to execute sessions.

Each server now caches the state of a session based on the last update time, allowing it to perform most operations with minimal context fetches in the common case where an agent is always being updated on the same server.

diff --git a/Engine/Source/Programs/Horde/HordeServer/ServerApp.cs b/Engine/Source/Programs/Horde/HordeServer/ServerApp.cs
--- a/Engine/Source/Programs/Horde/HordeServer/ServerApp.cs
+++ b/Engine/Source/Programs/Horde/HordeServer/ServerApp.cs
@@ -14,6 +14,7 @@ using HordeServer.Commands;
 using HordeServer.Plugins;
 using HordeServer.Storage;
 using HordeServer.Utilities;
+using Microsoft.Extensions.Caching.Memory;
 using Microsoft.Extensions.Configuration;
 using Microsoft.Extensions.Configuration.Memory;
 using Microsoft.Extensions.DependencyInjection;
diff --git a/Engine/Source/Programs/Horde/Plugins/Compute/HordeServer.Compute.Tests/Agents/AgentSchedulerTests.cs b/Engine/Source/Programs/Horde/Plugins/Compute/HordeServer.Compute.Tests/Agents/AgentSchedulerTests.cs
--- a/Engine/Source/Programs/Horde/Plugins/Compute/HordeServer.Compute.Tests/Agents/AgentSchedulerTests.cs
+++ b/Engine/Source/Programs/Horde/Plugins/Compute/HordeServer.Compute.Tests/Agents/AgentSchedulerTests.cs
@@ -2,12 +2,14 @@
 
 using EpicGames.Core;
 using EpicGames.Horde.Agents;
+using EpicGames.Horde.Agents.Leases;
 using EpicGames.Horde.Agents.Sessions;
 using HordeCommon;
 using HordeCommon.Rpc;
 using HordeCommon.Rpc.Messages;
 using HordeServer.Agents;
 using HordeServer.Agents.Sessions;
+using HordeServer.Utilities;
 using Microsoft.Extensions.DependencyInjection;
 
 namespace HordeServer.Tests.Agents
@@ -120,10 +122,10 @@ namespace HordeServer.Tests.Agents
 			// Update the cached queues
 			await scheduler.ForceUpdateFiltersAsync(CancellationToken.None);
 
-			SessionId[] sessionIds1 = await scheduler.GetAllFilteredSessionsAsync(queueHash1);
+			SessionId[] sessionIds1 = await scheduler.EnumerateFilteredSessionsAsync(queueHash1, SessionFilterType.Potential).ToArrayAsync();
 			Assert.AreEqual(3, sessionIds1.Length);
 
-			SessionId[] sessionIds2 = await scheduler.GetAllFilteredSessionsAsync(queueHash2);
+			SessionId[] sessionIds2 = await scheduler.EnumerateFilteredSessionsAsync(queueHash2, SessionFilterType.Potential).ToArrayAsync();
 			Assert.AreEqual(5, sessionIds2.Length);
 		}
 
@@ -170,10 +172,10 @@ namespace HordeServer.Tests.Agents
 			}
 
 			//
-			SessionId[] sessionIds1 = await scheduler.GetAllFilteredSessionsAsync(queueHash1);
+			SessionId[] sessionIds1 = await scheduler.EnumerateFilteredSessionsAsync(queueHash1, SessionFilterType.Potential).ToArrayAsync();
 			Assert.AreEqual(3, sessionIds1.Length);
 
-			SessionId[] sessionIds2 = await scheduler.GetAllFilteredSessionsAsync(queueHash2);
+			SessionId[] sessionIds2 = await scheduler.EnumerateFilteredSessionsAsync(queueHash2, SessionFilterType.Potential).ToArrayAsync();
 			Assert.AreEqual(5, sessionIds2.Length);
 		}
 
@@ -206,5 +208,163 @@ namespace HordeServer.Tests.Agents
 			await clock.AdvanceAsync(HordeServer.Agents.AgentScheduler.ExpireFiltersTime);
 			Assert.AreEqual(0, (await scheduler.GetFiltersAsync()).Length);
 		}
+
+		[TestMethod]
+		public async Task FilterAvailableTestAsync()
+		{
+			FakeClock clock = ServiceProvider.GetRequiredService<FakeClock>();
+			AgentScheduler scheduler = (AgentScheduler)AgentScheduler;
+			await scheduler.StartAsync(default);
+
+			// Create the first filter (non-shared windows agents)
+			RpcAgentRequirements reqs1 = new RpcAgentRequirements();
+			reqs1.Properties.Add("Platform=Windows");
+			reqs1.Shared = false;
+
+			IoHash reqsHash1 = await scheduler.CreateFilterAsync(reqs1);
+
+			// Create the second filter (windows agents with 1 core free)
+			RpcAgentRequirements reqs2 = new RpcAgentRequirements();
+			reqs2.Properties.Add("Platform=Windows");
+			reqs2.Resources.Add("LogicalCores", 1);
+			reqs2.Shared = true;
+
+			IoHash reqsHash2 = await scheduler.CreateFilterAsync(reqs2);
+
+			// Create the second filter (windows agents with 2 cores free)
+			RpcAgentRequirements reqs3 = new RpcAgentRequirements();
+			reqs3.Properties.Add("Platform=Windows");
+			reqs3.Resources.Add("LogicalCores", 2);
+			reqs3.Shared = true;
+
+			IoHash reqsHash3 = await scheduler.CreateFilterAsync(reqs3);
+
+			// Create the third filter (windows agents with at least 16gb RAM)
+			RpcAgentRequirements reqs4 = new RpcAgentRequirements();
+			reqs4.Properties.Add("Platform=Windows");
+			reqs4.Resources.Add("RAM", 16);
+			reqs4.Shared = true;
+
+			IoHash reqsHash4 = await scheduler.CreateFilterAsync(reqs4);
+
+			// Create the fourth filter (mac agents)
+			RpcAgentRequirements reqs5 = new RpcAgentRequirements();
+			reqs5.Properties.Add("Platform=MacOS");
+
+			IoHash reqsHash5 = await scheduler.CreateFilterAsync(reqs5);
+
+			// Update the clock so the scheduler caches the filters
+			await clock.AdvanceAsync(TimeSpan.FromMinutes(1.0));
+
+			// Create some sessions
+			RpcAgentCapabilities winCaps = new RpcAgentCapabilities();
+			winCaps.Properties.Add("Platform=Windows");
+			winCaps.Resources.Add("LogicalCores", 2);
+			winCaps.Resources.Add("RAM", 16);
+
+			RpcAgentCapabilities macCaps = new RpcAgentCapabilities();
+			macCaps.Properties.Add("Platform=MacOS");
+
+			RpcSession? session1 = await AgentScheduler.TryCreateSessionAsync(new AgentId("session1"), SessionIdUtils.GenerateNewId(), winCaps);
+			Assert.IsNotNull(session1);
+
+			RpcSession? session2 = await AgentScheduler.TryCreateSessionAsync(new AgentId("session2"), SessionIdUtils.GenerateNewId(), winCaps);
+			Assert.IsNotNull(session2);
+
+			RpcSession? session3 = await AgentScheduler.TryCreateSessionAsync(new AgentId("session3"), SessionIdUtils.GenerateNewId(), winCaps);
+			Assert.IsNotNull(session3);
+
+			RpcSession? session4 = await AgentScheduler.TryCreateSessionAsync(new AgentId("session4"), SessionIdUtils.GenerateNewId(), macCaps);
+			Assert.IsNotNull(session4);
+
+			async Task CheckFilterLengthsAsync(SessionFilterType type, int length1 = 3, int length2 = 3, int length3 = 3, int length4 = 3, int length5 = 1)
+			{
+				SessionId[] sessionIds1 = await scheduler.EnumerateFilteredSessionsAsync(reqsHash1, type).ToArrayAsync();
+				SessionId[] sessionIds2 = await scheduler.EnumerateFilteredSessionsAsync(reqsHash2, type).ToArrayAsync();
+				SessionId[] sessionIds3 = await scheduler.EnumerateFilteredSessionsAsync(reqsHash3, type).ToArrayAsync();
+				SessionId[] sessionIds4 = await scheduler.EnumerateFilteredSessionsAsync(reqsHash4, type).ToArrayAsync();
+				SessionId[] sessionIds5 = await scheduler.EnumerateFilteredSessionsAsync(reqsHash5, type).ToArrayAsync();
+				Assert.AreEqual(length1, sessionIds1.Length);
+				Assert.AreEqual(length2, sessionIds2.Length);
+				Assert.AreEqual(length3, sessionIds3.Length);
+				Assert.AreEqual(length4, sessionIds4.Length);
+				Assert.AreEqual(length5, sessionIds5.Length);
+			}
+
+			await CheckFilterLengthsAsync(SessionFilterType.Available);
+			await CheckFilterLengthsAsync(SessionFilterType.Potential);
+
+			// Add a lease to session 2, and check it's no longer reported as available in the exclusive queue (1), and two-cpu queue (3)
+			{
+				RpcSessionLease newLease = new RpcSessionLease();
+				newLease.Id = new LeaseId(BinaryIdUtils.CreateNew());
+				newLease.Resources.Add("LogicalCores", 1);
+
+				RpcSession newSession2 = new RpcSession(session2);
+				newSession2.Leases.Add(newLease);
+
+				session2 = await AgentScheduler.TryUpdateSessionAsync(session2, newSession2);
+				Assert.IsNotNull(session2);
+			}
+
+			await CheckFilterLengthsAsync(SessionFilterType.Available, length1: 2, length3: 2); // filter 1 (exclusive), filter 3 (2 cpus)
+			await CheckFilterLengthsAsync(SessionFilterType.Potential);
+
+			// Add another lease to session 2, and check that it's also removed from the one-cpu queue (1)
+			{
+				RpcSessionLease newLease = new RpcSessionLease();
+				newLease.Id = new LeaseId(BinaryIdUtils.CreateNew());
+				newLease.Resources.Add("LogicalCores", 1);
+
+				RpcSession newSession2 = new RpcSession(session2);
+				newSession2.Leases.Add(newLease);
+
+				session2 = await AgentScheduler.TryUpdateSessionAsync(session2, newSession2);
+				Assert.IsNotNull(session2);
+			}
+
+			await CheckFilterLengthsAsync(SessionFilterType.Available, length1: 2, length2: 2, length3: 2); // filter 1 (exclusive), filter 2 (1 cpu), filter 3 (2 cpus)
+			await CheckFilterLengthsAsync(SessionFilterType.Potential);
+
+			// Add a lease to session 3 that also requests ram, and and check that it's also removed from filter 4
+			{
+				RpcSessionLease newLease = new RpcSessionLease();
+				newLease.Id = new LeaseId(BinaryIdUtils.CreateNew());
+				newLease.Resources.Add("RAM", 1);
+
+				RpcSession newSession3 = new RpcSession(session3);
+				newSession3.Leases.Add(newLease);
+
+				session3 = await AgentScheduler.TryUpdateSessionAsync(session3, newSession3);
+				Assert.IsNotNull(session3);
+			}
+
+			await CheckFilterLengthsAsync(SessionFilterType.Available, length1: 1, length2: 2, length3: 2, length4: 2); // filter 1 (exclusive), filter 2 (1 cpu), filter 3 (2 cpus), filter 4 (16gb)
+			await CheckFilterLengthsAsync(SessionFilterType.Potential);
+
+			// Remove the leases from session 2
+			{
+				RpcSession newSession2 = new RpcSession(session2);
+				newSession2.Leases.Clear();
+
+				session2 = await AgentScheduler.TryUpdateSessionAsync(session2, newSession2);
+				Assert.IsNotNull(session2);
+			}
+
+			await CheckFilterLengthsAsync(SessionFilterType.Available, length1: 2, length4: 2); // filter 1 (exclusive), filter 4 (16gb)
+			await CheckFilterLengthsAsync(SessionFilterType.Potential);
+
+			// Remove the leases from session 3
+			{
+				RpcSession newSession3 = new RpcSession(session3);
+				newSession3.Leases.Clear();
+
+				session3 = await AgentScheduler.TryUpdateSessionAsync(session3, newSession3);
+				Assert.IsNotNull(session3);
+			}
+
+			await CheckFilterLengthsAsync(SessionFilterType.Available);
+			await CheckFilterLengthsAsync(SessionFilterType.Potential);
+		}
 	}
 }
diff --git a/Engine/Source/Programs/Horde/Plugins/Compute/HordeServer.Compute/Agents/AgentCollection.cs b/Engine/Source/Programs/Horde/Plugins/Compute/HordeServer.Compute/Agents/AgentCollection.cs
--- a/Engine/Source/Programs/Horde/Plugins/Compute/HordeServer.Compute/Agents/AgentCollection.cs
+++ b/Engine/Source/Programs/Horde/Plugins/Compute/HordeServer.Compute/Agents/AgentCollection.cs
@@ -333,7 +333,7 @@ namespace HordeServer.Agents
 				if (document != null)
 				{
 					Agent? agent = await CreateAgentObjectAsync(document, session, cancellationToken);
-					if (agent != null && agent.Session != null && agent.Session.ExpiryTicks == session.ExpiryTicks)
+					if (agent != null && agent.Session != null && agent.Session.UpdateTicks == session.UpdateTicks)
 					{
 						_logger.LogDebug("Terminating session {SessionId} for agent {Agent} (expiry time: {Time})", session.SessionId, session.AgentId, agent.Session?.ExpiryTime);
 						await TryUpdateSessionAsync(agent, new UpdateSessionOptions { Status = AgentStatus.Stopped }, cancellationToken);
@@ -1269,7 +1269,7 @@ namespace HordeServer.Agents
 				TaskCompletionSource tcs = _sessionIdToTcs.GetOrAdd(session.SessionId, newTcs);
 
 				RpcSession? newSession = await _scheduler.TryGetSessionAsync(session.SessionId, cancellationToken);
-				if (newSession != null && newSession.ExpiryTicks == session.ExpiryTicks)
+				if (newSession != null && newSession.UpdateTicks == session.UpdateTicks)
 				{
 					await tcs.Task.WaitAsync(cancellationToken);
 					newSession = await _scheduler.TryGetSessionAsync(session.SessionId, cancellationToken);
diff --git a/Engine/Source/Programs/Horde/Plugins/Compute/HordeServer.Compute/Agents/AgentScheduler.cs b/Engine/Source/Programs/Horde/Plugins/Compute/HordeServer.Compute/Agents/AgentScheduler.cs
--- a/Engine/Source/Programs/Horde/Plugins/Compute/HordeServer.Compute/Agents/AgentScheduler.cs
+++ b/Engine/Source/Programs/Horde/Plugins/Compute/HordeServer.Compute/Agents/AgentScheduler.cs
@@ -39,14 +39,8 @@ namespace HordeServer.Agents
 
 		record struct FilterCollectionKeys
 		{
-			public RedisSortedSetKey<IoHash> Created
-				=> new($"compute:filters:created");
-
-			public RedisSortedSetKey<IoHash> Deleted
-				=> new($"compute:filters:deleted");
-
-			public RedisHashKey<IoHash, long> Touched
-				=> new($"compute:filters:touched");
+			public RedisHashKey<IoHash, long> Current // Maps requirements hash to last-touched ticks
+				=> new($"compute:filters:current");
 
 			public FilterKeys this[IoHash requirementsHash]
 				=> new(requirementsHash);
@@ -54,17 +48,20 @@ namespace HordeServer.Agents
 
 		record struct FilterKeys(IoHash RequirementsHash)
 		{
+			public RedisSetKey<SessionId> AvailableSessions
+				=> new($"compute:filters:{RequirementsHash}:available");
+
+			public RedisSetKey<SessionId> PotentialSessions
+				=> new($"compute:filters:{RequirementsHash}:potential");
+
 			public RedisStringKey<RpcAgentRequirements> Requirements
 				=> new($"compute:filters:{RequirementsHash}:reqs");
-
-			public RedisSetKey<SessionId> Sessions
-				=> new($"compute:filters:{RequirementsHash}:sessions");
 		}
 
 		record struct SessionCollectionKeys
 		{
-			public RedisHashKey<SessionId, long> ExpiryTimes
-				=> new($"compute:sessions:update-ticks");
+			public RedisHashKey<SessionId, long> Current // Maps session id to last update ticks
+				=> new($"compute:sessions:current");
 
 			public SessionKeys this[SessionId sessionId]
 				=> new(sessionId);
@@ -75,7 +72,7 @@ namespace HordeServer.Agents
 			public RedisStringKey<RpcAgentCapabilities> Capabilities
 				=> new($"compute:sessions:{SessionId}:caps");
 
-			public RedisSetKey<IoHash> Filters
+			public RedisHashKey<IoHash, bool> Filters // Maps filter requirements hash to a flag indicating whether it's currently available
 				=> new($"compute:sessions:{SessionId}:filters");
 
 			public RedisStringKey<RpcSession> State
@@ -99,6 +96,16 @@ namespace HordeServer.Agents
 
 #pragma warning restore CA1822
 
+		record class LocalFilter(IoHash Hash, RpcAgentRequirements Requirements);
+
+		record class LocalCapabilities(IoHash Hash, RpcAgentCapabilities Message, IReadOnlySet<string> Properties)
+		{
+			public LocalCapabilities(IoHash hash, RpcAgentCapabilities message)
+				: this(hash, message, new HashSet<string>(message.Properties, StringComparer.Ordinal)) { }
+		}
+
+		record class LocalSession(long UpdateTicks, LocalCapabilities Capabilities, IReadOnlyDictionary<IoHash, bool> Filters);
+
 		public event Action<SessionId>? SessionUpdated;
 
 		readonly IRedisService _redisService;
@@ -106,14 +113,63 @@ namespace HordeServer.Agents
 		readonly ITicker _updateCachedFiltersTicker;
 		readonly ITicker _expireFiltersTicker;
 		readonly ILogger _logger;
-		readonly MemoryCache _memoryCache;
+		readonly MemoryCache _localSessionCache;
+		readonly MemoryCache _requirementsCache;
+
+		LocalFilter[] _cachedFilters = Array.Empty<LocalFilter>();
 
 		static readonly RedisChannel<SessionId> s_sessionUpdateChannel = new RedisChannel<SessionId>(RedisChannel.Literal("compute:sessions:update"));
 
-		record class CreatedFilter(IoHash Hash, RpcAgentRequirements Requirements, double Score);
+		void AddLocalSessionToCache(SessionId sessionId, LocalSession localSession)
+		{
+			using (ICacheEntry cacheEntry = _localSessionCache.CreateEntry(sessionId))
+			{
+				cacheEntry.SetSlidingExpiration(TimeSpan.FromMinutes(5.0));
+				cacheEntry.SetValue(localSession);
+			}
+		}
 
-		record class CachedFilters(long Ticks, CreatedFilter[] Created, SortedSetEntry<IoHash>[] Deleted);
-		CachedFilters _cachedFilters = new CachedFilters(0, Array.Empty<CreatedFilter>(), Array.Empty<SortedSetEntry<IoHash>>());
+		async ValueTask<LocalSession?> TryGetLocalSessionAsync(RpcSession session, CancellationToken cancellationToken)
+		{
+			// Try to get the current cached value
+			LocalSession? localSession;
+			if (!_localSessionCache.TryGetValue(session.SessionId, out localSession) || localSession == null || localSession.UpdateTicks != session.UpdateTicks)
+			{
+				IDatabase database = _redisService.GetDatabase();
+
+				// Update the capabilites
+				LocalCapabilities? capabilities;
+				if (localSession != null && localSession.Capabilities.Hash == session.CapabilitiesHash)
+				{
+					capabilities = localSession.Capabilities;
+				}
+				else
+				{
+					byte[]? data = (byte[]?)await database.StringGetAsync(Keys.Sessions[session.SessionId].Capabilities.Inner).WaitAsync(cancellationToken);
+					if (data == null)
+					{
+						return null;
+					}
+
+					IoHash hash = IoHash.Compute(data);
+					if (session.CapabilitiesHash != hash)
+					{
+						return null;
+					}
+
+					capabilities = new LocalCapabilities(hash, RpcAgentCapabilities.Parser.ParseFrom(data));
+				}
+
+				// Update the filters
+				HashEntry<IoHash, bool>[] entries = await database.HashGetAllAsync(Keys.Sessions[session.SessionId].Filters).WaitAsync(cancellationToken);
+				Dictionary<IoHash, bool> filters = entries.ToDictionary(x => x.Name, x => x.Value);
+
+				// Create the new session state
+				localSession = new LocalSession(session.UpdateTicks, capabilities, filters);
+				AddLocalSessionToCache(session.SessionId, localSession);
+			}
+			return localSession;
+		}
 
 		IAsyncDisposable? _updateEventSubscription;
 
@@ -126,15 +182,14 @@ namespace HordeServer.Agents
 			_clock = clock;
 			_updateCachedFiltersTicker = clock.AddTicker<AgentScheduler>(TimeSpan.FromSeconds(30.0), UpdateCachedFiltersTickAsync, logger);
 			_expireFiltersTicker = clock.AddSharedTicker($"{nameof(AgentScheduler)}.{nameof(ExpireFiltersAsync)}", TimeSpan.FromMinutes(5.0), ExpireFiltersAsync, logger);
-			_memoryCache = new MemoryCache(new MemoryCacheOptions());
+			_localSessionCache = new MemoryCache(new MemoryCacheOptions());
+			_requirementsCache = new MemoryCache(new MemoryCacheOptions());
 			_logger = logger;
 		}
 
 		/// <inheritdoc/>
 		public async ValueTask DisposeAsync()
 		{
-			_memoryCache.Dispose();
-
 			await _updateCachedFiltersTicker.DisposeAsync();
 			await _expireFiltersTicker.DisposeAsync();
 
@@ -143,6 +198,9 @@ namespace HordeServer.Agents
 				await _updateEventSubscription.DisposeAsync();
 				_updateEventSubscription = null;
 			}
+
+			_requirementsCache.Dispose();
+			_localSessionCache.Dispose();
 		}
 
 		void OnSessionUpdate(SessionId sessionId)
@@ -184,24 +242,27 @@ namespace HordeServer.Agents
 			newSession.AgentId = agentId;
 			newSession.SessionId = sessionId;
 			newSession.CapabilitiesHash = capabilitiesHash;
-			newSession.ExpiryTicks = (utcNow + RpcSession.ExpireAfterTime).Ticks;
+			newSession.UpdateTicks = utcNow.Ticks;
 			newSession.Status = RpcAgentStatus.Ok;
 
 			ITransaction transaction = _redisService.GetDatabase().CreateTransaction();
-			transaction.AddCondition(Keys.Sessions.ExpiryTimes.HashNotExists(sessionId));
-			_ = transaction.HashSetAsync(Keys.Sessions.ExpiryTimes, sessionId, newSession.ExpiryTicks, flags: CommandFlags.FireAndForget);
+			transaction.AddCondition(Keys.Sessions.Current.HashNotExists(sessionId));
+			_ = transaction.HashSetAsync(Keys.Sessions.Current, sessionId, newSession.UpdateTicks, flags: CommandFlags.FireAndForget);
 			_ = transaction.StringSetAsync(Keys.Sessions[sessionId].Capabilities.Inner, capabilitiesData, flags: CommandFlags.FireAndForget);
 			_ = transaction.StringSetAsync(Keys.Sessions[sessionId].State, newSession, flags: CommandFlags.FireAndForget);
 
-			if (!await UpdateFiltersAsync(_cachedFilters, transaction, newSession, capabilities, cancellationToken))
-			{
-				return null;
-			}
+			LocalCapabilities localCapabilities = new LocalCapabilities(capabilitiesHash, capabilities);
+			Dictionary<IoHash, bool> filters = FindMatchingFilters(localCapabilities, newSession);
+			UpdateFilters(transaction, sessionId, new Dictionary<IoHash, bool>(), filters);
+
 			if (!await transaction.ExecuteAsync().WaitAsync(cancellationToken))
 			{
 				return null;
 			}
 
+			LocalSession localSession = new LocalSession(newSession.UpdateTicks, localCapabilities, filters);
+			AddLocalSessionToCache(sessionId, localSession);
+
 			return newSession;
 		}
 
@@ -212,50 +273,11 @@ namespace HordeServer.Agents
 			return await database.StringGetAsync(Keys.Sessions[sessionId].State);
 		}
 
-		record class CapabilitiesLookup(IReadOnlySet<string> Properties, IReadOnlyDictionary<string, int> Resources)
-		{
-			public CapabilitiesLookup(RpcAgentCapabilities capabilities)
-				: this(new HashSet<string>(capabilities.Properties, StringComparer.Ordinal), new Dictionary<string, int>(capabilities.Resources, StringComparer.Ordinal)) { }
-		}
-
-		record class CapabilitiesCacheEntry(IoHash Hash, CapabilitiesLookup Capabilities, RpcAgentCapabilities Message);
-
-		async ValueTask<CapabilitiesCacheEntry?> TryGetCachedCapabilitiesAsync(RpcSession session, CancellationToken cancellationToken = default)
-		{
-			CapabilitiesCacheEntry? cachedCapabilities;
-			if (!_memoryCache.TryGetValue(session.SessionId, out cachedCapabilities) || cachedCapabilities == null || cachedCapabilities.Hash != session.CapabilitiesHash)
-			{
-				IDatabase database = _redisService.GetDatabase();
-
-				byte[]? data = (byte[]?)await database.StringGetAsync(Keys.Sessions[session.SessionId].Capabilities.Inner);
-				if (data == null)
-				{
-					return null;
-				}
-
-				IoHash hash = IoHash.Compute(data);
-				if (session.CapabilitiesHash != hash)
-				{
-					return null;
-				}
-
-				RpcAgentCapabilities capabilities = RpcAgentCapabilities.Parser.ParseFrom(data);
-				cachedCapabilities = new CapabilitiesCacheEntry(session.CapabilitiesHash, new CapabilitiesLookup(capabilities), capabilities);
-
-				using (ICacheEntry entry = _memoryCache.CreateEntry(session.SessionId))
-				{
-					entry.SetSlidingExpiration(TimeSpan.FromMinutes(30.0));
-					entry.Value = cachedCapabilities;
-				}
-			}
-			return cachedCapabilities;
-		}
-
 		/// <inheritdoc/>
 		public async Task<RpcAgentCapabilities?> TryGetCapabilitiesAsync(RpcSession session, CancellationToken cancellationToken = default)
 		{
-			CapabilitiesCacheEntry? cachedCapabilities = await TryGetCachedCapabilitiesAsync(session, cancellationToken);
-			return cachedCapabilities?.Message;
+			LocalSession? localSession = await TryGetLocalSessionAsync(session, cancellationToken);
+			return localSession?.Capabilities.Message;
 		}
 
 		/// <inheritdoc/>
@@ -263,65 +285,79 @@ namespace HordeServer.Agents
 		{
 			SessionId sessionId = session.SessionId;
 
+			// Get the local session state. This includes cached values for everything in Redis.
+			LocalSession? localSession = await TryGetLocalSessionAsync(session, cancellationToken);
+			if (localSession == null)
+			{
+				return null;
+			}
+
 			// Copy the session document if a new one was not specified, and update the last modified time
 			newSession ??= new RpcSession(session);
 
 			// Start building the update transaction
-			ITransaction transaction = _redisService.GetDatabase().CreateTransaction();
-			transaction.AddCondition(Keys.Sessions.ExpiryTimes.HashEqual(newSession.SessionId, session.ExpiryTicks));
+			IDatabase database = _redisService.GetDatabase();
+			ITransaction transaction = database.CreateTransaction();
+			transaction.AddCondition(Keys.Sessions.Current.HashEqual(newSession.SessionId, session.UpdateTicks));
 
 			// Handle agents transitioning to the stopping state
 			if (newSession.Status == RpcAgentStatus.Stopped)
 			{
-				// TODO: remove all leases
-				newSession.ExpiryTicks = Math.Min(newSession.ExpiryTicks, _clock.UtcNow.Ticks);
+				newSession.UpdateTicks = Math.Max(newSession.UpdateTicks + 1, Math.Min(newSession.UpdateTicks + RpcSession.ExpireAfterTime.Ticks, _clock.UtcNow.Ticks));
 
-				_ = transaction.HashDeleteAsync(Keys.Sessions.ExpiryTimes, newSession.SessionId, flags: CommandFlags.FireAndForget);
+				_ = transaction.HashDeleteAsync(Keys.Sessions.Current, newSession.SessionId, flags: CommandFlags.FireAndForget);
 				_ = transaction.KeyDeleteAsync(Keys.Sessions[sessionId].Capabilities, flags: CommandFlags.FireAndForget);
 				_ = transaction.KeyDeleteAsync(Keys.Sessions[sessionId].State, flags: CommandFlags.FireAndForget);
+				UpdateFilters(transaction, sessionId, localSession.Filters, new Dictionary<IoHash, bool>());
+
+				localSession = new LocalSession(newSession.UpdateTicks, localSession.Capabilities, new Dictionary<IoHash, bool>());
 			}
 			else
 			{
 				DateTime utcNow = _clock.UtcNow;
 
 				// Extend the expiry time, ensuring it only ever increases
-				newSession.ExpiryTicks = Math.Max(session.ExpiryTicks + 1, (utcNow + RpcSession.ExpireAfterTime).Ticks);
-				_ = transaction.HashSetAsync(Keys.Sessions.ExpiryTimes, session.SessionId, newSession.ExpiryTicks, flags: CommandFlags.FireAndForget);
+				newSession.UpdateTicks = Math.Max(session.UpdateTicks + 1, utcNow.Ticks);
+				_ = transaction.HashSetAsync(Keys.Sessions.Current, session.SessionId, newSession.UpdateTicks, flags: CommandFlags.FireAndForget);
 
 				// Update the capabilities
+				LocalCapabilities capabilities = localSession.Capabilities;
 				if (newCapabilities != null)
 				{
 					byte[] capabilitiesData = newCapabilities.ToByteArray();
 					IoHash capabilitiesHash = IoHash.Compute(capabilitiesData);
 
-					if (capabilitiesHash != newSession.CapabilitiesHash)
+					if (capabilitiesHash != session.CapabilitiesHash)
 					{
 						newSession.CapabilitiesHash = capabilitiesHash;
 						_ = transaction.StringSetAsync(Keys.Sessions[session.SessionId].Capabilities.Inner, capabilitiesData, flags: CommandFlags.FireAndForget);
+
+						capabilities = new LocalCapabilities(capabilitiesHash, newCapabilities);
 					}
 				}
 
-				// Add any new leases to the global state
-				foreach (RpcSessionLease newLease in newSession.Leases)
+				// Update the session state
+				_ = transaction.StringSetAsync(Keys.Sessions[session.SessionId].State, newSession, flags: CommandFlags.FireAndForget);
+
+				// Update the filters for this session
+				Dictionary<IoHash, bool> filters = FindMatchingFilters(capabilities, newSession);
+				UpdateFilters(transaction, sessionId, localSession.Filters, filters);
+
+				// Create the new local session object
+				localSession = new LocalSession(newSession.UpdateTicks, capabilities, filters);
+			}
+
+			// Add any new leases to the global state
+			foreach (RpcSessionLease newLease in newSession.Leases)
+			{
+				if (!session.Leases.Any(x => x.Id == newLease.Id))
 				{
-					if (!session.Leases.Any(x => x.Id == newLease.Id))
+					_ = transaction.SetAddAsync(Keys.Leases.Current, newLease.Id, flags: CommandFlags.FireAndForget);
+					if (newLease.ParentId != null)
 					{
-						_ = transaction.SetAddAsync(Keys.Leases.Current, newLease.Id, flags: CommandFlags.FireAndForget);
-						if (newLease.ParentId != null)
-						{
-							_ = transaction.SetAddAsync(Keys.Leases[newLease.ParentId.Value].Children, newLease.Id, flags: CommandFlags.FireAndForget);
-						}
+						_ = transaction.SetAddAsync(Keys.Leases[newLease.ParentId.Value].Children, newLease.Id, flags: CommandFlags.FireAndForget);
 					}
 				}
-
-				// Update the filters that this session belongs to
-				if (!await UpdateFiltersAsync(_cachedFilters, transaction, newSession, newCapabilities, cancellationToken))
-				{
-					return null;
-				}
-
-				// Update the session state
-				_ = transaction.StringSetAsync(Keys.Sessions[session.SessionId].State, newSession, flags: CommandFlags.FireAndForget);
 			}
 
 			// Remove any complete leases from the global state
@@ -343,13 +379,16 @@ namespace HordeServer.Agents
 				return null;
 			}
 
+			// Update the cached session state
+			AddLocalSessionToCache(sessionId, localSession);
+
 			// Trace the new expiry time for debugging
 			_logger.LogDebug("Updated session {SessionId} expiry time to {ExpiryTime}", sessionId, newSession.ExpiryTime);
 
 			// Notify watchers that the session state has changed
 			if (newSession.Status != session.Status || newSession.CapabilitiesHash != session.CapabilitiesHash || !newSession.Leases.Equals(session.Leases))
 			{
-				_ = _redisService.GetDatabase().PublishAsync(s_sessionUpdateChannel, sessionId, CommandFlags.FireAndForget);
+				_ = database.PublishAsync(s_sessionUpdateChannel, sessionId, CommandFlags.FireAndForget);
 			}
 			return newSession;
 		}
@@ -357,16 +396,16 @@ namespace HordeServer.Agents
 		/// <inheritdoc/>
 		public async IAsyncEnumerable<RpcSession> FindExpiredSessionsAsync([EnumeratorCancellation] CancellationToken cancellationToken = default)
 		{
-			DateTime utcNow = _clock.UtcNow;
+			long expireTicks = (_clock.UtcNow - RpcSession.ExpireAfterTime).Ticks;
 			IDatabase database = _redisService.GetDatabase();
 
-			HashEntry<SessionId, long>[] entries = await database.HashGetAllAsync(Keys.Sessions.ExpiryTimes);
-			foreach ((SessionId sessionId, long expiryTicks) in entries)
+			HashEntry<SessionId, long>[] entries = await database.HashGetAllAsync(Keys.Sessions.Current);
+			foreach ((SessionId sessionId, long updateTicks) in entries)
 			{
-				if (expiryTicks < utcNow.Ticks)
+				if (updateTicks < expireTicks)
 				{
 					RpcSession? session = await TryGetSessionAsync(sessionId, cancellationToken);
-					if (session != null && session.ExpiryTicks == expiryTicks)
+					if (session != null && session.UpdateTicks == updateTicks)
 					{
 						yield return session;
 					}
@@ -389,11 +428,9 @@ namespace HordeServer.Agents
 				DateTime utcNow = _clock.UtcNow;
 
 				ITransaction transaction = database.CreateTransaction();
-				transaction.AddCondition(Keys.Filters.Touched.HashNotExists(requirementsHash));
-				_ = transaction.HashSetAsync(Keys.Filters.Touched, requirementsHash, utcNow.Ticks);
+				transaction.AddCondition(Keys.Filters.Current.HashNotExists(requirementsHash));
+				_ = transaction.HashSetAsync(Keys.Filters.Current, requirementsHash, utcNow.Ticks);
 				_ = transaction.StringSetAsync(Keys.Filters[requirementsHash].Requirements.Inner, requirementsData, flags: CommandFlags.FireAndForget);
-				_ = transaction.SortedSetAddAsync(Keys.Filters.Created, requirementsHash, utcNow.Ticks, flags: CommandFlags.FireAndForget);
-				_ = transaction.SortedSetRemoveAsync(Keys.Filters.Deleted, requirementsHash, flags: CommandFlags.FireAndForget);
 
 				if (await transaction.ExecuteAsync().WaitAsync(cancellationToken))
 				{
@@ -401,8 +438,8 @@ namespace HordeServer.Agents
 				}
 
 				transaction = database.CreateTransaction();
-				transaction.AddCondition(Keys.Filters.Touched.HashNotExists(requirementsHash));
-				_ = transaction.HashSetAsync(Keys.Filters.Touched, requirementsHash, utcNow.Ticks, flags: CommandFlags.FireAndForget);
+				transaction.AddCondition(Keys.Filters.Current.HashNotExists(requirementsHash));
+				_ = transaction.HashSetAsync(Keys.Filters.Current, requirementsHash, utcNow.Ticks, flags: CommandFlags.FireAndForget);
 
 				if (await transaction.ExecuteAsync().WaitAsync(cancellationToken))
 				{
@@ -432,28 +469,23 @@ namespace HordeServer.Agents
 			DateTime utcNow = _clock.UtcNow;
 			IDatabase database = _redisService.GetDatabase();
 
-			ITransaction transaction = database.CreateTransaction();
-			Task<SortedSetEntry<IoHash>[]> createdEntriesTask = transaction.SortedSetRangeByRankWithScoresAsync(Keys.Filters.Created);
-			Task<SortedSetEntry<IoHash>[]> deletedEntriesTask = transaction.SortedSetRangeByRankWithScoresAsync(Keys.Filters.Deleted);
-			await transaction.ExecuteAsync().WaitAsync(cancellationToken);
-
-			SortedSetEntry<IoHash>[] createdEntries = await createdEntriesTask;
-			CreatedFilter[] createdFilters = new CreatedFilter[createdEntries.Length];
+			IoHash[] filterHashes = await database.HashKeysAsync(Keys.Filters.Current);
+			LocalFilter[] filters = new LocalFilter[filterHashes.Length];
 
-			for (int idx = 0; idx < createdEntries.Length; idx++)
+			for (int idx = 0; idx < filterHashes.Length; idx++)
 			{
-				SortedSetEntry<IoHash> createdEntry = createdEntries[idx];
+				IoHash filterHash = filterHashes[idx];
 
-				RpcAgentRequirements? requirements = await TryGetFilterRequirementsAsync(createdEntry.Element, cancellationToken);
+				RpcAgentRequirements? requirements = await TryGetFilterRequirementsAsync(filterHash, cancellationToken);
 				if (requirements == null)
 				{
 					return false;
 				}
 
-				createdFilters[idx] = new CreatedFilter(createdEntry.Element, requirements, createdEntry.Score);
+				filters[idx] = new LocalFilter(filterHash, requirements);
 			}
 
-			_cachedFilters = new CachedFilters(utcNow.Ticks, createdFilters, await deletedEntriesTask);
+			_cachedFilters = filters;
 			return true;
 		}
 
@@ -466,20 +498,19 @@ namespace HordeServer.Agents
 
 			DateTime expiryTime = _clock.UtcNow - ExpireFiltersTime;
 
-			HashEntry<IoHash, long>[] entries = await database.HashGetAllAsync(Keys.Filters.Touched).WaitAsync(cancellationToken);
+			HashEntry<IoHash, long>[] entries = await database.HashGetAllAsync(Keys.Filters.Current).WaitAsync(cancellationToken);
 			foreach ((IoHash requirementsHash, long ticks) in entries)
 			{
 				cancellationToken.ThrowIfCancellationRequested();
 				if (ticks < expiryTime.Ticks)
 				{
 					ITransaction transaction = database.CreateTransaction();
-					transaction.AddCondition(Keys.Filters.Touched.HashEqual(requirementsHash, ticks));
+					transaction.AddCondition(Keys.Filters.Current.HashEqual(requirementsHash, ticks));
 
-					_ = transaction.HashDeleteAsync(Keys.Filters.Touched, requirementsHash, flags: CommandFlags.FireAndForget);
-					_ = transaction.SortedSetRemoveAsync(Keys.Filters.Created, requirementsHash, flags: CommandFlags.FireAndForget);
-					_ = transaction.SortedSetAddAsync(Keys.Filters.Deleted, requirementsHash, _clock.UtcNow.Ticks, flags: CommandFlags.FireAndForget);
+					_ = transaction.HashDeleteAsync(Keys.Filters.Current, requirementsHash, flags: CommandFlags.FireAndForget);
+					_ = transaction.KeyDeleteAsync(Keys.Filters[requirementsHash].AvailableSessions, flags: CommandFlags.FireAndForget);
+					_ = transaction.KeyDeleteAsync(Keys.Filters[requirementsHash].PotentialSessions, flags: CommandFlags.FireAndForget);
 					_ = transaction.KeyDeleteAsync(Keys.Filters[requirementsHash].Requirements, flags: CommandFlags.FireAndForget);
-					_ = transaction.KeyDeleteAsync(Keys.Filters[requirementsHash].Sessions, flags: CommandFlags.FireAndForget);
 
 					_ = transaction.ExecuteAsync(CommandFlags.FireAndForget);
 				}
@@ -491,14 +522,14 @@ namespace HordeServer.Agents
 		{
 			IDatabase database = _redisService.GetDatabase();
 			DateTime utcNow = _clock.UtcNow;
-			await database.HashSetAsync(Keys.Filters.Touched, requirementsHash, utcNow.Ticks);
+			await database.HashSetAsync(Keys.Filters.Current, requirementsHash, utcNow.Ticks);
 		}
 
 		/// <inheritdoc/>
 		public async Task<IoHash[]> GetFiltersAsync(CancellationToken cancellationToken = default)
 		{
 			IDatabase database = _redisService.GetDatabase();
-			return await database.HashKeysAsync(Keys.Filters.Touched);
+			return await database.HashKeysAsync(Keys.Filters.Current);
 		}
 
 		public async Task ForceUpdateFiltersAsync(CancellationToken cancellationToken)
@@ -506,99 +537,147 @@ namespace HordeServer.Agents
 			IDatabase database = _redisService.GetDatabase();
 
 			await UpdateCachedFiltersTickAsync(cancellationToken);
-			CachedFilters cachedFilters = _cachedFilters;
 
-			SessionId[] sessionIds = await database.HashKeysAsync(Keys.Sessions.ExpiryTimes);
+			SessionId[] sessionIds = await database.HashKeysAsync(Keys.Sessions.Current);
 			foreach (SessionId sessionId in sessionIds)
 			{
 				RpcSession? session = await TryGetSessionAsync(sessionId, cancellationToken);
 				if (session != null)
 				{
-					ITransaction transaction = database.CreateTransaction();
-					await UpdateFiltersAsync(cachedFilters, transaction, session, null, cancellationToken);
-					await transaction.ExecuteAsync().WaitAsync(cancellationToken);
+					LocalSession? localSession = await TryGetLocalSessionAsync(session, cancellationToken);
+					if (localSession != null)
+					{
+						Dictionary<IoHash, bool> newFilters = FindMatchingFilters(localSession.Capabilities, session);
+
+						ITransaction transaction = database.CreateTransaction();
+						UpdateFilters(transaction, session.SessionId, localSession.Filters, newFilters);
+						await transaction.ExecuteAsync().WaitAsync(cancellationToken);
+
+						AddLocalSessionToCache(sessionId, localSession with { Filters = newFilters });
+					}
 				}
 			}
 		}
 
-		async Task<bool> UpdateFiltersAsync(CachedFilters cachedFilters, ITransaction transaction, RpcSession newSession, RpcAgentCapabilities? newCapabilities, CancellationToken cancellationToken)
+		void UpdateFilters(ITransaction transaction, SessionId sessionId, IReadOnlyDictionary<IoHash, bool> oldFilters, IReadOnlyDictionary<IoHash, bool> newFilters)
 		{
-			long minFilterTime = Math.Max(0, newSession.LastFilterUpdateTicks - TimeSpan.FromSeconds(30.0).Ticks);
-
-			// Add the session to any new matching filters
-			CreatedFilter[] createdFilters = cachedFilters.Created;
-
-			int createdIdx = createdFilters.BinarySearch(x => x.Score, minFilterTime);
-			if (createdIdx < 0)
+			// Create any new filters
+			foreach ((IoHash newFilter, bool newAvailable) in newFilters)
 			{
-				createdIdx = ~createdIdx;
-			}
-
-			if (createdIdx < createdFilters.Length)
-			{
-				CapabilitiesLookup? capabilities;
-				if (newCapabilities != null)
+				bool oldAvailable;
+				if (!oldFilters.TryGetValue(newFilter, out oldAvailable))
 				{
-					capabilities = new CapabilitiesLookup(newCapabilities);
-				}
-				else
-				{
-					capabilities = (await TryGetCachedCapabilitiesAsync(newSession, cancellationToken))?.Capabilities;
+					transaction.AddCondition(Keys.Filters.Current.HashExists(newFilter));
+					_ = transaction.HashSetAsync(Keys.Sessions[sessionId].Filters, newFilter, newAvailable, flags: CommandFlags.FireAndForget);
+					_ = transaction.SetAddAsync(Keys.Filters[newFilter].PotentialSessions, sessionId, flags: CommandFlags.FireAndForget);
+					if (newAvailable)
+					{
+						_ = transaction.SetAddAsync(Keys.Filters[newFilter].AvailableSessions, sessionId, flags: CommandFlags.FireAndForget);
+					}
 				}
-
-				if (capabilities == null)
+				else if (newAvailable != oldAvailable)
 				{
-					return false;
+					transaction.AddCondition(Keys.Filters.Current.HashExists(newFilter));
+					_ = transaction.HashSetAsync(Keys.Sessions[sessionId].Filters, newFilter, newAvailable, flags: CommandFlags.FireAndForget);
+					if (newAvailable)
+					{
+						_ = transaction.SetAddAsync(Keys.Filters[newFilter].AvailableSessions, sessionId, flags: CommandFlags.FireAndForget);
+					}
+					else
+					{
+						_ = transaction.SetRemoveAsync(Keys.Filters[newFilter].AvailableSessions, sessionId, flags: CommandFlags.FireAndForget);
+					}
 				}
+			}
 
-				for (; createdIdx < createdFilters.Length; createdIdx++)
+			// Remove any filters that no longer exist
+			foreach ((IoHash oldFilter, bool oldAvailable) in oldFilters)
+			{
+				if (!newFilters.ContainsKey(oldFilter))
 				{
-					CreatedFilter createdFilter = createdFilters[createdIdx];
-					if (MeetsRequirements(capabilities, createdFilter.Requirements))
+					_ = transaction.HashDeleteAsync(Keys.Sessions[sessionId].Filters, oldFilter, flags: CommandFlags.FireAndForget);
+					_ = transaction.SetRemoveAsync(Keys.Filters[oldFilter].PotentialSessions, sessionId, flags: CommandFlags.FireAndForget);
+					if (oldAvailable)
 					{
-						transaction.AddCondition(Keys.Filters.Touched.HashExists(createdFilter.Hash));
-						_ = transaction.SetAddAsync(Keys.Filters[createdFilter.Hash].Sessions, newSession.SessionId, flags: CommandFlags.FireAndForget);
-						_ = transaction.SetAddAsync(Keys.Sessions[newSession.SessionId].Filters, createdFilter.Hash, flags: CommandFlags.FireAndForget);
+						_ = transaction.SetRemoveAsync(Keys.Filters[oldFilter].AvailableSessions, sessionId, flags: CommandFlags.FireAndForget);
 					}
 				}
 			}
+		}
 
-			// Remove the session from any removed filters
-			SortedSetEntry<IoHash>[] deletedFilters = cachedFilters.Deleted;
+		Dictionary<IoHash, bool> FindMatchingFilters(LocalCapabilities capabilities, RpcSession session)
+		{
+			LocalFilter[] cachedFilters = _cachedFilters;
 
-			int deletedIdx = deletedFilters.BinarySearch(new SortedSetEntry<IoHash>(IoHash.Zero, minFilterTime));
-			if (deletedIdx < 0)
+			Dictionary<IoHash, bool> newFilters = new Dictionary<IoHash, bool>();
+			foreach (LocalFilter filter in cachedFilters)
 			{
-				deletedIdx = ~deletedIdx;
+				if (MeetsRequirements(capabilities, filter.Requirements))
+				{
+					bool available = IsAvailable(capabilities, session, filter.Requirements);
+					newFilters.Add(filter.Hash, available);
+				}
 			}
 
-			for (; deletedIdx < deletedFilters.Length; deletedIdx++)
+			return newFilters;
+		}
+
+		static bool MeetsRequirements(LocalCapabilities capabilities, RpcAgentRequirements requirements)
+		{
+			foreach (string property in requirements.Properties)
+			{
+				if (!capabilities.Properties.Contains(property))
+				{
+					return false;
+				}
+			}
+			foreach ((string name, int requiredCount) in requirements.Resources)
 			{
-				IoHash queueHash = deletedFilters[deletedIdx].Element;
-				transaction.AddCondition(Keys.Filters.Touched.HashNotExists(queueHash));
-				_ = transaction.SetRemoveAsync(Keys.Sessions[newSession.SessionId].Filters, queueHash, flags: CommandFlags.FireAndForget);
+				if (!capabilities.Message.Resources.TryGetValue(name, out int maxCount) || requiredCount > maxCount)
+				{
+					return false;
+				}
 			}
-
-			// Update the last queue time
-			newSession.LastFilterUpdateTicks = cachedFilters.Ticks;
 			return true;
 		}
 
-		static bool MeetsRequirements(CapabilitiesLookup capabilities, RpcAgentRequirements requirements)
+		static bool IsAvailable(LocalCapabilities capabilities, RpcSession session, RpcAgentRequirements requirements)
 		{
-			foreach (string property in requirements.Properties)
+			if (session.Leases.Any(x => x.Exclusive))
 			{
-				if (!capabilities.Properties.Contains(property))
+				return false;
+			}
+			if (!requirements.Shared && session.Leases.Any())
+			{
+				return false;
+			}
+
+			foreach ((string name, int requiredCount) in requirements.Resources)
+			{
+				int availableCount;
+				if (!capabilities.Message.Resources.TryGetValue(name, out availableCount))
+				{
+					return false;
+				}
+				foreach (RpcSessionLease lease in session.Leases)
+				{
+					if (lease.Resources.TryGetValue(name, out int count))
+					{
+						availableCount -= count;
+					}
+				}
+				if (availableCount < requiredCount)
 				{
 					return false;
 				}
 			}
+
 			return true;
 		}
 
 		void AddRequirementsToCache(IoHash requirementsHash, RpcAgentRequirements requirements)
 		{
-			using (ICacheEntry entry = _memoryCache.CreateEntry(requirementsHash))
+			using (ICacheEntry entry = _requirementsCache.CreateEntry(requirementsHash))
 			{
 				entry.SetSlidingExpiration(TimeSpan.FromMinutes(30.0));
 				entry.SetValue(requirements);
@@ -609,7 +688,7 @@ namespace HordeServer.Agents
 		public async ValueTask<RpcAgentRequirements?> TryGetFilterRequirementsAsync(IoHash requirementsHash, CancellationToken cancellationToken)
 		{
 			RpcAgentRequirements? requirements;
-			if (!_memoryCache.TryGetValue(requirementsHash, out requirements))
+			if (!_requirementsCache.TryGetValue(requirementsHash, out requirements))
 			{
 				IDatabase database = _redisService.GetDatabase();
 				requirements = await database.StringGetAsync(Keys.Filters[requirementsHash].Requirements);
@@ -623,11 +702,16 @@ namespace HordeServer.Agents
 		}
 
 		/// <inheritdoc/>
-		public async Task<SessionId[]> GetAllFilteredSessionsAsync(IoHash requirementsHash, CancellationToken cancellationToken = default)
+		public async IAsyncEnumerable<SessionId> EnumerateFilteredSessionsAsync(IoHash requirementsHash, SessionFilterType type, [EnumeratorCancellation] CancellationToken cancellationToken = default)
 		{
 			IDatabase database = _redisService.GetDatabase();
-			_ = database.HashSetAsync(Keys.Filters.Touched, requirementsHash, _clock.UtcNow.Ticks, flags: CommandFlags.FireAndForget);
-			return await database.SetMembersAsync(Keys.Filters[requirementsHash].Sessions).WaitAsync(cancellationToken);
+			_ = database.HashSetAsync(Keys.Filters.Current, requirementsHash, _clock.UtcNow.Ticks, flags: CommandFlags.FireAndForget);
+
+			RedisSetKey<SessionId> key = (type == SessionFilterType.Available)? Keys.Filters[requirementsHash].AvailableSessions : Keys.Filters[requirementsHash].PotentialSessions;
+			await foreach (SessionId sessionId in database.SetScanAsync(key).WithCancellation(cancellationToken))
+			{
+				yield return sessionId;
+			}
 		}
 
 		#endregion
diff --git a/Engine/Source/Programs/Horde/Plugins/Compute/HordeServer.Compute/Agents/IAgentScheduler.cs b/Engine/Source/Programs/Horde/Plugins/Compute/HordeServer.Compute/Agents/IAgentScheduler.cs
--- a/Engine/Source/Programs/Horde/Plugins/Compute/HordeServer.Compute/Agents/IAgentScheduler.cs
+++ b/Engine/Source/Programs/Horde/Plugins/Compute/HordeServer.Compute/Agents/IAgentScheduler.cs
@@ -8,6 +8,22 @@ using HordeCommon.Rpc.Messages;
 
 namespace HordeServer.Agents
 {
+	/// <summary>
+	/// Type of data to return from enumeration calls
+	/// </summary>
+	public enum SessionFilterType
+	{
+		/// <summary>
+		/// Return only sessions which can satisfy resource/exclusivity requirements
+		/// </summary>
+		Available = 0,
+
+		/// <summary>
+		/// Return all active sessions in the response, even if they are currently busy
+		/// </summary>
+		Potential = 1,
+	}
+
 	/// <summary>
 	/// Interface for a scheduler that manages agent sessions and leases.
 	/// </summary>
@@ -89,12 +105,13 @@ namespace HordeServer.Agents
 		ValueTask<RpcAgentRequirements?> TryGetFilterRequirementsAsync(IoHash requirementsHash, CancellationToken cancellationToken);
 
 		/// <summary>
-		/// Finds all sessions currently included in a filter
+		/// Enumerate all the sessions in a filter
 		/// </summary>
-		/// <param name="requirementsHash">Hash of the requirements object</param>
+		/// <param name="requirementsHash">Requirements hash for the filter</param>
+		/// <param name="sessionType">Whether to only include available sessions, or also include any sessions that are busy and cannot meet resource requirements</param>
 		/// <param name="cancellationToken">Cancellation token for the operation</param>
-		/// <returns>Array of sessions</returns>
-		Task<SessionId[]> GetAllFilteredSessionsAsync(IoHash requirementsHash, CancellationToken cancellationToken = default);
+		/// <returns>Sequence of sessions</returns>
+		IAsyncEnumerable<SessionId> EnumerateFilteredSessionsAsync(IoHash requirementsHash, SessionFilterType sessionType = SessionFilterType.Potential, CancellationToken cancellationToken = default);
 
 		#endregion
 		#region Stats
diff --git a/Engine/Source/Programs/Horde/Plugins/Compute/HordeServer.Compute/Protos/horde/messages/session.proto.cs b/Engine/Source/Programs/Horde/Plugins/Compute/HordeServer.Compute/Protos/horde/messages/session.proto.cs
--- a/Engine/Source/Programs/Horde/Plugins/Compute/HordeServer.Compute/Protos/horde/messages/session.proto.cs
+++ b/Engine/Source/Programs/Horde/Plugins/Compute/HordeServer.Compute/Protos/horde/messages/session.proto.cs
@@ -33,7 +33,7 @@ namespace HordeCommon.Rpc.Messages
 		}
 
 		public DateTime ExpiryTime
-			=> new DateTime(ExpiryTicks, DateTimeKind.Utc);
+			=> new DateTime(UpdateTicks + ((Status == RpcAgentStatus.Stopped)? 0 : ExpireAfterTime.Ticks), DateTimeKind.Utc);
 	}
 
 	partial class RpcSessionLease : IAgentLease
diff --git a/Engine/Source/Programs/Shared/EpicGames.Redis/RedisStringKey.cs b/Engine/Source/Programs/Shared/EpicGames.Redis/RedisStringKey.cs
--- a/Engine/Source/Programs/Shared/EpicGames.Redis/RedisStringKey.cs
+++ b/Engine/Source/Programs/Shared/EpicGames.Redis/RedisStringKey.cs
@@ -78,7 +78,7 @@ namespace EpicGames.Redis
 		public static async Task<TValue?> StringGetAsync<TValue>(this IDatabaseAsync target, RedisStringKey<TValue> key, CommandFlags flags = CommandFlags.None)
 		{
 			RedisValue value = await target.StringGetAsync(key.Inner, flags);
-			if (value.IsNullOrEmpty)
+			if (value.IsNull)
 			{
 				return default;
 			}
@@ -89,7 +89,7 @@ namespace EpicGames.Redis
 		public static async Task<TValue> StringGetAsync<TValue>(this IDatabaseAsync target, RedisStringKey<TValue> key, TValue defaultValue, CommandFlags flags = CommandFlags.None)
 		{
 			RedisValue value = await target.StringGetAsync(key.Inner, flags);
-			if (value.IsNullOrEmpty)
+			if (value.IsNull)
 			{
 				return defaultValue;
 			}
